[Mutable] Optimize EnsureBugetBelow function.

- Added mutable.EnableDetailedMemoryBudgetExceededLogging Cvar defaulted to false to enable more detailed but expensive logging when budget is exceeded.

#rnx

diff --git a/Engine/Plugins/Experimental/Mutable/Source/CustomizableObjectEditor/Private/MuCOE/SMutableConstantsWidget.cpp b/Engine/Plugins/Experimental/Mutable/Source/CustomizableObjectEditor/Private/MuCOE/SMutableConstantsWidget.cpp
--- a/Engine/Plugins/Experimental/Mutable/Source/CustomizableObjectEditor/Private/MuCOE/SMutableConstantsWidget.cpp
+++ b/Engine/Plugins/Experimental/Mutable/Source/CustomizableObjectEditor/Private/MuCOE/SMutableConstantsWidget.cpp
@@ -1127,7 +1127,7 @@ void SMutableConstantsWidget::LoadConstantImages()
 
 	for (int32 ImageIndex = 0; ImageIndex < ConstantsCount; ImageIndex++)
 	{
-		TSharedPtr<FMutableConstantImageElement> ConstantImageElement =MakeShared<FMutableConstantImageElement>();
+		TSharedPtr<FMutableConstantImageElement> ConstantImageElement = MakeShared<FMutableConstantImageElement>();
 
 		MutableProgramPtr->GetConstant(ImageIndex, ConstantImageElement->ImagePtr, 0,
 			[this](int32 x, int32 y, int32 m, mu::EImageFormat f, mu::EInitializationType i) { return new mu::Image(x, y, m, f, i); });
@@ -1160,7 +1160,7 @@ void SMutableConstantsWidget::LoadConstantMeshes()
 {
 	check (MutableProgramPtr);
 	
-	const int32 ConstantsCount = MutableProgramPtr->m_constantMeshes.Num();
+	const int32 ConstantsCount = MutableProgramPtr->ConstantMeshes.Num();
 	ConstantMeshElements.Empty(ConstantsCount);
 	
 	uint64 ConstantMeshesAccumulatedSize = 0;
@@ -1168,7 +1168,7 @@ void SMutableConstantsWidget::LoadConstantMeshes()
 	for (int32 MeshIndex = 0; MeshIndex < ConstantsCount; MeshIndex++)
 	{
 		TSharedPtr< FMutableConstantMeshElement> ConstantMeshElement = MakeShared<FMutableConstantMeshElement>();
-		ConstantMeshElement->MeshPtr = MutableProgramPtr->m_constantMeshes[MeshIndex].Value;
+		ConstantMeshElement->MeshPtr = MutableProgramPtr->ConstantMeshes[MeshIndex].Value;
 		ConstantMeshElement->IndexOnSourceVector = MeshIndex;
 		
 		ConstantMeshesAccumulatedSize += ConstantMeshElement->MeshPtr->GetDataSize();
diff --git a/Engine/Plugins/Experimental/Mutable/Source/MutableRuntime/Internal/MuR/ModelPrivate.h b/Engine/Plugins/Experimental/Mutable/Source/MutableRuntime/Internal/MuR/ModelPrivate.h
--- a/Engine/Plugins/Experimental/Mutable/Source/MutableRuntime/Internal/MuR/ModelPrivate.h
+++ b/Engine/Plugins/Experimental/Mutable/Source/MutableRuntime/Internal/MuR/ModelPrivate.h
@@ -218,18 +218,21 @@ namespace mu
 
 		//! Data for every rom.
 		TArray<FRomData> m_roms;
+	
+		/** Loaded roms worth tracking (only images and meshes for now). Stores the rom's data type.*/
+		TSparseArray<uint8> LoadedMemTrackedRoms;
 
 		//! Constant image mip data: the first is the index in m_roms for each image mip or -1 if it is always loaded.
-		TArray<TPair<int32, Ptr<const Image>>> m_constantImageLODs;
+		TArray<TPair<int32, Ptr<const Image>>> ConstantImageLODs;
 
-		//! Constant image mip chain indices: ranges in this array are defined in FImageLODRange and the indices here refer to m_constantImageLODs.
+		//! Constant image mip chain indices: ranges in this array are defined in FImageLODRange and the indices here refer to ConstantImageLODs.
 		TArray<uint32> m_constantImageLODIndices;
 
 		//! Constant image data.
 		TArray<FImageLODRange> m_constantImages;
 
         //! Constant mesh data: the first is the index in m_roms for each mesh or -1 if it is always loaded.
-		TArray<TPair<int32, Ptr<const Mesh>>> m_constantMeshes;
+		TArray<TPair<int32, Ptr<const Mesh>>> ConstantMeshes;
 
 		//! Constant ExtensionData
 		TArray<FExtensionDataConstant> m_constantExtensionData;
@@ -282,10 +285,10 @@ namespace mu
             arch << m_byteCode;
             arch << m_states;
 			arch << m_roms;
-			arch << m_constantImageLODs;
+			arch << ConstantImageLODs;
 			arch << m_constantImageLODIndices;
 			arch << m_constantImages;
-			arch << m_constantMeshes;
+			arch << ConstantMeshes;
 			arch << m_constantExtensionData;
 			arch << m_constantStrings;
             arch << m_constantLayouts;
@@ -307,10 +310,10 @@ namespace mu
             arch >> m_byteCode;
             arch >> m_states;
 			arch >> m_roms;
-			arch >> m_constantImageLODs;
+			arch >> ConstantImageLODs;
 			arch >> m_constantImageLODIndices;
 			arch >> m_constantImages;
-			arch >> m_constantMeshes;
+			arch >> ConstantMeshes;
 			arch >> m_constantExtensionData;
 			arch >> m_constantStrings;
             arch >> m_constantLayouts;
@@ -332,43 +335,49 @@ namespace mu
         void LogHistogram() const;
 
 		//! Return true if the given ROM is loaded.
-		inline bool IsRomLoaded( int32 RomIndex ) const
+		FORCEINLINE bool IsRomLoaded(int32 RomIndex) const
 		{
-			switch (m_roms[RomIndex].ResourceType)
+			const bool bIsRomLoaded = LoadedMemTrackedRoms.IsValidIndex(RomIndex);
+			check(bIsRomLoaded == Invoke([&]()
 			{
-			case DT_IMAGE:
-				return m_constantImageLODs[m_roms[RomIndex].ResourceIndex].Value.get() != nullptr;
-			case DT_MESH:
-				return m_constantMeshes[m_roms[RomIndex].ResourceIndex].Value.get() != nullptr;
-			default:
-				check(false);
-				break;
-			}
-			return false;
+				switch (m_roms[RomIndex].ResourceType)
+				{
+				case DT_IMAGE: return ConstantImageLODs[m_roms[RomIndex].ResourceIndex].Value.get() != nullptr;
+				case DT_MESH: return ConstantMeshes[m_roms[RomIndex].ResourceIndex].Value.get() != nullptr;
+				default: return false;
+				}
+			}));
+
+			return bIsRomLoaded;
 		}
 
 		/** Unload a rom resource. Return the size of the unloaded rom.*/
-		inline int32 UnloadRom(int32 RomIndex)
+		FORCEINLINE int32 UnloadRom(int32 RomIndex)
 		{
 			int32 RomSize = 0;
 
 			if (DebugRom && (DebugRomAll||RomIndex == DebugRomIndex))
 				UE_LOG(LogMutableCore, Log, TEXT("Unloading rom %d."), RomIndex);
 
+			if (LoadedMemTrackedRoms.IsValidIndex(RomIndex))
+			{
+				LoadedMemTrackedRoms.RemoveAt(RomIndex);
+			}
+
 			switch (m_roms[RomIndex].ResourceType)
 			{
 			case DT_IMAGE:
-				if (m_constantImageLODs[m_roms[RomIndex].ResourceIndex].Value)
+				if (ConstantImageLODs[m_roms[RomIndex].ResourceIndex].Value)
 				{
-					RomSize = m_constantImageLODs[m_roms[RomIndex].ResourceIndex].Value->GetDataSize();
-					m_constantImageLODs[m_roms[RomIndex].ResourceIndex].Value = nullptr;
+					RomSize = ConstantImageLODs[m_roms[RomIndex].ResourceIndex].Value->GetDataSize();
+					ConstantImageLODs[m_roms[RomIndex].ResourceIndex].Value = nullptr;
 				}
 				break;
 			case DT_MESH:
-				if (m_constantMeshes[m_roms[RomIndex].ResourceIndex].Value)
+				if (ConstantMeshes[m_roms[RomIndex].ResourceIndex].Value)
 				{
-					RomSize = m_constantMeshes[m_roms[RomIndex].ResourceIndex].Value->GetDataSize();
-					m_constantMeshes[m_roms[RomIndex].ResourceIndex].Value = nullptr;
+					RomSize = ConstantMeshes[m_roms[RomIndex].ResourceIndex].Value->GetDataSize();
+					ConstantMeshes[m_roms[RomIndex].ResourceIndex].Value = nullptr;
 				}
 				break;
 			default:
@@ -379,11 +388,27 @@ namespace mu
 			return RomSize;
 		}
 
+		FORCEINLINE void SetMeshRomValue(int32 RomIndex, const Ptr<Mesh>& Value)
+		{
+			check(m_roms[RomIndex].ResourceType == DT_MESH);
+			
+			LoadedMemTrackedRoms.EmplaceAt(RomIndex, (uint8)m_roms[RomIndex].ResourceType);
+			ConstantMeshes[m_roms[RomIndex].ResourceIndex].Value = Value;
+		}
+
+		FORCEINLINE void SetImageRomValue(int32 RomIndex, const Ptr<Image>& Value)
+		{
+			check(m_roms[RomIndex].ResourceType == DT_IMAGE);
+			
+			LoadedMemTrackedRoms.EmplaceAt(RomIndex, (uint8)m_roms[RomIndex].ResourceType);
+			ConstantImageLODs[m_roms[RomIndex].ResourceIndex].Value = Value;
+		}
+
 
 		int32 AddConstant(Ptr<const Mesh> pMesh)
 		{
 			// Uniques needs to be ensured outside
-			return m_constantMeshes.Add(TPair<int32, Ptr<const Mesh>>( -1, pMesh.get() ));
+			return ConstantMeshes.Add(TPair<int32, Ptr<const Mesh>>( -1, pMesh.get() ));
 		}
 
 		OP::ADDRESS AddConstant(Ptr<const ExtensionData> Data)
@@ -557,7 +582,7 @@ namespace mu
 
 			// Get the first mip
 			int32 ResultLODIndex = m_constantImageLODIndices[ResultLODIndexIndex];
-			Ptr<const Image> CurrentMip = m_constantImageLODs[ResultLODIndex].Value;
+			Ptr<const Image> CurrentMip = ConstantImageLODs[ResultLODIndex].Value;
 			check(CurrentMip);
 				
 			// Shortcut if we only want one mip
@@ -580,8 +605,8 @@ namespace mu
 					int32 TotalSize = 0;
 					for (int32 LOD = 0; LOD < FinalLODs; ++LOD)
 					{
-						int LODIndex = m_constantImageLODIndices[ResultLODIndexIndex + LOD];
-						int32 MipSizeBytes = m_constantImageLODs[LODIndex].Value->GetDataSize();
+						int32 LODIndex = m_constantImageLODIndices[ResultLODIndexIndex + LOD];
+						int32 MipSizeBytes = ConstantImageLODs[LODIndex].Value->GetDataSize();
 						TotalSize += MipSizeBytes;
 					}
 					Result->m_data.SetNum(TotalSize);
@@ -603,7 +628,7 @@ namespace mu
 					if (LOD + 1 < FinalLODs)
 					{
 						ResultLODIndex = m_constantImageLODIndices[ResultLODIndexIndex + LOD + 1];
-						CurrentMip = m_constantImageLODs[ResultLODIndex].Value;
+						CurrentMip = ConstantImageLODs[ResultLODIndex].Value;
 						check(CurrentMip);
 					}
 				}
@@ -613,9 +638,9 @@ namespace mu
 		}
 
 
-        void GetConstant(int32 ConstantIndex, MeshPtrConst& res ) const
+        void GetConstant(int32 ConstantIndex, MeshPtrConst& res) const
         {
-			res = m_constantMeshes[ConstantIndex].Value;
+			res = ConstantMeshes[ConstantIndex].Value;
 		}
 
 		void GetExtensionDataConstant(int32 ConstantIndex, ExtensionDataPtrConst& Result) const
diff --git a/Engine/Plugins/Experimental/Mutable/Source/MutableRuntime/Internal/MuR/SystemPrivate.h b/Engine/Plugins/Experimental/Mutable/Source/MutableRuntime/Internal/MuR/SystemPrivate.h
--- a/Engine/Plugins/Experimental/Mutable/Source/MutableRuntime/Internal/MuR/SystemPrivate.h
+++ b/Engine/Plugins/Experimental/Mutable/Source/MutableRuntime/Internal/MuR/SystemPrivate.h
@@ -14,6 +14,8 @@
 #include "MuR/ParametersPrivate.h"
 #include "MuR/MutableTrace.h"
 
+#include "Templates/UnrealTypeTraits.h"
+
 #if !(UE_BUILD_SHIPPING || UE_BUILD_TEST)
 #include "HAL/Thread.h"
 #include "HAL/PlatformTLS.h"
@@ -481,12 +483,23 @@ namespace mu
 		*/
 		CodeContainer<FOpExecutionData> OpExecutionData;
 
+		template<class ResourceType>
+		struct TResourceResult
+		{
+			static_assert(TIsDerivedFrom<ResourceType, Resource>::Value);
+
+			FCacheAddress OpAddress;
+			Ptr<const ResourceType> Value = nullptr;
+		};
+		static_assert(sizeof(TResourceResult<Image>) == 16);
+		static_assert(sizeof(TResourceResult<Mesh>) == 16);
+
 		/** */
 
 		TMemoryTrackedArray<FVector4f> ColorResults;
-		TMemoryTrackedArray<Ptr<const Image>> ImageResults;
+		TMemoryTrackedArray<TResourceResult<Image>> ImageResults;
+		TMemoryTrackedArray<TResourceResult<Mesh>> MeshResults;
 		TMemoryTrackedArray<Ptr<const Layout>> LayoutResults;
-		TMemoryTrackedArray<Ptr<const Mesh>> MeshResults;
 		TMemoryTrackedArray<Ptr<const Instance>> InstanceResults;
 		TMemoryTrackedArray<FProjector> ProjectorResults;
 		TMemoryTrackedArray<Ptr<const String>> StringResults;
@@ -541,9 +554,9 @@ namespace mu
 			{
 				// Insert dafault/null values
 				ColorResults.Add(FVector4f());
-				ImageResults.Add(nullptr);
+				ImageResults.Emplace();
 				LayoutResults.Add(nullptr);
-				MeshResults.Add(nullptr);
+				MeshResults.Emplace();
 				InstanceResults.Add(nullptr);
 				ProjectorResults.Add(FProjector());
 				StringResults.Add(nullptr);
@@ -563,11 +576,11 @@ namespace mu
 				switch ((DATATYPE)Data.DataType)
 				{
 				case DATATYPE::DT_IMAGE:					
-					ImageResults[DataTypeIndex] = nullptr;
+					ImageResults[DataTypeIndex].Value = nullptr;
 					break;
 
 				case DATATYPE::DT_MESH:
-					MeshResults[DataTypeIndex] = nullptr;
+					MeshResults[DataTypeIndex].Value = nullptr;
 					break;
 
 				case DATATYPE::DT_INSTANCE:
@@ -745,7 +758,7 @@ namespace mu
 			if (!Data) return nullptr;
 
 			check(Data->DataType == DATATYPE::DT_IMAGE);
-			Ptr<const Image> Result = ImageResults[Data->DataTypeIndex];
+			Ptr<const Image> Result = ImageResults[Data->DataTypeIndex].Value;
 
 			// We need to decrease the hit-count even if the result is null.
 			check(Data->OpHitCount > 0);
@@ -771,7 +784,7 @@ namespace mu
 			if (!Data) return nullptr;
 
 			check(Data->DataType == DATATYPE::DT_MESH);
-			Ptr<const Mesh> Result = MeshResults[Data->DataTypeIndex];
+			Ptr<const Mesh> Result = MeshResults[Data->DataTypeIndex].Value;
 
 			// We need to decrease the hit-count even if the result is null.
 			check(Data->OpHitCount > 0);
@@ -937,10 +950,10 @@ namespace mu
 			check(Data->DataTypeIndex != 0);
 		}
 
-		void SetImage(FCacheAddress at, Ptr<const Image> v)
+		void SetImage(FCacheAddress At, Ptr<const Image> Value)
 		{
-			check(at.At < OpExecutionData.size_code());
-			FOpExecutionData* Data = OpExecutionData.get_ptr(at);
+			check(At.At < OpExecutionData.size_code());
+			FOpExecutionData* Data = OpExecutionData.get_ptr(At);
 			check(Data->DataType == DATATYPE::DT_IMAGE || Data->DataType == DATATYPE::DT_NONE);
 			Data->DataType = DATATYPE::DT_IMAGE;
 			Data->IsValueValid = true;
@@ -948,21 +961,23 @@ namespace mu
 			if (!Data->DataTypeIndex)
 			{
 				Data->DataTypeIndex = ImageResults.Num();
-				ImageResults.Add(v);
+				ImageResults.Add(TResourceResult<Image>{At, Value});
 			}
 			else
 			{
-				ImageResults[Data->DataTypeIndex] = v;
+				ImageResults[Data->DataTypeIndex].Value = Value;
 			}
 			check(Data->DataTypeIndex != 0);
 
 			mu::UpdateLLMStats();
 		}
 
-		void SetMesh(FCacheAddress at, Ptr<const Mesh> v)
+		void SetMesh(FCacheAddress At, Ptr<const Mesh> Value)
 		{
-			check(at.At < OpExecutionData.size_code());
-			FOpExecutionData* Data = OpExecutionData.get_ptr(at);
+			check(At.At < OpExecutionData.size_code());
+
+			FOpExecutionData* Data = OpExecutionData.get_ptr(At);
+
 			check(Data->DataType == DATATYPE::DT_MESH || Data->DataType == DATATYPE::DT_NONE);
 			Data->DataType = DATATYPE::DT_MESH;
 			Data->IsValueValid = true;
@@ -970,11 +985,11 @@ namespace mu
 			if (!Data->DataTypeIndex)
 			{
 				Data->DataTypeIndex = MeshResults.Num();
-				MeshResults.Add(v);
+				MeshResults.Add(TResourceResult<Mesh>{At, Value});
 			}
 			else
 			{
-				MeshResults[Data->DataTypeIndex] = v;
+				MeshResults[Data->DataTypeIndex].Value = Value;
 			}
 			check(Data->DataTypeIndex != 0);
 
@@ -1565,14 +1580,14 @@ namespace mu
 			{
 				// Count streamable and currently-loaded resources
 				const FProgram& Program = Model->GetPrivate()->m_program;
-				for (const TPair<int32, Ptr<const Image>>& Rom : Program.m_constantImageLODs)
+				for (const TPair<int32, Ptr<const Image>>& Rom : Program.ConstantImageLODs)
 				{
 					if (Rom.Value && Rom.Key >= 0)
 					{
 						Result += Rom.Value->GetDataSize();
 					}
 				}
-				for (const TPair<int32, Ptr<const Mesh>>& Rom : Program.m_constantMeshes)
+				for (const TPair<int32, Ptr<const Mesh>>& Rom : Program.ConstantMeshes)
 				{
 					if (Rom.Value && Rom.Key >= 0)
 					{
@@ -1623,7 +1638,7 @@ namespace mu
 					switch ((*it).DataType)
 					{
 					case DATATYPE::DT_IMAGE:
-						Value = Instance.Cache->ImageResults[(*it).DataTypeIndex].get();
+						Value = Instance.Cache->ImageResults[(*it).DataTypeIndex].Value.get();
 						if (Value)
 						{
 							TargetSet->Add(Value);
@@ -1631,7 +1646,7 @@ namespace mu
 						break;
 
 					case DATATYPE::DT_MESH:
-						Value = Instance.Cache->MeshResults[(*it).DataTypeIndex].get();
+						Value = Instance.Cache->MeshResults[(*it).DataTypeIndex].Value.get();
 						if (Value)
 						{
 							TargetSet->Add(Value);
@@ -1681,20 +1696,20 @@ namespace mu
 				switch (Data.DataType)
 				{
 				case DATATYPE::DT_IMAGE:
-					Value = CurrentInstanceCache->ImageResults[Data.DataTypeIndex].get();
+					Value = CurrentInstanceCache->ImageResults[Data.DataTypeIndex].Value.get();
 					if (Value)
 					{
 						CacheResources.Remove(Value);
-						CurrentInstanceCache->ImageResults[Data.DataTypeIndex] = nullptr;
+						CurrentInstanceCache->ImageResults[Data.DataTypeIndex].Value = nullptr;
 					}
 					break;
 
 				case DATATYPE::DT_MESH:
-					Value = CurrentInstanceCache->MeshResults[Data.DataTypeIndex].get();
+					Value = CurrentInstanceCache->MeshResults[Data.DataTypeIndex].Value.get();
 					if (Value)
 					{
 						CacheResources.Remove(Value);
-						CurrentInstanceCache->MeshResults[Data.DataTypeIndex] = nullptr;
+						CurrentInstanceCache->MeshResults[Data.DataTypeIndex].Value = nullptr;
 					}
 					break;
 
@@ -1739,15 +1754,15 @@ namespace mu
 				switch (Data.DataType)
 				{
 				case DATATYPE::DT_IMAGE:
-					Value = CurrentInstanceCache->ImageResults[Data.DataTypeIndex].get();
+					Value = CurrentInstanceCache->ImageResults[Data.DataTypeIndex].Value.get();
 					CacheResources.Remove(Value);
-					CurrentInstanceCache->ImageResults[Data.DataTypeIndex] = nullptr;
+					CurrentInstanceCache->ImageResults[Data.DataTypeIndex].Value = nullptr;
 					break;
 
 				case DATATYPE::DT_MESH:
-					Value = CurrentInstanceCache->MeshResults[Data.DataTypeIndex].get();
+					Value = CurrentInstanceCache->MeshResults[Data.DataTypeIndex].Value.get();
 					CacheResources.Remove(Value);
-					CurrentInstanceCache->MeshResults[Data.DataTypeIndex] = nullptr;
+					CurrentInstanceCache->MeshResults[Data.DataTypeIndex].Value = nullptr;
 					break;
 
 				case DATATYPE::DT_LAYOUT:
diff --git a/Engine/Plugins/Experimental/Mutable/Source/MutableRuntime/Private/MuR/CodeRunner.cpp b/Engine/Plugins/Experimental/Mutable/Source/MutableRuntime/Private/MuR/CodeRunner.cpp
--- a/Engine/Plugins/Experimental/Mutable/Source/MutableRuntime/Private/MuR/CodeRunner.cpp
+++ b/Engine/Plugins/Experimental/Mutable/Source/MutableRuntime/Private/MuR/CodeRunner.cpp
@@ -928,7 +928,7 @@ namespace mu
             OP::ADDRESS cat = args.value;
 
             // Assume the ROM has been loaded previously
-            check(Program.m_constantMeshes[cat].Value)
+            check(Program.ConstantMeshes[cat].Value)
 
             Ptr<const Mesh> SourceConst;
 			Program.GetConstant(cat, SourceConst);
diff --git a/Engine/Plugins/Experimental/Mutable/Source/MutableRuntime/Private/MuR/CodeRunnerTasked.cpp b/Engine/Plugins/Experimental/Mutable/Source/MutableRuntime/Private/MuR/CodeRunnerTasked.cpp
--- a/Engine/Plugins/Experimental/Mutable/Source/MutableRuntime/Private/MuR/CodeRunnerTasked.cpp
+++ b/Engine/Plugins/Experimental/Mutable/Source/MutableRuntime/Private/MuR/CodeRunnerTasked.cpp
@@ -2140,16 +2140,17 @@ namespace mu
 			Runner->m_pSystem->StreamInterface->EndRead(RomLoadOp->m_streamID);									
 
 			const int32 ResIndex = Program.m_roms[RomIndex].ResourceIndex;
-			check(!Program.m_constantMeshes[ResIndex].Value)
+			check(!Program.ConstantMeshes[ResIndex].Value)
 			MUTABLE_CPUPROFILER_SCOPE(Unserialise);
 
-			InputMemoryStream stream(RomLoadOp->m_streamBuffer.GetData(), RomLoadOp->m_streamBuffer.Num());
-			InputArchive arch(&stream);
+			InputMemoryStream Stream(RomLoadOp->m_streamBuffer.GetData(), RomLoadOp->m_streamBuffer.Num());
+			InputArchive Arch(&Stream);
 
-			check(!Program.m_constantMeshes[ResIndex].Value);
-			Ptr<Mesh> Value = Mesh::StaticUnserialise(arch);
-			Program.m_constantMeshes[ResIndex].Value = Value;
-			check(Program.m_constantMeshes[ResIndex].Value);
+			check(!Program.ConstantMeshes[ResIndex].Value);
+			Ptr<Mesh> Value = Mesh::StaticUnserialise(Arch);
+			Program.SetMeshRomValue(RomIndex, Value);
+
+			check(Program.ConstantMeshes[ResIndex].Value);
 
 			Runner->RomLoadOps.Remove(*RomLoadOp);
 		}
@@ -2277,13 +2278,13 @@ namespace mu
 			const int32 CurrentIndexIndex = LODIndexIndex + LODIndex;
 			const int32 CurrentIndex = program.m_constantImageLODIndices[CurrentIndexIndex];
 
-			if (program.m_constantImageLODs[CurrentIndex].Key<0)
+			if (program.ConstantImageLODs[CurrentIndex].Key < 0)
 			{
 				// This data is always resident.
 				continue;
 			}
 
-			int32 RomIndex = program.m_constantImageLODs[CurrentIndex].Key;
+			int32 RomIndex = program.ConstantImageLODs[CurrentIndex].Key;
 			check(RomIndex < program.m_roms.Num());
 
 			++ModelCache->PendingOpsPerRom[RomIndex];
@@ -2379,7 +2380,7 @@ namespace mu
 			return;
 		}
 
-		FProgram& program = Runner->m_pModel->GetPrivate()->m_program;
+		FProgram& Program = Runner->m_pModel->GetPrivate()->m_program;
 		
 		FWorkingMemoryManager::FModelCacheEntry* ModelCache = Runner->m_pSystem->WorkingMemoryManager.FindModelCache(Runner->m_pModel.Get());
 
@@ -2392,16 +2393,17 @@ namespace mu
 
 				MUTABLE_CPUPROFILER_SCOPE(Unserialise);
 
-				InputMemoryStream stream(RomLoadOp->m_streamBuffer.GetData(), RomLoadOp->m_streamBuffer.Num());
-				InputArchive arch(&stream);
+				InputMemoryStream Stream(RomLoadOp->m_streamBuffer.GetData(), RomLoadOp->m_streamBuffer.Num());
+				InputArchive Arch(&Stream);
 
-				const int32 ResIndex = program.m_roms[RomIndex].ResourceIndex;
+				const int32 ResIndex = Program.m_roms[RomIndex].ResourceIndex;
 
 				// TODO: Try to reuse buffer from PooledImages.
-				check(!program.m_constantImageLODs[ResIndex].Value);
-				Ptr<Image> Value = Image::StaticUnserialise(arch);
-				program.m_constantImageLODs[ResIndex].Value = Value;
-				check(program.m_constantImageLODs[ResIndex].Value);
+				check(!Program.ConstantImageLODs[ResIndex].Value);
+				Ptr<Image> Value = Image::StaticUnserialise(Arch);
+
+				Program.SetImageRomValue(RomIndex, Value);
+				check(Program.ConstantImageLODs[ResIndex].Value);
 				
 				Runner->RomLoadOps.Remove(*RomLoadOp);
 			}
@@ -2413,16 +2415,16 @@ namespace mu
 		for (int32 LODIndex = 0; LODIndex < LODIndexCount; ++LODIndex)
 		{
 			int32 CurrentIndexIndex = LODIndexIndex + LODIndex;
-			int32 CurrentIndex = program.m_constantImageLODIndices[CurrentIndexIndex];
+			int32 CurrentIndex = Program.m_constantImageLODIndices[CurrentIndexIndex];
 
-			if (program.m_constantImageLODs[CurrentIndex].Key < 0)
+			if (Program.ConstantImageLODs[CurrentIndex].Key < 0)
 			{
 				// This data is always resident.
 				continue;
 			}
 
-			int32 RomIndex = program.m_constantImageLODs[CurrentIndex].Key;
-			check(RomIndex < program.m_roms.Num());
+			int32 RomIndex = Program.ConstantImageLODs[CurrentIndex].Key;
+			check(RomIndex < Program.m_roms.Num());
 
 			Runner->m_pSystem->WorkingMemoryManager.MarkRomUsed(RomIndex, Runner->m_pModel);
 			--ModelCache->PendingOpsPerRom[RomIndex];
@@ -2516,8 +2518,8 @@ namespace mu
 		case OP_TYPE::ME_CONSTANT:
 		{
 			OP::MeshConstantArgs args = program.GetOpArgs<OP::MeshConstantArgs>(item.At);
-			int32 RomIndex = program.m_constantMeshes[args.value].Key;
-			if (RomIndex >= 0 && !program.m_constantMeshes[args.value].Value )
+			int32 RomIndex = program.ConstantMeshes[args.value].Key;
+			if (RomIndex >= 0 && !program.ConstantMeshes[args.value].Value)
 			{
 				Issued = MakeShared<FLoadMeshRomTask>(item, RomIndex);
 			}
@@ -2565,7 +2567,7 @@ namespace mu
 			//for (int32 i=0; i<LODIndexCount; ++i)
 			//{
 			//	uint32 LODIndex = program.m_constantImageLODIndices[LODIndexIndex+i];
-			//	if ( !program.m_constantImageLODs[LODIndex].Value )
+			//	if ( !program.ConstantImageLODs[LODIndex].Value )
 			//	{
 			//		bAnyMissing = true;
 			//		break;
diff --git a/Engine/Plugins/Experimental/Mutable/Source/MutableRuntime/Private/MuR/Model.cpp b/Engine/Plugins/Experimental/Mutable/Source/MutableRuntime/Private/MuR/Model.cpp
--- a/Engine/Plugins/Experimental/Mutable/Source/MutableRuntime/Private/MuR/Model.cpp
+++ b/Engine/Plugins/Experimental/Mutable/Source/MutableRuntime/Private/MuR/Model.cpp
@@ -145,13 +145,13 @@ namespace mu
 
 		mu::FProgram& program = p->m_pD->m_program;
 
-		TArray<TPair<int32, mu::ImagePtrConst>> InitialImages = program.m_constantImageLODs;
-		TArray<TPair<int32, mu::MeshPtrConst>> InitialMeshes = program.m_constantMeshes;
+		TArray<TPair<int32, mu::ImagePtrConst>> InitialImages = program.ConstantImageLODs;
+		TArray<TPair<int32, mu::MeshPtrConst>> InitialMeshes = program.ConstantMeshes;
 
 		// Save images and unload from memory
-		for (int32 ResourceIndex = 0; ResourceIndex < program.m_constantImageLODs.Num(); ++ResourceIndex)
+		for (int32 ResourceIndex = 0; ResourceIndex < program.ConstantImageLODs.Num(); ++ResourceIndex)
 		{
-			TPair<int32, mu::ImagePtrConst>& ResData = program.m_constantImageLODs[ResourceIndex];
+			TPair<int32, mu::ImagePtrConst>& ResData = program.ConstantImageLODs[ResourceIndex];
 
 			// This shouldn't have been serialised with rom support before.
 			if (ResData.Key < 0)
@@ -179,9 +179,9 @@ namespace mu
 		}
 
 		// Save meshes and unload from memory
-		for (int32 ResourceIndex = 0; ResourceIndex < program.m_constantMeshes.Num(); ++ResourceIndex)
+		for (int32 ResourceIndex = 0; ResourceIndex < program.ConstantMeshes.Num(); ++ResourceIndex)
 		{
-			TPair<int32, mu::MeshPtrConst>& ResData = program.m_constantMeshes[ResourceIndex];
+			TPair<int32, mu::MeshPtrConst>& ResData = program.ConstantMeshes[ResourceIndex];
 
 			// This shouldn't have been serialised with rom support before.
 			if (ResData.Key < 0)
@@ -220,8 +220,8 @@ namespace mu
 		}
 
 		// Restore full data
-		program.m_constantImageLODs = InitialImages;
-		program.m_constantMeshes = InitialMeshes;
+		program.ConstantImageLODs = InitialImages;
+		program.ConstantMeshes = InitialMeshes;
 	}
 
 
@@ -251,19 +251,19 @@ namespace mu
     {
 		LLM_SCOPE_BYNAME(TEXT("MutableRuntime"));
 
-		for (int32 ResIndex = 0; ResIndex < m_pD->m_program.m_constantImageLODs.Num(); ++ResIndex)
+		for (int32 ResIndex = 0; ResIndex < m_pD->m_program.ConstantImageLODs.Num(); ++ResIndex)
 		{
-			if (m_pD->m_program.m_constantImageLODs[ResIndex].Key >= 0)
+			if (m_pD->m_program.ConstantImageLODs[ResIndex].Key >= 0)
 			{
-				m_pD->m_program.m_constantImageLODs[ResIndex].Value = nullptr;
+				m_pD->m_program.ConstantImageLODs[ResIndex].Value = nullptr;
 			}
 		}
 
-		for (int32 ResIndex = 0; ResIndex < m_pD->m_program.m_constantMeshes.Num(); ++ResIndex)
+		for (int32 ResIndex = 0; ResIndex < m_pD->m_program.ConstantMeshes.Num(); ++ResIndex)
 		{
-			if (m_pD->m_program.m_constantMeshes[ResIndex].Key >= 0)
+			if (m_pD->m_program.ConstantMeshes[ResIndex].Key >= 0)
 			{
-				m_pD->m_program.m_constantMeshes[ResIndex].Value = nullptr;
+				m_pD->m_program.ConstantMeshes[ResIndex].Value = nullptr;
 			}
 		}
 	}
diff --git a/Engine/Plugins/Experimental/Mutable/Source/MutableRuntime/Private/MuR/System.cpp b/Engine/Plugins/Experimental/Mutable/Source/MutableRuntime/Private/MuR/System.cpp
--- a/Engine/Plugins/Experimental/Mutable/Source/MutableRuntime/Private/MuR/System.cpp
+++ b/Engine/Plugins/Experimental/Mutable/Source/MutableRuntime/Private/MuR/System.cpp
@@ -42,6 +42,18 @@ TAutoConsoleVariable<bool> CVarTaskGraphBusyWait(
 	TEXT("Use TaskGraph BusyWait instead of simple Wait. Required to avoid hangs on platforms with low number of cores."),
 	ECVF_Default);
 
+namespace 
+{
+
+bool bEnableDetailedMemoryBudgetExceededLogging = false;
+static FAutoConsoleVariableRef CVarEnableDetailedMemoryBudgetExceededLogging (
+	TEXT("mutable.EnableDetailedMemoryBudgetExceededLogging"),
+	bEnableDetailedMemoryBudgetExceededLogging,
+	TEXT("If set to true, enables a more detailed logging when memory budget is exceeded. Only for Debug and Development builds."),
+	ECVF_Default);
+}
+
+
 namespace mu
 {
     static_assert( sizeof(mat4f) == 64, "UNEXPECTED_STRUCT_PACKING" );
@@ -478,19 +490,19 @@ namespace mu
 			if (Instance.InstanceID == instanceID)
 			{
 				// Make sure all the resources cached in the instance are removed from the tracking list
-				for (Ptr<const Image>& Data : Instance.Cache->ImageResults)
+				for (const FProgramCache::TResourceResult<Image>& Data : Instance.Cache->ImageResults)
 				{
-					if (Data)
+					if (Data.Value)
 					{
-						m_pD->WorkingMemoryManager.CacheResources.Remove(Data);
+						m_pD->WorkingMemoryManager.CacheResources.Remove(Data.Value);
 					}
 				}
 
-				for (Ptr<const Mesh>& Data : Instance.Cache->MeshResults)
+				for (const FProgramCache::TResourceResult<Mesh>& Data : Instance.Cache->MeshResults)
 				{
-					if (Data)
+					if (Data.Value)
 					{
-						m_pD->WorkingMemoryManager.CacheResources.Remove(Data);
+						m_pD->WorkingMemoryManager.CacheResources.Remove(Data.Value);
 					}
 				}
 
@@ -499,7 +511,7 @@ namespace mu
 			}
 		}
 
- 		int Removed = m_pD->WorkingMemoryManager.LiveInstances.RemoveAllSwap(
+ 		int32 Removed = m_pD->WorkingMemoryManager.LiveInstances.RemoveAllSwap(
 			[instanceID](const FLiveInstance& Instance)
 			{
 				return (Instance.InstanceID == instanceID);
@@ -1033,6 +1045,18 @@ namespace mu
 		return CurrentBytes > BudgetThresholdBytes;
 	}
 
+	namespace Private
+	{
+		enum class EBudgetBelowSearchFlags : uint8
+		{
+			None = 0,
+			Keep = 1 << 1,
+			Visited = 1 << 2,
+			FirstOccurance = 1 << 3
+		};
+
+		ENUM_CLASS_FLAGS(EBudgetBelowSearchFlags);
+	}
 
     //---------------------------------------------------------------------------------------------
 	bool FWorkingMemoryManager::EnsureBudgetBelow( uint64 AdditionalMemory )
@@ -1067,7 +1091,7 @@ namespace mu
 			}
 		}
 		
-		// Try to free a loaded roms
+		// Try to free loaded roms
 		if (!bFinished)
 		{
 			MUTABLE_CPUPROFILER_SCOPE(EnsureBudgetBelow_FreeRoms);
@@ -1086,40 +1110,45 @@ namespace mu
 				TSharedPtr<const Model> CacheModel = ModelCache.Model.Pin();
 				if (CacheModel)
 				{
-					mu::FProgram& program = CacheModel->GetPrivate()->m_program;
-					check(ModelCache.RomWeights.Num() == program.m_roms.Num());
+					mu::FProgram& Program = CacheModel->GetPrivate()->m_program;
+					check(ModelCache.RomWeights.Num() == Program.m_roms.Num());
+	
+					check(Program.LoadedMemTrackedRoms.GetMaxIndex() <= Program.m_roms.Num());
 
-					for (int32 RomIndex = 0; RomIndex < program.m_roms.Num(); ++RomIndex)
+					for (TSparseArray<uint8>::TConstIterator Iter = Program.LoadedMemTrackedRoms.CreateConstIterator(); Iter; ++Iter)
 					{
-						const FRomData& Rom = program.m_roms[RomIndex];
-						bool bIsLoaded = program.IsRomLoaded(RomIndex);
+						const int32 RomIndex = Iter.GetIndex();
+
+						const FRomData& Rom = Program.m_roms[RomIndex];
+						check(Program.IsRomLoaded(RomIndex));
+						check(Rom.ResourceType == (uint32)*Iter);
 
 						// We cannot unload a rom if some operation is expecting it.
-						bool bIsRomLocked = ModelCache.PendingOpsPerRom.IsValidIndex(RomIndex)
-							&& 
-							ModelCache.PendingOpsPerRom[RomIndex]>0;
-						if (bIsLoaded && !bIsRomLocked)
+						const bool bIsRomLocked = ModelCache.PendingOpsPerRom.IsValidIndex(RomIndex) && 
+												  ModelCache.PendingOpsPerRom[RomIndex] > 0;
+						if (!bIsRomLocked)
 						{
 							constexpr float FactorWeight = 100.0f;
 							constexpr float FactorTime = -1.0f;
-							float Priority = FactorWeight * float(ModelCache.RomWeights[RomIndex].Key)
-								+
-								FactorTime * float((RomTick - ModelCache.RomWeights[RomIndex].Value));
+							float Priority = FactorWeight * float(ModelCache.RomWeights[RomIndex].Get<0>()) +
+											 FactorTime * float((RomTick - ModelCache.RomWeights[RomIndex].Get<1>()));
 
-							FRomRef Ref = { CacheModel.Get(), RomIndex };
-							Candidates.Add(TPair<float, FRomRef>(Priority,Ref));
+							Candidates.Emplace(Priority, FRomRef{ CacheModel.Get(), RomIndex });
 						}
 					}
 				}
 			}
 
-			Candidates.Sort([](const TPair<float, FRomRef>& A, const TPair<float, FRomRef>& B) { return A.Key > B.Key; });
+			// Don't sort all candidates, make it a heap in O(N) time. We may not need to visit all elements.
+			auto CompareCandidates = [](const TPair<float, FRomRef>& A, const TPair<float, FRomRef>& B) { return A.Key < B.Key; };
+			Candidates.Heapify(CompareCandidates);
 
 			while (!bFinished && Candidates.Num())
 			{
 				MUTABLE_CPUPROFILER_SCOPE(EnsureBudgetBelow_UnloadRom);
 
-				TPair<float, FRomRef> Candidate = Candidates.Pop(EAllowShrinking::No);
+				TPair<float, FRomRef> Candidate;
+				Candidates.HeapPop(Candidate, CompareCandidates, EAllowShrinking::No);
 
 				// UE_LOG(LogMutableCore,Log, "Unloading rom because of memory budget: %d.", lowestPriorityRom);
 				int32 UnloadedSize = Candidate.Value.Model->GetPrivate()->m_program.UnloadRom(Candidate.Value.RomIndex);
@@ -1127,253 +1156,273 @@ namespace mu
 				bFinished = TotalBytes <= BudgetBytes;
 			}
 		}
-
+	
 		// Try to free cache 1 memory
 		if (!bFinished)
 		{
 			MUTABLE_CPUPROFILER_SCOPE(EnsureBudgetBelow_FreeCached);
 
+			TSet<const Resource*> RemovedResources;
+			RemovedResources.Reserve(1024);
+			
 			// From other live instances first
 			for (const FLiveInstance& Instance : LiveInstances)
 			{
-				if (Instance.Cache==CurrentInstanceCache)
+				if (Instance.Cache == CurrentInstanceCache)
 				{
 					// Ignore the current live instance.
 					continue;
 				}
 
 				// Gather all data in the cache for this instance
-				TArray<const Resource*> CacheUnique;
-				CacheUnique.Reserve(1024);
 				{
-					MUTABLE_CPUPROFILER_SCOPE(EnsureBudgetBelow_FreeCached_Gather_Other);
+					MUTABLE_CPUPROFILER_SCOPE(EnsureBudgetBelow_FreeCached_GatherAndRemove_Other);
 
-					CodeContainer<FProgramCache::FOpExecutionData>::iterator It = Instance.Cache->OpExecutionData.begin();
-					for (; It.IsValid(); ++It)
+					RemovedResources.Reset();
+					int32 ImageRemovedBytes = 0;
+					for (const FProgramCache::TResourceResult<Image>& ImageResult : Instance.Cache->ImageResults)
 					{
-						FProgramCache::FOpExecutionData& Data = *It;
-
-						if (!Data.DataTypeIndex)
+						if (!ImageResult.Value)
 						{
 							continue;
 						}
 
-						const Resource* Value = nullptr;
-						switch (Data.DataType)
+						if (!RemovedResources.Find(ImageResult.Value.get()))
 						{
-						case DATATYPE::DT_IMAGE:
-							Value = Instance.Cache->ImageResults[Data.DataTypeIndex].get();
-							if (Value)
-							{
-								CacheUnique.AddUnique(Value);
-							}
-							break;
+							ImageRemovedBytes += ImageResult.Value->GetDataSize();
+							bFinished = TotalBytes - ImageRemovedBytes <= BudgetBytes;
 
-						case DATATYPE::DT_MESH:
-							Value = Instance.Cache->MeshResults[Data.DataTypeIndex].get();
-							if (Value)
-							{
-								CacheUnique.AddUnique(Value);
-							}
-							break;
+							RemovedResources.Add(ImageResult.Value.get());
+							CacheResources.Remove(ImageResult.Value);
+						}
 
-						default:
+						if (bFinished)
+						{
 							break;
 						}
 					}
-				}
 
-				{
-					MUTABLE_CPUPROFILER_SCOPE(EnsureBudgetBelow_FreeCached_Free_Other);
-
-					while (!bFinished && CacheUnique.Num())
+					if (ImageRemovedBytes > 0)
 					{
-						// Free one
-						const Resource* Removed = CacheUnique.Pop(EAllowShrinking::No);
-
-						int32 RemovedDataSize = Removed->GetDataSize();
-
-						// Clear its cache references
-						CodeContainer<FProgramCache::FOpExecutionData>::iterator RemIt = Instance.Cache->OpExecutionData.begin();
-						for (; RemIt.IsValid(); ++RemIt)
+						for (const FProgramCache::TResourceResult<Image>& ImageResult : Instance.Cache->ImageResults)
 						{
-							FProgramCache::FOpExecutionData& Data = *RemIt;
-
-							if (!Data.DataTypeIndex)
-							{
-								continue;
-							}
-
-							const Resource* Value = nullptr;
-							switch (Data.DataType)
+							if (RemovedResources.Find(ImageResult.Value.get()))
 							{
-							case DATATYPE::DT_IMAGE:
-								Value = Instance.Cache->ImageResults[Data.DataTypeIndex].get();
-								break;
-
-							case DATATYPE::DT_MESH:
-								Value = Instance.Cache->MeshResults[Data.DataTypeIndex].get();
-								break;
-
-							default:
-								break;
-							}
-
-							if (Value == Removed)
-							{
-								CacheResources.Remove(Removed);
-								Instance.Cache->SetUnused(*RemIt);
+								Instance.Cache->SetUnused(Instance.Cache->OpExecutionData[ImageResult.OpAddress]);
 							}
 						}
-
-						TotalBytes -= RemovedDataSize;
-						bFinished = TotalBytes <= BudgetBytes;
 					}
 
+					TotalBytes -= ImageRemovedBytes;
 					if (bFinished)
 					{
 						break;
 					}
-				}
-			}
 
-			// From the current live instances. It is more involved: we have to make sure any data we want to free is not also
-			// in any cache (0 or 1) position with hit-count > 0.
-			if (CurrentInstanceCache && !bFinished)
-			{
-				// Gather all data in the cache for this instance
-				TArray<const Resource*> CacheUnique;
-				CacheUnique.Reserve(1024);
-				{
-					MUTABLE_CPUPROFILER_SCOPE(EnsureBudgetBelow_FreeCached_Gather_Current);
-
-					CodeContainer<FProgramCache::FOpExecutionData>::iterator It = CurrentInstanceCache->OpExecutionData.begin();
-					for (; It.IsValid(); ++It)
+					int32 MeshRemovedBytes = 0;
+					RemovedResources.Reset();
+					for (const FProgramCache::TResourceResult<Mesh>& MeshResult : Instance.Cache->MeshResults)
 					{
-						FProgramCache::FOpExecutionData& Data = *It;
-
-						if (!Data.DataTypeIndex || Data.OpHitCount > 0)
+						if (!MeshResult.Value)
 						{
 							continue;
 						}
 
-						const Resource* Value = nullptr;
-						switch (Data.DataType)
+						if (!RemovedResources.Find(MeshResult.Value.get()))
 						{
-						case DATATYPE::DT_IMAGE:
-							Value = CurrentInstanceCache->ImageResults[Data.DataTypeIndex].get();
-							if (Value)
-							{
-								CacheUnique.AddUnique(Value);
-							}
-							break;
+							MeshRemovedBytes = MeshResult.Value->GetDataSize();
+							bFinished = TotalBytes - MeshRemovedBytes <= BudgetBytes;
 
-						case DATATYPE::DT_MESH:
-							Value = CurrentInstanceCache->MeshResults[Data.DataTypeIndex].get();
-							if (Value)
-							{
-								CacheUnique.AddUnique(Value);
-							}
-							break;
+							RemovedResources.Add(MeshResult.Value.get());
+							CacheResources.Remove(MeshResult.Value);
+						}
 
-						default:
+						if (bFinished)
+						{
 							break;
 						}
 					}
-				}
-
-				{
-					MUTABLE_CPUPROFILER_SCOPE(EnsureBudgetBelow_FreeCached_Free_Current);
 
-					while (!bFinished && CacheUnique.Num())
+					if (MeshRemovedBytes > 0)
 					{
-						// Free one
-						const Resource* Removed = CacheUnique.Pop(EAllowShrinking::No);
-
-						// Does this data have any other cache references with op-hit-count bigger than 0?
-						bool bStillUsed = false;
-
-						CodeContainer<FProgramCache::FOpExecutionData>::iterator CheckIt = CurrentInstanceCache->OpExecutionData.begin();
-						for (; CheckIt.IsValid(); ++CheckIt)
+						for (const FProgramCache::TResourceResult<Mesh>& MeshResult : Instance.Cache->MeshResults)
 						{
-							const FProgramCache::FOpExecutionData& Data = *CheckIt;
-
-							if (!Data.DataTypeIndex)
+							if (RemovedResources.Find(MeshResult.Value.get()))
 							{
-								continue;
+								Instance.Cache->SetUnused(Instance.Cache->OpExecutionData[MeshResult.OpAddress]);
 							}
+						}
+					}
 
-							const Resource* Value = nullptr;
-							switch (Data.DataType)
-							{
-							case DATATYPE::DT_IMAGE:
-								Value = CurrentInstanceCache->ImageResults[Data.DataTypeIndex].get();
-								break;
+					TotalBytes -= MeshRemovedBytes;
+					if (bFinished)
+					{
+						break;
+					}
+				}
+			}
+		}
 
-							case DATATYPE::DT_MESH:
-								Value = CurrentInstanceCache->MeshResults[Data.DataTypeIndex].get();
-								break;
+		// From the current live instances. It is more involved: we have to make sure any data we want to free is not also
+		// in any cache (0 or 1) position with hit-count > 0.
+		if (!bFinished && CurrentInstanceCache)
+		{
+			MUTABLE_CPUPROFILER_SCOPE(EnsureBudgetBelow_FreeCached_Current);
+			// This removes images first, not sure if this can become a problem.
+			
+			using namespace Private;
+			using ESearchFlags = EBudgetBelowSearchFlags;
+			
+			auto SearchResourcesToRemove = 
+			[
+				TotalBytes = TotalBytes, 
+				BudgetBytes = BudgetBytes, 
+				CurrentInstanceCache = CurrentInstanceCache.Get()
+			](const auto& ResourceRange, TArray<ESearchFlags>& SearchFlags) -> int32
+			{
+				MUTABLE_CPUPROFILER_SCOPE(EnsureBudgetBelow_FreeCached_Current_SearchResources);
 
-							default:
-								break;
-							}
+				int32 RemovedBytes = 0;
 
-							if (Value == Removed && Data.OpHitCount > 0)
+				const int32 NumResources = ResourceRange.Num();
+				for (int32 ResourceIndex = 0; ResourceIndex < NumResources; ++ResourceIndex)
+				{
+					// Check if visited.
+					if (EnumHasAnyFlags(SearchFlags[ResourceIndex], ESearchFlags::Visited))
+					{
+						continue;
+					}
+	
+					// Null values will not have any flag set.
+					if (!ResourceRange[ResourceIndex].Value)
+					{
+						continue;
+					}
+
+					SearchFlags[ResourceIndex] = (ESearchFlags::Visited | ESearchFlags::FirstOccurance);
+					
+					if (CurrentInstanceCache->OpExecutionData[ResourceRange[ResourceIndex].OpAddress].OpHitCount > 0)
+					{
+						// Mark all occurences as visited to keep. 
+						for (int32 I = ResourceIndex; I < NumResources; ++I)
+						{
+							if (ResourceRange[I].Value == ResourceRange[ResourceIndex].Value)
 							{
-								bStillUsed = true;
-								break;
+								EnumAddFlags(SearchFlags[I], (ESearchFlags::Visited | ESearchFlags::Keep));
 							}
 						}
-
-						if (bStillUsed)
+					}
+					else
+					{
+						int32 I = ResourceIndex + 1;
+						for (; I < NumResources; ++I)
 						{
-							continue;
+							// Mark as visted.
+							if (ResourceRange[I].Value == ResourceRange[ResourceIndex].Value)
+							{
+								EnumAddFlags(SearchFlags[I], ESearchFlags::Visited);	
+								if (CurrentInstanceCache->OpExecutionData[ResourceRange[I].OpAddress].OpHitCount > 0)
+								{
+									// Still used, next step will mark to keep.
+									break;
+								}
+							}
 						}
 
-
-						int32 RemovedDataSize = Removed->GetDataSize();
-
-						// Clear its cache references
-						CodeContainer<FProgramCache::FOpExecutionData>::iterator RemIt = CurrentInstanceCache->OpExecutionData.begin();
-						for (; RemIt.IsValid(); ++RemIt)
+						if (I < NumResources)
 						{
-							FProgramCache::FOpExecutionData& Data = *RemIt;
-
-							if (!Data.DataTypeIndex)
+							// The image is still used, mark all occurences as visited and keep.
+							for (int32 J = ResourceIndex; J < NumResources; ++J)
 							{
-								continue;
+								if (ResourceRange[J].Value == ResourceRange[ResourceIndex].Value)
+								{
+									EnumAddFlags(SearchFlags[J], (ESearchFlags::Visited | ESearchFlags::Keep));	
+								}
 							}
+						}
+						else
+						{
+							// The resource is not used, see if we can stop searching. 
+							// In that case all occurences have been marked as visited.
+							RemovedBytes += ResourceRange[ResourceIndex].Value->GetDataSize();
 
-							const Resource* Value = nullptr;
-							switch (Data.DataType)
+							if (TotalBytes - RemovedBytes <= BudgetBytes)
 							{
-							case DATATYPE::DT_IMAGE:
-								Value = CurrentInstanceCache->ImageResults[Data.DataTypeIndex].get();
-								break;
-
-							case DATATYPE::DT_MESH:
-								Value = CurrentInstanceCache->MeshResults[Data.DataTypeIndex].get();
-								break;
-
-							default:
-								break;
+								return RemovedBytes;
 							}
+						}
+					}
+				}
 
-							if (Value == Removed)
-							{
-								MUTABLE_CPUPROFILER_SCOPE(EnsureBudgetBelow_FreeCached_Free_Current_ActualFree);
+				return RemovedBytes;
+			};
 
-								CacheResources.Remove(Removed);
-								CurrentInstanceCache->SetUnused(*RemIt);
-							}
+			const int32 MaxNumResources = FMath::Max(CurrentInstanceCache->ImageResults.Num(), CurrentInstanceCache->MeshResults.Num());
+			TArray<ESearchFlags> SearchFlags;
+			SearchFlags.SetNumUninitialized(MaxNumResources);
+			
+			if (!bFinished)
+			{
+				FMemory::Memzero(SearchFlags.GetData(), MaxNumResources*sizeof(ESearchFlags));
+				
+				const TArrayView<FProgramCache::TResourceResult<Image>>& Images = MakeArrayView(CurrentInstanceCache->ImageResults);
+				const int32 RemovedBytes = SearchResourcesToRemove(Images, SearchFlags); 
+				
+				if (RemovedBytes > 0)
+				{	
+					MUTABLE_CPUPROFILER_SCOPE(EnsureBudgetBelow_FreeCached_Current_FreeResources);
+			
+					const int32 NumImages = Images.Num();
+					for (int32 I = 0; I < NumImages; ++I)
+					{
+						// Remove the first occurence of any visited not to keep. 
+						if (SearchFlags[I] == (ESearchFlags::FirstOccurance | ESearchFlags::Visited))
+						{
+							CacheResources.Remove(Images[I].Value);
+						}
+					
+						// Set unused all references visited not marked to keep. 
+						if ((SearchFlags[I] & (ESearchFlags::Visited | ESearchFlags::Keep)) == ESearchFlags::Visited)
+						{
+							CurrentInstanceCache->SetUnused(CurrentInstanceCache->OpExecutionData[Images[I].OpAddress]);
 						}
+					}	
+				}
 
-						TotalBytes -= RemovedDataSize;
-						bFinished = TotalBytes <= BudgetBytes;
+				TotalBytes -= RemovedBytes;
+				bFinished = TotalBytes <= BudgetBytes;
+			}
+			
+			if (!bFinished)
+			{
+				FMemory::Memzero(SearchFlags.GetData(), MaxNumResources*sizeof(ESearchFlags));
+				
+				TArrayView<FProgramCache::TResourceResult<Mesh>> Meshes = MakeArrayView(CurrentInstanceCache->MeshResults);
+				const int32 RemovedBytes = SearchResourcesToRemove(Meshes, SearchFlags); 
+
+				if (RemovedBytes > 0)
+				{	
+					MUTABLE_CPUPROFILER_SCOPE(EnsureBudgetBelow_FreeCached_Current_FreeResources);
+					
+					const int32 NumMeshes = Meshes.Num();
+					for (int32 I = 0; I < NumMeshes; ++I)
+					{
+						// Remove the first occurence of any visited not to keep. 
+						if (SearchFlags[I] == (ESearchFlags::FirstOccurance | ESearchFlags::Visited))
+						{
+							CacheResources.Remove(Meshes[I].Value);
+						}
+						
+						if ((SearchFlags[I] & (ESearchFlags::Visited | ESearchFlags::Keep)) == ESearchFlags::Visited)
+						{
+							CurrentInstanceCache->SetUnused(CurrentInstanceCache->OpExecutionData[Meshes[I].OpAddress]);
+						}
 					}
 				}
-			}
 
+				TotalBytes -= RemovedBytes;
+				bFinished = TotalBytes <= BudgetBytes;
+			}
 		}
 
 		if (!bFinished)
@@ -1388,10 +1437,14 @@ namespace mu
 				// This is a good place to insert a brakpoint to detect callstacks with memory peaks
 				UE_LOG(LogMutableCore, Log, TEXT("Failed to keep memory budget. Budget: %d, Current: %d, New: %d"),
 					BudgetBytes / 1024, (TotalBytes - AdditionalMemory) / 1024, AdditionalMemory / 1024);
-
-				// We won't show correct internal or streaming buffer memory.
-				LogWorkingMemory(nullptr);
+				
+				if (bEnableDetailedMemoryBudgetExceededLogging)
+				{
+					// We won't show correct internal or streaming buffer memory.
+					LogWorkingMemory(nullptr);
+				}
 			}
+
 		}
 
         return bFinished;
diff --git a/Engine/Plugins/Experimental/Mutable/Source/MutableTools/Private/MuT/ASTOpConstantResource.cpp b/Engine/Plugins/Experimental/Mutable/Source/MutableTools/Private/MuT/ASTOpConstantResource.cpp
--- a/Engine/Plugins/Experimental/Mutable/Source/MutableTools/Private/MuT/ASTOpConstantResource.cpp
+++ b/Engine/Plugins/Experimental/Mutable/Source/MutableTools/Private/MuT/ASTOpConstantResource.cpp
@@ -130,7 +130,7 @@ namespace mu
 
 				if (MipIndex<0)
 				{
-					MipIndex = Program.m_constantImageLODs.Add(TPair<int32, Ptr<const Image>>(-1, pMip));
+					MipIndex = Program.ConstantImageLODs.Add(TPair<int32, Ptr<const Image>>(-1, pMip));
 					Options.ImageConstantMipMap.Add(pMip, MipIndex);
 				}
 
@@ -531,4 +531,4 @@ namespace mu
 		return nullptr;
 	}
 
-}
\ No newline at end of file
+}
diff --git a/Engine/Plugins/Experimental/Mutable/Source/MutableTools/Private/MuT/Compiler.cpp b/Engine/Plugins/Experimental/Mutable/Source/MutableTools/Private/MuT/Compiler.cpp
--- a/Engine/Plugins/Experimental/Mutable/Source/MutableTools/Private/MuT/Compiler.cpp
+++ b/Engine/Plugins/Experimental/Mutable/Source/MutableTools/Private/MuT/Compiler.cpp
@@ -490,9 +490,9 @@ namespace mu
 		int32 NumEmbedded = 0;
 
 		// Save images and unload from memory
-		for (int32 ResourceIndex = 0; ResourceIndex < program.m_constantImageLODs.Num(); ++ResourceIndex)
+		for (int32 ResourceIndex = 0; ResourceIndex < program.ConstantImageLODs.Num(); ++ResourceIndex)
 		{
-			TPair<int32, mu::ImagePtrConst>& ResData = program.m_constantImageLODs[ResourceIndex];
+			TPair<int32, mu::ImagePtrConst>& ResData = program.ConstantImageLODs[ResourceIndex];
 
 			// This shouldn't have been serialised with rom support before.
 			check(ResData.Key < 0);
@@ -525,9 +525,9 @@ namespace mu
 		}
 
 		// Save meshes and unload from memory
-		for (int32 ResourceIndex = 0; ResourceIndex < program.m_constantMeshes.Num(); ++ResourceIndex)
+		for (int32 ResourceIndex = 0; ResourceIndex < program.ConstantMeshes.Num(); ++ResourceIndex)
 		{
-			TPair<int32, mu::MeshPtrConst>& ResData = program.m_constantMeshes[ResourceIndex];
+			TPair<int32, mu::MeshPtrConst>& ResData = program.ConstantMeshes[ResourceIndex];
 
 			// This shouldn't have been serialised with rom support before.
 			check(ResData.Key < 0);
