[Mutable] Fix Edit and Extend Material Nodes ignoring Tags.

Unified GenerateSurfaceConditional and GenerateSurfaceModifier.

#rnx

diff --git a/Engine/Plugins/Experimental/Mutable/Source/MutableTools/Private/MuT/CodeGenerator_FirstPass.cpp b/Engine/Plugins/Experimental/Mutable/Source/MutableTools/Private/MuT/CodeGenerator_FirstPass.cpp
--- a/Engine/Plugins/Experimental/Mutable/Source/MutableTools/Private/MuT/CodeGenerator_FirstPass.cpp
+++ b/Engine/Plugins/Experimental/Mutable/Source/MutableTools/Private/MuT/CodeGenerator_FirstPass.cpp
@@ -254,6 +254,8 @@ namespace mu
 		if (Surface)
 		{
 			FSurface::FEdit edit;
+			edit.PositiveTags = m_currentPositiveTags;
+			edit.NegativeTags = m_currentNegativeTags;
             edit.node = Private;
             edit.condition = m_currentCondition.Last().objectCondition;
 			Surface->edits.Add(edit);
diff --git a/Engine/Plugins/Experimental/Mutable/Source/MutableTools/Private/MuT/CodeGenerator_FirstPass.h b/Engine/Plugins/Experimental/Mutable/Source/MutableTools/Private/MuT/CodeGenerator_FirstPass.h
--- a/Engine/Plugins/Experimental/Mutable/Source/MutableTools/Private/MuT/CodeGenerator_FirstPass.h
+++ b/Engine/Plugins/Experimental/Mutable/Source/MutableTools/Private/MuT/CodeGenerator_FirstPass.h
@@ -108,6 +108,12 @@ namespace mu
 			// All surface editing nodes that edit this surface
             struct FEdit
             {
+            	// List of tags that are required for the presence of this surface
+            	TArray<FString> PositiveTags;
+
+            	// List of tags that block the presence of this surface
+            	TArray<FString> NegativeTags;
+            	
                 //! Condition that enables the effects of this edit node on the surface
                 Ptr<ASTOp> condition;
 
diff --git a/Engine/Plugins/Experimental/Mutable/Source/MutableTools/Private/MuT/CodeGenerator_SecondPass.cpp b/Engine/Plugins/Experimental/Mutable/Source/MutableTools/Private/MuT/CodeGenerator_SecondPass.cpp
--- a/Engine/Plugins/Experimental/Mutable/Source/MutableTools/Private/MuT/CodeGenerator_SecondPass.cpp
+++ b/Engine/Plugins/Experimental/Mutable/Source/MutableTools/Private/MuT/CodeGenerator_SecondPass.cpp
@@ -67,6 +67,16 @@ namespace mu
 			}
 		}
 
+		auto PositiveTags = [this](size_t SurfaceIndex) -> TArray<FString>&
+		{
+			return m_pFirstPass->surfaces[SurfaceIndex].positiveTags;
+		};
+
+		auto NegativeTags = [this](size_t SurfaceIndex) -> TArray<FString>&
+		{
+			return m_pFirstPass->surfaces[SurfaceIndex].negativeTags;
+		};
+		
 		Ptr<ASTOp> c;
 
 		// Condition expression for all the surfaces that activate the tag
@@ -91,7 +101,9 @@ namespace mu
 			auto positiveTags = posTag;
 			positiveTags.insert(tagIndex);
 
-			Ptr<ASTOp> surfCondition = GenerateSurfaceCondition(surfIndex,
+			Ptr<ASTOp> surfCondition = GenerateSurfaceOrModifierCodition(surfIndex,
+				PositiveTags,
+				NegativeTags,
 				posSurf,
 				negSurf,
 				positiveTags,
@@ -161,7 +173,9 @@ namespace mu
 				auto positiveTags = posTag;
 				positiveTags.insert(tagIndex);
 
-				surfCondition = GenerateSurfaceCondition(editKey.Key,
+				surfCondition = GenerateSurfaceOrModifierCodition(editKey.Key,
+					PositiveTags,
+					NegativeTags,
 					posSurf,
 					negSurf,
 					positiveTags,
@@ -222,29 +236,29 @@ namespace mu
 	}
 
 	//---------------------------------------------------------------------------------------------
-	mu::Ptr<ASTOp> SecondPassGenerator::GenerateSurfaceCondition(size_t surfIndex,
+	mu::Ptr<ASTOp> SecondPassGenerator::GenerateSurfaceOrModifierCodition(size_t Index,
+		TFunction<const TArray<FString>&(size_t)> PositiveTags,
+		TFunction<const TArray<FString>&(size_t)> NegativeTags,
 		const set<size_t>& posSurf,
 		const set<size_t>& negSurf,
 		const set<size_t>& posTag,
 		const set<size_t>& negTag)
 	{
-		const auto& surf = m_pFirstPass->surfaces[surfIndex];
-
 		// If this surface is already in the list of positive surfaces, return true as condition
-		if (posSurf.find(surfIndex) != posSurf.end())
+		if (posSurf.find(Index) != posSurf.end())
 		{
 			return m_opPool.Add(new ASTOpConstantBool(true));
 		}
 
 		// If this surface is already in the list of negative surfaces, return false as condition
-		if (negSurf.find(surfIndex) != negSurf.end())
+		if (negSurf.find(Index) != negSurf.end())
 		{
 			return m_opPool.Add(new ASTOpConstantBool(false));
 		}
 
 		Ptr<ASTOp> c;
 
-		for (const auto& t : surf.positiveTags)
+		for (const auto& t : PositiveTags(Index))
 		{
 			auto it = m_pFirstPass->m_tags.FindByPredicate([&](const FirstPassGenerator::FTag& e) { return e.tag == t; });
 			if (!it)
@@ -257,7 +271,7 @@ namespace mu
 			size_t tagIndex = it - &m_pFirstPass->m_tags[0];
 
 			set<size_t> positiveSurfacesVisited = posSurf;
-			positiveSurfacesVisited.insert(surfIndex);
+			positiveSurfacesVisited.insert(Index);
 
 			Ptr<ASTOp> tagCondition = GenerateTagCondition(tagIndex,
 				positiveSurfacesVisited,
@@ -314,7 +328,7 @@ namespace mu
 		}
 
 
-		for (const auto& t : surf.negativeTags)
+		for (const auto& t : NegativeTags(Index))
 		{
 			auto it = m_pFirstPass->m_tags.FindByPredicate([&](const FirstPassGenerator::FTag& e) { return e.tag == t; });
 			if (!it)
@@ -327,7 +341,7 @@ namespace mu
 
 			set<size_t> positiveSurfacesVisited = negSurf;
 			set<size_t> negativeSurfacesVisited = posSurf;
-			negativeSurfacesVisited.insert(surfIndex);
+			negativeSurfacesVisited.insert(Index);
 			set<size_t> positiveTagsVisited = negTag;
 			set<size_t> negativeTagsVisited = posTag;
 			Ptr<ASTOp> tagCondition = GenerateTagCondition(tagIndex,
@@ -390,157 +404,6 @@ namespace mu
 		return c;
 	}
 
-	//---------------------------------------------------------------------------------------------
-	mu::Ptr<ASTOp> SecondPassGenerator::GenerateModifierCondition(size_t modIndex)
-	{
-		const auto& mod = m_pFirstPass->modifiers[modIndex];
-
-		Ptr<ASTOp> c;
-
-		bool done = false;
-		for (const auto& t : mod.positiveTags)
-		{
-			auto it = m_pFirstPass->m_tags.FindByPredicate([&](const FirstPassGenerator::FTag& e) { return e.tag == t; });
-			if (!it)
-			{
-				// This could happen if a tag is in a variation but noone defines it.
-
-				// Entire expression will be false
-				Ptr<ASTOpConstantBool> f = new ASTOpConstantBool(false);
-				c = m_opPool.Add(f);
-				done = true;
-
-				continue;
-			}
-
-			size_t tagIndex = it - &m_pFirstPass->m_tags[0];
-
-			set<size_t> empty;
-			Ptr<ASTOp> tagCondition = GenerateTagCondition(tagIndex, empty, empty, empty, empty);
-
-			if (!tagCondition)
-			{
-				// This tag is unconditionally activated, so there's no condition logic to add
-				continue;
-			}
-
-			// TODO: Optimise the tag condition here
-
-			// If the tag is a constant ...
-			bool isConstant = false;
-			bool constantValue = false;
-			if (tagCondition && tagCondition->GetOpType() == OP_TYPE::BO_CONSTANT)
-			{
-				const ASTOpConstantBool* constOp = static_cast<const ASTOpConstantBool*>(tagCondition.get());
-				isConstant = true;
-				constantValue = constOp->value;
-			}
-
-			if (tagCondition && !isConstant)
-			{
-				if (!c)
-				{
-					c = tagCondition;
-				}
-				else
-				{
-					Ptr<ASTOpFixed> o = new ASTOpFixed;
-					o->op.type = OP_TYPE::BO_AND;
-					o->SetChild(o->op.args.BoolBinary.a, tagCondition);
-					o->SetChild(o->op.args.BoolBinary.b, c);
-					c = m_opPool.Add(o);
-				}
-			}
-			else if (constantValue == true)
-			{
-				// No need to add it to the AND
-			}
-			else //if (constantValue==false)
-			{
-				// Entire expression will be false
-				Ptr<ASTOpConstantBool> f = new ASTOpConstantBool(false);
-
-				// No need to evaluate anything else.
-				c = m_opPool.Add(f);
-				done = true;
-				break;
-			}
-		}
-
-		if (!done)
-		{
-			for (const auto& t : mod.negativeTags)
-			{
-				auto it = m_pFirstPass->m_tags.FindByPredicate([&](const FirstPassGenerator::FTag& e) { return e.tag == t; });
-				if (!it)
-				{
-					// This could happen if a tag is in a variation but noone defines it.
-					continue;
-				}
-
-				size_t tagIndex = it - &m_pFirstPass->m_tags[0];
-
-				set<size_t> empty;
-				Ptr<ASTOp> tagCondition = GenerateTagCondition(tagIndex, empty, empty, empty, empty);
-
-				// No condition is equal to a conditional with a true constant
-				if (!tagCondition)
-				{
-					ASTOpConstantBool* ConstOp = new ASTOpConstantBool();
-					ConstOp->value = true;
-					tagCondition = ConstOp;
-				}
-
-				// TODO: Optimise the tag condition here
-
-				// If the tag is a constant ...
-				bool isConstant = false;
-				bool constantValue = false;
-				if (tagCondition && tagCondition->GetOpType() == OP_TYPE::BO_CONSTANT)
-				{
-					const ASTOpConstantBool* constOp = static_cast<const ASTOpConstantBool*>(tagCondition.get());
-					isConstant = true;
-					constantValue = constOp->value;
-				}
-
-
-				if (!isConstant && tagCondition)
-				{
-					Ptr<ASTOpFixed> n = new ASTOpFixed;
-					n->op.type = OP_TYPE::BO_NOT;
-					n->SetChild(n->op.args.BoolNot.source, tagCondition);
-
-					if (!c)
-					{
-						c = m_opPool.Add(n);
-					}
-					else
-					{
-						Ptr<ASTOpFixed> o = new ASTOpFixed;
-						o->op.type = OP_TYPE::BO_AND;
-						o->SetChild(o->op.args.BoolBinary.a, n);
-						o->SetChild(o->op.args.BoolBinary.b, c);
-						c = m_opPool.Add(o);
-					}
-				}
-				else if (isConstant && constantValue == true)
-				{
-					// No expression here means always true which becomes always false
-					Ptr<ASTOpConstantBool> f = new ASTOpConstantBool(false);
-
-					// No need to evaluate anything else.
-					c = m_opPool.Add(f);
-					done = true;
-					break;
-				}
-			}
-		}
-
-		//PartialOptimise( c, m_pCompilerOptions->OptimisationOptions );
-
-		return c;
-	}
-
 
 	//---------------------------------------------------------------------------------------------
 	bool SecondPassGenerator::Generate(
@@ -630,17 +493,72 @@ namespace mu
 		// Create the conditions for every surface, modifier and individual tag.
 		m_tagConditionGenerationCache.clear();
 
-		for (int32 s = 0; s < m_pFirstPass->surfaces.Num(); ++s)
+		for (int32 SurfaceIndex = 0; SurfaceIndex < m_pFirstPass->surfaces.Num(); ++SurfaceIndex)
 		{
-			set<size_t> empty;
-			Ptr<ASTOp> c = GenerateSurfaceCondition(s, empty, empty, empty, empty);
-			m_pFirstPass->surfaces[s].surfaceCondition = c;
+			FirstPassGenerator::FSurface& Surface = m_pFirstPass->surfaces[SurfaceIndex];
+
+			{
+				set<size_t> Empty;
+
+				auto PositiveTags = [this](size_t SurfaceIndex) -> const TArray<FString>&
+				{
+					return m_pFirstPass->surfaces[SurfaceIndex].positiveTags;
+				};
+
+				auto NegativeTags = [this](size_t SurfaceIndex) -> const TArray<FString>&
+				{
+					return m_pFirstPass->surfaces[SurfaceIndex].negativeTags;
+				};
+				
+				Ptr<ASTOp> c = GenerateSurfaceOrModifierCodition(SurfaceIndex, PositiveTags, NegativeTags, Empty, Empty, Empty, Empty);
+				m_pFirstPass->surfaces[SurfaceIndex].surfaceCondition = c;
+			}
+			
+			for (int32 EditIndex = 0; EditIndex < Surface.edits.Num(); ++EditIndex)
+			{
+				FirstPassGenerator::FSurface::FEdit& Edit = Surface.edits[EditIndex];
+				
+				set<size_t> Empty;
+
+				auto PositiveTags = [&Surface](size_t SurfaceIndex) -> const TArray<FString>&
+				{
+					return Surface.edits[SurfaceIndex].PositiveTags;
+				};
+
+				auto NegativeTags = [&Surface](size_t SurfaceIndex) -> const TArray<FString>&
+				{
+					return Surface.edits[SurfaceIndex].NegativeTags;
+				};
+				
+				Ptr<ASTOp> c = GenerateSurfaceOrModifierCodition(EditIndex, PositiveTags, NegativeTags, Empty, Empty, Empty, Empty);
+
+				Ptr<ASTOpFixed> OpAnd = new ASTOpFixed;
+				OpAnd->op.type = OP_TYPE::BO_AND;
+				OpAnd->SetChild(OpAnd->op.args.BoolBinary.a, Edit.condition);
+				OpAnd->SetChild(OpAnd->op.args.BoolBinary.b, c);
+				c = m_opPool.Add(OpAnd);
+				
+				Edit.condition = OpAnd;
+			}
 		}
 
-		for (int32 m = 0; m < m_pFirstPass->modifiers.Num(); ++m)
+		for (int32 ModifierIndex = 0; ModifierIndex < m_pFirstPass->modifiers.Num(); ++ModifierIndex)
 		{
-			Ptr<ASTOp> c = GenerateModifierCondition(m);
-			m_pFirstPass->modifiers[m].surfaceCondition = c;
+			set<size_t> Empty;
+
+			auto PositiveTags = [this](size_t SurfaceIndex) -> TArray<FString>&
+			{
+				return m_pFirstPass->modifiers[SurfaceIndex].positiveTags;
+			};
+
+			auto NegativeTags = [this](size_t SurfaceIndex) -> TArray<FString>&
+			{
+				return m_pFirstPass->modifiers[SurfaceIndex].negativeTags;
+			};
+				
+			Ptr<ASTOp> c = GenerateSurfaceOrModifierCodition(ModifierIndex, PositiveTags, NegativeTags, Empty, Empty, Empty, Empty);
+			
+			m_pFirstPass->modifiers[ModifierIndex].surfaceCondition = c;
 		}
 
 		for (int32 s = 0; s < m_pFirstPass->m_tags.Num(); ++s)
diff --git a/Engine/Plugins/Experimental/Mutable/Source/MutableTools/Private/MuT/CodeGenerator_SecondPass.h b/Engine/Plugins/Experimental/Mutable/Source/MutableTools/Private/MuT/CodeGenerator_SecondPass.h
--- a/Engine/Plugins/Experimental/Mutable/Source/MutableTools/Private/MuT/CodeGenerator_SecondPass.h
+++ b/Engine/Plugins/Experimental/Mutable/Source/MutableTools/Private/MuT/CodeGenerator_SecondPass.h
@@ -78,15 +78,21 @@ namespace mu
                                          const set<size_t>& posTag,
                                          const set<size_t>& negTag );
 
-        //!
-        Ptr<ASTOp> GenerateSurfaceCondition( size_t surfIndex,
-                                             const set<size_t>& posSurf,
-                                             const set<size_t>& negSurf,
-                                             const set<size_t>& posTag,
-                                             const set<size_t>& negTag );
-
-        //!
-        Ptr<ASTOp> GenerateModifierCondition( size_t modIndex );
+        /** Generate Surface, Edit or Modifier condition.
+    	 * @param Index Surface, Edit or Modifier index.
+    	 * @param positiveTags function that given the Surface, Edit or Modifier index, returns its positive tags.
+    	 * @param negativeTags function that given the Surface, Edit or Modifier, returns its negative tags.
+    	 * @param posSurf already visited Surfaces, Edits, or Modifiers that participate positively in the condition.
+    	 * @param negSurf already visited Surfaces, Edits, or Modifiers that participate negatively in the condition.
+      	 * @param posSurf Tags that already belong to the condition (positively).
+		 * @param negSurf Tags that already belong to the condition (negatively). */
+        Ptr<ASTOp> GenerateSurfaceOrModifierCodition(size_t Index,
+									        TFunction<const TArray<FString>&(size_t)> positiveTags,
+											TFunction<const TArray<FString>&(size_t)> negativeTags,
+                                            const set<size_t>& posSurf,
+                                            const set<size_t>& negSurf,
+                                            const set<size_t>& posTag,
+                                            const set<size_t>& negTag);
     };
 
 }
