[Mutable] Reversed order of pins in Node Switch and Node Texture Layer.

#rnx

diff --git a/Engine/Plugins/Experimental/Mutable/Source/CustomizableObjectEditor/Private/MuCOE/Nodes/CustomizableObjectNodeSwitchBase.cpp b/Engine/Plugins/Experimental/Mutable/Source/CustomizableObjectEditor/Private/MuCOE/Nodes/CustomizableObjectNodeSwitchBase.cpp
--- a/Engine/Plugins/Experimental/Mutable/Source/CustomizableObjectEditor/Private/MuCOE/Nodes/CustomizableObjectNodeSwitchBase.cpp
+++ b/Engine/Plugins/Experimental/Mutable/Source/CustomizableObjectEditor/Private/MuCOE/Nodes/CustomizableObjectNodeSwitchBase.cpp
@@ -69,8 +69,14 @@ void UCustomizableObjectNodeSwitchBase::AllocateDefaultPins(UCustomizableObjectN
 	OutputPin->bDefaultValueIsIgnored = true;
 
 	OutputPinReference = FEdGraphPinReference(OutputPin);
+
+	UEdGraphPin* SwitchParameterPin = CustomCreatePin(EGPD_Input, Schema->PC_Enum, FName(TEXT("Switch Parameter")));
+	SwitchParameterPin->bDefaultValueIsIgnored = true;
+	SwitchParameterPin->SetOwningNode(this);
+
+	SwitchParameterPinReference = FEdGraphPinReference(SwitchParameterPin);
 	
-	for (int LayerIndex = ReloadingElementsNames.Num() - 1; LayerIndex >= 0; --LayerIndex)
+	for (int32 LayerIndex = 0; LayerIndex < ReloadingElementsNames.Num(); ++LayerIndex)
 	{
 		FString PinName = GetPinPrefix(LayerIndex);
 		UEdGraphPin* InputPin = CustomCreatePin(EGPD_Input, GetCategory(), FName(*PinName));
@@ -79,12 +85,6 @@ void UCustomizableObjectNodeSwitchBase::AllocateDefaultPins(UCustomizableObjectN
 		InputPin->bDefaultValueIsIgnored = true;
 		InputPin->SetOwningNode(this);
 	}
-
-	UEdGraphPin* SwitchParameterPin = CustomCreatePin(EGPD_Input, Schema->PC_Enum, FName(TEXT("Switch Parameter")));
-	SwitchParameterPin->bDefaultValueIsIgnored = true;
-	SwitchParameterPin->SetOwningNode(this);
-
-	SwitchParameterPinReference = FEdGraphPinReference(SwitchParameterPin);
 }
 
 
diff --git a/Engine/Plugins/Experimental/Mutable/Source/CustomizableObjectEditor/Private/MuCOE/Nodes/CustomizableObjectNodeTextureLayer.cpp b/Engine/Plugins/Experimental/Mutable/Source/CustomizableObjectEditor/Private/MuCOE/Nodes/CustomizableObjectNodeTextureLayer.cpp
--- a/Engine/Plugins/Experimental/Mutable/Source/CustomizableObjectEditor/Private/MuCOE/Nodes/CustomizableObjectNodeTextureLayer.cpp
+++ b/Engine/Plugins/Experimental/Mutable/Source/CustomizableObjectEditor/Private/MuCOE/Nodes/CustomizableObjectNodeTextureLayer.cpp
@@ -53,7 +53,10 @@ void UCustomizableObjectNodeTextureLayer::AllocateDefaultPins(UCustomizableObjec
 
 	OutputPinReference = FEdGraphPinReference(OutputPin);
 
-	for (int LayerIndex = Layers.Num() - 1; LayerIndex >=0 ; --LayerIndex )
+	UEdGraphPin* BasePin = CustomCreatePin(EGPD_Input, Schema->PC_Image, FName(TEXT("Base")));
+	BasePin->bDefaultValueIsIgnored = true;
+	
+	for (int32 LayerIndex = 0; LayerIndex < Layers.Num() ; ++LayerIndex )
 	{
 		FString PinName = FString::Printf( TEXT("Layer %d "), LayerIndex );
 		UEdGraphPin* LayerPin = CustomCreatePin(EGPD_Input, Schema->PC_Image, *PinName);
@@ -63,9 +66,6 @@ void UCustomizableObjectNodeTextureLayer::AllocateDefaultPins(UCustomizableObjec
 		UEdGraphPin* MaskPin = CustomCreatePin(EGPD_Input, Schema->PC_Image, *PinName);
 		MaskPin->bDefaultValueIsIgnored = true;
 	}
-
-	UEdGraphPin* BasePin = CustomCreatePin(EGPD_Input, Schema->PC_Image, FName(TEXT("Base")));
-	BasePin->bDefaultValueIsIgnored = true;
 }
 
 
