[Mutable] Clear CachedPlatformData on PostSaveRoot to free some resources at cook time.

#rnx

diff --git a/Engine/Plugins/Experimental/Mutable/Source/CustomizableObject/Classes/MuCO/CustomizableObject.h b/Engine/Plugins/Experimental/Mutable/Source/CustomizableObject/Classes/MuCO/CustomizableObject.h
--- a/Engine/Plugins/Experimental/Mutable/Source/CustomizableObject/Classes/MuCO/CustomizableObject.h
+++ b/Engine/Plugins/Experimental/Mutable/Source/CustomizableObject/Classes/MuCO/CustomizableObject.h
@@ -784,6 +784,7 @@ public:
 	FString GetDesc() override;
 	bool IsEditorOnly() const override;
 	void PreSave(FObjectPreSaveContext ObjectSaveContext) override;
+	void PostSaveRoot(FObjectPostSaveRootContext ObjectSaveContext) override;
 	void PostRename(UObject* OldOuter, const FName OldName) override;	
 	void BeginCacheForCookedPlatformData(const ITargetPlatform* TargetPlatform) override;
 	bool IsCachedCookedPlatformDataLoaded(const ITargetPlatform* TargetPlatform) override;
diff --git a/Engine/Plugins/Experimental/Mutable/Source/CustomizableObject/Private/MuCO/CustomizableObject.cpp b/Engine/Plugins/Experimental/Mutable/Source/CustomizableObject/Private/MuCO/CustomizableObject.cpp
--- a/Engine/Plugins/Experimental/Mutable/Source/CustomizableObject/Private/MuCO/CustomizableObject.cpp
+++ b/Engine/Plugins/Experimental/Mutable/Source/CustomizableObject/Private/MuCO/CustomizableObject.cpp
@@ -144,7 +144,6 @@ void UCustomizableObject::PreSave(FObjectPreSaveContext ObjectSaveContext)
 		}
 	}
 
-#if WITH_EDITOR
 	if (ObjectSaveContext.IsCooking() && !bIsChildObject)
 	{
 		const ITargetPlatform* TargetPlatform = ObjectSaveContext.GetTargetPlatform();
@@ -170,9 +169,24 @@ void UCustomizableObject::PreSave(FObjectPreSaveContext ObjectSaveContext)
 			GetPrivate()->ClearCompiledData(true);
 		}
 	}
-#endif
 }
 
+
+void UCustomizableObject::PostSaveRoot(FObjectPostSaveRootContext ObjectSaveContext)
+{
+	MUTABLE_CPUPROFILER_SCOPE(UCustomizableObject::PostSaveRoot);
+
+	Super::PostSaveRoot(ObjectSaveContext);
+
+	if (ObjectSaveContext.IsCooking())
+	{
+		// Free cached data after saving;
+		const ITargetPlatform* TargetPlatform = ObjectSaveContext.GetTargetPlatform();
+		GetPrivate()->CachedPlatformsData.Remove(TargetPlatform->PlatformName());
+	}
+}
+
+
 bool UCustomizableObjectPrivate::TryUpdateIsChildObject()
 {
 	FCustomizableObjectCompilerBase* Compiler = UCustomizableObjectSystem::GetInstance()->GetNewCompiler();
