[Mutable] Fix desynchronization between RenderSections and UserSectionsData.

#rnx

diff --git a/Engine/Plugins/Experimental/Mutable/Source/CustomizableObject/Private/MuCO/CustomizableObjectInstance.cpp b/Engine/Plugins/Experimental/Mutable/Source/CustomizableObject/Private/MuCO/CustomizableObjectInstance.cpp
--- a/Engine/Plugins/Experimental/Mutable/Source/CustomizableObject/Private/MuCO/CustomizableObjectInstance.cpp
+++ b/Engine/Plugins/Experimental/Mutable/Source/CustomizableObject/Private/MuCO/CustomizableObjectInstance.cpp
@@ -6723,14 +6723,13 @@ void UCustomizableInstancePrivate::RegenerateImportedModels()
 				ImportedSection.OriginalDataSectionIndex = OriginalIndex++;
 
 				FSkelMeshSourceSectionUserData& SectionUserData = ImportedModel->LODModels[LODIndex].UserSectionsData.FindOrAdd(ImportedSection.OriginalDataSectionIndex);
+				SectionUserData.bCastShadow = RenderSection.bCastShadow;
+				SectionUserData.bDisabled = RenderSection.bDisabled;
 
 				SectionUserData.CorrespondClothAssetIndex = RenderSection.CorrespondClothAssetIndex;
 				SectionUserData.ClothingData.AssetGuid = RenderSection.ClothingData.AssetGuid;
 				SectionUserData.ClothingData.AssetLodIndex = RenderSection.ClothingData.AssetLodIndex;
 				
-				//TODO(Max): MTBL-1779
-				//SectionUserData.bCastShadow = RenderSection.bCastShadow;
-
 				ImportedModel->LODModels[LODIndex].SyncronizeUserSectionsDataArray();
 
 				// DDC keys
