[Mutable] Copy and paste parameters now uses UE UStruct serialization.

Serialization to FString.

- Unified all printing of parameters.

#rnx

diff --git a/Engine/Plugins/Experimental/Mutable/Source/CustomizableObject/Internal/MuCO/CustomizableObjectInstancePrivate.h b/Engine/Plugins/Experimental/Mutable/Source/CustomizableObject/Internal/MuCO/CustomizableObjectInstancePrivate.h
--- a/Engine/Plugins/Experimental/Mutable/Source/CustomizableObject/Internal/MuCO/CustomizableObjectInstancePrivate.h
+++ b/Engine/Plugins/Experimental/Mutable/Source/CustomizableObject/Internal/MuCO/CustomizableObjectInstancePrivate.h
@@ -375,9 +375,6 @@ public:
 	bool bShowOnlyRuntimeParameters = true;
 	bool bShowOnlyRelevantParameters = true;
 	bool bShowUISections = false;
-
-	/** Flag to know if a property of this instance changed in the editor */
-	bool bEditorPropertyChanged = false;
 	
 	// TEMP VARIABLE to check the Min desired LODs for this instance
 	TWeakObjectPtr<UCustomizableObjectInstanceUsage> NearestToActor;
diff --git a/Engine/Plugins/Experimental/Mutable/Source/CustomizableObject/Private/MuCO/CustomizableObjectInstance.cpp b/Engine/Plugins/Experimental/Mutable/Source/CustomizableObject/Private/MuCO/CustomizableObjectInstance.cpp
--- a/Engine/Plugins/Experimental/Mutable/Source/CustomizableObject/Private/MuCO/CustomizableObjectInstance.cpp
+++ b/Engine/Plugins/Experimental/Mutable/Source/CustomizableObject/Private/MuCO/CustomizableObjectInstance.cpp
@@ -356,8 +356,6 @@ void UCustomizableObjectInstance::PostEditChangeProperty(FPropertyChangedEvent&
 {
 	Super::PostEditChangeProperty(PropertyChangedEvent);
 
-	GetPrivate()->bEditorPropertyChanged = true;
-
 	const FName PropertyName = PropertyChangedEvent.MemberProperty ? PropertyChangedEvent.MemberProperty->GetFName() : NAME_None;
 
 	if (PropertyName == GET_MEMBER_NAME_CHECKED(UCustomizableObjectInstance, TextureParameterDeclarations))
diff --git a/Engine/Plugins/Experimental/Mutable/Source/CustomizableObject/Private/MuCO/CustomizableObjectInstanceDescriptor.cpp b/Engine/Plugins/Experimental/Mutable/Source/CustomizableObject/Private/MuCO/CustomizableObjectInstanceDescriptor.cpp
--- a/Engine/Plugins/Experimental/Mutable/Source/CustomizableObject/Private/MuCO/CustomizableObjectInstanceDescriptor.cpp
+++ b/Engine/Plugins/Experimental/Mutable/Source/CustomizableObject/Private/MuCO/CustomizableObjectInstanceDescriptor.cpp
@@ -659,84 +659,12 @@ mu::ParametersPtr FCustomizableObjectInstanceDescriptor::GetParameters() const
 
 FString FCustomizableObjectInstanceDescriptor::ToString() const
 {
-	TStringBuilder<3000> Builder;
-	
-	Builder.Appendf(TEXT("CustomizableObject=%s\n"), *CustomizableObject.GetFullName());
-	Builder.Appendf(TEXT("State=%i\n"), State);
-	Builder.Appendf(TEXT("BuildParameterRelevancy=%i\n"), GetBuildParameterRelevancy());
-	Builder.Appendf(TEXT("MinLOD=%i\n"), MinLOD);
-	Builder.Appendf(TEXT("MaxLOD=%i\n"), MaxLOD);
-
-	for (int32 Index = 0; Index < RequestedLODLevels.Num(); ++Index)
-	{
-		Builder.Appendf(TEXT("RequestredLODLevels [%i]=%i\n"), Index, RequestedLODLevels[Index]);
-	}
-
-	Builder.Appendf(TEXT("BOOL PARAMETERS:\n"));
-	for (const FCustomizableObjectBoolParameterValue& Parameter : BoolParameters)
-	{
-		Builder.Appendf(TEXT("%s=%i\n"), *Parameter.ParameterName, Parameter.ParameterValue);
-	}
-
-	Builder.Appendf(TEXT("INT PARAMETERS:\n"));
-	for (const FCustomizableObjectIntParameterValue& Parameter : IntParameters)
-	{
-		Builder.Appendf(TEXT("%s=\"%s\"\n"), *Parameter.ParameterName, *Parameter.ParameterValueName);
+	const UScriptStruct* Struct = StaticStruct();
+	FString ExportedText;
 
-		for (int32 Index = 0; Index < Parameter.ParameterRangeValueNames.Num(); ++Index)
-		{
-			Builder.Appendf(TEXT("%s [%i]=\"%s\"\n"), *Parameter.ParameterName, Index, *Parameter.ParameterRangeValueNames[Index]);			
-		}
-	}
-
-	Builder.Appendf(TEXT("FLOAT PARAMETETS:\n"));
-	for (const FCustomizableObjectFloatParameterValue& Parameter : FloatParameters)
-	{
-		Builder.Appendf(TEXT("%s=%f\n"), *Parameter.ParameterName, Parameter.ParameterValue);
-
-		for (int32 Index = 0; Index < Parameter.ParameterRangeValues.Num(); ++Index)
-		{
-			Builder.Appendf(TEXT("%s [%i]=%f\n"), *Parameter.ParameterName, Index, Parameter.ParameterRangeValues[Index]);			
-		}
-	}
-
-	Builder.Appendf(TEXT("TEXTURE PARAMETERS:\n"));
-	for (const FCustomizableObjectTextureParameterValue& Parameter : TextureParameters)
-	{
-		Builder.Appendf(TEXT("%s=%s\n"), *Parameter.ParameterName, *Parameter.ParameterValue.ToString());
-
-		for (int32 Index = 0; Index < Parameter.ParameterRangeValues.Num(); ++Index)
-		{
-			Builder.Appendf(TEXT("%s [%i]=%s\n"), *Parameter.ParameterName, Index, *Parameter.ParameterRangeValues[Index].ToString());			
-		}
-	}
-
-	Builder.Appendf(TEXT("VECTOR PARAMETERS:\n"));
-	for (const FCustomizableObjectVectorParameterValue& Parameter : VectorParameters)
-	{
-		Builder.Appendf(TEXT("%s=%s\n"), *Parameter.ParameterName, *Parameter.ParameterValue.ToString());			
-	}
-
-	Builder.Appendf(TEXT("PROJECTOR PARAMETERS:\n"));
-	for (const FCustomizableObjectProjectorParameterValue& Parameter : ProjectorParameters)
-	{
-		auto StructToString = [](const UScriptStruct& ScriptStruct, const void* Data) {
-			FString String;
-			ScriptStruct.ExportText(String, &Data, &Data, nullptr, PPF_ExternalEditor | PPF_IncludeTransient, nullptr);
-			return String;
-		};
-
-		Builder.Appendf(TEXT("%s=%s\n"), *Parameter.ParameterName, *StructToString(*Parameter.Value.StaticStruct(), &Parameter.Value));
-
-		for (int32 Index = 0; Index < Parameter.RangeValues.Num(); ++Index)
-		{
-			const FCustomizableObjectProjector& RangeValue = Parameter.RangeValues[Index];
-			
-			Builder.Appendf(TEXT("%s [%i]=%s\n"), *Parameter.ParameterName, Index, *StructToString(*RangeValue.StaticStruct(), &RangeValue));			
-		}
-	}
+	Struct->ExportText(ExportedText, this, nullptr, nullptr, (PPF_ExportsNotFullyQualified | PPF_Copy | PPF_Delimited | PPF_IncludeTransient), nullptr);
 
-	return Builder.ToString();
+	return ExportedText;
 }
 
 
diff --git a/Engine/Plugins/Experimental/Mutable/Source/CustomizableObject/Private/MuCO/CustomizableObjectSystem.cpp b/Engine/Plugins/Experimental/Mutable/Source/CustomizableObject/Private/MuCO/CustomizableObjectSystem.cpp
--- a/Engine/Plugins/Experimental/Mutable/Source/CustomizableObject/Private/MuCO/CustomizableObjectSystem.cpp
+++ b/Engine/Plugins/Experimental/Mutable/Source/CustomizableObject/Private/MuCO/CustomizableObjectSystem.cpp
@@ -982,8 +982,6 @@ void UpdateSkeletalMesh(const TSharedRef<FUpdateContextPrivate>& Context)
 			}
 		}
 	}
-
-	CustomizableObjectInstance->GetPrivate()->bEditorPropertyChanged = false;
 }
 
 
diff --git a/Engine/Plugins/Experimental/Mutable/Source/CustomizableObject/Private/MuCO/LogBenchmarkUtil.cpp b/Engine/Plugins/Experimental/Mutable/Source/CustomizableObject/Private/MuCO/LogBenchmarkUtil.cpp
--- a/Engine/Plugins/Experimental/Mutable/Source/CustomizableObject/Private/MuCO/LogBenchmarkUtil.cpp
+++ b/Engine/Plugins/Experimental/Mutable/Source/CustomizableObject/Private/MuCO/LogBenchmarkUtil.cpp
@@ -55,7 +55,11 @@ TSharedPtr<FArchive> CreateFile()
 	TSharedPtr<FArchive> Archive = MakeShareable(IFileManager::Get().CreateFileWriter(*Filename, FILEWRITE_AllowRead | FILEWRITE_NoFail));
 	check(Archive);
 
-	const FString HeaderRow = TEXT("ID_CO,ID_COI,ID_UpdateType,ID_UpdateResult,Time_Queue,Time_Update,Time_TaskGetMesh,Time_TaskLockCache,Time_TaskGetImages,Time_TaskConvertResources,Time_TaskCallbacks,Memory_Update,Memory_Update_Real,Time_TaskUpdateImage,Memory_TaskUpdateImage,Memory_TaskUpdateImage_Real");
+	const FString HeaderRow = TEXT(
+		"ID_CO,ID_COI,ID_UpdateType,ID_Descriptor,ID_UpdateResult,Time_Queue,Time_Update,Time_TaskGetMesh,Time_TaskLockCache,"
+		"Time_TaskGetImages,Time_TaskConvertResources,Time_TaskCallbacks,Memory_Update,Memory_Update_Real,"
+		"Time_TaskUpdateImage,Memory_TaskUpdateImage,Memory_TaskUpdateImage_Real");
+
 	LogBenchmarkUtil::Write(*Archive, HeaderRow);
 
 	return Archive;
@@ -309,6 +313,7 @@ void FLogBenchmarkUtil::FinishUpdateMesh(const TSharedRef<FUpdateContextPrivate>
 	const FString ID_CO = Context->Instance->GetCustomizableObject()->GetPathName();
 	const FString ID_COI = Instance->GetPathName();
 	const FString ID_UpdateType = TEXT("Mesh");
+	const FString ID_Descriptor = Context->GetCapturedDescriptor().ToString();
 	const FString ID_UpdateResult = StaticEnum<EUpdateResult>()->GetValueAsString(Context->UpdateResult);
 	const double Time_Queue = Context->QueueTime * 1000;
 	const double Time_Update = Context->UpdateTime * 1000;
@@ -321,7 +326,7 @@ void FLogBenchmarkUtil::FinishUpdateMesh(const TSharedRef<FUpdateContextPrivate>
 	const double Memory_UpdateEndPeakMB = (Context->UpdateEndPeakBytes / 1024.0) / 1024.0;
 	const double Memory_UpdateEndRealPeakMB = (Context->UpdateEndRealPeakBytes / 1024.0) / 1024.0;
 	
-	const FString UpdateString = FString::Printf(TEXT("%s,%s,%s,%s,%f,%f,%f,%f,%f,%f,%f,%f,%f"), *ID_CO, *ID_COI, *ID_UpdateType, *ID_UpdateResult, Time_Queue, Time_Update, Time_TaskGetMesh, Time_TaskLockCache, Time_TaskGetImages, Time_TaskConvertResources, Time_TaskCallbacks, Memory_UpdateEndPeakMB, Memory_UpdateEndRealPeakMB);
+	const FString UpdateString = FString::Printf(TEXT("%s,%s,%s,%s,%s,%f,%f,%f,%f,%f,%f,%f,%f,%f"), *ID_CO, *ID_COI, *ID_UpdateType, *ID_Descriptor, *ID_UpdateResult, Time_Queue, Time_Update, Time_TaskGetMesh, Time_TaskLockCache, Time_TaskGetImages, Time_TaskConvertResources, Time_TaskCallbacks, Memory_UpdateEndPeakMB, Memory_UpdateEndRealPeakMB);
 	LogBenchmarkUtil::Write(*Archive, UpdateString);
 	Archive->Flush();
 }
@@ -344,7 +349,7 @@ void FLogBenchmarkUtil::FinishUpdateImage(const FString& CustomizableObjectPathN
 	const double Memory_TaskUpdateImagePeakMB = (TaskUpdateImageMemoryPeak / 1024.0) / 1024.0;
 	const double Memory_TaskUpdateImageRealPeakMB = (TaskUpdateImageRealMemoryPeak / 1024.0) / 1024.0;
 
-	const FString UpdateString = FString::Printf(TEXT("%s,%s,%s,,,,,,,,,,,%f,%f,%f"), *ID_CO, *ID_COI, *ID_UpdateType, Time_TaskUpdateImage, Memory_TaskUpdateImagePeakMB,Memory_TaskUpdateImageRealPeakMB);
+	const FString UpdateString = FString::Printf(TEXT("%s,%s,%s,,,,,,,,,,,,%f,%f,%f"), *ID_CO, *ID_COI, *ID_UpdateType, Time_TaskUpdateImage, Memory_TaskUpdateImagePeakMB,Memory_TaskUpdateImageRealPeakMB);
 	LogBenchmarkUtil::Write(*Archive, UpdateString);
 	Archive->Flush();	
 }
diff --git a/Engine/Plugins/Experimental/Mutable/Source/CustomizableObject/Private/MuCO/LogInformationUtil.cpp b/Engine/Plugins/Experimental/Mutable/Source/CustomizableObject/Private/MuCO/LogInformationUtil.cpp
--- a/Engine/Plugins/Experimental/Mutable/Source/CustomizableObject/Private/MuCO/LogInformationUtil.cpp
+++ b/Engine/Plugins/Experimental/Mutable/Source/CustomizableObject/Private/MuCO/LogInformationUtil.cpp
@@ -124,13 +124,11 @@ void LogInformationUtil::LogShowInstanceDataFull(const UCustomizableObjectInstan
 	FString LogData = "\n\n";
 
 	LogData += FString::Printf(TEXT("CustomizableObjectInstance %s, "), *CustomizableObjectInstance->GetName());
-	
-	PrintBoolParameters(CustomizableObjectInstance->GetBoolParameters(), LogData);
-	PrintIntParameters(CustomizableObjectInstance->GetIntParameters(), LogData);
-	PrintFloatParameters(CustomizableObjectInstance->GetFloatParameters(), LogData);
-	PrintVectorParameters(CustomizableObjectInstance->GetVectorParameters(), LogData);
-	PrintProjectorParameters(CustomizableObjectInstance->GetProjectorParameters(), LogData);
 
+	LogData += "Descriptor:\n";
+	LogData += CustomizableObjectInstance->GetPrivate()->CommittedDescriptor.ToString();
+	LogData += "\n";
+	
 	if (ShowMaterialInfo)
 	{
 		PrintGeneratedMaterial(CustomizableObjectInstance->GetPrivate()->GeneratedMaterials, LogData);
@@ -203,134 +201,6 @@ void LogInformationUtil::LogShowInstanceDataFull(const UCustomizableObjectInstan
 }
 
 
-void LogInformationUtil::PrintBoolParameters(const TArray<struct FCustomizableObjectBoolParameterValue>& BoolParameters, FString& Log)
-{
-	if (BoolParameters.Num() == 0)
-	{
-		return;
-	}
-
-	Log += "Boolean parameters:\n";
-
-	const int Max = BoolParameters.Num();
-	int i;
-
-	for (i = 0; i < Max; ++i)
-	{
-		Log += FString::Printf(TEXT("        %s="), *BoolParameters[i].ParameterName);
-		Log += FString::Printf(TEXT("%d\n"), BoolParameters[i].ParameterValue);
-	}
-
-	Log += "\n";
-}
-
-
-void LogInformationUtil::PrintIntParameters(const TArray<struct FCustomizableObjectIntParameterValue>& IntParameters, FString& Log)
-{
-	if (IntParameters.Num() == 0)
-	{
-		return;
-	}
-
-	FString MessageChunk;
-
-	MessageChunk += "\tInt parameters:\n";
-
-	const int Max = IntParameters.Num();
-	int i;
-
-	for (i = 0; i < Max; ++i)
-	{
-		MessageChunk += FString::Printf(TEXT("        %s="), *IntParameters[i].ParameterName);
-		MessageChunk += FString::Printf(TEXT("%s\n"), *IntParameters[i].ParameterValueName);
-		FillToLength(MessageChunk,25);
-	}
-
-	MessageChunk += "\n";
-
-	Log += MessageChunk;
-}
-
-
-void LogInformationUtil::PrintFloatParameters(const TArray<struct FCustomizableObjectFloatParameterValue>& FloatParameters, FString& Log)
-{
-	if (FloatParameters.Num() == 0)
-	{
-		return;
-	}
-
-	Log += "Float parameters:\n";
-
-	const int Max = FloatParameters.Num();
-	int i;
-
-	for (i = 0; i < Max; ++i)
-	{
-		Log += FString::Printf(TEXT("        %s="), *FloatParameters[i].ParameterName);
-		Log += FString::Printf(TEXT("%.2f\n"), FloatParameters[i].ParameterValue);
-	}
-
-	Log += "\n";
-}
-
-
-void LogInformationUtil::PrintVectorParameters(const TArray<struct FCustomizableObjectVectorParameterValue>& VectorParameters, FString& Log)
-{
-	if (VectorParameters.Num() == 0)
-	{
-		return;
-	}
-
-	Log += "Vector parameters:\n";
-
-	const int Max = VectorParameters.Num();
-	int i;
-
-	for (i = 0; i < Max; ++i)
-	{
-		Log += FString::Printf(TEXT("        %s="), *VectorParameters[i].ParameterName);
-		Log += FString::Printf(TEXT("(%.2f,"), VectorParameters[i].ParameterValue.R);
-		Log += FString::Printf(TEXT("%.2f,"), VectorParameters[i].ParameterValue.G);
-		Log += FString::Printf(TEXT("%.2f,"), VectorParameters[i].ParameterValue.B);
-		Log += FString::Printf(TEXT("%.2f)\n"), VectorParameters[i].ParameterValue.A);
-	}
-
-	Log += "\n";
-}
-
-
-void LogInformationUtil::PrintProjectorParameters(const TArray<struct FCustomizableObjectProjectorParameterValue>& ProjectorParameters, FString& Log)
-{
-	if (ProjectorParameters.Num() == 0)
-	{
-		return;
-	}
-
-	Log += "Projector parameters:\n";
-
-	const int Max = ProjectorParameters.Num();
-	int i;
-
-	for (i = 0; i < Max; ++i)
-	{
-		Log += FString::Printf(TEXT("        %s has "), *ProjectorParameters[i].ParameterName);
-		Log += FString::Printf(TEXT(" position=(%.2f,"), ProjectorParameters[i].Value.Position.X);
-		Log += FString::Printf(TEXT("%.2f,"), ProjectorParameters[i].Value.Position.Y);
-		Log += FString::Printf(TEXT("%.2f)"), ProjectorParameters[i].Value.Position.Z);
-		Log += FString::Printf(TEXT(" direction=(%.2f,"), ProjectorParameters[i].Value.Direction.X);
-		Log += FString::Printf(TEXT("%.2f,"), ProjectorParameters[i].Value.Direction.Y);
-		Log += FString::Printf(TEXT("%.2f)"), ProjectorParameters[i].Value.Direction.Z);
-		Log += FString::Printf(TEXT(" up=(%.2f,"), ProjectorParameters[i].Value.Up.X);
-		Log += FString::Printf(TEXT("%.2f,"), ProjectorParameters[i].Value.Up.Y);
-		Log += FString::Printf(TEXT("%.2f)"), ProjectorParameters[i].Value.Up.Z);
-		Log += FString::Printf(TEXT(" scale=(%.2f,"), ProjectorParameters[i].Value.Scale.X);
-		Log += FString::Printf(TEXT("%.2f)\n"), ProjectorParameters[i].Value.Scale.Y);
-	}
-
-	Log += "\n";
-}
-
-
 void LogInformationUtil::PrintTextureData(const UTexture2D* Texture, FString& Log, bool UseTabulation)
 {
 	if (!Texture)
diff --git a/Engine/Plugins/Experimental/Mutable/Source/CustomizableObject/Private/MuCO/LogInformationUtil.h b/Engine/Plugins/Experimental/Mutable/Source/CustomizableObject/Private/MuCO/LogInformationUtil.h
--- a/Engine/Plugins/Experimental/Mutable/Source/CustomizableObject/Private/MuCO/LogInformationUtil.h
+++ b/Engine/Plugins/Experimental/Mutable/Source/CustomizableObject/Private/MuCO/LogInformationUtil.h
@@ -24,26 +24,6 @@ public:
 	// Show data about all UCustomizableObjectInstance existing elements
 	static void LogShowInstanceData(const UCustomizableObjectInstance* CustomizableObjectInstance, FString& LogData);
 
-	// Utility method just to log information, adds to the given FString containing the
-	// log message the log information of BoolParameters
-	static void PrintBoolParameters(const TArray<struct FCustomizableObjectBoolParameterValue>& BoolParameters, FString& Log);
-
-	// Utility method just to log information, adds to the given FString containing the
-	// log message the log information of IntParameters
-	static void PrintIntParameters(const TArray<struct FCustomizableObjectIntParameterValue>& IntParameters, FString& Log);
-
-	// Utility method just to log information, adds to the given FString containing the
-	// log message the log information of FloatParameters
-	static void PrintFloatParameters(const TArray<struct FCustomizableObjectFloatParameterValue>& FloatParameters, FString& Log);
-
-	// Utility method just to log information, adds to the given FString containing the
-	// log message the log information of VectorParameters
-	static void PrintVectorParameters(const TArray<struct FCustomizableObjectVectorParameterValue>& VectorParameters, FString& Log);
-
-	// Utility method just to log information, adds to the given FString containing the
-	// log message the log information of ProjectorParameters
-	static void PrintProjectorParameters(const TArray<struct FCustomizableObjectProjectorParameterValue>& ProjectorParameters, FString& Log);
-
 	// Utility method just to log information, adds to the given FString containing the
 	// log message the log information of GeneratedMaterials
 	static void PrintGeneratedMaterial(const TArray<struct FGeneratedMaterial> GeneratedMaterials, FString& Log);
diff --git a/Engine/Plugins/Experimental/Mutable/Source/CustomizableObjectEditor/Private/MuCOE/CustomizableInstanceDetails.cpp b/Engine/Plugins/Experimental/Mutable/Source/CustomizableObjectEditor/Private/MuCOE/CustomizableInstanceDetails.cpp
--- a/Engine/Plugins/Experimental/Mutable/Source/CustomizableObjectEditor/Private/MuCOE/CustomizableInstanceDetails.cpp
+++ b/Engine/Plugins/Experimental/Mutable/Source/CustomizableObjectEditor/Private/MuCOE/CustomizableInstanceDetails.cpp
@@ -1849,11 +1849,8 @@ void FCustomizableInstanceDetails::OnBoolParameterChanged(ECheckBoxState InCheck
 
 FReply FCustomizableInstanceDetails::OnCopyAllParameters()
 {
-	FBufferArchive ParametersBuffer;
-	CustomInstance->SaveDescriptor(ParametersBuffer, true);
-
-	FString StringifyedData = FString::FromHexBlob(ParametersBuffer.GetData(), ParametersBuffer.Num());
-	FPlatformApplicationMisc::ClipboardCopy(StringifyedData.GetCharArray().GetData());
+	const FString ExportedText = CustomInstance->GetPrivate()->GetDescriptor().ToString();
+	FPlatformApplicationMisc::ClipboardCopy(*ExportedText);
 
 	return FReply::Handled();
 }
@@ -1861,27 +1858,31 @@ FReply FCustomizableInstanceDetails::OnCopyAllParameters()
 
 FReply FCustomizableInstanceDetails::OnPasteAllParameters()
 {
-	FString StringifiedData;
-	FPlatformApplicationMisc::ClipboardPaste(StringifiedData);
+	CustomInstance->Modify();
 
-	if (!StringifiedData.IsEmpty())
+	FString ClipText;
+	FPlatformApplicationMisc::ClipboardPaste(ClipText);
+	
+	FCustomizableObjectInstanceDescriptor& Descriptor = CustomInstance->GetPrivate()->GetDescriptor();
+	const UScriptStruct* Struct = Descriptor.StaticStruct();
+
+	const int32 MinLOD = Descriptor.GetMinLod();
+	const int32 MaxLOD = Descriptor.GetMaxLod();
+	const TArray<uint16> RequestedLODLevels = Descriptor.GetRequestedLODLevels();
+	
+	if (Struct->ImportText(*ClipText, &Descriptor, nullptr, 0, GLog, GetPathNameSafe(Struct)))
 	{
-		TArray<uint8> DescriptorData;
-		// Each char represents a hex number i.e. half byte.
-		DescriptorData.SetNum(StringifiedData.Len() * 2);
-		FString::ToHexBlob(StringifiedData, DescriptorData.GetData(), DescriptorData.Num());
-		FMemoryReader FromBinary = FMemoryReader(DescriptorData, true); //true, free data after done;
-		FromBinary.Seek(0);
-
-		CustomInstance->LoadDescriptor(FromBinary);
-		CustomInstance->GetPrivate()->SetSelectedParameterProfileDirty();
-		CustomInstance->UpdateSkeletalMeshAsync(true, true);
-		CustomInstance->PostEditChange();
-
-		// Non-continuous change: collect garbage.
-		GEngine->ForceGarbageCollection();
+		// Keep current LOD
+		Descriptor.SetMinLod(MinLOD);
+		Descriptor.SetMaxLod(MaxLOD);
+		Descriptor.SetRequestedLODLevels(RequestedLODLevels);
+		
+		CustomInstance->UpdateSkeletalMeshAsync(true, true);		
 	}
 
+	CustomInstance->GetPrivate()->SetSelectedParameterProfileDirty();
+
+	
 	return FReply::Handled();
 }
 
