[Mutable] Fix UCustomizableObjectSystem::LockObject deadlock.

The deadlock occurred when compiling a CO at the same time as generating another CO with optimization=None.

- Removed unnecessary delta times.

#rnx

diff --git a/Engine/Plugins/Experimental/Mutable/Source/CustomizableObject/Classes/MuCO/CustomizableObjectSystem.h b/Engine/Plugins/Experimental/Mutable/Source/CustomizableObject/Classes/MuCO/CustomizableObjectSystem.h
--- a/Engine/Plugins/Experimental/Mutable/Source/CustomizableObject/Classes/MuCO/CustomizableObjectSystem.h
+++ b/Engine/Plugins/Experimental/Mutable/Source/CustomizableObject/Classes/MuCO/CustomizableObjectSystem.h
@@ -342,7 +342,7 @@ private:
 	bool Tick(float DeltaTime);
 
 	/** Returns the number of remaining operations. */
-	int32 TickInternal(float DeltaTime);
+	int32 TickInternal();
 
 	// If there is an on-going operation, advance it.
 	void AdvanceCurrentOperation();
diff --git a/Engine/Plugins/Experimental/Mutable/Source/CustomizableObject/Private/MuCO/CustomizableObjectSystem.cpp b/Engine/Plugins/Experimental/Mutable/Source/CustomizableObject/Private/MuCO/CustomizableObjectSystem.cpp
--- a/Engine/Plugins/Experimental/Mutable/Source/CustomizableObject/Private/MuCO/CustomizableObjectSystem.cpp
+++ b/Engine/Plugins/Experimental/Mutable/Source/CustomizableObject/Private/MuCO/CustomizableObjectSystem.cpp
@@ -1404,7 +1404,19 @@ bool UCustomizableObjectSystem::LockObject(const class UCustomizableObject* InOb
 				{
 					Streamer->CancelStreamingForObject(InObject);
 				});
-			Task.Wait();
+
+			
+			if (GetMutableDefault<UCustomizableObjectSettings>()->bEnableStreamingManager)
+			{
+				while (!Task.IsCompleted())				
+				{
+					GetPrivate()->UpdateResourceStreaming(0, false);
+				}
+			}
+			else
+			{
+				Task.Wait();
+			}
 		}
 
 		Private->MutablePendingInstanceWork.RemoveUpdatesForObject(InObject);
@@ -2926,12 +2938,12 @@ void UCustomizableObjectSystem::AdvanceCurrentOperation()
 
 bool UCustomizableObjectSystem::Tick(float DeltaTime)
 {
-	TickInternal(DeltaTime);
+	TickInternal();
 	return true;
 }
 
 
-int32 UCustomizableObjectSystem::TickInternal(float DeltaTime)
+int32 UCustomizableObjectSystem::TickInternal()
 {
 	MUTABLE_CPUPROFILER_SCOPE(UCustomizableObjectSystem::TickInternal)
 	
@@ -3146,6 +3158,11 @@ int32 UCustomizableObjectSystem::TickInternal(float DeltaTime)
 
 #if WITH_EDITOR
 	TickRecompileCustomizableObjects();
+
+	if (GetPrivate()->ImageProvider)
+	{
+		GetPrivate()->ImageProvider->Tick();		
+	}
 #endif
 
 #if WITH_EDITORONLY_DATA
@@ -3905,15 +3922,13 @@ bool UCustomizableObjectSystem::IsMutableAnimInfoDebuggingEnabled() const
 
 void UCustomizableObjectSystemPrivate::UpdateResourceStreaming(float DeltaTime, bool bProcessEverything)
 {
-	GetPublic()->TickInternal(DeltaTime);
+	GetPublic()->TickInternal();
 }
 
 
 int32 UCustomizableObjectSystemPrivate::BlockTillAllRequestsFinished(float TimeLimit, bool bLogResults)
 {
 	const double BlockEndTime = FPlatformTime::Seconds() + TimeLimit;
-	double StartTime;
-	double DeltaTime = 0.0;
 
 	int32 RemainingWork = TNumericLimits<int32>::Max();
 	
@@ -3921,24 +3936,19 @@ int32 UCustomizableObjectSystemPrivate::BlockTillAllRequestsFinished(float TimeL
 	{
 		while (RemainingWork > 0)
 		{
-			StartTime = FPlatformTime::Seconds();
-			RemainingWork = GetPublic()->TickInternal(DeltaTime);
-			DeltaTime = FPlatformTime::Seconds() - StartTime;
+			RemainingWork = GetPublic()->TickInternal();
 		}
 	}
 	else
 	{
 		while (RemainingWork > 0)
 		{			
-			StartTime = FPlatformTime::Seconds();
-
-			if (StartTime > BlockEndTime)
+			if (FPlatformTime::Seconds() > BlockEndTime)
 			{
 				return RemainingWork;
 			}
 			
-			RemainingWork = GetPublic()->TickInternal(DeltaTime);
-			DeltaTime = FPlatformTime::Seconds() - StartTime;
+			RemainingWork = GetPublic()->TickInternal();
 		}
 	}
 
diff --git a/Engine/Plugins/Experimental/Mutable/Source/CustomizableObject/Private/MuCO/UnrealMutableImageProvider.cpp b/Engine/Plugins/Experimental/Mutable/Source/CustomizableObject/Private/MuCO/UnrealMutableImageProvider.cpp
--- a/Engine/Plugins/Experimental/Mutable/Source/CustomizableObject/Private/MuCO/UnrealMutableImageProvider.cpp
+++ b/Engine/Plugins/Experimental/Mutable/Source/CustomizableObject/Private/MuCO/UnrealMutableImageProvider.cpp
@@ -85,18 +85,9 @@ mu::EImageFormat GetMutablePixelFormat(EPixelFormat InTextureFormat)
 }
 
 
-FUnrealMutableImageProvider::FUnrealMutableImageProvider()
-{
-#if WITH_EDITOR
-	TickDelegate = FTickerDelegate::CreateRaw(this, &FUnrealMutableImageProvider::Tick);
-	TickDelegateHandle = FTSTicker::GetCoreTicker().AddTicker(TickDelegate, 0.f);
-#endif 
-}
-
-
 #if WITH_EDITOR
 
-bool FUnrealMutableImageProvider::Tick(float DeltaTime)
+bool FUnrealMutableImageProvider::Tick()
 {
 	FReferencedImageRequest* Request=nullptr;
 
@@ -486,7 +477,7 @@ TTuple<FGraphEventRef, TFunction<void()>> FUnrealMutableImageProvider::GetRefere
 		// This may happen in the mutable debugger.
 		while (!Request->CompletionEvent.IsCompleted())
 		{
-			Tick(0.1f);
+			Tick();
 		}
 	}
 	else
diff --git a/Engine/Plugins/Experimental/Mutable/Source/CustomizableObject/Private/MuCO/UnrealMutableImageProvider.h b/Engine/Plugins/Experimental/Mutable/Source/CustomizableObject/Private/MuCO/UnrealMutableImageProvider.h
--- a/Engine/Plugins/Experimental/Mutable/Source/CustomizableObject/Private/MuCO/UnrealMutableImageProvider.h
+++ b/Engine/Plugins/Experimental/Mutable/Source/CustomizableObject/Private/MuCO/UnrealMutableImageProvider.h
@@ -22,9 +22,6 @@ class FUnrealMutableImageProvider : public mu::ImageParameterGenerator, public F
 {
 
 public:
-
-	FUnrealMutableImageProvider();
-
 	// mu::ImageParameterGenerator interface
 	// Thread: worker
 #ifdef MUTABLE_USE_NEW_TASKGRAPH
@@ -103,10 +100,11 @@ private:
 	FCriticalSection ExternalImagesLock;
 
 #if WITH_EDITOR
+public:
 	/** The provider may tick to support referenced images. */
-	bool Tick(float DeltaTime);
-	FTSTicker::FDelegateHandle TickDelegateHandle;
-	FTickerDelegate TickDelegate;
+	bool Tick();
+
+private:
 	struct FReferencedImageRequest
 	{
 		FReferencedImageRequest(const UE::Tasks::FTaskEvent&& Event) : CompletionEvent(Event) {}
