[Mutable] Enable block compressed texture mipmap and resize optimizations.

#rnx

diff --git a/Engine/Plugins/Experimental/Mutable/Source/MutableRuntime/Private/MuR/OpImageMipmap.cpp b/Engine/Plugins/Experimental/Mutable/Source/MutableRuntime/Private/MuR/OpImageMipmap.cpp
--- a/Engine/Plugins/Experimental/Mutable/Source/MutableRuntime/Private/MuR/OpImageMipmap.cpp
+++ b/Engine/Plugins/Experimental/Mutable/Source/MutableRuntime/Private/MuR/OpImageMipmap.cpp
@@ -8,7 +8,7 @@
 namespace
 {
 
-bool bEnableCompressedMipGenerationMemoryOptimizations = false;
+bool bEnableCompressedMipGenerationMemoryOptimizations = true;
 static FAutoConsoleVariableRef CVarEnableCompressedMipGenerationMemoryOptimizations (
 	TEXT("mutable.EnableCompressedMipGenerationMemoryOptimizations"),
 	bEnableCompressedMipGenerationMemoryOptimizations,
diff --git a/Engine/Plugins/Experimental/Mutable/Source/MutableRuntime/Private/MuR/OpImageResize.cpp b/Engine/Plugins/Experimental/Mutable/Source/MutableRuntime/Private/MuR/OpImageResize.cpp
--- a/Engine/Plugins/Experimental/Mutable/Source/MutableRuntime/Private/MuR/OpImageResize.cpp
+++ b/Engine/Plugins/Experimental/Mutable/Source/MutableRuntime/Private/MuR/OpImageResize.cpp
@@ -7,7 +7,7 @@
 namespace
 {
 
-bool bEnableImageResizeMemoryOptimizations = false;
+bool bEnableImageResizeMemoryOptimizations = true;
 static FAutoConsoleVariableRef CVarEnableImageResizeMemoryOptimizations (
 	TEXT("mutable.EnableImageResizeMemoryOptimizations"),
 	bEnableImageResizeMemoryOptimizations,
