[mutable] Added disk compilation tests for Mutable Editor Tests
- Added flag to only run the compilation part of the COValidation Test commandlet.
- Added argument to override the setting for disk compilation of the target CO.

diff --git a/Engine/Plugins/Experimental/Mutable/Source/MutableValidation/Private/MuV/CustomizableObjectValidationCommandlet.cpp b/Engine/Plugins/Experimental/Mutable/Source/MutableValidation/Private/MuV/CustomizableObjectValidationCommandlet.cpp
--- a/Engine/Plugins/Experimental/Mutable/Source/MutableValidation/Private/MuV/CustomizableObjectValidationCommandlet.cpp
+++ b/Engine/Plugins/Experimental/Mutable/Source/MutableValidation/Private/MuV/CustomizableObjectValidationCommandlet.cpp
@@ -76,6 +76,14 @@ int32 UCustomizableObjectValidationCommandlet::Main(const FString& Params)
 		// Set the target compilation platform based on what the caller wants
 		CompilationOptions.TargetPlatform = TargetCompilationPlatform;
 
+		// Disk cache usage for compilation operation : Override if the user provided an argument with a different value than the default one of the CO
+		bool bUseDiskCacheForCompilation = false;
+		if (FParse::Bool(*Params,TEXT("UseDiskCompilation="),bUseDiskCacheForCompilation))
+		{
+			// Set the disk cache configuration
+			CompilationOptions.bUseDiskCompilation = bUseDiskCacheForCompilation;
+		}
+
 		bWasCoCompilationSuccessful = CompileCustomizableObject(ToTestCustomizableObject, true, &CompilationOptions);
 	}
 	// -------------------------------------------------------------------------------------------------------------- //
