[Mutable] Remove MaxLOD to simplify the LOD management system.

#rnx

diff --git a/Engine/Plugins/Experimental/Mutable/Source/CustomizableObject/Classes/MuCO/CustomizableObject.h b/Engine/Plugins/Experimental/Mutable/Source/CustomizableObject/Classes/MuCO/CustomizableObject.h
--- a/Engine/Plugins/Experimental/Mutable/Source/CustomizableObject/Classes/MuCO/CustomizableObject.h
+++ b/Engine/Plugins/Experimental/Mutable/Source/CustomizableObject/Classes/MuCO/CustomizableObject.h
@@ -572,7 +572,7 @@ public:
 	UPROPERTY(EditAnywhere, Category = CustomizableObject)
 	bool bEnableUseRefSkeletalMeshAsPlaceholder = true;
 
-	/** Use the Instance MinLOD, MaxLOD and RequestedLODs in the descriptor when performing the initial generation (ignore LOD Management). */
+	/** Use the Instance MinLOD, and RequestedLODs in the descriptor when performing the initial generation (ignore LOD Management). */
 	UPROPERTY(Category = "CustomizableObject", EditAnywhere, DisplayName = "Preserve User LODs On First Generation")
 	bool bPreserveUserLODsOnFirstGeneration = false;
 
diff --git a/Engine/Plugins/Experimental/Mutable/Source/CustomizableObject/Classes/MuCO/CustomizableObjectInstance.h b/Engine/Plugins/Experimental/Mutable/Source/CustomizableObject/Classes/MuCO/CustomizableObjectInstance.h
--- a/Engine/Plugins/Experimental/Mutable/Source/CustomizableObject/Classes/MuCO/CustomizableObjectInstance.h
+++ b/Engine/Plugins/Experimental/Mutable/Source/CustomizableObject/Classes/MuCO/CustomizableObjectInstance.h
@@ -566,21 +566,15 @@ public:
 	/** Return the min LOD that will be used in the next Instance update. */
 	int32 GetMinLODToLoad() const;
 
-	/** Return the max LOD that will be used in the next Instance update. */
-	int32 GetMaxLODToLoad() const;
-	
 	int32 GetNumLODsAvailable() const;
 
 	/** Return the min LOD this Instance is using. */
 	int32 GetCurrentMinLOD() const;
 
-	/** Return the max LOD this Instance is using. */
-	int32 GetCurrentMaxLOD() const;
-
 	/** Sets an array of LODs to generate per component. Mutable will generate those plus the currently generated LODs (if any).
 	  * Requires mutable.EnableOnlyGenerateRequestedLODs and CurrentInstanceLODManagement->IsOnlyGenerateRequestedLODLevelsEnabled() to be true.
 	  * @param InMinLOD - MinLOD to generate.
-	  * @param InMaxLOD - MaxLOD to generate.
+	  * @param InMaxLOD - MaxLOD to generate - DEPRECATED.
 	  * @param InRequestedLODsPerComponent - Array with bitmasks of requested LODs per component with range from [0 .. CO->GetComponentCount()].
 	  * @param InOutRequestedUpdates - Map from Instance to Update data that stores a request for the Instance to be updated, which will be either processed or discarded by priority (to be rerequested the next tick) */
 	void SetRequestedLODs(int32 InMinLOD, int32 InMaxLOD, const TArray<uint16>& InRequestedLODsPerComponent, FMutableInstanceUpdateMap& InOutRequestedUpdates);
diff --git a/Engine/Plugins/Experimental/Mutable/Source/CustomizableObject/Internal/MuCO/CustomizableObjectSystemPrivate.h b/Engine/Plugins/Experimental/Mutable/Source/CustomizableObject/Internal/MuCO/CustomizableObjectSystemPrivate.h
--- a/Engine/Plugins/Experimental/Mutable/Source/CustomizableObject/Internal/MuCO/CustomizableObjectSystemPrivate.h
+++ b/Engine/Plugins/Experimental/Mutable/Source/CustomizableObject/Internal/MuCO/CustomizableObjectSystemPrivate.h
@@ -60,7 +60,6 @@ public:
 
 	// These are the LODs that would be applied if this candidate is chosen
 	int32 MinLOD = 0;
-	int32 MaxLOD = INT32_MAX;
 
 	/** Array of RequestedLODs per component to generate if this candidate is chosen */
 	TArray<uint16> RequestedLODLevels;
@@ -69,13 +68,12 @@ public:
 	{
 		const FCustomizableObjectInstanceDescriptor& Descriptor = InCustomizableObjectInstance->GetDescriptor();
 		MinLOD = Descriptor.GetMinLod();
-		MaxLOD = Descriptor.GetMaxLod();
 		RequestedLODLevels = Descriptor.GetRequestedLODLevels();
 	}
 
-	FMutableUpdateCandidate(const UCustomizableObjectInstance* InCustomizableObjectInstance, const int32 InMinLOD, const int32 InMaxLOD,
+	FMutableUpdateCandidate(const UCustomizableObjectInstance* InCustomizableObjectInstance, const int32 InMinLOD,
 		const TArray<uint16>& InRequestedLODLevels) :
-		CustomizableObjectInstance(const_cast<UCustomizableObjectInstance*>(InCustomizableObjectInstance)), MinLOD(InMinLOD), MaxLOD(InMaxLOD),
+		CustomizableObjectInstance(const_cast<UCustomizableObjectInstance*>(InCustomizableObjectInstance)), MinLOD(InMinLOD),
 		RequestedLODLevels(InRequestedLODLevels) {}
 
 	bool HasBeenIssued() const;
@@ -486,10 +484,6 @@ public:
 
 	void SetMinLOD(int32 MinLOD);
 
-	int32 GetMaxLOD() const;
-
-	void SetMaxLOD(int32 MaxLOD);
-
 	const TArray<uint16>& GetRequestedLODs() const;
 
 	void SetRequestedLODs(TArray<uint16>& RequestedLODs);
diff --git a/Engine/Plugins/Experimental/Mutable/Source/CustomizableObject/Internal/MuCO/DescriptorHash.h b/Engine/Plugins/Experimental/Mutable/Source/CustomizableObject/Internal/MuCO/DescriptorHash.h
--- a/Engine/Plugins/Experimental/Mutable/Source/CustomizableObject/Internal/MuCO/DescriptorHash.h
+++ b/Engine/Plugins/Experimental/Mutable/Source/CustomizableObject/Internal/MuCO/DescriptorHash.h
@@ -28,7 +28,6 @@ private:
 
 public:
 	int32 MinLOD = 0;
-	int32 MaxLOD = INT32_MAX;
 
 	// Array of bitmasks that indicate which LODs of each component have been requested
 	TArray<uint16> RequestedLODsPerComponent;
diff --git a/Engine/Plugins/Experimental/Mutable/Source/CustomizableObject/Private/MuCO/CustomizableInstanceLODManagement.cpp b/Engine/Plugins/Experimental/Mutable/Source/CustomizableObject/Private/MuCO/CustomizableInstanceLODManagement.cpp
--- a/Engine/Plugins/Experimental/Mutable/Source/CustomizableObject/Private/MuCO/CustomizableInstanceLODManagement.cpp
+++ b/Engine/Plugins/Experimental/Mutable/Source/CustomizableObject/Private/MuCO/CustomizableInstanceLODManagement.cpp
@@ -292,7 +292,6 @@ void UCustomizableInstanceLODManagement::UpdateInstanceDistsAndLODs(FMutableInst
 		struct FLODTracker
 		{
 			int32 MinLOD = MAX_int32;
-			int32 MaxLOD = MAX_int32;
 
 			bool bInitialized = false;
 			TArray<uint16> RequestedLODsPerComponent;
@@ -382,11 +381,6 @@ void UCustomizableInstanceLODManagement::UpdateInstanceDistsAndLODs(FMutableInst
 					{
 						COI->SetMinSquareDistToPlayer(-1.f);
 					}
-					else
-					{
-						LODTracker.MaxLOD = COI->GetNumLODsAvailable() - 1;
-					}
-
 
 					// Use the component minLOD to set the minimum LOD to generate
 					LODTracker.MinLOD = FMath::Min(LODTracker.MinLOD, Parent->bOverrideMinLod ? Parent->MinLodModel : 0);
@@ -418,13 +412,7 @@ void UCustomizableInstanceLODManagement::UpdateInstanceDistsAndLODs(FMutableInst
 				// Limit MinLOD
 				It.Value.MinLOD = FMath::Min(It.Value.MinLOD, CustomizableObject->GetNumLODs() - 1);
 
-				// If it's the player generate only the first LOD
-				if (It.Key->GetMinSquareDistToPlayer() == -1.f && It.Value.MaxLOD == MAX_int32)
-				{
-					It.Value.MaxLOD = It.Value.MinLOD;
-				}
-
-				It.Key->SetRequestedLODs(It.Value.MinLOD, It.Value.MaxLOD, It.Value.RequestedLODsPerComponent, InOutRequestedUpdates);
+				It.Key->SetRequestedLODs(It.Value.MinLOD, 0, It.Value.RequestedLODsPerComponent, InOutRequestedUpdates);
 			}
 		}
 	}
diff --git a/Engine/Plugins/Experimental/Mutable/Source/CustomizableObject/Private/MuCO/CustomizableObjectInstance.cpp b/Engine/Plugins/Experimental/Mutable/Source/CustomizableObject/Private/MuCO/CustomizableObjectInstance.cpp
--- a/Engine/Plugins/Experimental/Mutable/Source/CustomizableObject/Private/MuCO/CustomizableObjectInstance.cpp
+++ b/Engine/Plugins/Experimental/Mutable/Source/CustomizableObject/Private/MuCO/CustomizableObjectInstance.cpp
@@ -1625,7 +1625,6 @@ bool UCustomizableInstancePrivate::DoComponentsNeedUpdate(UCustomizableObjectIns
 	UCustomizableObject* CustomizableObject = Public->GetCustomizableObject();
 	check(CustomizableObject);
 
-	const int32 NumLODs = OperationData->InstanceUpdateData.LODs.Num();
 	const int32 NumComponents = CustomizableObject->GetComponentCount();
 
 	TArray<bool> ComponentWithMesh;
@@ -1635,7 +1634,7 @@ bool UCustomizableInstancePrivate::DoComponentsNeedUpdate(UCustomizableObjectIns
 	MeshIDs.Init(MAX_uint64, NumComponents * MAX_MESH_LOD_COUNT);
 
 	// Gather the Mesh Ids of all components, and validate the integrity of the meshes to generate. 
-	for (int32 LODIndex = OperationData->GetMinLOD(); LODIndex <= OperationData->GetMaxLOD() && LODIndex < NumLODs; ++LODIndex)
+	for (int32 LODIndex = OperationData->GetMinLOD(); LODIndex < OperationData->NumLODsAvailable; ++LODIndex)
 	{
 		const FInstanceUpdateData::FLOD& LOD = OperationData->InstanceUpdateData.LODs[LODIndex];
 
@@ -1660,7 +1659,7 @@ bool UCustomizableInstancePrivate::DoComponentsNeedUpdate(UCustomizableObjectIns
 			{
 				UE_LOG(LogMutable, Error, TEXT("Failed to generate SkeletalMesh for CO Instance [%s]. CO [%s] has invalid geometry for LOD [%d] Component [%d]."),
 					*Public->GetName(), *CustomizableObject->GetName(),
-					OperationData->GetMaxLOD(), ComponentIndex);
+					LODIndex, ComponentIndex);
 				bHasInvalidMesh = true;
 				continue;
 			}
@@ -1688,11 +1687,11 @@ bool UCustomizableInstancePrivate::DoComponentsNeedUpdate(UCustomizableObjectIns
 		}
 
 		// Components with mesh must have valid geometry at CurrentMaxLOD
-		if (ComponentWithMesh[ComponentIndex] && MeshIDs[ComponentIndex * MAX_MESH_LOD_COUNT + OperationData->GetMaxLOD()] == MAX_uint64)
+		if (ComponentWithMesh[ComponentIndex] && MeshIDs[ComponentIndex * MAX_MESH_LOD_COUNT + OperationData->NumLODsAvailable - 1] == MAX_uint64)
 		{
 			UE_LOG(LogMutable, Error, TEXT("Failed to generate SkeletalMesh for CO Instance [%s]. CO [%s] is missing geometry for LOD [%d] Component [%d]."),
 				*Public->GetName(), *CustomizableObject->GetName(),
-				OperationData->GetMaxLOD(), ComponentIndex);
+				OperationData->NumLODsAvailable - 1, ComponentIndex);
 			bHasInvalidMesh = true;
 			continue;
 		}
@@ -2799,11 +2798,9 @@ void FMutableUpdateCandidate::Issue()
 void FMutableUpdateCandidate::ApplyLODUpdateParamsToInstance(FUpdateContextPrivate& Context)
 {
 	CustomizableObjectInstance->Descriptor.MinLOD = MinLOD;
-	CustomizableObjectInstance->Descriptor.MaxLOD = MaxLOD;
 	CustomizableObjectInstance->Descriptor.RequestedLODLevels = RequestedLODLevels;
 
 	Context.SetMinLOD(MinLOD);
-	Context.SetMaxLOD(MaxLOD);
 	Context.SetRequestedLODs(RequestedLODLevels);
 }
 
@@ -3224,7 +3221,7 @@ void UCustomizableInstancePrivate::InitSkeletalMeshData(const TSharedRef<FUpdate
 		SkeletalMesh->AllocateResourceForRendering();
 
 		FSkeletalMeshRenderData* RenderData = SkeletalMesh->GetResourceForRendering();
-		for (int32 LODIndex = 0; LODIndex <= OperationData->GetMaxLOD(); ++LODIndex)
+		for (int32 LODIndex = 0; LODIndex < OperationData->NumLODsAvailable; ++LODIndex)
 		{
 			RenderData->LODRenderData.Add(new FSkeletalMeshLODRenderData());
 
@@ -3492,7 +3489,7 @@ void UCustomizableInstancePrivate::BuildOrCopyElementData(const TSharedRef<FUpda
 {
 	MUTABLE_CPUPROFILER_SCOPE(UCustomizableInstancePrivate::BuildOrCopyElementData);
 
-	for (int32 LODIndex = OperationData->GetMaxLOD(); LODIndex >= FirstLODAvailable; --LODIndex)
+	for (int32 LODIndex = OperationData->NumLODsAvailable - 1; LODIndex >= FirstLODAvailable; --LODIndex)
 	{
 		const FInstanceUpdateData::FLOD& LOD = OperationData->InstanceUpdateData.LODs[LODIndex];
 
@@ -3549,8 +3546,8 @@ void UCustomizableInstancePrivate::BuildOrCopyMorphTargetsData(const TSharedRef<
 
 	const int32 SkeletalMeshMorphTargetsNum = SkeletalMesh->GetMorphTargets().Num();
 
-	int32 LastValidLODIndex = OperationData->GetMaxLOD();
-	for (int32 LODIndex = OperationData->GetMaxLOD(); LODIndex >= FirstLODAvailable; --LODIndex)
+	int32 LastValidLODIndex = OperationData->NumLODsAvailable - 1;
+	for (int32 LODIndex = LastValidLODIndex; LODIndex >= FirstLODAvailable; --LODIndex)
 	{
 		const FInstanceUpdateData::FLOD& LOD = OperationData->InstanceUpdateData.LODs[LODIndex];
 
@@ -3710,7 +3707,7 @@ void UCustomizableInstancePrivate::BuildOrCopyClothingData(const TSharedRef<FUpd
 	{
 		MUTABLE_CPUPROFILER_SCOPE(DiscoverSectionsWithCloth);
 
-		for (int32 LODIndex = OperationData->GetMinLOD(); LODIndex <= OperationData->GetMaxLOD(); ++LODIndex)
+		for (int32 LODIndex = OperationData->GetMinLOD(); LODIndex < OperationData->NumLODsAvailable; ++LODIndex)
 		{
 			const FInstanceUpdateData::FLOD& LOD = OperationData->InstanceUpdateData.LODs[LODIndex];
 			if (ComponentIndex >= LOD.ComponentCount)
@@ -3836,7 +3833,7 @@ void UCustomizableInstancePrivate::BuildOrCopyClothingData(const TSharedRef<FUpd
 			// Only initilialize once, multiple sections with cloth could point to the same cloth asset.
 			if (!DstAssetData.LodMap.Num())
 			{
-				DstAssetData.LodMap.Init( INDEX_NONE, OperationData->GetMaxLOD() + 1 );
+				DstAssetData.LodMap.Init( INDEX_NONE, OperationData->NumLODsAvailable);
 
 				DstAssetData.LodData.SetNum( SrcAssetData.LodData.Num() );
 				DstAssetData.UsedBoneNames = SrcAssetData.UsedBoneNames;
@@ -4360,7 +4357,7 @@ void UCustomizableInstancePrivate::BuildOrCopyClothingData(const TSharedRef<FUpd
 	{
 		int32 NumSectionsWithClothProcessed = 0;
 
-		for (int32 LODIndex = OperationData->GetMinLOD(); LODIndex <= OperationData->GetMaxLOD(); ++LODIndex)
+		for (int32 LODIndex = OperationData->GetMinLOD(); LODIndex < OperationData->NumLODsAvailable; ++LODIndex)
 		{
 			const FInstanceUpdateData::FLOD& LOD = OperationData->InstanceUpdateData.LODs[LODIndex];
 			if (ComponentIndex >= LOD.ComponentCount)
@@ -4413,7 +4410,7 @@ void UCustomizableInstancePrivate::BuildOrCopyClothingData(const TSharedRef<FUpd
 	{
 		MUTABLE_CPUPROFILER_SCOPE(InitClothRenderData)
 		// Based on FSkeletalMeshLODModel::GetClothMappingData().
-		for (int32 LODIndex = OperationData->GetMinLOD(); LODIndex <= OperationData->GetMaxLOD(); ++LODIndex)
+		for (int32 LODIndex = OperationData->GetMinLOD(); LODIndex < OperationData->NumLODsAvailable; ++LODIndex)
 		{
 			FSkeletalMeshLODRenderData& LODModel = RenderResource->LODRenderData[LODIndex];
 	
@@ -4532,8 +4529,8 @@ bool UCustomizableInstancePrivate::BuildOrCopyRenderData(const TSharedRef<FUpdat
 	check(CustomizableObject);
 	const FModelResources& ModelResources = CustomizableObject->GetPrivate()->GetModelResources();
 
-	int32 LastValidLODIndex = OperationData->GetMaxLOD();
-	for (int32 LODIndex = OperationData->GetMaxLOD(); LODIndex >= FirstLODAvailable; --LODIndex)
+	int32 LastValidLODIndex = OperationData->NumLODsAvailable - 1;
+	for (int32 LODIndex = LastValidLODIndex; LODIndex >= FirstLODAvailable; --LODIndex)
 	{
 		MUTABLE_CPUPROFILER_SCOPE(BuildOrCopyRenderData_LODLoop);
 
@@ -4543,7 +4540,7 @@ bool UCustomizableInstancePrivate::BuildOrCopyRenderData(const TSharedRef<FUpdat
 		if (ComponentIndex >=  LOD.ComponentCount)
 		{
 			// Interrupt the generation if the LOD is empty and it should have been generated.
-			if (LODIndex >= OperationData->GetMinLOD() && LODIndex <= OperationData->GetMaxLOD())
+			if (LODIndex >= OperationData->GetMinLOD() && LODIndex < OperationData->NumLODsAvailable)
 			{
 				UE_LOG(LogMutable, Warning, TEXT("Building instance: generated mesh [%s] has LOD [%d] with no component.")
 					, *SkeletalMesh->GetName()
@@ -5314,7 +5311,7 @@ void UCustomizableInstancePrivate::BuildMaterials(const TSharedRef<FUpdateContex
 
 		MUTABLE_CPUPROFILER_SCOPE(BuildMaterials_LODLoop);
 
-		for (int32 LODIndex = FirstLODAvailable; LODIndex <= OperationData->GetMaxLOD(); LODIndex++)
+		for (int32 LODIndex = FirstLODAvailable; LODIndex < OperationData->NumLODsAvailable; LODIndex++)
 		{
 			const FInstanceUpdateData::FLOD& LOD = OperationData->InstanceUpdateData.LODs[LODIndex];
 			const FInstanceUpdateData::FComponent& Component = OperationData->InstanceUpdateData.Components[LOD.FirstComponent + ComponentIndex];
@@ -5827,8 +5824,8 @@ void UCustomizableInstancePrivate::BuildMaterials(const TSharedRef<FUpdateContex
 
 		{
 			// Copy data from valid LODs into the skipped ones.
-			int32 LastValidLODIndex = OperationData->GetMaxLOD();
-			for (int32 LODIndex = OperationData->GetMaxLOD(); LODIndex >= FirstLODAvailable; --LODIndex)
+			int32 LastValidLODIndex = OperationData->NumLODsAvailable - 1;
+			for (int32 LODIndex = LastValidLODIndex; LODIndex >= FirstLODAvailable; --LODIndex)
 			{
 				// Copy information from the LastValidLODIndex
 				if (LODsSkipped[LODIndex])
@@ -6015,12 +6012,6 @@ int32 UCustomizableObjectInstance::GetMinLODToLoad() const
 }
 
 
-int32 UCustomizableObjectInstance::GetMaxLODToLoad() const
-{
-	return Descriptor.MaxLOD;	
-}
-
-
 int32 UCustomizableObjectInstance::GetNumLODsAvailable() const
 {
 	return GetPrivate()->GetNumLODsAvailable();
@@ -6033,21 +6024,16 @@ int32 UCustomizableObjectInstance::GetCurrentMinLOD() const
 }
 
 
-int32 UCustomizableObjectInstance::GetCurrentMaxLOD() const
-{
-	return GetPrivate()->CommittedDescriptor.GetMaxLod();
-}
-
 #if !UE_BUILD_SHIPPING
-static bool bIgnoreMinMaxLOD = false;
+static bool bIgnoreMinLOD = false;
 FAutoConsoleVariableRef CVarMutableIgnoreMinMaxLOD(
 	TEXT("Mutable.IgnoreMinMaxLOD"),
-	bIgnoreMinMaxLOD,
+	bIgnoreMinLOD,
 	TEXT("The limits on the number of LODs to generate will be ignored."));
 #endif
 
 
-void UCustomizableObjectInstance::SetRequestedLODs(int32 InMinLOD, int32 InMaxLOD, const TArray<uint16>& InRequestedLODsPerComponent, 
+void UCustomizableObjectInstance::SetRequestedLODs(int32 InMinLOD, int32 , const TArray<uint16>& InRequestedLODsPerComponent, 
 	                                               FMutableInstanceUpdateMap& InOutRequestedUpdates)
 {
 	check(PrivateData);
@@ -6075,11 +6061,10 @@ void UCustomizableObjectInstance::SetRequestedLODs(int32 InMinLOD, int32 InMaxLO
 	}
 	
 #if !UE_BUILD_SHIPPING
-	// Ignore Min/Max LOD limits. Mainly used for debug
-	if (bIgnoreMinMaxLOD)
+	// Ignore Min LOD limits. Mainly used for debug
+	if (bIgnoreMinLOD)
 	{
 		InMinLOD = 0;
-		InMaxLOD = MAX_MESH_LOD_COUNT;
 	}
 #endif
 
@@ -6087,19 +6072,15 @@ void UCustomizableObjectInstance::SetRequestedLODs(int32 InMinLOD, int32 InMaxLO
 
 	// Clamp Min LOD
 	const int32 MinLODIdx = GetCustomizableObject()->GetPrivate()->GetMinLODIndex();
-	InMinLOD = FMath::Min(FMath::Max(InMinLOD, MinLODIdx), MinLODIdx + PrivateData->NumMaxLODsToStream);
-
-	// Clamp Max LOD
-	InMaxLOD = FMath::Max(FMath::Min3(InMaxLOD, PrivateData->NumLODsAvailable - 1, (int32)MAX_MESH_LOD_COUNT), InMinLOD);
+	InMinLOD = FMath::Min(FMath::Max(InMinLOD, MinLODIdx), (int32)PrivateData->NumMaxLODsToStream);
 
-	const bool bMinMaxLODChanged = Descriptor.MaxLOD != InMaxLOD || Descriptor.MinLOD != InMinLOD;
+	const bool bMinLODChanged = Descriptor.MinLOD != InMinLOD;
 
 	const bool bIsDowngradeLODUpdate = GetCurrentMinLOD() >= 0 && InMinLOD > GetCurrentMinLOD();
 	PrivateData->SetCOInstanceFlags(bIsDowngradeLODUpdate ? PendingLODsDowngrade : ECONone);
 
 	// Save the new LODs
 	MutableUpdateCandidate.MinLOD = InMinLOD;
-	MutableUpdateCandidate.MaxLOD = InMaxLOD;
 	MutableUpdateCandidate.RequestedLODLevels = Descriptor.GetRequestedLODLevels();
 
 	bool bUpdateRequestedLODs = false;
@@ -6112,7 +6093,7 @@ void UCustomizableObjectInstance::SetRequestedLODs(int32 InMinLOD, int32 InMaxLO
 
 		const bool bIgnoreGeneratedLODs = GeneratedLODsPerComponent.Num() != ComponentCount;
 
-		if (bMinMaxLODChanged || bIgnoreGeneratedLODs || Descriptor.GetRequestedLODLevels() != InRequestedLODsPerComponent)
+		if (bMinLODChanged || bIgnoreGeneratedLODs || Descriptor.GetRequestedLODLevels() != InRequestedLODsPerComponent)
 		{
 			check(InRequestedLODsPerComponent.Num() == ComponentCount);
 
@@ -6120,8 +6101,8 @@ void UCustomizableObjectInstance::SetRequestedLODs(int32 InMinLOD, int32 InMaxLO
 			for (int32 ComponentIndex = 0; ComponentIndex < ComponentCount; ++ComponentIndex)
 			{
 				// Find the first requested LOD. We'll generate [FirstRequestedLOD ... MaxLOD].
-				int32 FirstRequestedLOD = MutableUpdateCandidate.MaxLOD;
-				for (int32 LODIndex = 0; LODIndex < MutableUpdateCandidate.MaxLOD; ++LODIndex)
+				int32 FirstRequestedLOD = ((int32)PrivateData->NumMaxLODsToStream) - 1;
+				for (int32 LODIndex = 0; LODIndex < FirstRequestedLOD; ++LODIndex)
 				{
 					if ((!bIgnoreGeneratedLODs && GeneratedLODsPerComponent[ComponentIndex] & (1 << LODIndex))
 						||
@@ -6129,12 +6110,11 @@ void UCustomizableObjectInstance::SetRequestedLODs(int32 InMinLOD, int32 InMaxLO
 					{
 						// First RequestedLOD that fall within the range
 						FirstRequestedLOD = LODIndex >= MutableUpdateCandidate.MinLOD ? LODIndex : MutableUpdateCandidate.MinLOD;
-						break;
 					}
 				}
 
 				// Generate at least the MaxLOD
-				int32 RequestedLODs = 1 << MutableUpdateCandidate.MaxLOD;
+				int32 RequestedLODs = 1 << FirstRequestedLOD;
 
 				// Mark all LODs up until MAX_MESH_LOD_COUNT since MaxLOD can be set to Max_int32
 				for (int32 LODIndex = FirstRequestedLOD; LODIndex < MAX_MESH_LOD_COUNT; ++LODIndex)
@@ -6150,7 +6130,7 @@ void UCustomizableObjectInstance::SetRequestedLODs(int32 InMinLOD, int32 InMaxLO
 		}
 	}
 
-	if (bMinMaxLODChanged || bUpdateRequestedLODs)
+	if (bMinLODChanged || bUpdateRequestedLODs)
 	{
 		// TODO: Remove this flag as it will become redundant with the new InOutRequestedUpdates system
 		PrivateData->SetCOInstanceFlags(PendingLODsUpdate);
diff --git a/Engine/Plugins/Experimental/Mutable/Source/CustomizableObject/Private/MuCO/CustomizableObjectInstanceDescriptor.cpp b/Engine/Plugins/Experimental/Mutable/Source/CustomizableObject/Private/MuCO/CustomizableObjectInstanceDescriptor.cpp
--- a/Engine/Plugins/Experimental/Mutable/Source/CustomizableObject/Private/MuCO/CustomizableObjectInstanceDescriptor.cpp
+++ b/Engine/Plugins/Experimental/Mutable/Source/CustomizableObject/Private/MuCO/CustomizableObjectInstanceDescriptor.cpp
@@ -682,7 +682,6 @@ void FCustomizableObjectInstanceDescriptor::ReloadParameters()
 
 	SetState(FMath::Clamp(GetState(), 0, CustomizableObject->GetStateCount() - 1));
 	
-	MaxLOD = FMath::Max(0, CustomizableObject->GetNumLODs() - 1);
 	RequestedLODLevels.Init(0, CustomizableObject->GetComponentCount());
 
 	TArray<FCustomizableObjectBoolParameterValue> OldBoolParameters = BoolParameters;
@@ -1082,18 +1081,6 @@ void FCustomizableObjectInstanceDescriptor::SetMinLod(int32 InMinLOD)
 }
 
 
-int32 FCustomizableObjectInstanceDescriptor::GetMaxLod() const
-{
-	return MaxLOD;
-}
-
-
-void FCustomizableObjectInstanceDescriptor::SetMaxLod(int32 InMaxLOD)
-{
-	MaxLOD = InMaxLOD;
-}
-
-
 void FCustomizableObjectInstanceDescriptor::SetRequestedLODLevels(const TArray<uint16>& InRequestedLODLevels)
 {
 	RequestedLODLevels = InRequestedLODLevels;
diff --git a/Engine/Plugins/Experimental/Mutable/Source/CustomizableObject/Private/MuCO/CustomizableObjectSystem.cpp b/Engine/Plugins/Experimental/Mutable/Source/CustomizableObject/Private/MuCO/CustomizableObjectSystem.cpp
--- a/Engine/Plugins/Experimental/Mutable/Source/CustomizableObject/Private/MuCO/CustomizableObjectSystem.cpp
+++ b/Engine/Plugins/Experimental/Mutable/Source/CustomizableObject/Private/MuCO/CustomizableObjectSystem.cpp
@@ -186,19 +186,6 @@ void FUpdateContextPrivate::SetMinLOD(int32 MinLOD)
 }
 
 
-int32 FUpdateContextPrivate::GetMaxLOD() const
-{
-	return CapturedDescriptor.GetMaxLod();
-}
-
-
-void FUpdateContextPrivate::SetMaxLOD(int32 MaxLOD)
-{
-	CapturedDescriptor.SetMaxLod(MaxLOD);
-	CapturedDescriptorHash.MaxLOD = MaxLOD;
-}
-
-
 const TArray<uint16>& FUpdateContextPrivate::GetRequestedLODs() const
 {
 	return CapturedDescriptor.GetRequestedLODLevels();
@@ -1227,7 +1214,7 @@ void UCustomizableObjectSystemPrivate::EnqueueUpdateSkeletalMesh(const TSharedRe
 	{
 		if (InstancePrivate->HasCOInstanceFlags(PendingLODsUpdate))
 		{
-			UE_LOG(LogMutable, Verbose, TEXT("LOD change: %d, %d -> %d, %d"), Instance->GetCurrentMinLOD(), Instance->GetCurrentMaxLOD(), Instance->GetMinLODToLoad(), Instance->GetMaxLODToLoad());
+			UE_LOG(LogMutable, Verbose, TEXT("Min LOD change: %d -> %d"), Instance->GetCurrentMinLOD(), Instance->GetMinLODToLoad());
 		}
 
 		if (const FMutablePendingInstanceUpdate* QueueElem = MutablePendingInstanceWork.GetUpdate(Instance))
@@ -1636,20 +1623,13 @@ namespace impl
 		Operation->NumLODsAvailable = Operation->MutableInstance->GetLODCount();
 
 		int32 CurrentMinLOD = FMath::Max(Operation->GetMinLOD(), Operation->FirstLODAvailable);
-		int32 CurrentMaxLOD = Operation->GetMaxLOD();
 
 		if (CurrentMinLOD >= Operation->NumLODsAvailable)
 		{
 			CurrentMinLOD = Operation->NumLODsAvailable - 1;
-			CurrentMaxLOD = CurrentMinLOD;
-		}
-		else if (CurrentMaxLOD >= Operation->NumLODsAvailable)
-		{
-			CurrentMaxLOD = Operation->NumLODsAvailable - 1;
 		}
 
 		Operation->SetMinLOD(CurrentMinLOD);
-		Operation->SetMaxLOD(CurrentMaxLOD);
 
 		// Initialize RequestedLODs to zero if not set
 		TArray<uint16> RequestedLODs = Operation->GetRequestedLODs();
@@ -1658,7 +1638,7 @@ namespace impl
 		for (int32 ComponentIndex = 0; ComponentIndex < Operation->NumComponents; ++ComponentIndex)
 		{
 			// Ensure we're generating at least one LOD
-			for (int32 LODIndex = CurrentMaxLOD; LODIndex < MAX_MESH_LOD_COUNT; ++LODIndex)
+			for (int32 LODIndex = Operation->NumLODsAvailable - 1; LODIndex < MAX_MESH_LOD_COUNT; ++LODIndex)
 			{
 				RequestedLODs[ComponentIndex] |= (1 << LODIndex);
 			}
@@ -1712,11 +1692,11 @@ namespace impl
 		TArray<int32> SurfacesSharedId;
 
 		// Generate the mesh and gather all the required resource Ids
-		OperationData->InstanceUpdateData.LODs.SetNum(Instance->GetLODCount());
-		for (int32 MutableLODIndex = 0; MutableLODIndex < Instance->GetLODCount(); ++MutableLODIndex)
+		OperationData->InstanceUpdateData.LODs.SetNum(OperationData->NumLODsAvailable);
+		for (int32 MutableLODIndex = 0; MutableLODIndex < OperationData->NumLODsAvailable; ++MutableLODIndex)
 		{
 			// Skip LODs outside the range we want to generate
-			if (MutableLODIndex < OperationData->GetMinLOD() || MutableLODIndex > OperationData->GetMaxLOD())
+			if (MutableLODIndex < OperationData->GetMinLOD())
 			{
 				continue;
 			}
@@ -2631,7 +2611,7 @@ namespace impl
 
 			const TArray<uint16>& RequestedLODs = Operation->GetRequestedLODs();
 			
-			for (int32 LODIndex = Operation->GetMinLOD(); LODIndex <= Operation->GetMaxLOD(); ++LODIndex)
+			for (int32 LODIndex = Operation->GetMinLOD(); LODIndex < Operation->NumLODsAvailable; ++LODIndex)
 			{
 				const bool bGenerateLOD = RequestedLODs.IsValidIndex(ComponentIndex) ? (RequestedLODs[ComponentIndex] & (1 << LODIndex)) != 0 : true;
 				if (bGenerateLOD)
diff --git a/Engine/Plugins/Experimental/Mutable/Source/CustomizableObject/Private/MuCO/DescriptorHash.cpp b/Engine/Plugins/Experimental/Mutable/Source/CustomizableObject/Private/MuCO/DescriptorHash.cpp
--- a/Engine/Plugins/Experimental/Mutable/Source/CustomizableObject/Private/MuCO/DescriptorHash.cpp
+++ b/Engine/Plugins/Experimental/Mutable/Source/CustomizableObject/Private/MuCO/DescriptorHash.cpp
@@ -54,7 +54,6 @@ FDescriptorHash::FDescriptorHash(const FCustomizableObjectInstanceDescriptor& De
 	}
 	
 	MinLOD = Descriptor.MinLOD;
-	MaxLOD = Descriptor.MaxLOD;	
 
 	RequestedLODsPerComponent = Descriptor.RequestedLODLevels;
 }
@@ -64,7 +63,6 @@ bool FDescriptorHash::operator==(const FDescriptorHash& Other) const
 {
 	return Hash == Other.Hash &&
 		MinLOD == Other.MinLOD &&
-		MaxLOD == Other.MaxLOD &&
 		RequestedLODsPerComponent == Other.RequestedLODsPerComponent;
 }
 
@@ -77,7 +75,7 @@ bool FDescriptorHash::operator!=(const FDescriptorHash& Other) const
 
 bool FDescriptorHash::IsSubset(const FDescriptorHash& Other) const
 {
-	if (Hash != Other.Hash || MinLOD < Other.MinLOD || MaxLOD != Other.MaxLOD)
+	if (Hash != Other.Hash || MinLOD < Other.MinLOD)
 	{
 		return false;
 	}
@@ -96,7 +94,7 @@ bool FDescriptorHash::IsSubset(const FDescriptorHash& Other) const
 	{
 		const int32 RequestedLODs = RequestedLODsPerComponent[ComponentIndex];
 		const int32 OtherRequestedLODs = Other.RequestedLODsPerComponent[ComponentIndex];
-		for (uint16 LODIndex = MinLOD; LODIndex <= MaxLOD; ++LODIndex)
+		for (uint16 LODIndex = MinLOD; LODIndex <= MAX_MESH_LOD_COUNT; ++LODIndex)
 		{
 			// To be a subset all bits set in RequestedLODs must be set in OtherRequestedLODs. OtherRequestedLODs can have additional bits set
 			if (RequestedLODs & (1 << LODIndex) && !(OtherRequestedLODs & (1 << LODIndex)))
@@ -116,7 +114,6 @@ FString FDescriptorHash::ToString() const
 
 	Builder.Appendf(TEXT("(Hash=%u,"), Hash);
 	Builder.Appendf(TEXT("MinLOD=%i,"), MinLOD);
-	Builder.Appendf(TEXT("MaxLOD=%i,"), MaxLOD);
 	Builder.Appendf(TEXT("RequestredLODLevels=["));
 
 	for (const uint16 RequestedLODs : RequestedLODsPerComponent)
diff --git a/Engine/Plugins/Experimental/Mutable/Source/CustomizableObject/Private/MuCO/LogBenchmarkUtil.cpp b/Engine/Plugins/Experimental/Mutable/Source/CustomizableObject/Private/MuCO/LogBenchmarkUtil.cpp
--- a/Engine/Plugins/Experimental/Mutable/Source/CustomizableObject/Private/MuCO/LogBenchmarkUtil.cpp
+++ b/Engine/Plugins/Experimental/Mutable/Source/CustomizableObject/Private/MuCO/LogBenchmarkUtil.cpp
@@ -94,8 +94,8 @@ void FLogBenchmarkUtil::GetInstancesStats(int32& OutNumInstances, int32& OutNumB
 		
 		++OutNumInstances;
 
-		OutNumInstancesLOD0 += Instance->GetCurrentMinLOD() <= 0 && Instance->GetCurrentMaxLOD() >= 0;
-		OutNumInstancesLOD1 += Instance->GetCurrentMinLOD() <= 1 && Instance->GetCurrentMaxLOD() >= 1;
+		OutNumInstancesLOD0 += Instance->GetCurrentMinLOD() == 0;
+		OutNumInstancesLOD1 += Instance->GetCurrentMinLOD() == 1;
 		OutNumInstancesLOD2 += Instance->GetCurrentMinLOD() >= 2;
 
 		if (Instance->GetPrivate()->SkeletalMeshStatus == ESkeletalMeshStatus::Success)
diff --git a/Engine/Plugins/Experimental/Mutable/Source/CustomizableObject/Private/MuCO/LogInformationUtil.cpp b/Engine/Plugins/Experimental/Mutable/Source/CustomizableObject/Private/MuCO/LogInformationUtil.cpp
--- a/Engine/Plugins/Experimental/Mutable/Source/CustomizableObject/Private/MuCO/LogInformationUtil.cpp
+++ b/Engine/Plugins/Experimental/Mutable/Source/CustomizableObject/Private/MuCO/LogInformationUtil.cpp
@@ -112,9 +112,7 @@ void LogInformationUtil::LogShowInstanceData(const UCustomizableObjectInstance*
 	LogData += FString::Printf(TEXT("bShouldUpdateLODs=%d "), InstanceFlags & ECOInstanceFlags::PendingLODsUpdate);
 
 	LogData += FString::Printf(TEXT("CurrentMinLOD=%d "), CustomizableObjectInstance->GetCurrentMinLOD());
-	LogData += FString::Printf(TEXT("CurrentMaxLOD=%d "), CustomizableObjectInstance->GetCurrentMaxLOD());
 	LogData += FString::Printf(TEXT("MinLODToLoad=%d "), CustomizableObjectInstance->GetMinLODToLoad());
-	LogData += FString::Printf(TEXT("MaxLODToLoad=%d\n"), CustomizableObjectInstance->GetMaxLODToLoad());
 }
 
 
@@ -158,9 +156,7 @@ void LogInformationUtil::LogShowInstanceDataFull(const UCustomizableObjectInstan
 	LogData += MessageChunk;
 
 	MessageChunk = FString::Printf(TEXT("        CurrentMinLOD = %d\n"), CustomizableObjectInstance->GetCurrentMinLOD());
-	MessageChunk += FString::Printf(TEXT("        CurrentMaxLOD = %d\n"), CustomizableObjectInstance->GetCurrentMaxLOD());
 	MessageChunk += FString::Printf(TEXT("        MinLODToLoad = %d\n"), CustomizableObjectInstance->GetMinLODToLoad());
-	MessageChunk += FString::Printf(TEXT("        MaxLODToLoad = %d\n"), CustomizableObjectInstance->GetMaxLODToLoad());
 	MessageChunk += FString::Printf(TEXT("        bShouldUpdateLODs = %d\n"), InstanceFlags & ECOInstanceFlags::PendingLODsUpdate);
 	LogData += MessageChunk;
 
diff --git a/Engine/Plugins/Experimental/Mutable/Source/CustomizableObject/Public/MuCO/CustomizableObjectInstanceDescriptor.h b/Engine/Plugins/Experimental/Mutable/Source/CustomizableObject/Public/MuCO/CustomizableObjectInstanceDescriptor.h
--- a/Engine/Plugins/Experimental/Mutable/Source/CustomizableObject/Public/MuCO/CustomizableObjectInstanceDescriptor.h
+++ b/Engine/Plugins/Experimental/Mutable/Source/CustomizableObject/Public/MuCO/CustomizableObjectInstanceDescriptor.h
@@ -73,9 +73,9 @@ struct CUSTOMIZABLEOBJECT_API FCustomizableObjectInstanceDescriptor
 
 	void SetMinLod(int32 InMinLOD);
 
-	int32 GetMaxLod() const;
+	int32 GetMaxLod() const { return MAX_int32; }; // DEPRECATED
 
-	void SetMaxLod(int32 InMaxLOD);
+	void SetMaxLod(int32 InMaxLOD) {}; // DEPRECATED
 
 	void SetRequestedLODLevels(const TArray<uint16>& InRequestedLODLevels);
 
@@ -382,7 +382,6 @@ private:
 
 	/** These are the LODs Mutable can generate, they MUST NOT be used in an update (Mutable thread). */
 	int32 MinLOD = 0;
-	int32 MaxLOD = INT32_MAX;
 
 	/** Array of RequestedLODs per component to generate, they MUST NOT be used in an update (Mutable thread). */
 	TArray<uint16> RequestedLODLevels;
diff --git a/Engine/Plugins/Experimental/Mutable/Source/CustomizableObjectEditor/Private/MuCOE/CustomizableInstanceDetails.cpp b/Engine/Plugins/Experimental/Mutable/Source/CustomizableObjectEditor/Private/MuCOE/CustomizableInstanceDetails.cpp
--- a/Engine/Plugins/Experimental/Mutable/Source/CustomizableObjectEditor/Private/MuCOE/CustomizableInstanceDetails.cpp
+++ b/Engine/Plugins/Experimental/Mutable/Source/CustomizableObjectEditor/Private/MuCOE/CustomizableInstanceDetails.cpp
@@ -1867,14 +1867,12 @@ FReply FCustomizableInstanceDetails::OnPasteAllParameters()
 	const UScriptStruct* Struct = Descriptor.StaticStruct();
 
 	const int32 MinLOD = Descriptor.GetMinLod();
-	const int32 MaxLOD = Descriptor.GetMaxLod();
 	const TArray<uint16> RequestedLODLevels = Descriptor.GetRequestedLODLevels();
 	
 	if (Struct->ImportText(*ClipText, &Descriptor, nullptr, 0, GLog, GetPathNameSafe(Struct)))
 	{
 		// Keep current LOD
 		Descriptor.SetMinLod(MinLOD);
-		Descriptor.SetMaxLod(MaxLOD);
 		Descriptor.SetRequestedLODLevels(RequestedLODLevels);
 		
 		CustomInstance->UpdateSkeletalMeshAsync(true, true);		
