[mutable] Avoid using FPipe in the task graph when compiling because of possible bugs.

#rnx

diff --git a/Engine/Plugins/Experimental/Mutable/Source/MutableTools/Private/MuT/CodeOptimiser.cpp b/Engine/Plugins/Experimental/Mutable/Source/MutableTools/Private/MuT/CodeOptimiser.cpp
--- a/Engine/Plugins/Experimental/Mutable/Source/MutableTools/Private/MuT/CodeOptimiser.cpp
+++ b/Engine/Plugins/Experimental/Mutable/Source/MutableTools/Private/MuT/CodeOptimiser.cpp
@@ -24,7 +24,6 @@
 #include "MuR/MutableRuntimeModule.h"
 
 #include "Tasks/Task.h"
-#include "Tasks/Pipe.h"
 
 #include <unordered_set>
 
@@ -617,8 +616,6 @@ namespace mu
 	{
 		MUTABLE_CPUPROFILER_SCOPE(ConstantGenerator);
 
-		UE::Tasks::FPipe ASTPipe(TEXT("ASTPipe"));
-
 		// don't do this if constant optimization has been disabled, usually for debugging.
 		if (!InOptions->OptimisationOptions.bConstReduction)
 		{
@@ -782,13 +779,17 @@ namespace mu
 
 		if (bUseConcurrency)
 		{
+			/** Protect access to the original AST being optimized. */
+			FCriticalSection ASTAccessLock;
+
 			// Launch the tasks. Do it from a task in the ASTPipe to make sure to avoid race conditions in the main AST manipulation.
-			UE::Tasks::FTaskEvent EverythingComplete(TEXT("ConstantGenerator_EverythingComplete"));
-			UE::Tasks::FTask LaunchTask = ASTPipe.Launch(TEXT("ConstantGeneratorLaunchTasks"), 
-				[&ConstantSubgraphs, &ASTPipe, &EverythingComplete, &GetRequisites, Pass, InOptions]()
+			UE::Tasks::FTask LaunchTask = UE::Tasks::Launch(TEXT("ConstantGeneratorLaunchTasks"), 
+				[&ConstantSubgraphs, &GetRequisites, &ASTAccessLock, Pass, InOptions]()
 				{
 					MUTABLE_CPUPROFILER_SCOPE(ConstantGenerator_LaunchTasks);
 
+					FScopeLock Lock(&ASTAccessLock);
+
 					FImageOperator ImOp = FImageOperator::GetDefault(InOptions->ImageFormatFunc);
 
 					// Traverse list of constants to generate. It is ordered in a bottom-up way.
@@ -814,9 +815,11 @@ namespace mu
 
 							UE::Tasks::FTaskEvent ReferenceCompletionEvent = InOptions->OptimisationOptions.ReferencedResourceProvider(ImageID, ResolveImage);
 
-							UE::Tasks::FTask CompleteTask = ASTPipe.Launch(TEXT("MutableResolveComplete"),
-								[SubgraphRoot, InOptions, ResolveImage]()
+							UE::Tasks::FTask CompleteTask = UE::Tasks::Launch(TEXT("MutableResolveComplete"),
+								[SubgraphRoot, InOptions, ResolveImage, &ASTAccessLock]()
 								{
+									FScopeLock Lock(&ASTAccessLock);
+
 									Ptr<ASTOpConstantResource> ConstantOp = new ASTOpConstantResource;
 									ConstantOp->type = OP_TYPE::IM_CONSTANT;
 									ConstantOp->SetValue(ResolveImage->get(), InOptions->OptimisationOptions.bUseDiskCache);
@@ -838,8 +841,10 @@ namespace mu
 							FConstantTask* TaskPtr = Task.Get();
 
 							// Launch the preparation on the AST-modification pipe
-							UE::Tasks::FTask PrepareTask = ASTPipe.Launch(TEXT("MutableConstantPrepare"), [TaskPtr]()
+							UE::Tasks::FTask PrepareTask = UE::Tasks::Launch(TEXT("MutableConstantPrepare"), [TaskPtr, &ASTAccessLock]()
 								{
+									FScopeLock Lock(&ASTAccessLock);
+
 									// We need the clone because linking modifies ASTOp state and also to be safe for concurrency.
 									TaskPtr->SourceCloned = ASTOp::DeepClone(TaskPtr->Source);
 								},
@@ -855,8 +860,10 @@ namespace mu
 								LowLevelTasks::ETaskPriority::BackgroundHigh);
 
 							// Launch the completion on the AST-modification pipe
-							UE::Tasks::FTask CompleteTask = ASTPipe.Launch(TEXT("MutableConstantComplete"), [TaskPtr = MoveTemp(Task)]()
+							UE::Tasks::FTask CompleteTask = UE::Tasks::Launch(TEXT("MutableConstantComplete"), [TaskPtr = MoveTemp(Task), &ASTAccessLock]()
 								{
+									FScopeLock Lock(&ASTAccessLock);
+
 									ASTOp::Replace(TaskPtr->Source, TaskPtr->Result);
 									TaskPtr->Source = nullptr;
 									TaskPtr->Result = nullptr;
@@ -869,7 +876,8 @@ namespace mu
 
 						ConstantSubgraphs[Index].Root = nullptr;
 						SubgraphCompletionEvent.Trigger();
-						EverythingComplete.AddPrerequisites(SubgraphCompletionEvent);
+						
+						UE::Tasks::AddNested(SubgraphCompletionEvent);
 					}
 
 				});
@@ -878,8 +886,6 @@ namespace mu
 			{
 				MUTABLE_CPUPROFILER_SCOPE(ConstantGenerator_WaitPending);
 				LaunchTask.Wait();
-				EverythingComplete.Trigger();
-				EverythingComplete.Wait();
 			}
 		}
 
