[Mutable] Add on demand realtime morph loading.

- Modification to the CO BulkData so it supports different data types.
- Now total number of meshes with morphs in a CO is limited to uint16_MAX.
- Asset and data streaming completion events now use the new Tasks API.
- Added typed streameable resource id.

#rnx

diff --git a/Engine/Plugins/Experimental/Mutable/Source/CustomizableObject/Classes/MuCO/CustomizableObject.h b/Engine/Plugins/Experimental/Mutable/Source/CustomizableObject/Classes/MuCO/CustomizableObject.h
--- a/Engine/Plugins/Experimental/Mutable/Source/CustomizableObject/Classes/MuCO/CustomizableObject.h
+++ b/Engine/Plugins/Experimental/Mutable/Source/CustomizableObject/Classes/MuCO/CustomizableObject.h
@@ -9,6 +9,7 @@
 #include "MuCO/CustomizableObjectStreamedResourceData.h"
 #include "MuCO/CustomizableObjectParameterTypeDefinitions.h"
 #include "MuCO/CustomizableObjectUIData.h"
+#include "Templates/TypeCompatibleBytes.h"
 
 #include "CustomizableObject.generated.h"
 
@@ -209,58 +210,6 @@ struct FCompilationOptions
 };
 
 
-USTRUCT()
-struct FMorphTargetInfo
-{
-	GENERATED_USTRUCT_BODY()
-
-	UPROPERTY()
-	FName Name;
-
-	UPROPERTY()
-	int32 LodNum = 0;
-
-	friend FArchive& operator<<(FArchive& Ar, FMorphTargetInfo& Info)
-	{
-		Ar << Info.Name;
-		Ar << Info.LodNum;
-
-		return Ar;
-	}
-};
-
-
-USTRUCT()
-struct FMorphTargetVertexData
-{
-	GENERATED_USTRUCT_BODY()
-
-	UPROPERTY()
-	FVector3f PositionDelta = FVector3f::ZeroVector;
-
-	UPROPERTY()
-	FVector3f TangentZDelta = FVector3f::ZeroVector;
-	
-	UPROPERTY()
-	int32 MorphIndex = 0;
-
-	// Unused padding so memory footprint is the same the aligned.
-	int32 Padding = 0;
-
-	friend FArchive& operator<<(FArchive& Ar, FMorphTargetVertexData& Data)
-	{
-		Ar << Data.PositionDelta;
-		Ar << Data.TangentZDelta;
-		Ar << Data.MorphIndex;
-
-		return Ar;
-	}
-};
-template<> struct TCanBulkSerialize<FMorphTargetVertexData> { enum { Value = true }; };
-
-static_assert(sizeof(FMorphTargetVertexData) == 32, "");
-
-
 // A USTRUCT version of FMeshToMeshVertData in SkeletalMeshTypes.h
 // We are taking advantage of the padding data to store from which asset this data comes from
 // maintaining the same memory footprint than the original.
@@ -378,7 +327,7 @@ struct CUSTOMIZABLEOBJECT_API FMutableStreamableBlock
 
 	UPROPERTY()
 	uint32 Size = 0;
-
+	
 	UPROPERTY()
 	uint64 Offset = 0;
 
@@ -392,7 +341,7 @@ struct CUSTOMIZABLEOBJECT_API FMutableStreamableBlock
 	}
 };
 template<> struct TCanBulkSerialize<FMutableStreamableBlock> { enum { Value = true }; };
-
+static_assert(sizeof(FMutableStreamableBlock) == 8*2);
 
 USTRUCT()
 struct FMutableLODSettings
@@ -457,6 +406,8 @@ public:
 
 	/**  */
 	const FString& GetBulkFilePrefix() const { return BulkFilePrefix; }
+	
+	TUniquePtr<IAsyncReadFileHandle> OpenFileAsyncRead(uint32 FileId) const;
 
 #if WITH_EDITOR
 
@@ -475,9 +426,19 @@ private:
 
 #if WITH_EDITOR
 
+	enum class EDataType : uint8
+	{
+		None = 0,
+		Model,
+		RealTimeMorph
+	};
+
 	struct FBlock
 	{
-		// \TODO: needed?
+		/** Data Type*/
+		EDataType DataType;
+	
+		/** Used on some data types as the index to the block stored in the CustomizableObject */
 		uint32 Id;
 
 		/** Size of the data block. */
@@ -489,6 +450,8 @@ private:
 
 	struct FFile
 	{
+		EDataType DataType;
+
 		/** Id generated from a hash of the file content + offset to avoid collisions. */
 		uint32 Id;
 
@@ -501,7 +464,6 @@ private:
 
 	/** Helper to retrieve the BulkData from within the CookAdditionalFilesOverride */
 	TObjectPtr<UCustomizableObject> CustomizableObject;
-
 #endif
 
 	/** Prefix to locate bulkfiles for loading, using the file ids in each FMutableStreamableBlock. */
@@ -540,12 +502,6 @@ public:
 	UPROPERTY(EditAnywhere, Category = CustomizableObject, meta = (DisplayName = "LOD Settings"))
 	FMutableLODSettings LODSettings;
 	
-	UPROPERTY(Transient)
-	TArray<FMorphTargetInfo> ContributingMorphTargetsInfo;
-	
-	UPROPERTY(Transient)
-	TArray<FMorphTargetVertexData> MorphTargetReconstructionData;
-
 	UPROPERTY(Transient)
 	TArray<FCustomizableObjectClothConfigData> ClothSharedConfigsData;	
 
@@ -835,7 +791,7 @@ public:
 	TArray<FName> LowPriorityTextures;
 
 	/** Map of Hash to Streaming blocks, used to stream a block of data representing a resource from the BulkData */
-	UPROPERTY()
+	UPROPERTY()	
 	TMap<uint64, FMutableStreamableBlock> HashToStreamableBlock; // TODO UE-205600, move to FModelResources
 
 	// Customizable Object Population data start ------------------------------------------------------
diff --git a/Engine/Plugins/Experimental/Mutable/Source/CustomizableObject/Internal/MuCO/CustomizableObjectInstancePrivate.h b/Engine/Plugins/Experimental/Mutable/Source/CustomizableObject/Internal/MuCO/CustomizableObjectInstancePrivate.h
--- a/Engine/Plugins/Experimental/Mutable/Source/CustomizableObject/Internal/MuCO/CustomizableObjectInstancePrivate.h
+++ b/Engine/Plugins/Experimental/Mutable/Source/CustomizableObject/Internal/MuCO/CustomizableObjectInstancePrivate.h
@@ -221,8 +221,8 @@ public:
 
 	void ReuseTexture(UTexture2D* Texture, TSharedRef<FTexturePlatformData, ESPMode::ThreadSafe>& PlatformData);
 
-	// Return an event that will be fired when the assets  have been loaded. It returns null if no asset needs loading.
-	FGraphEventRef LoadAdditionalAssetsAsync(const TSharedRef<FUpdateContextPrivate>& OperationData, FStreamableManager &StreamableManager);
+	/** Returns the task that will be called when all assets and data are loaded, may be already completed if no assets or data needs loading */
+	UE::Tasks::FTask LoadAdditionalAssetsAndDataAsync(const TSharedRef<FUpdateContextPrivate>& OperationData, FStreamableManager &StreamableManager);
 	void AdditionalAssetsAsyncLoaded(UCustomizableObjectInstance* Public);
 
 	void TickUpdateCloseCustomizableObjects(UCustomizableObjectInstance& Public, FMutableInstanceUpdateMap& InOutRequestedUpdates);
@@ -300,7 +300,7 @@ public:
 
 	void SetState(int32 InState);
 
-	void AdditionalAssetsAsyncLoaded(FGraphEventRef CompletionEvent);
+	void AdditionalAssetsAsyncLoaded(UE::Tasks::FTaskEvent CompletionEvent);
 
 	FCustomizableObjectInstanceDescriptor& GetDescriptor() const;
 	
diff --git a/Engine/Plugins/Experimental/Mutable/Source/CustomizableObject/Internal/MuCO/CustomizableObjectPrivate.h b/Engine/Plugins/Experimental/Mutable/Source/CustomizableObject/Internal/MuCO/CustomizableObjectPrivate.h
--- a/Engine/Plugins/Experimental/Mutable/Source/CustomizableObject/Internal/MuCO/CustomizableObjectPrivate.h
+++ b/Engine/Plugins/Experimental/Mutable/Source/CustomizableObject/Internal/MuCO/CustomizableObjectPrivate.h
@@ -73,6 +73,24 @@ struct MutableCompiledDataStreamHeader
 	}
 };
 
+struct FCustomizableObjectStreameableResourceId
+{
+	enum class EType : uint8
+	{
+		None                  = 0,
+		AssetUserData         = 1,
+		RealTimeMorphTarget   = 2,
+	};
+
+	uint32 Id   : 24;
+	uint32 Type : 8;
+
+	friend bool operator==(FCustomizableObjectStreameableResourceId A, FCustomizableObjectStreameableResourceId B)
+	{
+		return BitCast<uint32>(A) == BitCast<uint32>(B);
+	}
+};
+static_assert(sizeof(FCustomizableObjectStreameableResourceId) == sizeof(uint32));
 
 USTRUCT()
 struct FMutableModelParameterValue
@@ -464,6 +482,30 @@ struct CUSTOMIZABLEOBJECT_API FMutableSkinWeightProfileInfo
 #endif
 };
 
+USTRUCT()
+struct FMorphTargetVertexData
+{
+	GENERATED_USTRUCT_BODY()
+
+	UPROPERTY()
+	FVector3f PositionDelta = FVector3f::ZeroVector;
+
+	UPROPERTY()
+	FVector3f TangentZDelta = FVector3f::ZeroVector;
+	
+	UPROPERTY()
+	uint32 MorphNameIndex = 0;
+
+	friend FArchive& operator<<(FArchive& Ar, FMorphTargetVertexData& Data)
+	{
+		Ar << Data.PositionDelta;
+		Ar << Data.TangentZDelta;
+		Ar << Data.MorphNameIndex;
+
+		return Ar;
+	}
+};
+template<> struct TCanBulkSerialize<FMorphTargetVertexData> { enum { Value = true }; };
 
 // Referenced materials, skeletons, passthrough textures...
 USTRUCT()
@@ -471,7 +513,8 @@ struct FModelResources
 {
 	GENERATED_BODY()
 
-	/** All the SkeletalMeshes generated for this CustomizableObject instances will use the Reference Skeletal Mesh
+	/** 
+	 * All the SkeletalMeshes generated for this CustomizableObject instances will use the Reference Skeletal Mesh
 	 * properties for everything that Mutable doesn't create or modify. This struct stores the information used from
 	 * the Reference Skeletal Meshes to avoid having them loaded at all times. This includes data like LOD distances,
 	 * LOD render data settings, Mesh sockets, Bounding volumes, etc.
@@ -528,8 +571,20 @@ struct FModelResources
 	/** State UI metadata information for all the dependencies of this Customizable Object */
 	UPROPERTY()
 	TMap<FString, FParameterUIData> StateUIDataMap;
-};
 
+	UPROPERTY()
+	TArray<FName> RealTimeMorphTargetNames;
+
+	UPROPERTY()
+	TArray<FMutableStreamableBlock> RealTimeMorphStreamableBlocks;
+
+#if WITH_EDITORONLY_DATA
+	UPROPERTY()
+	TArray<FMorphTargetVertexData> EditorOnlyMorphTargetReconstructionData;
+#endif
+	
+
+};
 
 struct CUSTOMIZABLEOBJECT_API FMutableCachedPlatformData
 {
@@ -538,6 +593,9 @@ struct CUSTOMIZABLEOBJECT_API FMutableCachedPlatformData
 
 	/** */
 	TArray64<uint8> StreamableData;
+
+	/** */
+	TArray64<uint8> MorphData;
 };
 
 
@@ -546,6 +604,8 @@ class CUSTOMIZABLEOBJECT_API UCustomizableObjectPrivate : public UObject
 {
 	GENERATED_BODY()
 
+	
+
 	TSharedPtr<mu::Model, ESPMode::ThreadSafe> MutableModel;
 
 	/** Stores resources to be used by MutableModel In-Game. Cooked resources. */
@@ -553,8 +613,10 @@ class CUSTOMIZABLEOBJECT_API UCustomizableObjectPrivate : public UObject
 	FModelResources ModelResources;
 
 #if WITH_EDITORONLY_DATA
-	/** Stores resources to be used by MutableModel in the Editor. Editor resources.
-	 * Editor-Only to avoid packaging assets referenced by editor compilations. */
+	/** 
+	 * Stores resources to be used by MutableModel in the Editor. Editor resources.
+	 * Editor-Only to avoid packaging assets referenced by editor compilations. 
+	 */
 	UPROPERTY(Transient)
 	FModelResources ModelResourcesEditor;
 #endif
@@ -635,7 +697,7 @@ public:
 	void LoadCompiledDataFromDisk();
 
 	/** Cache platform data for cook */
-	void CachePlatformData(const ITargetPlatform* InTargetPlatform, const TArray64<uint8>& InObjectBytes, const TArray64<uint8>& InBulkBytes);
+	void CachePlatformData(const ITargetPlatform* InTargetPlatform, const TArray64<uint8>& InObjectBytes, const TArray64<uint8>& InBulkBytes, const TArray64<uint8>& InMorphBytes);
 	
 	/** Loads data previously compiled in BeginCacheForCookedPlatformData onto the UProperties in *this,
 	  * in preparation for saving the cooked package for *this or for a CustomizableObjectInstance using *this.
@@ -722,6 +784,6 @@ public:
 	// This is a manual version number for the binary blobs in this asset.
 	// Increasing it invalidates all the previously compiled models.
 	// Warning: If while merging code both versions have changed, take the highest+1.
-	static constexpr int32 CurrentSupportedVersion = 429;
+	static constexpr int32 CurrentSupportedVersion = 431;
 };
 
diff --git a/Engine/Plugins/Experimental/Mutable/Source/CustomizableObject/Internal/MuCO/CustomizableObjectSystemPrivate.h b/Engine/Plugins/Experimental/Mutable/Source/CustomizableObject/Internal/MuCO/CustomizableObjectSystemPrivate.h
--- a/Engine/Plugins/Experimental/Mutable/Source/CustomizableObject/Internal/MuCO/CustomizableObjectSystemPrivate.h
+++ b/Engine/Plugins/Experimental/Mutable/Source/CustomizableObject/Internal/MuCO/CustomizableObjectSystemPrivate.h
@@ -296,8 +296,7 @@ struct FMutableTask
 	UE::Tasks::FTask Dependency;
 #endif
 
-	/** From the traditional event graph system, still used to wait for async loads. */
-	FGraphEventRef GraphDependency0;
+	UE::Tasks::FTask GraphDependency0;
 	
 #ifdef MUTABLE_USE_NEW_TASKGRAPH
 	UE::Tasks::FTask GraphDependency1;
@@ -311,12 +310,12 @@ struct FMutableTask
 #ifdef MUTABLE_USE_NEW_TASKGRAPH
 		return (!Dependency.IsValid() || Dependency.IsCompleted())
 			&&
-			(!GraphDependency0 || GraphDependency0->IsComplete())
+			(!GraphDependency0.IsValid() || GraphDependency0.IsCompleted())
 			&&
 			(GraphDependency1.IsCompleted());
 #else
 		return 
-			(!GraphDependency0 || GraphDependency0->IsComplete())
+			(!GraphDependency0.IsValid() || GraphDependency0.IsCompleted())
 			&&
 			(!GraphDependency1 || GraphDependency1->IsComplete());
 #endif
@@ -328,7 +327,7 @@ struct FMutableTask
 #ifdef MUTABLE_USE_NEW_TASKGRAPH
 		Dependency = {};
 #endif
-		GraphDependency0 = nullptr;
+		GraphDependency0 = {};
 
 #ifdef MUTABLE_USE_NEW_TASKGRAPH
 		GraphDependency1 = {};
@@ -431,6 +430,8 @@ struct FInstanceUpdateData
 	TArray<FScalar> Scalars;
 
 	TArray<uint16> BoneMaps;
+	
+	TMap<uint32, TArray<FMorphTargetVertexData>> MorphTargetsVertexData;
 
 	struct FSkeletonData
 	{
diff --git a/Engine/Plugins/Experimental/Mutable/Source/CustomizableObject/Private/MuCO/CustomizableObject.cpp b/Engine/Plugins/Experimental/Mutable/Source/CustomizableObject/Private/MuCO/CustomizableObject.cpp
--- a/Engine/Plugins/Experimental/Mutable/Source/CustomizableObject/Private/MuCO/CustomizableObject.cpp
+++ b/Engine/Plugins/Experimental/Mutable/Source/CustomizableObject/Private/MuCO/CustomizableObject.cpp
@@ -363,8 +363,6 @@ void UCustomizableObjectPrivate::ClearCompiledData(bool bIsCooking)
 {
 	GetModelResources(bIsCooking) = FModelResources();
 
-	GetPublic()->ContributingMorphTargetsInfo.Empty();
-	GetPublic()->MorphTargetReconstructionData.Empty();
 	GetPublic()->ClothMeshToMeshVertData.Empty();
 	GetPublic()->ContributingClothingAssetsData.Empty();
 	GetPublic()->ClothSharedConfigsData.Empty();
@@ -390,8 +388,8 @@ void SerializeStreamedResources(FArchive& Ar, UObject* Object, TArray<FCustomiza
 		for (const FCustomizableObjectStreamedResourceData& ResourceData : StreamedResources)
 		{
 			const FCustomizableObjectResourceData& Data = ResourceData.GetLoadedData();
-			uint32 Type = (uint32)Data.Type;
-			Ar << Type;
+			uint32 ResourceDataType = (uint32)Data.Type;
+			Ar << ResourceDataType;
 
 			switch (Data.Type)
 			{
@@ -560,9 +558,9 @@ void UCustomizableObjectPrivate::SaveCompiledData(FArchive& MemoryWriter, bool b
 	MemoryWriter << LocalModelResources.ParameterUIDataMap;
 	MemoryWriter << LocalModelResources.StateUIDataMap;
 
-	MemoryWriter << GetPublic()->ContributingMorphTargetsInfo;
-	MemoryWriter << GetPublic()->MorphTargetReconstructionData;
-	
+	MemoryWriter << LocalModelResources.RealTimeMorphTargetNames;
+	MemoryWriter << LocalModelResources.RealTimeMorphStreamableBlocks;
+
 	MemoryWriter << GetPublic()->ClothMeshToMeshVertData;
 	MemoryWriter << GetPublic()->ContributingClothingAssetsData;
 	MemoryWriter << GetPublic()->ClothSharedConfigsData; 
@@ -583,6 +581,11 @@ void UCustomizableObjectPrivate::SaveCompiledData(FArchive& MemoryWriter, bool b
 	MemoryWriter << CustomizableObjectPathMap;
 	MemoryWriter << GroupNodeMap;
 	MemoryWriter << ParticipatingObjects;
+
+	if (!bIsCooking)
+	{
+		MemoryWriter << LocalModelResources.EditorOnlyMorphTargetReconstructionData;
+	}
 }
 
 void UCustomizableObjectPrivate::LoadCompiledData(FArchive& MemoryReader, const ITargetPlatform* InTargetPlatform, bool bIsCooking)
@@ -655,6 +658,7 @@ void UCustomizableObjectPrivate::LoadCompiledData(FArchive& MemoryReader, const
 			LocalModelResource.PhysicsAssets.Add(TSoftObjectPtr<UPhysicsAsset>(FSoftObjectPath(StringRef)));
 		}
 
+
 		int32 NumAnimBps = 0;
 		MemoryReader << NumAnimBps;
 
@@ -678,9 +682,8 @@ void UCustomizableObjectPrivate::LoadCompiledData(FArchive& MemoryReader, const
 		MemoryReader << LocalModelResource.ParameterUIDataMap;
 		MemoryReader << LocalModelResource.StateUIDataMap;
 
-
-		MemoryReader << GetPublic()->ContributingMorphTargetsInfo;
-		MemoryReader << GetPublic()->MorphTargetReconstructionData;
+		MemoryReader << LocalModelResource.RealTimeMorphTargetNames;
+		MemoryReader << LocalModelResource.RealTimeMorphStreamableBlocks;
 
 		MemoryReader << GetPublic()->ClothMeshToMeshVertData;
 		MemoryReader << GetPublic()->ContributingClothingAssetsData;
@@ -706,6 +709,11 @@ void UCustomizableObjectPrivate::LoadCompiledData(FArchive& MemoryReader, const
 			MemoryReader << GroupNodeMap;
 			MemoryReader << ParticipatingObjects;
 
+			if (!bIsCooking)
+			{
+				MemoryReader << LocalModelResource.EditorOnlyMorphTargetReconstructionData;
+			}
+
 			if (!bIsCooking)
 			{
 				const FAssetRegistryModule& AssetRegistryModule = FModuleManager::LoadModuleChecked<FAssetRegistryModule>("AssetRegistry");
@@ -741,9 +749,9 @@ void UCustomizableObjectPrivate::LoadCompiledData(FArchive& MemoryReader, const
 
 		if (bModelSerialized && !bInvalidateModel)
 		{
-			UnrealMutableInputStream stream(MemoryReader);
-			mu::InputArchive arch(&stream);
-			LoadedModel = mu::Model::StaticUnserialise( arch );
+			UnrealMutableInputStream Stream(MemoryReader);
+			mu::InputArchive Arch(&Stream);
+			LoadedModel = mu::Model::StaticUnserialise(Arch);
 		}
 	}
 	
@@ -821,7 +829,7 @@ void UCustomizableObjectPrivate::LoadCompiledDataFromDisk()
 }
 
 
-void UCustomizableObjectPrivate::CachePlatformData(const ITargetPlatform* InTargetPlatform, const TArray64<uint8>& InModelBytes, const TArray64< uint8>& InBulkBytes)
+void UCustomizableObjectPrivate::CachePlatformData(const ITargetPlatform* InTargetPlatform, const TArray64<uint8>& InModelBytes, const TArray64<uint8>& InBulkBytes, const TArray64<uint8>& InMorphBytes)
 {
 	MUTABLE_CPUPROFILER_SCOPE(CachePlatformData)
 
@@ -837,8 +845,11 @@ void UCustomizableObjectPrivate::CachePlatformData(const ITargetPlatform* InTarg
 	Data.ModelData.Append(InModelBytes);
 
 	// Cache streamable bulk data
-	Data.StreamableData.SetNumUninitialized(InBulkBytes.Num(),EAllowShrinking::No);
-	FMemory::Memcpy(Data.StreamableData.GetData(), InBulkBytes.GetData(), InBulkBytes.Num() );
+	Data.StreamableData.SetNumUninitialized(InBulkBytes.Num(), EAllowShrinking::No);
+	FMemory::Memcpy(Data.StreamableData.GetData(), InBulkBytes.GetData(), InBulkBytes.Num());
+
+	Data.MorphData.SetNumUninitialized(InMorphBytes.Num(), EAllowShrinking::No);
+	FMemory::Memcpy(Data.MorphData.GetData(), InMorphBytes.GetData(), Data.MorphData.Num());
 }
 
 
@@ -1006,13 +1017,7 @@ void UCustomizableObjectPrivate::SaveEmbeddedData(FArchive& Ar)
 	Ar << InternalVersion;
 
 	if (GetModel())
-	{
-		// Serialize morph data
-		{
-			Ar << GetPublic()->ContributingMorphTargetsInfo;
-			Ar << GetPublic()->MorphTargetReconstructionData;
-		}
-		
+	{	
 		{
 			Ar << GetPublic()->ClothMeshToMeshVertData;
 			Ar << GetPublic()->ContributingClothingAssetsData;
@@ -1023,9 +1028,9 @@ void UCustomizableObjectPrivate::SaveEmbeddedData(FArchive& Ar)
 		{
 			GetModel()->UnloadExternalData();
 
-			UnrealMutableOutputStream stream(Ar);
-			mu::OutputArchive arch(&stream);
-			mu::Model::Serialise(GetModel().Get(), arch);
+			UnrealMutableOutputStream Stream(Ar);
+			mu::OutputArchive Arch(&Stream);
+			mu::Model::Serialise(GetModel().Get(), Arch);
 		}
 
 		UE_LOG(LogMutable, Verbose, TEXT("Saved embedded data for Customizable Object [%s] now at position %d."), *GetName(), int(Ar.Tell()));
@@ -1046,13 +1051,7 @@ void UCustomizableObjectPrivate::LoadEmbeddedData(FArchive& Ar)
 	check(CurrentSupportedVersion == InternalVersion);
 
 	if(CurrentSupportedVersion == InternalVersion)
-	{
-		// Load morph data
-		{
-			Ar << GetPublic()->ContributingMorphTargetsInfo;
-			Ar << GetPublic()->MorphTargetReconstructionData;
-		}
-		
+	{	
 		{
 			Ar << GetPublic()->ClothMeshToMeshVertData;
 			Ar << GetPublic()->ContributingClothingAssetsData;
@@ -1060,14 +1059,14 @@ void UCustomizableObjectPrivate::LoadEmbeddedData(FArchive& Ar)
 		}
 		
 		// Load model
-		UnrealMutableInputStream stream(Ar);
-		mu::InputArchive arch(&stream);
-		TSharedPtr<mu::Model, ESPMode::ThreadSafe> Model = mu::Model::StaticUnserialise( arch );
+		UnrealMutableInputStream Stream(Ar);
+		mu::InputArchive Arch(&Stream);
+		TSharedPtr<mu::Model, ESPMode::ThreadSafe> Model = mu::Model::StaticUnserialise(Arch);
 
 		// Create parameter properties
 		UpdateParameterPropertiesFromModel(Model);
 
-		SetModel( Model, FGuid());
+		SetModel(Model, FGuid());
 	}
 }
 
@@ -1934,6 +1933,7 @@ TSharedPtr<const mu::Model, ESPMode::ThreadSafe> UCustomizableObjectPrivate::Get
 	return MutableModel;
 }
 
+
 const FModelResources& UCustomizableObjectPrivate::GetModelResources() const
 {
 #if WITH_EDITORONLY_DATA
@@ -2025,6 +2025,19 @@ void UCustomizableObjectBulk::PostLoad()
 	BulkFilePrefix = PackageFilename;
 }
 
+TUniquePtr<IAsyncReadFileHandle> UCustomizableObjectBulk::OpenFileAsyncRead(uint32 FileId) const
+{
+	check(IsInGameThread());
+
+	FString FilePath = FString::Printf(TEXT("%s-%08x.mut"), *BulkFilePrefix, FileId);
+
+	IAsyncReadFileHandle* Result = FPlatformFileManager::Get().GetPlatformFile().OpenAsyncRead(*FilePath);
+	
+	// Result being null does not mean the file does not exist. A request has to be made. Let the callee deal with it.
+	//UE_CLOG(!Result, LogMutable, Warning, TEXT("CustomizableObjectBulkData: Failed to open file [%s]."), *FilePath);
+
+	return TUniquePtr<IAsyncReadFileHandle>(Result);
+}
 
 #if WITH_EDITOR
 
@@ -2043,28 +2056,40 @@ void UCustomizableObjectBulk::CookAdditionalFilesOverride(const TCHAR* PackageFi
 	FMutableCachedPlatformData* PlatformData = CustomizableObject->GetPrivate()->CachedPlatformsData.Find(TargetPlatform->PlatformName());
 	check(PlatformData);
 
-	// Source data pointer
-	const uint8* SourceData = PlatformData->StreamableData.GetData();
-
 	const int32 NumBulkDataFiles = BulkDataFiles.Num();
 	for(int32 FileIndex = 0; FileIndex < NumBulkDataFiles; ++FileIndex)
 	{
 		const FFile& CurrentFile = BulkDataFiles[FileIndex];
 
+		const uint8* SourceData = nullptr;
+		if (CurrentFile.DataType == EDataType::Model)
+		{
+			SourceData = PlatformData->StreamableData.GetData();
+		}
+		else if (CurrentFile.DataType == EDataType::RealTimeMorph)
+		{
+			SourceData = PlatformData->MorphData.GetData();
+		}
+		else
+		{
+			checkf(false, TEXT("Unknown file DataType found."));
+			continue;
+		}
+
 		int64 FileSize = 0;
-		for ( const FBlock& Block: CurrentFile.Blocks )
+		for (const FBlock& Block : CurrentFile.Blocks)
 		{
 			FileSize += Block.Size;
 		}
 
 		// Generate the bulk data file in memory
 		TArray64<uint8> FileBulkData;
-		FileBulkData.SetNum( FileSize );
+		FileBulkData.SetNum(FileSize);
 		uint8* FileData = FileBulkData.GetData();
 
 		for (const FBlock& Block : CurrentFile.Blocks)
 		{
-			FMemory::Memcpy(FileData,SourceData+Block.Offset,Block.Size);
+			FMemory::Memcpy(FileData, SourceData + Block.Offset, Block.Size);
 			FileData += Block.Size;
 		}
 
@@ -2076,7 +2101,6 @@ void UCustomizableObjectBulk::CookAdditionalFilesOverride(const TCHAR* PackageFi
 	}
 }
 
-
 void UCustomizableObjectBulk::PrepareBulkData(UCustomizableObject* InOuter, const ITargetPlatform* TargetPlatform)
 {
 	CustomizableObject = InOuter;
@@ -2094,7 +2118,6 @@ void UCustomizableObjectBulk::PrepareBulkData(UCustomizableObject* InOuter, cons
 	const uint64 MaxChunkSize = UCustomizableObjectSystem::GetInstance()->GetMaxChunkSizeForPlatform(TargetPlatform);
 	TargetBulkDataFileBytes = FMath::Min(TargetBulkDataFileBytes, MaxChunkSize);
 
-	const int32 NumBlocks = Model->GetRomCount();
 
 	// TODO:
 	// To avoid influence of the order of the streamed data (their index), classify it recursively based on hash values
@@ -2109,44 +2132,88 @@ void UCustomizableObjectBulk::PrepareBulkData(UCustomizableObject* InOuter, cons
 	};
 
 	FClassifyNode RootNode;
-	RootNode.Blocks.Reserve(NumBlocks);
-
 	// Create blocks data, filtering out the ones that are too big and will go on its own file in any case.
-	uint64 SourceOffset = 0;
-	for (int32 BlockIndex = 0; BlockIndex < NumBlocks; ++BlockIndex)
 	{
-		uint32 BlockId = Model->GetRomId(BlockIndex);
-		const FMutableStreamableBlock& StreamableBlock = CustomizableObject->HashToStreamableBlock[BlockId];
-		const uint32 BlockSize = StreamableBlock.Size;
+		const int32 NumBlocks = Model->GetRomCount();
+		RootNode.Blocks.Reserve(NumBlocks);
+		
+		uint64 SourceOffset = 0;
+		for (int32 BlockIndex = 0; BlockIndex < NumBlocks; ++BlockIndex)
+		{
+			uint32 BlockId = Model->GetRomId(BlockIndex);
+			const FMutableStreamableBlock& StreamableBlock = CustomizableObject->HashToStreamableBlock[BlockId];
+			const uint32 BlockSize = StreamableBlock.Size;
 
-		FBlock CurrentBlock = { BlockId, BlockSize, SourceOffset };
+			FBlock CurrentBlock = { EDataType::Model, BlockId, BlockSize, SourceOffset };
+			if (BlockSize > TargetBulkDataFileBytes)
+			{
+				// It will go to its own file
+				BulkDataFiles.Add(FFile{EDataType::Model, 0, {CurrentBlock}});
+			}
+			else
+			{
+				// It may merge with other small blocks
+				RootNode.Blocks.Add(CurrentBlock);
+			}
 
-		if (BlockSize > TargetBulkDataFileBytes)
-		{
-			// It will go to its own file
-			BulkDataFiles.Add(FFile{ 0, {CurrentBlock} });
+			SourceOffset += BlockSize;
 		}
-		else
+	}
+
+	// TODO: This should create a new classification branch when the tree is implemented.
+	// For now append after Model roms. 
+	{
+		uint64 SourceOffset = 0;
+		
+		constexpr bool bGetCooked = false;
+		const TArray<FMutableStreamableBlock>& RealTimeMorphTargetsBlocks = 
+					CustomizableObject->GetPrivate()->GetModelResources(bGetCooked).RealTimeMorphStreamableBlocks;
+
+		const int32 NumBlocks = RealTimeMorphTargetsBlocks.Num();
+		for (int32 BlockIndex = 0; BlockIndex < NumBlocks; ++BlockIndex)
 		{
-			// It may merge with other small blocks
-			RootNode.Blocks.Add(CurrentBlock);
-		}
+			const FMutableStreamableBlock& StreamableBlock = RealTimeMorphTargetsBlocks[BlockIndex];
+			const uint32 BlockSize = StreamableBlock.Size;
+			
+			check(SourceOffset == StreamableBlock.Offset);
+			FBlock CurrentBlock = { EDataType::RealTimeMorph, BlockIndex, BlockSize, SourceOffset };
+			
+			if (BlockSize > TargetBulkDataFileBytes)
+			{
+				// It will go to its own file
+				BulkDataFiles.Add(FFile{ EDataType::RealTimeMorph, 0, {CurrentBlock} });
+			}
+			else
+			{
+				// It may merge with other small blocks
+				RootNode.Blocks.Add(CurrentBlock);
+			}
 
-		SourceOffset += BlockSize;
+			SourceOffset += BlockSize;
+		}
 	}
-	
+
 	// Temp: Group by order in the array
 	for (int32 BlockIndex = 0; BlockIndex < RootNode.Blocks.Num(); )
 	{
-		FFile CurrentFile;
 		int32 CurrentFileSize = 0;
+		
+		FFile CurrentFile;
+		CurrentFile.DataType = RootNode.Blocks[BlockIndex].DataType;
 
 		while(BlockIndex < RootNode.Blocks.Num())
 		{
 			FBlock CurrentBlock = RootNode.Blocks[BlockIndex];
 
 			// Next file?
-			if (CurrentFileSize>0 && CurrentFileSize + CurrentBlock.Size > TargetBulkDataFileBytes)
+			// Store different data types in different files. Blocks should be sorted by DataType
+			// so data is properly packeted
+			if (CurrentFile.DataType != CurrentBlock.DataType)
+			{
+				break;
+			}
+			
+			if (CurrentFileSize > 0 && CurrentFileSize + CurrentBlock.Size > TargetBulkDataFileBytes)
 			{
 				break;
 			}
@@ -2163,7 +2230,7 @@ void UCustomizableObjectBulk::PrepareBulkData(UCustomizableObject* InOuter, cons
 	}
 
 	// Create the file list
-	for ( int32 FileIndex=0; FileIndex<BulkDataFiles.Num(); ++FileIndex )
+	for (int32 FileIndex = 0; FileIndex < BulkDataFiles.Num(); ++FileIndex)
 	{
 		// Generate the id for this file
 		FFile& CurrentFile = BulkDataFiles[FileIndex];
@@ -2171,7 +2238,7 @@ void UCustomizableObjectBulk::PrepareBulkData(UCustomizableObject* InOuter, cons
 		for (int32 FileBlockIndex = 0; FileBlockIndex < CurrentFile.Blocks.Num(); ++FileBlockIndex)
 		{
 			FBlock ThisBlock = CurrentFile.Blocks[FileBlockIndex];
-			FileId = (FileBlockIndex == 0) ? ThisBlock.Id : HashCombine( FileId, ThisBlock.Id);
+			FileId = (FileBlockIndex == 0) ? ThisBlock.Id : HashCombine(HashCombine(FileId, ThisBlock.Id), uint32(CurrentFile.DataType));
 		}
 
 		// Ensure the FileId is unique
@@ -2193,17 +2260,43 @@ void UCustomizableObjectBulk::PrepareBulkData(UCustomizableObject* InOuter, cons
 		// Set it to the editor-only file descriptor
 		CurrentFile.Id = FileId;
 
-		// Set it to all streamable blocks
-		uint32 OffsetInFile = 0;
-		for (int32 FileBlockIndex = 0; FileBlockIndex < CurrentFile.Blocks.Num(); ++FileBlockIndex)
+		if (CurrentFile.DataType == EDataType::Model)
 		{
-			FBlock ThisBlock = CurrentFile.Blocks[FileBlockIndex];
+			// Set it to all streamable blocks
+			uint32 OffsetInFile = 0;
+			for (int32 FileBlockIndex = 0; FileBlockIndex < CurrentFile.Blocks.Num(); ++FileBlockIndex)
+			{
+				FBlock ThisBlock = CurrentFile.Blocks[FileBlockIndex];
 
-			FMutableStreamableBlock& StreamableBlock = CustomizableObject->HashToStreamableBlock[ThisBlock.Id];
-			check(StreamableBlock.Size == ThisBlock.Size);
-			StreamableBlock.FileId = FileId;
-			StreamableBlock.Offset = OffsetInFile;
-			OffsetInFile += ThisBlock.Size;
+				FMutableStreamableBlock& StreamableBlock = CustomizableObject->HashToStreamableBlock[ThisBlock.Id];
+				check(StreamableBlock.Size == ThisBlock.Size);
+				StreamableBlock.FileId = FileId;
+				StreamableBlock.Offset = OffsetInFile;
+				OffsetInFile += ThisBlock.Size;
+			}
+		}
+		else if (CurrentFile.DataType == EDataType::RealTimeMorph)
+		{
+			constexpr bool bGetCooked = true;
+			TArray<FMutableStreamableBlock>& MorphBlocks = 
+					CustomizableObject->GetPrivate()->GetModelResources(bGetCooked).RealTimeMorphStreamableBlocks;
+			// Set it to all streamable blocks
+			uint32 OffsetInFile = 0;
+			for (int32 FileBlockIndex = 0; FileBlockIndex < CurrentFile.Blocks.Num(); ++FileBlockIndex)
+			{
+				FBlock ThisBlock = CurrentFile.Blocks[FileBlockIndex];
+
+				FMutableStreamableBlock& StreamableBlock = MorphBlocks[ThisBlock.Id];
+				check(StreamableBlock.Size == ThisBlock.Size);
+				StreamableBlock.FileId = FileId;
+				StreamableBlock.Offset = OffsetInFile;
+				OffsetInFile += ThisBlock.Size;
+			}
+		}
+		else
+		{
+			UE_LOG(LogMutable, Error, TEXT("Unknown DataType found while fixing streaming block files ids."));
+			check(false);
 		}
 	}
 }
diff --git a/Engine/Plugins/Experimental/Mutable/Source/CustomizableObject/Private/MuCO/CustomizableObjectInstance.cpp b/Engine/Plugins/Experimental/Mutable/Source/CustomizableObject/Private/MuCO/CustomizableObjectInstance.cpp
--- a/Engine/Plugins/Experimental/Mutable/Source/CustomizableObject/Private/MuCO/CustomizableObjectInstance.cpp
+++ b/Engine/Plugins/Experimental/Mutable/Source/CustomizableObject/Private/MuCO/CustomizableObjectInstance.cpp
@@ -1622,6 +1622,8 @@ void SetMeshUVChannelDensity(FMeshUVChannelInfo& UVChannelInfo, float Density =
 
 bool UCustomizableInstancePrivate::DoComponentsNeedUpdate(UCustomizableObjectInstance* Public, const TSharedRef<FUpdateContextPrivate>& OperationData, bool& bHasInvalidMesh)
 {
+	MUTABLE_CPUPROFILER_SCOPE(UCustomizableInstancePrivate::DoComponentsNeedUpdate);
+
 	UCustomizableObject* CustomizableObject = Public->GetCustomizableObject();
 	check(CustomizableObject);
 
@@ -3523,32 +3525,33 @@ void UCustomizableInstancePrivate::BuildOrCopyMorphTargetsData(const TSharedRef<
 	}
 
 	const UCustomizableObject* CustomizableObject = CustomizableObjectInstance->GetCustomizableObject();
-	const TArray<FMorphTargetInfo>& ContributingMorphTargetsInfo = CustomizableObject->ContributingMorphTargetsInfo;
-	const TArray<FMorphTargetVertexData>& MorphTargetReconstructionData = CustomizableObject->MorphTargetReconstructionData;
+	const TArray<FName>& MorphTargetNames = CustomizableObject->GetPrivate()->GetModelResources().RealTimeMorphTargetNames;
+
+	const TMap<uint32, TArray<FMorphTargetVertexData>>& ResourceIdToVertexDataMap = 
+			OperationData->InstanceUpdateData.MorphTargetsVertexData;
 
-	if (!(MorphTargetReconstructionData.Num() && ContributingMorphTargetsInfo.Num()))
+	if (MorphTargetNames.IsEmpty() || ResourceIdToVertexDataMap.IsEmpty())
 	{
 		return;
 	}
 
 	TArray<int32> SectionMorphTargetVertices;
-	SectionMorphTargetVertices.SetNumZeroed(ContributingMorphTargetsInfo.Num());
 
-	SkeletalMesh->GetMorphTargets().Empty(ContributingMorphTargetsInfo.Num());
-	for (const FMorphTargetInfo& MorphTargetInfo : ContributingMorphTargetsInfo)
-	{
-		UMorphTarget* NewMorphTarget = NewObject<UMorphTarget>(SkeletalMesh, MorphTargetInfo.Name);
+	SectionMorphTargetVertices.SetNumZeroed(MorphTargetNames.Num());
 
-		NewMorphTarget->BaseSkelMesh = SkeletalMesh;
-		NewMorphTarget->GetMorphLODModels().SetNum(MorphTargetInfo.LodNum);
-		SkeletalMesh->GetMorphTargets().Add(NewMorphTarget);
-	}
+	TArray<uint8> UsedMorphSet;
+	UsedMorphSet.SetNum(MorphTargetNames.Num());
+
+	TArray<TArray<FMorphTargetLODModel>> MorphsData;
+	MorphsData.SetNum(MorphTargetNames.Num());
 
-	const int32 SkeletalMeshMorphTargetsNum = SkeletalMesh->GetMorphTargets().Num();
+	const int32 NumLODs = OperationData->NumLODsAvailable;
 
 	int32 LastValidLODIndex = OperationData->NumLODsAvailable - 1;
 	for (int32 LODIndex = LastValidLODIndex; LODIndex >= FirstLODAvailable; --LODIndex)
 	{
+		int32 NumNotFoundLoadedMorphsResources = 0;
+
 		const FInstanceUpdateData::FLOD& LOD = OperationData->InstanceUpdateData.LODs[LODIndex];
 
 		if (LODIndex >= OperationData->GetMinLOD() && OperationData->InstanceUpdateData.Components[LOD.FirstComponent + ComponentIndex].bGenerated)
@@ -3570,7 +3573,10 @@ void UCustomizableInstancePrivate::BuildOrCopyMorphTargetsData(const TSharedRef<
 			int32 VertexMorphsCountBufferIndex, VertexMorphsCountBufferChannel;
 			MeshSet.FindChannel(mu::MBS_OTHER, 1, &VertexMorphsCountBufferIndex, &VertexMorphsCountBufferChannel);
 
-			if (VertexMorphsInfoIndexBufferIndex < 0 || VertexMorphsCountBufferIndex < 0)
+			int32 VertexMorphsResourceIdBufferIndex, VertexMorphsResourceIdBufferChannel;
+			MeshSet.FindChannel(mu::MBS_OTHER, 2, &VertexMorphsResourceIdBufferIndex, &VertexMorphsResourceIdBufferChannel);
+
+			if (VertexMorphsInfoIndexBufferIndex < 0 || VertexMorphsCountBufferIndex < 0 || VertexMorphsResourceIdBufferIndex < 0)
 			{
 				continue;
 			}
@@ -3578,45 +3584,119 @@ void UCustomizableInstancePrivate::BuildOrCopyMorphTargetsData(const TSharedRef<
 			const int32* const VertexMorphsInfoIndexBuffer = reinterpret_cast<const int32*>(MeshSet.GetBufferData(VertexMorphsInfoIndexBufferIndex));
 			TArrayView<const int32> VertexMorphsInfoIndexView(VertexMorphsInfoIndexBuffer, MeshSet.GetElementCount());
 
-			const int32* const VertexMorphsCountBuffer = reinterpret_cast<const int32*>(MeshSet.GetBufferData(VertexMorphsCountBufferIndex));
-			TArrayView<const int32> VertexMorphsCountView(VertexMorphsCountBuffer, MeshSet.GetElementCount());
+			const uint16* const VertexMorphsCountBuffer = reinterpret_cast<const uint16*>(MeshSet.GetBufferData(VertexMorphsCountBufferIndex));
+			TArrayView<const uint16> VertexMorphsCountView(VertexMorphsCountBuffer, MeshSet.GetElementCount());
+			
+			const uint16* const VertexMorphsResourceIdBuffer = reinterpret_cast<const uint16*>(MeshSet.GetBufferData(VertexMorphsResourceIdBufferIndex));
+			TArrayView<const uint16> VertexMorphsResourceIdView(VertexMorphsResourceIdBuffer, MeshSet.GetElementCount());
 
 			const int32 SurfaceCount = Component.Mesh->GetSurfaceCount();
 			for (int32 Section = 0; Section < SurfaceCount; ++Section)
 			{
 				// Reset SectionMorphTargets.
-				for (auto& Elem : SectionMorphTargetVertices)
+				for (int32& Elem : SectionMorphTargetVertices)
 				{
 					Elem = 0;
 				}
 
-				int FirstVertex, VerticesCount, FirstIndex, IndiciesCount;
+				int32 FirstVertex, VerticesCount, FirstIndex, IndiciesCount;
 				Component.Mesh->GetSurface(Section, &FirstVertex, &VerticesCount, &FirstIndex, &IndiciesCount, nullptr, nullptr, nullptr);
 
-				for (int32 VertexIdx = FirstVertex; VertexIdx < FirstVertex + VerticesCount; ++VertexIdx)
+				for (int32 VertexIdx = FirstVertex; VertexIdx < FirstVertex + VerticesCount;)
 				{
-					const int32 MorphCount = VertexMorphsCountView[VertexIdx];
-					if (MorphCount <= 0)
+					// Find a span with the same VertexMorphResourceId to amortise the cost of finding 
+					// in the loaded resources map. It is expected to find large consecutive mesh sections pointing to
+					// the same loaded resource.
+					
+					const int32 SpanStart = VertexIdx++;
+					const uint16 CurrentResourceId = VertexMorphsResourceIdView[SpanStart];
+					
+					// Vertex with no morphs are marked with TNumericLimits<uint16>::Max(), skip vertex if the case.
+					if (CurrentResourceId == TNumericLimits<uint16>::Max())
 					{
 						continue;
 					}
 
-					TArrayView<const FMorphTargetVertexData> MorphsVertexDataView(&(MorphTargetReconstructionData[VertexMorphsInfoIndexView[VertexIdx]]), MorphCount);
-					for (const FMorphTargetVertexData& SourceVertex : MorphsVertexDataView)
+					for (; VertexIdx < FirstVertex + VerticesCount; ++VertexIdx)
+					{
+						const int32 VertexResourceId = VertexMorphsResourceIdView[VertexIdx];
+						// we can skip vertices with no morph without breaking the span.
+						if (VertexResourceId == TNumericLimits<uint16>::Max())
+						{
+							continue;
+						}
+
+						if (CurrentResourceId != VertexResourceId)
+						{
+							break;
+						}
+					}
+					const int32 SpanEnd = VertexIdx;
+
+					const TArray<FMorphTargetVertexData>* MorphTargetReconstructionData = ResourceIdToVertexDataMap.Find(CurrentResourceId); 
+					
+					if (!MorphTargetReconstructionData)
+					{
+						++NumNotFoundLoadedMorphsResources;
+						continue;
+					}
+
+					const TArray<FMorphTargetVertexData>& SpanMorphData = *MorphTargetReconstructionData;
+
+
+					// This assumes the number of vertex in an span will be large compared to the number of
+					// morphs. Maybe the allocation could be done in a different pass as an optimization.
+					FMemory::Memzero(UsedMorphSet.GetData(), UsedMorphSet.Num());
+					for (int32 SpanVertexIdx = SpanStart; SpanVertexIdx < SpanEnd; ++SpanVertexIdx)
+					{
+						const int32 MorphCount = VertexMorphsCountView[SpanVertexIdx];
+
+						TArrayView<const FMorphTargetVertexData> MorphsVertexDataView = MakeArrayView(
+								SpanMorphData.GetData() + VertexMorphsInfoIndexView[SpanVertexIdx], 
+								MorphCount);
+
+						for (const FMorphTargetVertexData& MorphVertexData : MorphsVertexDataView)
+						{
+							check((uint32)UsedMorphSet.Num() > MorphVertexData.MorphNameIndex);
+							UsedMorphSet[MorphVertexData.MorphNameIndex] = 1;
+						}
+					}
+
+					const int32 NumMorphs = UsedMorphSet.Num();
+					for (int32 MorphIndex = 0; MorphIndex < NumMorphs; ++MorphIndex)
 					{
-						// check(SkeletalMeshMorphTargetsNum <= SourceVertex.MorphIndex);
-						if (SkeletalMeshMorphTargetsNum <= SourceVertex.MorphIndex)
+						if (!UsedMorphSet[MorphIndex])
 						{
 							continue;
 						}
 
-						FMorphTargetLODModel& DestMorphLODModel = SkeletalMesh->GetMorphTargets()[SourceVertex.MorphIndex]->GetMorphLODModels()[LODIndex];
+						if (MorphsData[MorphIndex].IsEmpty())
+						{
+							MorphsData[MorphIndex].SetNum(NumLODs);
+						}
+					}
 
-						DestMorphLODModel.Vertices.Emplace(
-							FMorphTargetDelta{ SourceVertex.PositionDelta, SourceVertex.TangentZDelta, static_cast<uint32>(VertexIdx) });
+					for (int32 SpanVertexIdx = SpanStart; SpanVertexIdx < SpanEnd; ++SpanVertexIdx)
+					{
+						const uint16 MorphCount = VertexMorphsCountView[SpanVertexIdx];
+						if (MorphCount == 0)
+						{
+							continue;
+						}
 
-						++SectionMorphTargetVertices[SourceVertex.MorphIndex];
+						TArrayView<const FMorphTargetVertexData> MorphsVertexDataView = MakeArrayView(
+								SpanMorphData.GetData() + VertexMorphsInfoIndexView[SpanVertexIdx], 
+								MorphCount);
 
+						for (const FMorphTargetVertexData& SourceVertex : MorphsVertexDataView)
+						{
+							FMorphTargetLODModel& DestMorphLODModel = MorphsData[SourceVertex.MorphNameIndex][LODIndex];
+
+							DestMorphLODModel.Vertices.Emplace(
+									FMorphTargetDelta{ SourceVertex.PositionDelta, SourceVertex.TangentZDelta, static_cast<uint32>(SpanVertexIdx) });
+
+							++SectionMorphTargetVertices[SourceVertex.MorphNameIndex];
+						}
 					}
 				}
 
@@ -3625,13 +3705,34 @@ void UCustomizableInstancePrivate::BuildOrCopyMorphTargetsData(const TSharedRef<
 				{
 					if (SectionMorphTargetVertices[MorphIdx] > 0)
 					{
-						FMorphTargetLODModel& MorphTargetLodModel = SkeletalMesh->GetMorphTargets()[MorphIdx]->GetMorphLODModels()[LODIndex];
+						FMorphTargetLODModel& MorphTargetLodModel = MorphsData[MorphIdx][LODIndex];
 
 						MorphTargetLodModel.SectionIndices.Add(Section);
-						MorphTargetLodModel.NumVertices = SectionMorphTargetVertices[MorphIdx];
+						MorphTargetLodModel.NumVertices += SectionMorphTargetVertices[MorphIdx];
 					}
 				}
 			}
+
+			if (NumNotFoundLoadedMorphsResources > 0)
+			{
+				UE_LOG(LogMutable, Warning, TEXT("Needed realtime morph reconstruction data was not loaded properly. Some realtime morphs may not work correctly."));
+			}
+			
+			// Generate the SkeletalMesh data structures. The previous step could be done in an async task.
+			SkeletalMesh->GetMorphTargets().Empty();
+			const int32 NumMorphs = MorphTargetNames.Num(); 
+			for (int32 I = 0; I < NumMorphs; ++I)
+			{
+				if (MorphsData[I].IsEmpty())
+				{
+					continue;
+				} 
+
+				UMorphTarget* NewMorphTarget = NewObject<UMorphTarget>(SkeletalMesh, MorphTargetNames[I]);
+				NewMorphTarget->BaseSkelMesh = SkeletalMesh;
+				NewMorphTarget->GetMorphLODModels() = MoveTemp(MorphsData[I]);
+				SkeletalMesh->GetMorphTargets().Add(NewMorphTarget);
+			}
 		}
 		else
 		{
@@ -3641,7 +3742,8 @@ void UCustomizableInstancePrivate::BuildOrCopyMorphTargetsData(const TSharedRef<
 			const int32 NumMorphTargets = SkeletalMesh->GetMorphTargets().Num();
 			for (int32 MorphTargetIndex = 0; MorphTargetIndex < NumMorphTargets; ++MorphTargetIndex)
 			{
-				SkeletalMesh->GetMorphTargets()[MorphTargetIndex]->GetMorphLODModels()[LODIndex] = SkeletalMesh->GetMorphTargets()[MorphTargetIndex]->GetMorphLODModels()[LastValidLODIndex];
+				SkeletalMesh->GetMorphTargets()[MorphTargetIndex]->GetMorphLODModels()[LODIndex] = 
+						SkeletalMesh->GetMorphTargets()[MorphTargetIndex]->GetMorphLODModels()[LastValidLODIndex];
 			}
 		}
 	}
@@ -3732,7 +3834,7 @@ void UCustomizableInstancePrivate::BuildOrCopyClothingData(const TSharedRef<FUpd
 				const int32 ClothingDataBufferIndex = [&MeshSet]()
 				{
 					int32 BufferIndex, Channel;
-					MeshSet.FindChannel(mu::MBS_OTHER, 2, &BufferIndex, &Channel);
+					MeshSet.FindChannel(mu::MBS_OTHER, 3, &BufferIndex, &Channel);
 
 					if (BufferIndex > 0)
 					{
@@ -4690,10 +4792,10 @@ FAutoConsoleVariableRef CVarMutableHighPriorityLoading(
 	bEnableHighPriorityLoading,
 	TEXT("If enabled, the request to load additional assets will have high priority."));
 
-
-FGraphEventRef UCustomizableInstancePrivate::LoadAdditionalAssetsAsync(const TSharedRef<FUpdateContextPrivate>& OperationData, FStreamableManager& StreamableManager)
+UE::Tasks::FTask UCustomizableInstancePrivate::LoadAdditionalAssetsAndDataAsync(
+		const TSharedRef<FUpdateContextPrivate>& OperationData, FStreamableManager& StreamableManager)
 {
-	MUTABLE_CPUPROFILER_SCOPE(UCustomizableInstancePrivate::LoadAdditionalAssetsAsync);
+	MUTABLE_CPUPROFILER_SCOPE(UCustomizableInstancePrivate::LoadAdditionalAssetsAndDataAsync);
 
 	UCustomizableObject* CustomizableObject = GetPublic()->GetCustomizableObject();
 
@@ -4702,6 +4804,7 @@ FGraphEventRef UCustomizableInstancePrivate::LoadAdditionalAssetsAsync(const TSh
 	FGraphEventRef Result = nullptr;
 
 	TArray<FSoftObjectPath> AssetsToStream;
+	TArray<uint32> RealTimeMorphStreamableBlocksToStream;
 
 	TArray<FInstanceUpdateData::FLOD>& LODs = OperationData->InstanceUpdateData.LODs;
 	TArray<FInstanceUpdateData::FComponent>& Components = OperationData->InstanceUpdateData.Components;
@@ -4716,7 +4819,7 @@ FGraphEventRef UCustomizableInstancePrivate::LoadAdditionalAssetsAsync(const TSh
 	GatheredAnimBPs.Empty();
 	AnimBPGameplayTags.Reset();
 	AnimBpPhysicsAssets.Reset();
-
+	
 	for (const FInstanceUpdateData::FSurface& Surface : OperationData->InstanceUpdateData.Surfaces)
 	{
 		const uint32 MaterialIndex = Surface.MaterialIndex;
@@ -4797,24 +4900,46 @@ FGraphEventRef UCustomizableInstancePrivate::LoadAdditionalAssetsAsync(const TSh
 
 			FCustomizableInstanceComponentData* ComponentData = GetComponentData(Component.Id);
 
-			const TArray<int32>& StreamedResources = MutableMesh->GetStreamedResources();
-			const TArray<FCustomizableObjectStreamedResourceData>& StreamedResourcesData = CustomizableObject->StreamedResourceData;
+			const TArray<uint32>& StreamedResources = MutableMesh->GetStreamedResources();
 
-			for (int32 ResourceIndex : StreamedResources)
+			for (uint32 ResourceId : StreamedResources)
 			{
-				if (!StreamedResourcesData.IsValidIndex(ResourceIndex))
+				const TArray<FCustomizableObjectStreamedResourceData>& StreamedResourcesData = CustomizableObject->StreamedResourceData;
+				FCustomizableObjectStreameableResourceId TypedResourceId = BitCast<FCustomizableObjectStreameableResourceId>(ResourceId);	
+	
+				if (TypedResourceId.Type == (uint8)FCustomizableObjectStreameableResourceId::EType::AssetUserData)
 				{
-					UE_LOG(LogMutable, Error, TEXT("Invalid streamed resource index. Max Index [%d]. Resource Index [%d]."), StreamedResourcesData.Num(), ResourceIndex);
-					continue; 
-				}
+					const uint32 ResourceIndex = TypedResourceId.Id;
+					if (!StreamedResourcesData.IsValidIndex(ResourceIndex))
+					{
+						UE_LOG(LogMutable, Error, TEXT("Invalid streamed resource index. Max Index [%d]. Resource Index [%d]."), StreamedResourcesData.Num(), ResourceIndex);
+						continue; 
+					}
+
+					const FCustomizableObjectStreamedResourceData& StreamedResource = StreamedResourcesData[ResourceIndex];
+					if (!StreamedResource.IsLoaded())
+					{
+						AssetsToStream.AddUnique(StreamedResource.GetPath().ToSoftObjectPath());
+					}
 
-				const FCustomizableObjectStreamedResourceData& StreamedResource = StreamedResourcesData[ResourceIndex];
-				if (!StreamedResource.IsLoaded())
+					ComponentData->StreamedResourceIndex.Add(ResourceIndex);
+				}
+				else if (TypedResourceId.Type == (uint8)FCustomizableObjectStreameableResourceId::EType::RealTimeMorphTarget)
 				{
-					AssetsToStream.AddUnique(StreamedResource.GetPath().ToSoftObjectPath());
+					const TArray<FMutableStreamableBlock>& MorphsStremeableBlocks = ModelResources.RealTimeMorphStreamableBlocks;
+					if (MorphsStremeableBlocks.IsValidIndex(TypedResourceId.Id))
+					{
+						RealTimeMorphStreamableBlocksToStream.AddUnique(TypedResourceId.Id);
+					}
+					else
+					{
+						UE_LOG(LogMutable, Error, TEXT("Invalid streamed real time morph target data block [%d] found."), TypedResourceId.Id);
+					}
+				}
+				else
+				{
+					check(false);
 				}
-
-				ComponentData->StreamedResourceIndex.Add(ResourceIndex);
 			}
 
 			const bool bReplacePhysicsAssets = HasCOInstanceFlags(ReplacePhysicsAssets);
@@ -4948,24 +5073,180 @@ FGraphEventRef UCustomizableInstancePrivate::LoadAdditionalAssetsAsync(const TSh
 		AssetsToStream.Add(TextureRef.ToSoftObjectPath());
 	}
 
+	TArray<UE::Tasks::FTaskEvent> StreamingCompletionEvents;
 	if (AssetsToStream.Num() > 0)
+	{	
+		UE::Tasks::FTaskEvent AssetAsyncLoadCompletionEvent = StreamingCompletionEvents.Emplace_GetRef(TEXT("AssetAsyncLoadCompletionEvent"));
+
+		StreamingHandle = StreamableManager.RequestAsyncLoad(
+				AssetsToStream, 
+				FStreamableDelegate::CreateUObject(this, &UCustomizableInstancePrivate::AdditionalAssetsAsyncLoaded, AssetAsyncLoadCompletionEvent),
+				bEnableHighPriorityLoading ? FStreamableManager::AsyncLoadHighPriority : FStreamableManager::DefaultAsyncLoadPriority);
+	}
+
+	
+	// File handles will end up owned by the gather task.
+	TArray<TUniquePtr<IAsyncReadFileHandle>> OpenFileHandles;
+	TArray<UE::Tasks::TTask<TUniquePtr<IAsyncReadRequest>>> ReadRequestTasks;
+	
+	bool bHasInvalidMesh = false;
+	bool bUpdateMeshes = DoComponentsNeedUpdate(GetPublic(), OperationData, bHasInvalidMesh);
+
+	if (RealTimeMorphStreamableBlocksToStream.Num() && bUpdateMeshes)
 	{
-		check(!StreamingHandle);
+#if WITH_EDITOR
+		// On editor the data is always loaded, load directly form the ModelResources.
+		MUTABLE_CPUPROFILER_SCOPE(RealTimeMorphStreamingEditor);
+		for (uint32 BlockId : RealTimeMorphStreamableBlocksToStream)
+		{	
+			const FMutableStreamableBlock& Block = ModelResources.RealTimeMorphStreamableBlocks[BlockId]; 
+			
+			const TArray<FMorphTargetVertexData>& SourceData = ModelResources.EditorOnlyMorphTargetReconstructionData;
+			TArray<FMorphTargetVertexData>& DestData = OperationData->InstanceUpdateData.MorphTargetsVertexData.Emplace(BlockId);
+			
+			const uint32 NumElems = Block.Size / sizeof(FMorphTargetVertexData);
+			const uint32 OffsetInElems = Block.Offset / sizeof(FMorphTargetVertexData);
+			DestData.SetNumUninitialized(NumElems);
+
+			check(SourceData.Num()*sizeof(FMorphTargetVertexData) >= Block.Offset + Block.Size);
+			FMemory::Memcpy(DestData.GetData(), SourceData.GetData() + OffsetInElems, Block.Size);
+		}
+#else	
+		MUTABLE_CPUPROFILER_SCOPE(RealTimeMorphStreaming);
+		struct FBlockReadInfo
+		{
+			uint64 Offset;
+			IAsyncReadFileHandle* FileHandle;
+			TArrayView<uint8> AllocatedMemoryView;
+			uint32 FileId;
+		};
+
+		TArray<FBlockReadInfo> BlockReadInfos;
+		BlockReadInfos.Reserve(16);
+		
+		const UCustomizableObjectBulk* BulkData = CustomizableObject->GetPrivate()->GetStreamableBulkData();
+		UE_CLOG(!BulkData, LogMutable, Error, TEXT("BulkData object for CustomizableObject [%s] not found."), 
+				*CustomizableObject->GetFName().ToString());
+		check(BulkData);
+
+		TArray<uint32> OpenFilesIds;
+		const int32 NumMorphBlocks = RealTimeMorphStreamableBlocksToStream.Num();
+		for (int32 I = 0; I < NumMorphBlocks; ++I)
+		{
+			MUTABLE_CPUPROFILER_SCOPE(RealTimeMorphStreamingRequest_Alloc);
+
+			const int32 BlockId = RealTimeMorphStreamableBlocksToStream[I];
+			const FMutableStreamableBlock& Block = ModelResources.RealTimeMorphStreamableBlocks[BlockId]; 
+		
+			TArray<FMorphTargetVertexData>& ReadDestData = OperationData->InstanceUpdateData.MorphTargetsVertexData.FindOrAdd(BlockId);
+
+			// Only request blocks once.
+			if (ReadDestData.Num())
+			{
+				continue;
+			}
+
+			check(Block.Size % sizeof(FMorphTargetVertexData) == 0);
+			uint32 NumElems = Block.Size / sizeof(FMorphTargetVertexData);
+
+			ReadDestData.SetNumUninitialized(NumElems);
+
+			int32 FileHandleIndex = OpenFilesIds.Find(Block.FileId);
+			if (FileHandleIndex == INDEX_NONE && BulkData)
+			{
+				TUniquePtr<IAsyncReadFileHandle> ReadFileHandle = BulkData->OpenFileAsyncRead(Block.FileId);
+
+				OpenFileHandles.Emplace(MoveTemp(ReadFileHandle));
+				FileHandleIndex = OpenFilesIds.Add(Block.FileId);
+			}	
+
+			BlockReadInfos.Emplace(FBlockReadInfo
+			{ 
+				Block.Offset,
+				OpenFileHandles[FileHandleIndex].Get(), 
+				MakeArrayView(reinterpret_cast<uint8*>(ReadDestData.GetData()), ReadDestData.Num()*sizeof(FMorphTargetVertexData)),
+				Block.FileId
+			});
+		}
+
+		for (const FBlockReadInfo& BlockReadInfo : BlockReadInfos)
+		{
+			if (!BlockReadInfo.FileHandle)
+			{
+				continue;
+			}	
+	
+			ReadRequestTasks.Emplace(UE::Tasks::Launch(TEXT("CustomizableObjectInstanceReadRequestTask"),
+			[
+				OwnedOperationData = OperationData.ToSharedPtr(), // Keep a reference to make sure allocated memory is always alive.
+				ReadDataReadyEvent = StreamingCompletionEvents.Emplace_GetRef(TEXT("AsyncReadDataReadyEvent")),
+				Block              = BlockReadInfo,
+				Priority 		   = bEnableHighPriorityLoading ? AIOP_High : AIOP_Normal
+			]() -> TUniquePtr<IAsyncReadRequest>
+			{
+				MUTABLE_CPUPROFILER_SCOPE(CustomizableInstanceLoadBlocksAsyncRead_Request);
+				FAsyncFileCallBack ReadRequestCallBack = 
+						[OwnedOperationData, ReadDataReadyEvent, FileId = Block.FileId](bool bWasCancelled, IAsyncReadRequest*) mutable
+						{
+							if (bWasCancelled)
+							{
+								UE_LOG(LogMutable, Warning, TEXT("An AsyncReadRequest to file %08x was cancelled. The file may not exist."), FileId);
+							}
+							ReadDataReadyEvent.Trigger();
+						}; 
+
+				return TUniquePtr<IAsyncReadRequest>(Block.FileHandle->ReadRequest(
+						Block.Offset, 
+						(int64)Block.AllocatedMemoryView.Num(),
+						Priority,
+						&ReadRequestCallBack,
+						Block.AllocatedMemoryView.GetData()));
+			}, 
+			bEnableHighPriorityLoading ? UE::Tasks::ETaskPriority::High : UE::Tasks::ETaskPriority::Normal));
+		}	
+
+#endif
 
-		Result = FGraphEvent::CreateGraphEvent();
-		StreamingHandle = StreamableManager.RequestAsyncLoad(AssetsToStream, FStreamableDelegate::CreateUObject(this, &UCustomizableInstancePrivate::AdditionalAssetsAsyncLoaded, Result),
-			bEnableHighPriorityLoading ? FStreamableManager::AsyncLoadHighPriority : FStreamableManager::DefaultAsyncLoadPriority);
 	}
 
-	return Result;
+	if (AssetsToStream.Num() > 0 || OpenFileHandles.Num() > 0)
+	{
+		return UE::Tasks::Launch(TEXT("GatherStreamingRequestsCompletionTask"),
+				[
+					ReadRequestTasks = MoveTemp(ReadRequestTasks),
+					OpenFileHandles  = MoveTemp(OpenFileHandles)
+				]() mutable 
+				{
+					for (UE::Tasks::TTask<TUniquePtr<IAsyncReadRequest>>& ReadRequestTask : ReadRequestTasks)
+					{
+						// GetResult() may wait for the task to complete, this should not be a problem as this task
+						// prerequisites guarantee ReadRequestTask has at least started execution.
+						TUniquePtr<IAsyncReadRequest>& ReadRequest = ReadRequestTask.GetResult();
+						if (ReadRequest)
+						{
+							ReadRequest->WaitCompletion();
+						}
+					}
+					
+					ReadRequestTasks.Empty();
+					OpenFileHandles.Empty();
+				},
+				StreamingCompletionEvents,
+				bEnableHighPriorityLoading ? UE::Tasks::ETaskPriority::High : UE::Tasks::ETaskPriority::Normal);
+	}
+	else
+	{
+		check(ReadRequestTasks.Num() == 0);
+		return UE::Tasks::MakeCompletedTask<void>(TEXT("GatherStreamingRequestsCompletionTaskComplete"));
+	}
 }
 
-
-void UCustomizableInstancePrivate::AdditionalAssetsAsyncLoaded( FGraphEventRef CompletionEvent )
+void UCustomizableInstancePrivate::AdditionalAssetsAsyncLoaded(UE::Tasks::FTaskEvent CompletionEvent)
 {
 	// TODO: Do we need this separated?
+	//check(IsInGameThread())
 	AdditionalAssetsAsyncLoaded(GetPublic());
-	CompletionEvent->DispatchSubsequents(); // TODO: we know it is game thread?
+	CompletionEvent.Trigger(); // TODO: we know it is game thread?
 
 	StreamingHandle = nullptr;
 }
@@ -4976,14 +5257,12 @@ FCustomizableObjectInstanceDescriptor& UCustomizableInstancePrivate::GetDescript
 	return GetPublic()->Descriptor;
 }
 
-
 const TArray<TObjectPtr<UMaterialInterface>>* UCustomizableObjectInstance::GetOverrideMaterials(int32 ComponentIndex) const
 {
 	FCustomizableInstanceComponentData* ComponentData = PrivateData->GetComponentData(ComponentIndex);
 	return ComponentData ? &ComponentData->OverrideMaterials : nullptr;
 }
 
-
 void UCustomizableInstancePrivate::AdditionalAssetsAsyncLoaded(UCustomizableObjectInstance* Public)
 {
 	MUTABLE_CPUPROFILER_SCOPE(UCustomizableInstancePrivate::AdditionalAssetsAsyncLoaded);
diff --git a/Engine/Plugins/Experimental/Mutable/Source/CustomizableObject/Private/MuCO/CustomizableObjectSystem.cpp b/Engine/Plugins/Experimental/Mutable/Source/CustomizableObject/Private/MuCO/CustomizableObjectSystem.cpp
--- a/Engine/Plugins/Experimental/Mutable/Source/CustomizableObject/Private/MuCO/CustomizableObjectSystem.cpp
+++ b/Engine/Plugins/Experimental/Mutable/Source/CustomizableObject/Private/MuCO/CustomizableObjectSystem.cpp
@@ -2495,12 +2495,9 @@ namespace impl
 
 		// Next Task: Load Unreal Assets
 		//-------------------------------------------------------------
-		FGraphEventRef Game_LoadUnrealAssets = ObjectInstancePrivateData->LoadAdditionalAssetsAsync(OperationData, UCustomizableObjectSystem::GetInstance()->GetPrivate()->StreamableManager);
-		if (Game_LoadUnrealAssets)
-		{
-			Game_LoadUnrealAssets->SetDebugName(TEXT("LoadAdditionalAssetsAsync"));
-		}
-
+		UE::Tasks::FTask Game_LoadUnrealAssets = ObjectInstancePrivateData->LoadAdditionalAssetsAndDataAsync(
+				OperationData, UCustomizableObjectSystem::GetInstance()->GetPrivate()->StreamableManager);
+		
 		// Next-next Task: Convert Resources
 		//-------------------------------------------------------------
 		SystemPrivateData->AddGameThreadTask(
diff --git a/Engine/Plugins/Experimental/Mutable/Source/CustomizableObjectEditor/Private/MuCOE/CustomizableObjectCompileRunnable.cpp b/Engine/Plugins/Experimental/Mutable/Source/CustomizableObjectEditor/Private/MuCOE/CustomizableObjectCompileRunnable.cpp
--- a/Engine/Plugins/Experimental/Mutable/Source/CustomizableObjectEditor/Private/MuCOE/CustomizableObjectCompileRunnable.cpp
+++ b/Engine/Plugins/Experimental/Mutable/Source/CustomizableObjectEditor/Private/MuCOE/CustomizableObjectCompileRunnable.cpp
@@ -263,6 +263,20 @@ FCustomizableObjectSaveDDRunnable::FCustomizableObjectSaveDDRunnable(UCustomizab
 		FMemoryWriter64 MemoryWriter(Bytes);
 		CustomizableObject->GetPrivate()->SaveCompiledData(MemoryWriter, Options.bIsCooking);
 	}
+#if WITH_EDITORONLY_DATA
+	else
+	{
+		// Do a copy of the MorphData generated at compile time. Only needed when cooking.
+		
+		static_assert(TCanBulkSerialize<FMorphTargetVertexData>::Value);
+		constexpr bool bGetCookedFalse = false;
+		const TArray<FMorphTargetVertexData>& MorphVertexData = 
+				CustomizableObject->GetPrivate()->GetModelResources(bGetCookedFalse).EditorOnlyMorphTargetReconstructionData;
+
+		MorphDataBytes.SetNum(MorphVertexData.Num() * sizeof(FMorphTargetVertexData));
+		FMemory::Memcpy(MorphDataBytes.GetData(), MorphVertexData.GetData(), MorphDataBytes.Num());
+	}
+#endif // WITH_EDITORONLY_DATA
 }
 
 
@@ -270,6 +284,9 @@ uint32 FCustomizableObjectSaveDDRunnable::Run()
 {
 	MUTABLE_CPUPROFILER_SCOPE(FCustomizableObjectSaveDDRunnable::Run)
 
+	// MorphDataBytes has data only if cooking. 
+	check(!!Options.bIsCooking || MorphDataBytes.IsEmpty());
+
 	bool bModelSerialized = Model.Get() != nullptr;
 
 	if (Options.bIsCooking)
@@ -283,9 +300,11 @@ uint32 FCustomizableObjectSaveDDRunnable::Run()
 		{
 			FUnrealMutableModelBulkWriter Streamer(&ModelMemoryWriter, &StreamableMemoryWriter);
 			mu::Model::Serialise(Model.Get(), Streamer);
+
+			//MorphData is already in the corresponding buffer copied from the compilation thread.
 		}
 	}
-	else if(bModelSerialized) // Save CO data + mu::Model and streamable resources to disk
+	else if (bModelSerialized) // Save CO data + mu::Model and streamable resources to disk
 	{
 		// Create folder...
 		IFileManager& FileManager = IFileManager::Get();
@@ -311,8 +330,8 @@ uint32 FCustomizableObjectSaveDDRunnable::Run()
 		if (bFilesDeleted)
 		{
 			// Create file writers...
-			TUniquePtr<FArchive> ModelMemoryWriter( FileManager.CreateFileWriter(*CompileDataFullFileName) );
-			TUniquePtr<FArchive> StreamableMemoryWriter( FileManager.CreateFileWriter(*StreamableDataFullFileName) );
+			TUniquePtr<FArchive> ModelMemoryWriter(FileManager.CreateFileWriter(*CompileDataFullFileName));
+			TUniquePtr<FArchive> StreamableMemoryWriter(FileManager.CreateFileWriter(*StreamableDataFullFileName));
 			check(ModelMemoryWriter);
 			check(StreamableMemoryWriter);
 
@@ -361,6 +380,10 @@ TArray64<uint8>& FCustomizableObjectSaveDDRunnable::GetBulkBytes()
 	return BulkDataBytes;
 }
 
+TArray64<uint8>& FCustomizableObjectSaveDDRunnable::GetMorphBytes()
+{
+	return MorphDataBytes;
+}
 
 bool FCustomizableObjectSaveDDRunnable::IsCompleted() const
 {
diff --git a/Engine/Plugins/Experimental/Mutable/Source/CustomizableObjectEditor/Private/MuCOE/CustomizableObjectCompileRunnable.h b/Engine/Plugins/Experimental/Mutable/Source/CustomizableObjectEditor/Private/MuCOE/CustomizableObjectCompileRunnable.h
--- a/Engine/Plugins/Experimental/Mutable/Source/CustomizableObjectEditor/Private/MuCOE/CustomizableObjectCompileRunnable.h
+++ b/Engine/Plugins/Experimental/Mutable/Source/CustomizableObjectEditor/Private/MuCOE/CustomizableObjectCompileRunnable.h
@@ -98,6 +98,9 @@ public:
 	// Bytes where the streamed data files are stored
 	TArray64<uint8>& GetBulkBytes();
 
+	// Bytes where the streamed morph data files are stored
+	TArray64<uint8>& GetMorphBytes();
+
 	//
 	bool IsCompleted() const;
 
@@ -124,6 +127,9 @@ private:
 	// Bytes where the streamed data files are stored
 	TArray64<uint8> BulkDataBytes;
 
+	// Bytes store streameable files coming form the CO itself.
+	TArray64<uint8> MorphDataBytes;
+
 	bool bIsCooking = false;
 
 	// Whether the thread has finished running
diff --git a/Engine/Plugins/Experimental/Mutable/Source/CustomizableObjectEditor/Private/MuCOE/CustomizableObjectCompiler.cpp b/Engine/Plugins/Experimental/Mutable/Source/CustomizableObjectEditor/Private/MuCOE/CustomizableObjectCompiler.cpp
--- a/Engine/Plugins/Experimental/Mutable/Source/CustomizableObjectEditor/Private/MuCOE/CustomizableObjectCompiler.cpp
+++ b/Engine/Plugins/Experimental/Mutable/Source/CustomizableObjectEditor/Private/MuCOE/CustomizableObjectCompiler.cpp
@@ -955,11 +955,36 @@ void FCustomizableObjectCompiler::CompileInternal(UCustomizableObject* Object, c
 		ModelResources.ParameterUIDataMap = MoveTemp(GenerationContext.ParameterUIDataMap);
 		ModelResources.StateUIDataMap = MoveTemp(GenerationContext.StateUIDataMap);
 
-		// Morph target generated data does not need extra processing so move semantics can be used
-		// to avoid a possibly expensive copy.
-		Object->ContributingMorphTargetsInfo = MoveTemp(GenerationContext.ContributingMorphTargetsInfo);
-		Object->MorphTargetReconstructionData = MoveTemp(GenerationContext.MorphTargetReconstructionData);
+		ModelResources.RealTimeMorphTargetNames = MoveTemp(GenerationContext.RealTimeMorphTargetsNames);
+
+		if (GenerationContext.RealTimeMorphTargetPerMeshData.Num() >= TNumericLimits<uint16>::Max())
+		{
+			UE_LOG(LogMutable, Warning, TEXT("Maximum number of meshes with realtime morph targets reached. Some morphs may not work."));
+		}
+
+		// Create the RealTimeMorphsTargets Blocks from the per mesh Morph data.
+		uint64 RealTimeMorphDataSize = 0;
+		for (const TArray<FMorphTargetVertexData>& VertexDataArray : GenerationContext.RealTimeMorphTargetPerMeshData)
+		{
+			RealTimeMorphDataSize += VertexDataArray.Num();
+		}
 		
+		ModelResources.RealTimeMorphStreamableBlocks.Empty(32);
+		ModelResources.EditorOnlyMorphTargetReconstructionData.Empty(RealTimeMorphDataSize);
+
+		uint64 RealTimeMorphDataOffset = 0;
+		for (const TArray<FMorphTargetVertexData>& VertexDataArray : GenerationContext.RealTimeMorphTargetPerMeshData)
+		{
+			ModelResources.RealTimeMorphStreamableBlocks.Emplace(FMutableStreamableBlock
+					{
+						uint32(0),
+						(uint32)VertexDataArray.Num()*sizeof(FMorphTargetVertexData), 
+						RealTimeMorphDataOffset, 
+					});
+
+			RealTimeMorphDataOffset += VertexDataArray.Num()*sizeof(FMorphTargetVertexData);
+			ModelResources.EditorOnlyMorphTargetReconstructionData.Append(VertexDataArray);
+		}
 		
 		// Clothing	
 		Object->ClothMeshToMeshVertData = MoveTemp(GenerationContext.ClothMeshToMeshVertData);
@@ -974,7 +999,7 @@ void FCustomizableObjectCompiler::CompileInternal(UCustomizableObject* Object, c
 		auto IsSharedConfigData = [](const FCustomizableObjectClothConfigData& ConfigData) -> bool
 		{
 			 const UClass* ConfigClass = FindObject<UClass>(nullptr, *ConfigData.ClassPath);
-			 return ConfigClass ? static_cast<bool>( Cast<UClothSharedConfigCommon>(ConfigClass->GetDefaultObject() ) ) : false;
+			 return ConfigClass ? static_cast<bool>(Cast<UClothSharedConfigCommon>(ConfigClass->GetDefaultObject())) : false;
 		};
 		
 		// Find shared configs to be used (One of each type) 
@@ -1261,7 +1286,11 @@ void FCustomizableObjectCompiler::FinishSavingDerivedData()
 
 	if (Options.bIsCooking)
 	{
-		CurrentObject->GetPrivate()->CachePlatformData(SaveDDTask->GetTargetPlatform(), SaveDDTask->GetModelBytes(), SaveDDTask->GetBulkBytes());
+		CurrentObject->GetPrivate()->CachePlatformData(
+				SaveDDTask->GetTargetPlatform(), 
+				SaveDDTask->GetModelBytes(), 
+				SaveDDTask->GetBulkBytes(),
+				SaveDDTask->GetMorphBytes());
 	}
 
 	// Order matters
diff --git a/Engine/Plugins/Experimental/Mutable/Source/CustomizableObjectEditor/Private/MuCOE/GenerateMutableSource/GenerateMutableSource.h b/Engine/Plugins/Experimental/Mutable/Source/CustomizableObjectEditor/Private/MuCOE/GenerateMutableSource/GenerateMutableSource.h
--- a/Engine/Plugins/Experimental/Mutable/Source/CustomizableObjectEditor/Private/MuCOE/GenerateMutableSource/GenerateMutableSource.h
+++ b/Engine/Plugins/Experimental/Mutable/Source/CustomizableObjectEditor/Private/MuCOE/GenerateMutableSource/GenerateMutableSource.h
@@ -47,6 +47,7 @@ struct FParameterUIData;
 struct FMutableRefSkeletalMeshData;
 struct FMutableRefSocket;
 struct FMutableSkinWeightProfileInfo;
+struct FMorphTargetVertexData;
 
 struct FGeneratedImageProperties
 {
@@ -700,8 +701,8 @@ struct FMutableGraphGenerationContext
 	TMap<class UCustomizableObjectNodeMeshClipWithMesh*, TArray<mu::Ptr<mu::NodeModifierMeshClipWithMesh>>> MapClipMeshNodeToMutableClipMeshNodeArray;
 
 	// Data used for MorphTarget reconstruction.
-	TArray<FMorphTargetInfo> ContributingMorphTargetsInfo;
-	TArray<FMorphTargetVertexData> MorphTargetReconstructionData;
+	TArray<FName> RealTimeMorphTargetsNames;
+	TArray<TArray<FMorphTargetVertexData>> RealTimeMorphTargetPerMeshData;
 
 	// Data used for Clothing reconstruction.
 	TArray<FCustomizableObjectMeshToMeshVertData> ClothMeshToMeshVertData;
diff --git a/Engine/Plugins/Experimental/Mutable/Source/CustomizableObjectEditor/Private/MuCOE/GenerateMutableSource/GenerateMutableSourceMesh.cpp b/Engine/Plugins/Experimental/Mutable/Source/CustomizableObjectEditor/Private/MuCOE/GenerateMutableSource/GenerateMutableSourceMesh.cpp
--- a/Engine/Plugins/Experimental/Mutable/Source/CustomizableObjectEditor/Private/MuCOE/GenerateMutableSource/GenerateMutableSourceMesh.cpp
+++ b/Engine/Plugins/Experimental/Mutable/Source/CustomizableObjectEditor/Private/MuCOE/GenerateMutableSource/GenerateMutableSourceMesh.cpp
@@ -851,7 +851,7 @@ mu::MeshPtr ConvertSkeletalMeshToMutable(const USkeletalMesh* InSkeletalMesh, co
 	MutableMesh->GetVertexBuffers().SetElementCount(VertexCount);
 
 	const int32 VertexBuffersCount = 1 +
-		(GenerationContext.Options.bRealTimeMorphTargetsEnabled ? 2 : 0) +
+		(GenerationContext.Options.bRealTimeMorphTargetsEnabled ? 3 : 0) +
 		(GenerationContext.Options.bClothingEnabled ? 1 : 0);
 
 	MutableMesh->GetVertexBuffers().SetBufferCount(VertexBuffersCount);
@@ -1007,7 +1007,7 @@ mu::MeshPtr ConvertSkeletalMeshToMutable(const USkeletalMesh* InSkeletalMesh, co
 	int32 nextBufferIndex = 1;
 	if (GenerationContext.Options.bRealTimeMorphTargetsEnabled)
 	{
-		nextBufferIndex += 2;
+		nextBufferIndex += 3;
 
 		// This call involves resolving every TObjectPtr<UMorphTarget> to a UMorphTarget*, so
 		// cache the result here to avoid calling it repeatedly.
@@ -1116,52 +1116,77 @@ mu::MeshPtr ConvertSkeletalMeshToMutable(const USkeletalMesh* InSkeletalMesh, co
 			}
 		}
 
-		// MorphTarget vertex info index.
+		// MorphTarget vertex block offset.
 		{
 			using namespace mu;
-			const int ElementSize = sizeof(int32);
-			const int ChannelCount = 1;
+			const int32 ElementSize = sizeof(int32);
+			const int32 ChannelCount = 1;
 			const MESH_BUFFER_SEMANTIC Semantics[ChannelCount] = { MBS_OTHER };
-			const int SemanticIndices[ChannelCount] = { 0 };
+			const int32 SemanticIndices[ChannelCount] = { 0 };
 			const MESH_BUFFER_FORMAT Formats[ChannelCount] = { MBF_INT32 };
-			int Components[ChannelCount] = { 1 };
-			const int Offsets[ChannelCount] = { 0 };
+			int32 Components[ChannelCount] = { 1 };
+			const int32 Offsets[ChannelCount] = { 0 };
 
 			MutableMesh->GetVertexBuffers().SetBuffer(1, ElementSize, ChannelCount, Semantics, SemanticIndices, Formats, Components, Offsets);
 		}
 
-		// MorphTarget vertex info count.
+		// MorphTarget vertex morph count.
 		{
 			using namespace mu;
-			const int ElementSize = sizeof(int32);
-			const int ChannelCount = 1;
+			const int32 ElementSize = sizeof(uint16);
+			const int32 ChannelCount = 1;
 			const MESH_BUFFER_SEMANTIC Semantics[ChannelCount] = { MBS_OTHER };
-			const int SemanticIndices[ChannelCount] = { 1 };
-			const MESH_BUFFER_FORMAT Formats[ChannelCount] = { MBF_INT32 };
-			int Components[ChannelCount] = { 1 };
-			const int Offsets[ChannelCount] = { 0 };
+			const int32 SemanticIndices[ChannelCount] = { 1 };
+			const MESH_BUFFER_FORMAT Formats[ChannelCount] = { MBF_UINT16 };
+			int32 Components[ChannelCount] = { 1 };
+			const int32 Offsets[ChannelCount] = { 0 };
 
 			MutableMesh->GetVertexBuffers().SetBuffer(2, ElementSize, ChannelCount, Semantics, SemanticIndices, Formats, Components, Offsets);
 		}
 
+		// MorphTarget vertex block id.
+		{
+			using namespace mu;
+			const int32 ElementSize = sizeof(uint16);
+			const int32 ChannelCount = 1;
+			const MESH_BUFFER_SEMANTIC Semantics[ChannelCount] = { MBS_OTHER };
+			const int32 SemanticIndices[ChannelCount] = { 2 };
+			const MESH_BUFFER_FORMAT Formats[ChannelCount] = { MBF_UINT16 };
+			int32 Components[ChannelCount] = { 1 };
+			const int32 Offsets[ChannelCount] = { 0 };
+
+			MutableMesh->GetVertexBuffers().SetBuffer(3, ElementSize, ChannelCount, Semantics, SemanticIndices, Formats, Components, Offsets);
+		}
+
+
 		// Setup MorphTarget reconstruction data.
 
-		TArrayView<int32> VertexMorphsCountBufferView(reinterpret_cast<int32*>(MutableMesh->GetVertexBuffers().GetBufferData(2)), VertexCount);
-		for (int32& Elem : VertexMorphsCountBufferView)
+		TArrayView<int32> VertexMorphsOffsetBufferView(reinterpret_cast<int32*>(MutableMesh->GetVertexBuffers().GetBufferData(1)), VertexCount);
+		TArrayView<uint16> VertexMorphsCountBufferView(reinterpret_cast<uint16*>(MutableMesh->GetVertexBuffers().GetBufferData(2)), VertexCount);
+		TArrayView<uint16> VertexMorphsResourceIdBufferView(reinterpret_cast<uint16*>(MutableMesh->GetVertexBuffers().GetBufferData(3)), VertexCount);
+	
+		for (int32& Elem : VertexMorphsOffsetBufferView)
+		{
+			Elem = -1;
+		}
+
+		for (uint16& Elem : VertexMorphsCountBufferView)
 		{
 			Elem = 0;
 		}
 
-		TArrayView<int32>  VertexMorphsInfoIndexBufferView(reinterpret_cast<int32*>(MutableMesh->GetVertexBuffers().GetBufferData(1)), VertexCount);
-		for (int32& Elem : VertexMorphsInfoIndexBufferView)
+		for (uint16& Elem : VertexMorphsResourceIdBufferView)
 		{
-			Elem = -1;
+			Elem = TNumericLimits<uint16>::Max(); 
 		}
 
 		if (UsedMorphTargets.Num())
 		{
 			const double StartTime = FPlatformTime::Seconds();
 
+			TArray<FMorphTargetVertexData> MorphsMeshData;
+			MorphsMeshData.Reserve(32);
+
 			TArray<FMorphTargetVertexData> MorphsUsed;
 			for (int32 VertexIdx = VertexStart; VertexIdx < VertexStart + VertexCount && VertexIdx < Vertices.Num(); ++VertexIdx)
 			{
@@ -1176,8 +1201,7 @@ mu::MeshPtr ConvertSkeletalMeshToMutable(const USkeletalMesh* InSkeletalMesh, co
 
 					const TArray<FMorphTargetLODModel>& MorphLODModels = MorphTarget->GetMorphLODModels();
 
-					if (LODIndex >= MorphLODModels.Num()
-						|| !MorphLODModels[LODIndex].SectionIndices.Contains(SectionIndex))
+					if (LODIndex >= MorphLODModels.Num() || !MorphLODModels[LODIndex].SectionIndices.Contains(SectionIndex))
 					{
 						continue;
 					}
@@ -1185,10 +1209,8 @@ mu::MeshPtr ConvertSkeletalMeshToMutable(const USkeletalMesh* InSkeletalMesh, co
 					// The vertices should be sorted by SourceIdx
 					check(MorphLODModels[LODIndex].Vertices.Num() < 2 || MorphLODModels[LODIndex].Vertices[0].SourceIdx < MorphLODModels[LODIndex].Vertices.Last().SourceIdx);
 
-					const int32 VertexFoundIndex = Algo::BinarySearchBy(
-						MorphLODModels[LODIndex].Vertices,
-						(uint32)VertexIdx,
-						[](const FMorphTargetDelta& Element) { return Element.SourceIdx; });
+					const int32 VertexFoundIndex = Algo::BinarySearchBy(MorphLODModels[LODIndex].Vertices, (uint32)VertexIdx,
+							[](const FMorphTargetDelta& Element) { return Element.SourceIdx; });
 
 					if (VertexFoundIndex == INDEX_NONE)
 					{
@@ -1198,27 +1220,40 @@ mu::MeshPtr ConvertSkeletalMeshToMutable(const USkeletalMesh* InSkeletalMesh, co
 					const FMorphTargetDelta& VertexFound = MorphLODModels[LODIndex].Vertices[VertexFoundIndex];
 					const FName MorphTargetName = MorphTarget->GetFName();
 
-					TArray<FMorphTargetInfo>& ContributingMorphTargetsInfo = GenerationContext.ContributingMorphTargetsInfo;
+					TArray<FName>& RealTimeMorphTargetsNames = GenerationContext.RealTimeMorphTargetsNames;
 
-					int32 DestMorphTargetIdx = ContributingMorphTargetsInfo.IndexOfByPredicate(
-						[&MorphTargetName](auto& MorphTargetInfo) { return MorphTargetName == MorphTargetInfo.Name; });
+					int32 MorphTargetNameIndex = RealTimeMorphTargetsNames.Find(MorphTargetName);
 
-					DestMorphTargetIdx = DestMorphTargetIdx != INDEX_NONE
-						? DestMorphTargetIdx
-						: ContributingMorphTargetsInfo.Emplace(FMorphTargetInfo{ MorphTargetName, GenerationContext.CurrentLOD + 1 });
+					MorphTargetNameIndex = MorphTargetNameIndex != INDEX_NONE
+							? MorphTargetNameIndex
+							: RealTimeMorphTargetsNames.Emplace(MorphTargetName);
 
-					FMorphTargetInfo& MorphTargetInfo = ContributingMorphTargetsInfo[DestMorphTargetIdx];
-					MorphTargetInfo.LodNum = FMath::Max(MorphTargetInfo.LodNum, GenerationContext.CurrentLOD + 1);
-
-					MorphsUsed.Emplace(FMorphTargetVertexData{ VertexFound.PositionDelta, VertexFound.TangentZDelta, DestMorphTargetIdx });
+					MorphsUsed.Emplace(FMorphTargetVertexData{VertexFound.PositionDelta, VertexFound.TangentZDelta, (uint32)MorphTargetNameIndex});
 				}
 
 				if (MorphsUsed.Num())
 				{
-					VertexMorphsInfoIndexBufferView[VertexIdx - VertexStart] = GenerationContext.MorphTargetReconstructionData.Num();
-					VertexMorphsCountBufferView[VertexIdx - VertexStart] = MorphsUsed.Num();
 
-					GenerationContext.MorphTargetReconstructionData.Append(MorphsUsed);
+					VertexMorphsOffsetBufferView[VertexIdx - VertexStart] = MorphsMeshData.Num();
+					VertexMorphsCountBufferView[VertexIdx - VertexStart] = (uint16)MorphsUsed.Num();
+					VertexMorphsResourceIdBufferView[VertexIdx - VertexStart] = (uint16)GenerationContext.RealTimeMorphTargetPerMeshData.Num();
+					
+					MorphsMeshData.Append(MorphsUsed);
+				}
+			}
+
+			// Only commit the morph if there is data.
+			if (MorphsMeshData.Num())
+			{
+				if (GenerationContext.RealTimeMorphTargetPerMeshData.Num() < TNumericLimits<uint16>::Max())
+				{
+					FCustomizableObjectStreameableResourceId StreamedMorphResource;
+					StreamedMorphResource.Id = GenerationContext.RealTimeMorphTargetPerMeshData.Num();
+					StreamedMorphResource.Type = static_cast<uint8>(FCustomizableObjectStreameableResourceId::EType::RealTimeMorphTarget);
+
+					MutableMesh->AddStreamedResource(BitCast<uint32>(StreamedMorphResource));
+						
+					GenerationContext.RealTimeMorphTargetPerMeshData.Emplace_GetRef() = MoveTemp(MorphsMeshData);
 				}
 			}
 
@@ -1231,13 +1266,13 @@ mu::MeshPtr ConvertSkeletalMeshToMutable(const USkeletalMesh* InSkeletalMesh, co
 	{
 		{
 			using namespace mu;
-			const int ElementSize = sizeof(int32);
-			const int ChannelCount = 1;
+			const int32 ElementSize = sizeof(int32);
+			const int32 ChannelCount = 1;
 			const MESH_BUFFER_SEMANTIC Semantics[ChannelCount] = { MBS_OTHER };
-			const int SemanticIndices[ChannelCount] = { GenerationContext.Options.bRealTimeMorphTargetsEnabled ? 2 : 0 };
+			const int32 SemanticIndices[ChannelCount] = { GenerationContext.Options.bRealTimeMorphTargetsEnabled ? 3 : 0 };
 			const MESH_BUFFER_FORMAT Formats[ChannelCount] = { MBF_INT32 };
-			int Components[ChannelCount] = { 1 };
-			const int Offsets[ChannelCount] = { 0 };
+			int32 Components[ChannelCount] = { 1 };
+			const int32 Offsets[ChannelCount] = { 0 };
 
 			MutableMesh->GetVertexBuffers().SetBuffer(nextBufferIndex, ElementSize, ChannelCount, Semantics, SemanticIndices, Formats, Components, Offsets);
 		}
@@ -2886,7 +2921,7 @@ mu::NodeMeshPtr GenerateMutableSourceMesh(const UEdGraphPin* Pin,
 				MeshUniqueTags += AnimBPTag;
 			}
 
-			TArray<int32> StreamedResources;
+			TArray<FCustomizableObjectStreameableResourceId> StreamedResources;
 
 			if (GenerationContext.Object->bEnableAssetUserDataMerge)
 			{
@@ -2902,7 +2937,16 @@ mu::NodeMeshPtr GenerateMutableSourceMesh(const UEdGraphPin* Pin,
 						}
 
 						const int32 ResourceIndex = GenerationContext.AddAssetUserDataToStreamedResources(AssetUserData);
-						StreamedResources.Add(ResourceIndex);
+						if (ResourceIndex >= 0)
+						{
+							check(ResourceIndex < (1 << 24) - 1);
+							
+							FCustomizableObjectStreameableResourceId ResourceId;
+							ResourceId.Id = (uint32)ResourceIndex;
+							ResourceId.Type = static_cast<uint8>(FCustomizableObjectStreameableResourceId::EType::AssetUserData);
+
+							StreamedResources.Add(ResourceId);
+						}
 
 						MeshUniqueTags += AssetUserData->GetPathName();
 					}
@@ -2982,9 +3026,9 @@ mu::NodeMeshPtr GenerateMutableSourceMesh(const UEdGraphPin* Pin,
 					AddTagToMutableMeshUnique(*MutableMesh, GamePlayTag);
 				}
 
-				for (int32 ResourceIndex : StreamedResources)
+				for (FCustomizableObjectStreameableResourceId ResourceId : StreamedResources)
 				{
-					MutableMesh->AddStreamedResource(ResourceIndex);
+					MutableMesh->AddStreamedResource(BitCast<uint32>(ResourceId));
 				}
 
 				AddSocketTagsToMesh(TypedNodeSkel->SkeletalMesh, MutableMesh, GenerationContext);
diff --git a/Engine/Plugins/Experimental/Mutable/Source/CustomizableObjectEditor/Private/MuCOE/GenerateMutableSource/GenerateMutableSourceSurface.cpp b/Engine/Plugins/Experimental/Mutable/Source/CustomizableObjectEditor/Private/MuCOE/GenerateMutableSource/GenerateMutableSourceSurface.cpp
--- a/Engine/Plugins/Experimental/Mutable/Source/CustomizableObjectEditor/Private/MuCOE/GenerateMutableSource/GenerateMutableSourceSurface.cpp
+++ b/Engine/Plugins/Experimental/Mutable/Source/CustomizableObjectEditor/Private/MuCOE/GenerateMutableSource/GenerateMutableSourceSurface.cpp
@@ -90,7 +90,7 @@ void SetSurfaceFormat( FMutableGraphGenerationContext& GenerationContext,
 
 	if (bWithRealTimeMorphs)
 	{
-		MutableBufferCount += 2;
+		MutableBufferCount += 3;
 	}
 
 	if (bWithClothing)
@@ -134,29 +134,41 @@ void SetSurfaceFormat( FMutableGraphGenerationContext& GenerationContext,
 	// MorphTarget vertex tracking info buffers
 	if (bWithRealTimeMorphs)
 	{
+		using namespace mu;
 		{
-			using namespace mu;
-			const int ElementSize = sizeof(int32);
-			constexpr int ChannelCount = 1;
+			const int32 ElementSize = sizeof(int32);
+			constexpr int32 ChannelCount = 1;
 			const MESH_BUFFER_SEMANTIC Semantics[ChannelCount] = { MBS_OTHER };
-			const int SemanticIndices[ChannelCount] = { 0 };
+			const int32 SemanticIndices[ChannelCount] = { 0 };
 			const MESH_BUFFER_FORMAT Formats[ChannelCount] = { MBF_INT32 };
-			const int Components[ChannelCount] = { 1 };
-			const int Offsets[ChannelCount] = { 0 };
+			const int32 Components[ChannelCount] = { 1 };
+			const int32 Offsets[ChannelCount] = { 0 };
 
 			OutVertexBufferFormat.SetBuffer(CurrentVertexBuffer, ElementSize, ChannelCount, Semantics, SemanticIndices, Formats, Components, Offsets);
 			++CurrentVertexBuffer;
 		}
 
 		{
-			using namespace mu;
-			const int ElementSize = sizeof(int32);
-			constexpr int ChannelCount = 1;
+			const int32 ElementSize = sizeof(uint16);
+			constexpr int32 ChannelCount = 1;
 			const MESH_BUFFER_SEMANTIC Semantics[ChannelCount] = { MBS_OTHER };
-			const int SemanticIndices[ChannelCount] = { 1 };
-			const MESH_BUFFER_FORMAT Formats[ChannelCount] = { MBF_INT32 };
-			const int Components[ChannelCount] = { 1 };
-			const int Offsets[ChannelCount] = { 0 };
+			const int32 SemanticIndices[ChannelCount] = { 1 };
+			const MESH_BUFFER_FORMAT Formats[ChannelCount] = { MBF_UINT16 };
+			const int32 Components[ChannelCount] = { 1 };
+			const int32 Offsets[ChannelCount] = { 0 };
+
+			OutVertexBufferFormat.SetBuffer(CurrentVertexBuffer, ElementSize, ChannelCount, Semantics, SemanticIndices, Formats, Components, Offsets);
+			++CurrentVertexBuffer;
+		}
+
+		{
+			const int32 ElementSize = sizeof(uint16);
+			constexpr int32 ChannelCount = 1;
+			const MESH_BUFFER_SEMANTIC Semantics[ChannelCount] = { MBS_OTHER };
+			const int32 SemanticIndices[ChannelCount] = { 2 };
+			const MESH_BUFFER_FORMAT Formats[ChannelCount] = { MBF_UINT16 };
+			const int32 Components[ChannelCount] = { 1 };
+			const int32 Offsets[ChannelCount] = { 0 };
 
 			OutVertexBufferFormat.SetBuffer(CurrentVertexBuffer, ElementSize, ChannelCount, Semantics, SemanticIndices, Formats, Components, Offsets);
 			++CurrentVertexBuffer;
@@ -167,13 +179,13 @@ void SetSurfaceFormat( FMutableGraphGenerationContext& GenerationContext,
 	if (bWithClothing)
 	{
 		using namespace mu;
-		const int ElementSize = sizeof(int32);
-		constexpr int ChannelCount = 1;
+		const int32 ElementSize = sizeof(int32);
+		constexpr int32 ChannelCount = 1;
 		const MESH_BUFFER_SEMANTIC Semantics[ChannelCount] = { MBS_OTHER };
-		const int SemanticIndices[ChannelCount] = { bWithRealTimeMorphs ? 2 : 0 };
+		const int32 SemanticIndices[ChannelCount] = { bWithRealTimeMorphs ? 3 : 0 };
 		const MESH_BUFFER_FORMAT Formats[ChannelCount] = { MBF_INT32 };
-		const int Components[ChannelCount] = { 1 };
-		const int Offsets[ChannelCount] = { 0 };
+		const int32 Components[ChannelCount] = { 1 };
+		const int32 Offsets[ChannelCount] = { 0 };
 
 		OutVertexBufferFormat.SetBuffer(CurrentVertexBuffer, ElementSize, ChannelCount, Semantics, SemanticIndices, Formats, Components, Offsets);
 		++CurrentVertexBuffer;
diff --git a/Engine/Plugins/Experimental/Mutable/Source/CustomizableObjectEditor/Private/MuCOE/GenerateMutableSource/GenerateMutableSourceTable.cpp b/Engine/Plugins/Experimental/Mutable/Source/CustomizableObjectEditor/Private/MuCOE/GenerateMutableSource/GenerateMutableSourceTable.cpp
--- a/Engine/Plugins/Experimental/Mutable/Source/CustomizableObjectEditor/Private/MuCOE/GenerateMutableSource/GenerateMutableSourceTable.cpp
+++ b/Engine/Plugins/Experimental/Mutable/Source/CustomizableObjectEditor/Private/MuCOE/GenerateMutableSource/GenerateMutableSourceTable.cpp
@@ -263,7 +263,17 @@ bool FillTableColumn(const UCustomizableObjectNodeTable* TableNode,	mu::TablePtr
 						}
 
 						const int32 ResourceIndex = GenerationContext.AddAssetUserDataToStreamedResources(AssetUserData);
-						StreamedResources.Add(ResourceIndex);
+						
+						if (ResourceIndex >= 0)
+						{
+							check(ResourceIndex < (1 << 24) - 1);
+
+							FCustomizableObjectStreameableResourceId ResourceId;
+							ResourceId.Id = GenerationContext.AddAssetUserDataToStreamedResources(AssetUserData);
+							ResourceId.Type = (uint8)FCustomizableObjectStreameableResourceId::EType::AssetUserData;
+
+							StreamedResources.Add(BitCast<uint32>(ResourceId));
+						}
 
 						MeshUniqueTags += AssetUserData->GetPathName();
 					}
diff --git a/Engine/Plugins/Experimental/Mutable/Source/MutableRuntime/Private/MuR/Mesh.cpp b/Engine/Plugins/Experimental/Mutable/Source/MutableRuntime/Private/MuR/Mesh.cpp
--- a/Engine/Plugins/Experimental/Mutable/Source/MutableRuntime/Private/MuR/Mesh.cpp
+++ b/Engine/Plugins/Experimental/Mutable/Source/MutableRuntime/Private/MuR/Mesh.cpp
@@ -13,6 +13,7 @@
 namespace mu
 {
 
+
 //---------------------------------------------------------------------------------------------
 void Mesh::Serialise( const Mesh* p, OutputArchive& arch )
 {
@@ -530,17 +531,14 @@ void Mesh::SetTag( int tagIndex, const FString& Name )
 
 
 //---------------------------------------------------------------------------------------------
-void Mesh::AddStreamedResource(int32 ResourceIndex)
+void Mesh::AddStreamedResource(uint32 ResourceId)
 {
-	if (ensure(ResourceIndex >= 0))
-	{
-		StreamedResources.AddUnique(ResourceIndex);
-	}
+	StreamedResources.AddUnique(ResourceId);
 }
 
 
 //---------------------------------------------------------------------------------------------
-const TArray<int32>& Mesh::GetStreamedResources() const
+const TArray<uint32>& Mesh::GetStreamedResources() const
 {
 	return StreamedResources;
 }
diff --git a/Engine/Plugins/Experimental/Mutable/Source/MutableRuntime/Public/MuR/Mesh.h b/Engine/Plugins/Experimental/Mutable/Source/MutableRuntime/Public/MuR/Mesh.h
--- a/Engine/Plugins/Experimental/Mutable/Source/MutableRuntime/Public/MuR/Mesh.h
+++ b/Engine/Plugins/Experimental/Mutable/Source/MutableRuntime/Public/MuR/Mesh.h
@@ -157,7 +157,6 @@ namespace mu
 	
 	ENUM_CLASS_FLAGS(EMeshCopyFlags);
 
-
     //! \brief Mesh object containing any number of buffers with any number of channels.
     //! The buffers can be per-index or per-vertex.
     //! The mesh also includes layout information for every texture channel for internal usage, and
@@ -285,10 +284,10 @@ namespace mu
         void SetTag( int tagIndex, const FString& Name );
 
 		//!
-		void AddStreamedResource(int32 ResourceIndex);
+		void AddStreamedResource(uint32 ResourceId);
 
 		//!
-		const TArray<int32>& GetStreamedResources() const;
+		const TArray<uint32>& GetStreamedResources() const;
 
 		//!
 		int32 FindBonePose(uint16 BoneId) const;
@@ -386,8 +385,8 @@ namespace mu
 		//!
 		TArray<FString> m_tags;
 
-		//!
-		TArray<int32> StreamedResources;
+		// Opaque handle to external resources.
+		TArray<uint32> StreamedResources;
 
 		struct FBonePose
 		{
@@ -406,8 +405,7 @@ namespace mu
 			//!
 			inline bool operator==(const FBonePose& Other) const
 			{
-				return BoneUsageFlags == Other.BoneUsageFlags
-					&& BoneId == Other.BoneId;
+				return BoneUsageFlags == Other.BoneUsageFlags && BoneId == Other.BoneId;
 			}
 		};
 		// This is the pose used by this mesh fragment, used to update the transforms of the final skeleton
