[Mutable] Fix CO instances infinite updates when editing an instance in the viewport.

#rnx

diff --git a/Engine/Plugins/Experimental/Mutable/Source/CustomizableObject/Private/MuCO/CustomizableObjectInstance.cpp b/Engine/Plugins/Experimental/Mutable/Source/CustomizableObject/Private/MuCO/CustomizableObjectInstance.cpp
--- a/Engine/Plugins/Experimental/Mutable/Source/CustomizableObject/Private/MuCO/CustomizableObjectInstance.cpp
+++ b/Engine/Plugins/Experimental/Mutable/Source/CustomizableObject/Private/MuCO/CustomizableObjectInstance.cpp
@@ -6380,7 +6380,7 @@ void UCustomizableObjectInstance::SetRequestedLODs(int32 InMinLOD, int32 , const
 			for (int32 ComponentIndex = 0; ComponentIndex < ComponentCount; ++ComponentIndex)
 			{
 				// Find the first requested LOD. We'll generate [FirstRequestedLOD ... MaxLOD].
-				int32 FirstRequestedLOD = ((int32)PrivateData->NumMaxLODsToStream) - 1;
+				int32 FirstRequestedLOD = FMath::Max(((int32)PrivateData->NumMaxLODsToStream) - 1, 0);
 				for (int32 LODIndex = 0; LODIndex < FirstRequestedLOD; ++LODIndex)
 				{
 					if ((!bIgnoreGeneratedLODs && GeneratedLODsPerComponent[ComponentIndex] & (1 << LODIndex))
