[mutable] Fixed bug where input nodes did not appear in the context menu generated from an input pin.
+Removed function that generated input node pins as orphan on beginconstruct as we do not want this behavior anymore.

#rnx

diff --git a/Engine/Plugins/Experimental/Mutable/Source/CustomizableObjectEditor/Classes/MuCOE/Nodes/CustomizableObjectNodeExternalPin.h b/Engine/Plugins/Experimental/Mutable/Source/CustomizableObjectEditor/Classes/MuCOE/Nodes/CustomizableObjectNodeExternalPin.h
--- a/Engine/Plugins/Experimental/Mutable/Source/CustomizableObjectEditor/Classes/MuCOE/Nodes/CustomizableObjectNodeExternalPin.h
+++ b/Engine/Plugins/Experimental/Mutable/Source/CustomizableObjectEditor/Classes/MuCOE/Nodes/CustomizableObjectNodeExternalPin.h
@@ -17,20 +17,6 @@ class UEdGraphPin;
 class UObject;
 
 
-/**
- * Custom remap pins by name action.
- *
- * If the node can not find the exposed pin, set the external pin as deprecated.
- */
-UCLASS()
-class UCustomizableObjectNodeRemapPinsCustomExternalPin : public UCustomizableObjectNodeRemapPinsByName
-{
-	GENERATED_BODY()
-
-public:
-	virtual void RemapPins(const UCustomizableObjectNode& Node, const TArray<UEdGraphPin*>& OldPins, const TArray<UEdGraphPin*>& NewPins, TMap<UEdGraphPin*, UEdGraphPin*>& PinsToRemap, TArray<UEdGraphPin*>& PinsToOrphan) override;
-};
-
 UCLASS()
 class CUSTOMIZABLEOBJECTEDITOR_API UCustomizableObjectNodeExternalPin : public UCustomizableObjectNode
 {
@@ -46,7 +32,6 @@ public:
 	FText GetTooltipText() const override;
 	
 	// UCustomizableObjectNode interface
-	virtual void BeginConstruct() override;
 	virtual void BackwardsCompatibleFixup() override;
 	virtual void PostBackwardsCompatibleFixup() override;
 	void AllocateDefaultPins(UCustomizableObjectNodeRemapPins* RemapPins) override;
@@ -55,15 +40,7 @@ public:
 	virtual bool CanConnect(const UEdGraphPin* InOwnedInputPin, const UEdGraphPin* InOutputPin, bool& bOutIsOtherNodeBlocklisted, bool& bOutArePinsCompatible) const override;
 
 	void UpdateReferencedNodeId(const FGuid& NewGuid) override;
-
-	// Own interface
 	
-	/** Create the external pin. 
-	 *
-	 * @param NodeExposePin Used to obtain expose pin friendly name. If it is nullptr, the name will be set to a default.
-	 */
-	UEdGraphPin* CreateExternalPin(const UCustomizableObjectNodeExposePin* NodeExposePin);
-
 	/** Set the linked Node Expose Pin node guid. */
 	void SetExternalObjectNodeId(FGuid Guid);
 
diff --git a/Engine/Plugins/Experimental/Mutable/Source/CustomizableObjectEditor/Private/MuCOE/Nodes/CustomizableObjectNodeExternalPin.cpp b/Engine/Plugins/Experimental/Mutable/Source/CustomizableObjectEditor/Private/MuCOE/Nodes/CustomizableObjectNodeExternalPin.cpp
--- a/Engine/Plugins/Experimental/Mutable/Source/CustomizableObjectEditor/Private/MuCOE/Nodes/CustomizableObjectNodeExternalPin.cpp
+++ b/Engine/Plugins/Experimental/Mutable/Source/CustomizableObjectEditor/Private/MuCOE/Nodes/CustomizableObjectNodeExternalPin.cpp
@@ -6,30 +6,11 @@
 #include "MuCOE/CustomizableObjectPin.h"
 #include "MuCOE/EdGraphSchema_CustomizableObject.h"
 #include "MuCOE/Nodes/CustomizableObjectNodeExposePin.h"
+#include "MuCOE/RemapPins/CustomizableObjectNodeRemapPinsByPosition.h"
 
 #define LOCTEXT_NAMESPACE "CustomizableObjectEditor"
 
 
-void UCustomizableObjectNodeRemapPinsCustomExternalPin::RemapPins(const UCustomizableObjectNode& Node, const TArray<UEdGraphPin*>& OldPins, const TArray<UEdGraphPin*>& NewPins, TMap<UEdGraphPin*, UEdGraphPin*>& PinsToRemap, TArray<UEdGraphPin*>& PinsToOrphan)
-{
-	const UCustomizableObjectNodeExternalPin* ExternalPinNode = Cast<UCustomizableObjectNodeExternalPin>(&Node);
-	
-	if (OldPins.Num()) 
-	{
-		UEdGraphPin* const OldPin = OldPins[0];
-		if (NewPins.Num())
-		{
-			PinsToRemap.Add(OldPin, NewPins[0]);
-		}
-
-		if (!ExternalPinNode->GetNodeExposePin())
-		{
-			PinsToOrphan.Add(OldPin);
-		}
-	}
-}
-
-
 void UCustomizableObjectNodeExternalPin::Serialize(FArchive& Ar)
 {
 	Super::Serialize(Ar);
@@ -108,29 +89,12 @@ UCustomizableObjectNodeExposePin* UCustomizableObjectNodeExternalPin::GetNodeExp
 }
 
 
-UEdGraphPin* UCustomizableObjectNodeExternalPin::CreateExternalPin(const UCustomizableObjectNodeExposePin* NodeExposePin)
+void UCustomizableObjectNodeExternalPin::AllocateDefaultPins(UCustomizableObjectNodeRemapPins* RemapPins)
 {
-	FName PinName = NodeExposePin ? FName(NodeExposePin->GetNodeName()) : FName("Object");
+	FName PinName = GetNodeExposePin() ? FName(GetNodeExposePin()->GetNodeName()) : FName("Object");
 	const bool bIsArrayPinCategory = PinType == UEdGraphSchema_CustomizableObject::PC_GroupProjector;
 
-	return CustomCreatePin(EGPD_Output, PinType, PinName, bIsArrayPinCategory);
-}
-
-
-void UCustomizableObjectNodeExternalPin::BeginConstruct()
-{
-	// Create a pin already orphaned since the node does not have a linked Expose Pin node yet.
-	UEdGraphPin* Pin = CreateExternalPin(GetNodeExposePin());
-	OrphanPin(*Pin);
-}
-
-
-void UCustomizableObjectNodeExternalPin::AllocateDefaultPins(UCustomizableObjectNodeRemapPins* RemapPins)
-{
-	if (UCustomizableObjectNodeExposePin* NodeExposePin = GetNodeExposePin())
-	{
-		CreateExternalPin(NodeExposePin);
-	}
+	CustomCreatePin(EGPD_Output, PinType, PinName, bIsArrayPinCategory);
 }
 
 
@@ -195,7 +159,7 @@ void UCustomizableObjectNodeExternalPin::UpdateReferencedNodeId(const FGuid& New
 
 UCustomizableObjectNodeRemapPins* UCustomizableObjectNodeExternalPin::CreateRemapPinsDefault() const
 {
-	return NewObject<UCustomizableObjectNodeRemapPinsCustomExternalPin>();
+	return NewObject<UCustomizableObjectNodeRemapPinsByPosition>();
 }
 
 
