[Mutable] Fix crash and re-enable StreamingManager.

The crash was due to having recursive UCustomizableObjectSystem::InternalTick. Replaced the recursive tick call with a subset function of the tick.

#rnx

diff --git a/Engine/Plugins/Experimental/Mutable/Source/CustomizableObject/Internal/MuCO/CustomizableObjectSystemPrivate.h b/Engine/Plugins/Experimental/Mutable/Source/CustomizableObject/Internal/MuCO/CustomizableObjectSystemPrivate.h
--- a/Engine/Plugins/Experimental/Mutable/Source/CustomizableObject/Internal/MuCO/CustomizableObjectSystemPrivate.h
+++ b/Engine/Plugins/Experimental/Mutable/Source/CustomizableObject/Internal/MuCO/CustomizableObjectSystemPrivate.h
@@ -757,7 +757,10 @@ public:
 
 	static void HideOnScreenCompileWarnings(const UCustomizableObjectPrivate& ObjectPrivate);
 #endif
-	
+
+	/** Tick, in the Game Thread, anything that could potentially block the Mutable Thread. */
+	void TickMutableThreadDependencies();
+
 	/** Mutable TaskGraph system (Mutable Thread). */
 	FMutableTaskGraph MutableTaskGraph;
 	
diff --git a/Engine/Plugins/Experimental/Mutable/Source/CustomizableObject/Private/MuCO/CustomizableObjectSettings.h b/Engine/Plugins/Experimental/Mutable/Source/CustomizableObject/Private/MuCO/CustomizableObjectSettings.h
--- a/Engine/Plugins/Experimental/Mutable/Source/CustomizableObject/Private/MuCO/CustomizableObjectSettings.h
+++ b/Engine/Plugins/Experimental/Mutable/Source/CustomizableObject/Private/MuCO/CustomizableObjectSettings.h
@@ -17,5 +17,5 @@ class UCustomizableObjectSettings : public UObject
 public:
 	/** If true, use the new StreamManager ticker. If false, use the old FTSTicker. */
 	UPROPERTY(Config)
-	bool bEnableStreamingManager = false;
+	bool bEnableStreamingManager = true;
 };
diff --git a/Engine/Plugins/Experimental/Mutable/Source/CustomizableObject/Private/MuCO/CustomizableObjectSystem.cpp b/Engine/Plugins/Experimental/Mutable/Source/CustomizableObject/Private/MuCO/CustomizableObjectSystem.cpp
--- a/Engine/Plugins/Experimental/Mutable/Source/CustomizableObject/Private/MuCO/CustomizableObjectSystem.cpp
+++ b/Engine/Plugins/Experimental/Mutable/Source/CustomizableObject/Private/MuCO/CustomizableObjectSystem.cpp
@@ -1397,7 +1397,7 @@ bool UCustomizableObjectSystem::LockObject(const class UCustomizableObject* InOb
 			{
 				while (!Task.IsCompleted())				
 				{
-					GetPrivate()->UpdateResourceStreaming(0, false);
+					GetPrivate()->TickMutableThreadDependencies();
 				}
 			}
 			else
@@ -3135,11 +3135,7 @@ int32 UCustomizableObjectSystem::TickInternal()
 
 #if WITH_EDITOR
 	TickRecompileCustomizableObjects();
-
-	if (GetPrivate()->ImageProvider)
-	{
-		GetPrivate()->ImageProvider->Tick();		
-	}
+	GetPrivate()->TickMutableThreadDependencies();
 #endif
 
 #if WITH_EDITORONLY_DATA
@@ -3721,10 +3717,16 @@ void UCustomizableObjectSystemPrivate::HideOnScreenCompileWarnings(const UCustom
 		GEngine->RemoveOnScreenDebugMessage(reinterpret_cast<uint64>(&ObjectPrivate) + KEY_OFFSET_COMPILATION_OUT_OF_DATE);
 	}
 }
-
 #endif
 
 
+void UCustomizableObjectSystemPrivate::TickMutableThreadDependencies()
+{
+	check(IsInGameThread());
+	ImageProvider->Tick();
+}
+
+
 uint64 UCustomizableObjectSystem::GetMaxChunkSizeForPlatform(const ITargetPlatform* TargetPlatform)
 {
 	const FString& PlatformName = TargetPlatform ? TargetPlatform->IniPlatformName() : FPlatformProperties::IniPlatformName();
