[mutable] Added metadata info to table node options (rows)

#rnx

diff --git a/Engine/Plugins/Experimental/Mutable/Source/CustomizableObjectEditor/Classes/MuCOE/Nodes/CustomizableObjectNodeTable.h b/Engine/Plugins/Experimental/Mutable/Source/CustomizableObjectEditor/Classes/MuCOE/Nodes/CustomizableObjectNodeTable.h
--- a/Engine/Plugins/Experimental/Mutable/Source/CustomizableObjectEditor/Classes/MuCOE/Nodes/CustomizableObjectNodeTable.h
+++ b/Engine/Plugins/Experimental/Mutable/Source/CustomizableObjectEditor/Classes/MuCOE/Nodes/CustomizableObjectNodeTable.h
@@ -219,6 +219,10 @@ public:
 	UPROPERTY(EditAnywhere, Category = UI, meta = (DisplayName = "Parameter UI Metadata"))
 	FMutableParamUIMetadata ParamUIMetadata;
 
+	/** Name of the column that contains the MutableUIMetadata of the row options. */
+	UPROPERTY(EditAnywhere, Category = UI)
+	FName ParamUIMetadataColumn;
+
 	/** Map to relate a Structure Column with its Data */
 	UPROPERTY()
 	TMap<FGuid, FTableNodeColumnData> ColumnDataMap;
diff --git a/Engine/Plugins/Experimental/Mutable/Source/CustomizableObjectEditor/Private/MuCOE/GenerateMutableSource/GenerateMutableSourceTable.cpp b/Engine/Plugins/Experimental/Mutable/Source/CustomizableObjectEditor/Private/MuCOE/GenerateMutableSource/GenerateMutableSourceTable.cpp
--- a/Engine/Plugins/Experimental/Mutable/Source/CustomizableObjectEditor/Private/MuCOE/GenerateMutableSource/GenerateMutableSourceTable.cpp
+++ b/Engine/Plugins/Experimental/Mutable/Source/CustomizableObjectEditor/Private/MuCOE/GenerateMutableSource/GenerateMutableSourceTable.cpp
@@ -741,14 +741,56 @@ void GenerateTableParameterUIData(const UDataTable* DataTable, const UCustomizab
 		RestrictRowNamesToSelectedOption(RowNames, *TableNode, GenerationContext);
 
 		FParameterUIData ParameterUIData(TableNode->ParameterName, TableNode->ParamUIMetadata, EMutableParameterType::Int);
+		ParameterUIData.IntegerParameterGroupType = TableNode->bAddNoneOption ? ECustomizableObjectGroupType::COGT_ONE_OR_NONE : ECustomizableObjectGroupType::COGT_ONE;
+		FParameterUIData& ParameterUIDataRef = GenerationContext.ParameterUIDataMap.Add(TableNode->ParameterName, ParameterUIData);
 
-		for (int32 NameIndex = 0; NameIndex < RowNames.Num(); ++NameIndex)
+		if (TableNode->ParamUIMetadataColumn.IsNone())
 		{
-			ParameterUIData.ArrayIntegerParameterOption.Add(FIntegerParameterUIData(RowNames[NameIndex].ToString(), FMutableParamUIMetadata()));
+			return;
 		}
 
-		ParameterUIData.IntegerParameterGroupType = TableNode->bAddNoneOption ? ECustomizableObjectGroupType::COGT_ONE_OR_NONE : ECustomizableObjectGroupType::COGT_ONE;
-		GenerationContext.ParameterUIDataMap.Add(TableNode->ParameterName, ParameterUIData);
+		FProperty* ColumnProperty = DataTable->FindTableProperty(TableNode->ParamUIMetadataColumn);
+
+		if (!ColumnProperty)
+		{
+			FString msg = "Couldn't find Options UI Metadata Column [" + TableNode->ParamUIMetadataColumn.ToString() + "] in the Structure of the Node.";
+			GenerationContext.Compiler->CompilerLog(FText::FromString(msg), TableNode);
+
+			return;
+		}
+
+		FString WrongTypeMessage = "Column with name [" + TableNode->ParamUIMetadataColumn.ToString() + "] is not a Mutable Param UI Metadata type.";
+
+		if (const FStructProperty* StructProperty = CastField<FStructProperty>(ColumnProperty))
+		{
+			if (StructProperty->Struct != FMutableParamUIMetadata::StaticStruct())
+			{
+				GenerationContext.Compiler->CompilerLog(FText::FromString(WrongTypeMessage), TableNode);
+
+				return;
+			}
+			
+			for (int32 NameIndex = 0; NameIndex < RowNames.Num(); ++NameIndex)
+			{
+				// Getting Row Data
+				if (const uint8* RowData = DataTable->FindRowUnchecked(RowNames[NameIndex]))
+				{
+					// Getting Cell Data
+					if (const uint8* CellData = ColumnProperty->ContainerPtrToValuePtr<uint8>(RowData, 0))
+					{
+						// Getting cell value
+						FMutableParamUIMetadata Value = *(FMutableParamUIMetadata*)CellData;
+						ParameterUIDataRef.ArrayIntegerParameterOption.Add(FIntegerParameterUIData(RowNames[NameIndex].ToString(), Value));
+					}
+				}
+			}
+		}
+		else
+		{
+			GenerationContext.Compiler->CompilerLog(FText::FromString(WrongTypeMessage), TableNode);
+
+			return;
+		}
 	}
 }
 
diff --git a/Engine/Plugins/Experimental/Mutable/Source/CustomizableObjectEditor/Private/MuCOE/Nodes/CustomizableObjectNodeTable.cpp b/Engine/Plugins/Experimental/Mutable/Source/CustomizableObjectEditor/Private/MuCOE/Nodes/CustomizableObjectNodeTable.cpp
--- a/Engine/Plugins/Experimental/Mutable/Source/CustomizableObjectEditor/Private/MuCOE/Nodes/CustomizableObjectNodeTable.cpp
+++ b/Engine/Plugins/Experimental/Mutable/Source/CustomizableObjectEditor/Private/MuCOE/Nodes/CustomizableObjectNodeTable.cpp
@@ -284,6 +284,7 @@ void UCustomizableObjectNodeTable::PostEditChangeProperty(FPropertyChangedEvent&
 	{
 		if (PropertyThatChanged->GetName() == TEXT("Table") || PropertyThatChanged->GetName() == TEXT("Structure"))
 		{
+			ParamUIMetadataColumn = NAME_None;
 			ReconstructNode();
 		}
 		else if (PropertyThatChanged->GetName() == TEXT("DefaultImageMode"))
diff --git a/Engine/Plugins/Experimental/Mutable/Source/CustomizableObjectEditor/Private/MuCOE/Nodes/CustomizableObjectNodeTableDetails.cpp b/Engine/Plugins/Experimental/Mutable/Source/CustomizableObjectEditor/Private/MuCOE/Nodes/CustomizableObjectNodeTableDetails.cpp
--- a/Engine/Plugins/Experimental/Mutable/Source/CustomizableObjectEditor/Private/MuCOE/Nodes/CustomizableObjectNodeTableDetails.cpp
+++ b/Engine/Plugins/Experimental/Mutable/Source/CustomizableObjectEditor/Private/MuCOE/Nodes/CustomizableObjectNodeTableDetails.cpp
@@ -13,6 +13,7 @@
 #include "MuCOE/Nodes/CustomizableObjectNodeTable.h"
 #include "MuCOE/SCustomizableObjectNodeLayoutBlocksEditor.h"
 #include "MuCOE/UnrealEditorPortabilityHelpers.h"
+#include "Styling/SlateColor.h"
 #include "Widgets/Input/STextComboBox.h"
 #include "Widgets/Input/SButton.h"
 
@@ -41,6 +42,7 @@ void FCustomizableObjectNodeTableDetails::CustomizeDetails(const TSharedPtr<IDet
 	{
 		IDetailCategoryBuilder& CustomizableObjectCategory = DetailBuilder->EditCategory("TableProperties");
 		IDetailCategoryBuilder& UICategory = DetailBuilder->EditCategory("UI");
+		DetailBuilder->HideProperty("ParamUIMetadataColumn");
 		IDetailCategoryBuilder& AnimationCategory = DetailBuilder->EditCategory("AnimationProperties");
 		IDetailCategoryBuilder& LayoutCategory = DetailBuilder->EditCategory("DefaultMeshLayoutEditor");
 
@@ -48,6 +50,27 @@ void FCustomizableObjectNodeTableDetails::CustomizeDetails(const TSharedPtr<IDet
 		Node->PostReconstructNodeDelegate.AddSP(this, &FCustomizableObjectNodeTableDetails::OnNodePinValueChanged);
 
 		GenerateMeshColumnComboBoxOptions();
+		TSharedPtr<FString> CurrentMutableMetadataColumn = GenerateMutableMetaDataColumnComboBoxOptions();
+
+		UICategory.AddCustomRow(LOCTEXT("MutableUIMetadataColumn_Selector","MutableUIMetadataColumn"))
+		.NameContent()
+		[
+			SNew(STextBlock)
+			.Text(LOCTEXT("MutableUIMetadataColumn_SelectorText","Options UI Metadata Column"))
+			.ToolTipText(LOCTEXT("MutableUIMetadataColumn_SelectorTooltip","Select a column that contains a Parameter UI Metadata for each Parameter Option (table row)."))
+			.Font(DetailBuilder->GetDetailFont())
+		]
+		.ValueContent()
+		[
+			SAssignNew(MutableMetaDataComboBox,STextComboBox)
+			.InitiallySelectedItem(CurrentMutableMetadataColumn)
+			.OptionsSource(&MutableMetaDataColumnsOptionNames)
+			.OnComboBoxOpening(this, &FCustomizableObjectNodeTableDetails::OnOpenMutableMetadataComboBox)
+			.OnSelectionChanged(this, &FCustomizableObjectNodeTableDetails::OnMutableMetaDataColumnComboBoxSelectionChanged)
+			.Font(DetailBuilder->GetDetailFont())
+			.ColorAndOpacity(this, &FCustomizableObjectNodeTableDetails::GetMetadataUIComboBoxTextColor, &MutableMetaDataColumnsOptionNames)
+		]
+		.OverrideResetToDefault(FResetToDefaultOverride::Create(FSimpleDelegate::CreateSP(this, &FCustomizableObjectNodeTableDetails::OnMutableMetaDataColumnComboBoxSelectionReset)));
 
 		AnimationCategory.AddCustomRow(LOCTEXT("AnimationProperties", "Animation Properties"))
 		[
@@ -592,4 +615,99 @@ FReply FCustomizableObjectNodeTableDetails::OnClearButtonPressed()
 }
 
 
+TSharedPtr<FString> FCustomizableObjectNodeTableDetails::GenerateMutableMetaDataColumnComboBoxOptions()
+{
+	const UScriptStruct* TableStruct = Node->GetTableNodeStruct();
+	TSharedPtr<FString> CurrentSelection;
+	MutableMetaDataColumnsOptionNames.Reset();
+
+	if (!TableStruct)
+	{
+		return CurrentSelection;
+	}
+
+	// Iterating struct Options
+	for (TFieldIterator<FProperty> It(TableStruct); It; ++It)
+	{
+		FProperty* ColumnProperty = *It;
+
+		if (!ColumnProperty)
+		{
+			continue;
+		}
+
+		if (const FStructProperty* StructProperty = CastField<FStructProperty>(ColumnProperty))
+		{
+			if (StructProperty->Struct == FMutableParamUIMetadata::StaticStruct())
+			{
+				TSharedPtr<FString> Option = MakeShareable(new FString(DataTableUtils::GetPropertyExportName(ColumnProperty)));
+				MutableMetaDataColumnsOptionNames.Add(Option);
+
+				if (*Option == Node->ParamUIMetadataColumn)
+				{
+					CurrentSelection = MutableMetaDataColumnsOptionNames.Last();
+				}
+			}
+		}
+	}
+
+	if (!Node->ParamUIMetadataColumn.IsNone() && !CurrentSelection)
+	{
+		MutableMetaDataColumnsOptionNames.Add(MakeShareable(new FString(Node->ParamUIMetadataColumn.ToString())));
+		CurrentSelection = MutableMetaDataColumnsOptionNames.Last();
+	}
+
+	return CurrentSelection;
+}
+
+
+void FCustomizableObjectNodeTableDetails::OnOpenMutableMetadataComboBox()
+{
+	TSharedPtr<FString> CurrentSelection = GenerateMutableMetaDataColumnComboBoxOptions();
+
+	if (MutableMetaDataComboBox.IsValid())
+	{
+		MutableMetaDataComboBox->ClearSelection();
+		MutableMetaDataComboBox->RefreshOptions();
+		MutableMetaDataComboBox->SetSelectedItem(CurrentSelection);
+	}
+}
+
+
+void FCustomizableObjectNodeTableDetails::OnMutableMetaDataColumnComboBoxSelectionChanged(TSharedPtr<FString> Selection, ESelectInfo::Type SelectInfo)
+{
+	if (Selection && Node->ParamUIMetadataColumn != FName(*Selection) 
+		&& (SelectInfo == ESelectInfo::OnKeyPress || SelectInfo == ESelectInfo::OnMouseClick))
+	{
+		Node->ParamUIMetadataColumn = FName(*Selection);
+		Node->MarkPackageDirty();
+	}
+}
+
+
+FSlateColor FCustomizableObjectNodeTableDetails::GetMetadataUIComboBoxTextColor(TArray<TSharedPtr<FString>>* CurrentOptions) const
+{	
+	if (Node->FindTableProperty(Node->GetTableNodeStruct(), Node->ParamUIMetadataColumn) || Node->ParamUIMetadataColumn.IsNone())
+	{
+		return FSlateColor::UseForeground();
+	}
+
+	// Table Struct null or does not contain the selected property anymore
+	return FSlateColor(FLinearColor(0.9f, 0.05f, 0.05f, 1.0f));
+}
+
+void FCustomizableObjectNodeTableDetails::OnMutableMetaDataColumnComboBoxSelectionReset()
+{
+	Node->ParamUIMetadataColumn = NAME_None;
+
+	if (MutableMetaDataComboBox.IsValid())
+	{
+		GenerateMutableMetaDataColumnComboBoxOptions();
+		MutableMetaDataComboBox->ClearSelection();
+		MutableMetaDataComboBox->RefreshOptions();
+	}
+	
+}
+
+
 #undef LOCTEXT_NAMESPACE
diff --git a/Engine/Plugins/Experimental/Mutable/Source/CustomizableObjectEditor/Private/MuCOE/Nodes/CustomizableObjectNodeTableDetails.h b/Engine/Plugins/Experimental/Mutable/Source/CustomizableObjectEditor/Private/MuCOE/Nodes/CustomizableObjectNodeTableDetails.h
--- a/Engine/Plugins/Experimental/Mutable/Source/CustomizableObjectEditor/Private/MuCOE/Nodes/CustomizableObjectNodeTableDetails.h
+++ b/Engine/Plugins/Experimental/Mutable/Source/CustomizableObjectEditor/Private/MuCOE/Nodes/CustomizableObjectNodeTableDetails.h
@@ -15,6 +15,7 @@ class FReply;
 class SButton;
 class SCustomizableObjectNodeLayoutBlocksEditor;
 class UCustomizableObjectNodeTable;
+struct FSlateColor;
 
 /** Copy Material node details panel. Hides all properties from the inheret Material node. */
 class FCustomizableObjectNodeTableDetails : public IDetailCustomization
@@ -68,6 +69,22 @@ private:
 	// Callback to clear the animation combobox selections
 	FReply OnClearButtonPressed();
 
+	// Generates MutableMetadata columns combobox options
+	// Returns the current selected option or a null pointer
+	TSharedPtr<FString> GenerateMutableMetaDataColumnComboBoxOptions();
+
+	// Callback to regenerate the combobox options
+	void OnOpenMutableMetadataComboBox();
+
+	// OnComboBoxSelectionChanged Callback for Layout ComboBox
+	void OnMutableMetaDataColumnComboBoxSelectionChanged(TSharedPtr<FString> Selection, ESelectInfo::Type SelectInfo);
+
+	// Sets the combo box selection color
+	FSlateColor GetMetadataUIComboBoxTextColor(TArray<TSharedPtr<FString>>* CurrentOptions) const;
+
+	// OnComboBoxSelectionChanged Callback for Layout ComboBox
+	void OnMutableMetaDataColumnComboBoxSelectionReset();
+
 	// Pointer to the node represented in this details
 	TWeakObjectPtr<UCustomizableObjectNodeTable> Node;
 
@@ -101,6 +118,12 @@ private:
 	// Array with the name of the Animation Tags columns as combobox options
 	TArray<TSharedPtr<FString>> AnimTagsOptionNames;
 	
+	// Array with the name of the MutableMetaData columns
+	TArray<TSharedPtr<FString>> MutableMetaDataColumnsOptionNames;
+
+	// ComboBox widget to select a MutableMetaDatacolumn from the NodeTable
+	TSharedPtr<STextComboBox> MutableMetaDataComboBox;
+	
 	// Button to clear selections of the animation comboboxes
 	TSharedPtr<SButton> ClearButton;
 
