[mutable] Updated MutableGraphViewer with 43 new node types (so they can be displayed in the tree view)
- Added 43 types of nodes to the MutableGraphViewer method that generates the tree.
- Added dummy treeElement placed for cases where the type of the mutable node has not been explicitly set as a possible mutable graph node to be part of the MutableGraphView treeview.
- Cyborg, Bandit and Weapon should now be fully (or close to be) debuggable using the MutableGraphViewer.

diff --git a/Engine/Plugins/Experimental/Mutable/Source/CustomizableObjectEditor/Private/MuCOE/SMutableGraphViewer.cpp b/Engine/Plugins/Experimental/Mutable/Source/CustomizableObjectEditor/Private/MuCOE/SMutableGraphViewer.cpp
--- a/Engine/Plugins/Experimental/Mutable/Source/CustomizableObjectEditor/Private/MuCOE/SMutableGraphViewer.cpp
+++ b/Engine/Plugins/Experimental/Mutable/Source/CustomizableObjectEditor/Private/MuCOE/SMutableGraphViewer.cpp
@@ -17,12 +17,46 @@
 #include "Widgets/Input/SNumericDropDown.h"
 #include "Widgets/Views/STreeView.h"
 #include "ScopedTransaction.h"
+#include "MuT/NodeColourConstant.h"
+#include "MuT/NodeColourFromScalars.h"
+#include "MuT/NodeColourParameter.h"
+#include "MuT/NodeColourSampleImage.h"
+#include "MuT/NodeColourSwitch.h"
+#include "MuT/NodeComponentEdit.h"
+#include "MuT/NodeImageFormat.h"
+#include "MuT/NodeImageInterpolate.h"
+#include "MuT/NodeImageInvert.h"
+#include "MuT/NodeImageLayer.h"
+#include "MuT/NodeImageLayerColour.h"
+#include "MuT/NodeImageMipmap.h"
+#include "MuT/NodeImageMultiLayer.h"
+#include "MuT/NodeImagePlainColour.h"
+#include "MuT/NodeImageProject.h"
+#include "MuT/NodeImageResize.h"
+#include "MuT/NodeImageSwitch.h"
+#include "MuT/NodeImageSwizzle.h"
+#include "MuT/NodeImageTable.h"
 #include "MuT/NodeObjectGroup.h"
 #include "MuT/NodeObjectNew.h"
 #include "MuT/NodeSurfaceEdit.h"
 #include "MuT/NodeSurfaceSwitch.h"
 #include "MuT/NodeSurfaceVariation.h"
 #include "MuT/Streams.h"
+#include "MuT/NodeLOD.h"
+#include "MuT/NodeMeshConstant.h"
+#include "MuT/NodeMeshFormat.h"
+#include "MuT/NodeMeshFragment.h"
+#include "MuT/NodeMeshMakeMorph.h"
+#include "MuT/NodeMeshMorph.h"
+#include "MuT/NodeMeshTable.h"
+#include "MuT/NodeModifierMeshClipDeform.h"
+#include "MuT/NodeModifierMeshClipMorphPlane.h"
+#include "MuT/NodeModifierMeshClipWithUVMask.h"
+#include "MuT/NodeScalarConstant.h"
+#include "MuT/NodeScalarCurve.h"
+#include "MuT/NodeScalarSwitch.h"
+#include "MuT/NodeScalarTable.h"
+
 
 // This is necessary because of problems with rtti information in other platforms. In any case, this part of the debugger is only useful in the standard editor.
 #if PLATFORM_WINDOWS
@@ -32,6 +66,45 @@
 #include "MuT/NodeSurfaceEditPrivate.h"
 #include "MuT/NodeSurfaceSwitchPrivate.h"
 #include "MuT/NodeSurfaceVariationPrivate.h"
+#include "MuT/NodeLODPrivate.h"
+#include "MuT/NodeComponentPrivate.h"
+#include "MuT/NodeModifierPrivate.h"
+#include "MuT/NodeComponentNewPrivate.h"
+#include "MuT/NodeComponentEditPrivate.h"
+#include "MuT/NodeImageFormatPrivate.h"
+#include "MuT/NodeMeshFormatPrivate.h"
+#include "MuT/NodePatchImagePrivate.h"
+#include "MuT/NodeMeshConstantPrivate.h"
+#include "MuT/NodeModifierMeshClipMorphPlanePrivate.h"
+#include "MuT/NodePatchMeshPrivate.h"
+#include "MuT/NodeImageSwitchPrivate.h"
+#include "MuT/NodeImageLayerColourPrivate.h"
+#include "MuT/NodeImageLayerPrivate.h"
+#include "MuT/NodeImageMipmapPrivate.h"
+#include "MuT/NodeImageResizePrivate.h"
+#include "MuT/NodeModifierMeshClipDeformPrivate.h"
+#include "MuT/NodeModifierMeshClipWithMeshPrivate.h"
+#include "MuT/NodeModifierMeshClipWithUVMaskPrivate.h"
+#include "MuT/NodeColourParameterPrivate.h"
+#include "MuT/NodeColourSampleImagePrivate.h"
+#include "MuT/NodeImageInterpolatePrivate.h"
+#include "MuT/NodeImagePlainColourPrivate.h"
+#include "MuT/NodeImageProjectPrivate.h"
+#include "MuT/NodeMeshFragmentPrivate.h"
+#include "MuT/NodeMeshMorphPrivate.h"
+#include "MuT/NodeScalarEnumParameterPrivate.h"
+#include "MuT/NodeScalarParameterPrivate.h"
+#include "MuT/NodeColourPrivate.h"
+#include "MuT/NodeMeshMakeMorphPrivate.h"
+#include "MuT/NodeScalarCurvePrivate.h"
+#include "MuT/NodeProjectorPrivate.h"
+#include "MuT/NodeColourSwitchPrivate.h"
+#include "MuT/NodeImageInvertPrivate.h"
+#include "MuT/NodeImageSwizzlePrivate.h"
+#include "MuT/NodeImageMultiLayerPrivate.h"
+#include "MuT/NodeMeshTablePrivate.h"
+#include "MuT/NodeColourFromScalarsPrivate.h"
+#include "MuT/NodeScalarSwitchPrivate.h"
 #endif
 
 class FExtender;
@@ -44,7 +117,6 @@ struct FSlateBrush;
 
 #define LOCTEXT_NAMESPACE "SMutableDebugger"
 
-
 // \todo: multi-column tree
 namespace MutableGraphTreeViewColumns
 {
@@ -60,18 +132,27 @@ public:
 	{
 		RowItem = InRowItem;
 
-		const char* TypeName = RowItem->MutableNode->GetType()->m_strName;
-
-		FString LabelString = RowItem->Prefix.IsEmpty() 
-			? StringCast<TCHAR>(TypeName).Get() 
-			: FString::Printf( TEXT("%s : %s"), *RowItem->Prefix, StringCast<TCHAR>(TypeName).Get() );
-
-		FText MainLabel = FText::FromString(LabelString);
-		if (RowItem->DuplicatedOf)
+		FText MainLabel = FText::GetEmpty();
+		if (RowItem->MutableNode)
+		{
+			const char* TypeName = RowItem->MutableNode->GetType()->m_strName;
+			
+			const FString LabelString = RowItem->Prefix.IsEmpty() 
+				? StringCast<TCHAR>(TypeName).Get() 
+				: FString::Printf( TEXT("%s : %s"), *RowItem->Prefix, StringCast<TCHAR>(TypeName).Get() );
+
+			MainLabel = FText::FromString(LabelString);
+			if (RowItem->DuplicatedOf)
+			{
+				MainLabel = FText::FromString( FString::Printf(TEXT("%s (Duplicated)"), StringCast<TCHAR>(TypeName).Get()));
+			}
+		}
+		else
 		{
-			MainLabel = FText::FromString( FString::Printf(TEXT("%s (Duplicated)"), StringCast<TCHAR>(TypeName).Get()));
+			MainLabel = FText::FromString( *RowItem->Prefix);
 		}
 
+
 		this->ChildSlot
 		[
 			SNew(SHorizontalBox)
@@ -290,12 +371,18 @@ void SMutableGraphViewer::GetChildrenForInfo(TSharedPtr<FMutableGraphTreeElement
 				}
 			}
 		}
+		else
+		{
+			// No mutable node has been provided so create a dummy tree element
+			TSharedPtr<FMutableGraphTreeElement> Item = MakeShareable(new FMutableGraphTreeElement(nullptr, nullptr , Prefix));
+			OutChildren.Add(Item);
+		}
 		++InputIndex;
 	};
 
 	if (ParentNode->GetType() == mu::NodeObjectNew::GetStaticType())
 	{
-		mu::NodeObjectNew* ObjectNew = reinterpret_cast<mu::NodeObjectNew*>(ParentNode);
+		mu::NodeObjectNew* ObjectNew = StaticCast<mu::NodeObjectNew*>(ParentNode);
 		mu::NodeObjectNew::Private* Private = ObjectNew->GetPrivate();
 		for (int32 l = 0; l < Private->m_lods.Num(); ++l)
 		{
@@ -310,7 +397,7 @@ void SMutableGraphViewer::GetChildrenForInfo(TSharedPtr<FMutableGraphTreeElement
 
 	else if (ParentNode->GetType() == mu::NodeObjectGroup::GetStaticType())
 	{
-		mu::NodeObjectGroup* ObjectGroup = reinterpret_cast<mu::NodeObjectGroup*>(ParentNode);
+		mu::NodeObjectGroup* ObjectGroup = StaticCast<mu::NodeObjectGroup*>(ParentNode);
 		mu::NodeObjectGroup::Private* Private = ObjectGroup->GetPrivate();
 		for (int32 l = 0; l < Private->m_children.Num(); ++l)
 		{
@@ -320,7 +407,7 @@ void SMutableGraphViewer::GetChildrenForInfo(TSharedPtr<FMutableGraphTreeElement
 
 	else if (ParentNode->GetType() == mu::NodeSurfaceNew::GetStaticType())
 	{
-		mu::NodeSurfaceNew* SurfaceNew = reinterpret_cast<mu::NodeSurfaceNew*>(ParentNode);
+		mu::NodeSurfaceNew* SurfaceNew = StaticCast<mu::NodeSurfaceNew*>(ParentNode);
 		mu::NodeSurfaceNew::Private* Private = SurfaceNew->GetPrivate();
 		for (int32 l = 0; l < Private->m_meshes.Num(); ++l)
 		{
@@ -345,7 +432,7 @@ void SMutableGraphViewer::GetChildrenForInfo(TSharedPtr<FMutableGraphTreeElement
 
 	else if (ParentNode->GetType() == mu::NodeSurfaceEdit::GetStaticType())
 	{
-		mu::NodeSurfaceEdit* SurfaceEdit = reinterpret_cast<mu::NodeSurfaceEdit*>(ParentNode);
+		mu::NodeSurfaceEdit* SurfaceEdit = StaticCast<mu::NodeSurfaceEdit*>(ParentNode);
 		mu::NodeSurfaceEdit::Private* Private = SurfaceEdit->GetPrivate();
 		AddChildFunc(Private->m_pMesh.get(), TEXT("MESH"));
 		AddChildFunc(Private->m_pMorph.get(), TEXT("MORPH"));
@@ -360,10 +447,9 @@ void SMutableGraphViewer::GetChildrenForInfo(TSharedPtr<FMutableGraphTreeElement
 
 	else if (ParentNode->GetType() == mu::NodeSurfaceSwitch::GetStaticType())
 	{
-		mu::NodeSurfaceSwitch* SurfaceSwitch = reinterpret_cast<mu::NodeSurfaceSwitch*>(ParentNode);
+		mu::NodeSurfaceSwitch* SurfaceSwitch = StaticCast<mu::NodeSurfaceSwitch*>(ParentNode);
 		mu::NodeSurfaceSwitch::Private* Private = SurfaceSwitch->GetPrivate();
 		AddChildFunc(Private->Parameter.get(), TEXT("PARAM"));
-
 		for (int32 l = 0; l < Private->Options.Num(); ++l)
 		{
 			AddChildFunc(Private->Options[l].get(), FString::Printf(TEXT("OPTION [%d]"), l));
@@ -372,7 +458,7 @@ void SMutableGraphViewer::GetChildrenForInfo(TSharedPtr<FMutableGraphTreeElement
 
 	else if (ParentNode->GetType() == mu::NodeSurfaceVariation::GetStaticType())
 	{
-		mu::NodeSurfaceVariation* SurfaceVar = reinterpret_cast<mu::NodeSurfaceVariation*>(ParentNode);
+		mu::NodeSurfaceVariation* SurfaceVar = StaticCast<mu::NodeSurfaceVariation*>(ParentNode);
 		mu::NodeSurfaceVariation::Private* Private = SurfaceVar->GetPrivate();
 		for (int32 l = 0; l < Private->m_defaultSurfaces.Num(); ++l)
 		{
@@ -396,11 +482,377 @@ void SMutableGraphViewer::GetChildrenForInfo(TSharedPtr<FMutableGraphTreeElement
 			}
 		}
 	}
+	
+	else if (ParentNode->GetType() == mu::NodeLOD::GetStaticType())
+	{
+		mu::NodeLOD* LodVar = StaticCast<mu::NodeLOD*>(ParentNode);
+		mu::NodeLOD::Private* Private = LodVar->GetPrivate();
+
+		for (int32 Component = 0; Component < Private->m_components.Num(); Component++)
+		{
+			AddChildFunc(Private->m_components[Component].get(), FString::Printf(TEXT("COMP [%d]"),  Component));
+		}
+		for (int32 Modifier = 0; Modifier < Private->m_modifiers.Num(); Modifier++)
+		{
+			AddChildFunc(Private->m_modifiers[Modifier].get(), FString::Printf(TEXT("MOD [%d]"), Modifier));
+		}
+	}
+	
+	else if (ParentNode->GetType() == mu::NodeComponentNew::GetStaticType())
+	{
+		mu::NodeComponentNew* ComponentVar = StaticCast<mu::NodeComponentNew*>(ParentNode);
+		mu::NodeComponentNew::Private* Private = ComponentVar->GetPrivate();
+		for (int32 Surface = 0; Surface < Private->m_surfaces.Num(); Surface++)
+		{
+			AddChildFunc(Private->m_surfaces[Surface].get(), FString::Printf(TEXT("SURF [%d]"), Surface));
+		}
+	}
+
+	else if (ParentNode->GetType() == mu::NodeComponentEdit::GetStaticType())
+	{
+		mu::NodeComponentEdit* ComponentEditVar = StaticCast<mu::NodeComponentEdit*>(ParentNode);
+		mu::NodeComponentEdit::Private* Private = ComponentEditVar->GetPrivate();
+		for (int32 Surface = 0; Surface < Private->m_surfaces.Num(); Surface++)
+		{
+			AddChildFunc(Private->m_surfaces[Surface].get(), FString::Printf(TEXT("SURF [%d]"), Surface));
+		}
+	}
+	
+
+	else if (ParentNode->GetType() == mu::NodeMeshConstant::GetStaticType())
+	{
+		mu::NodeMeshConstant* MeshConstantVar = StaticCast<mu::NodeMeshConstant*>(ParentNode);
+		mu::NodeMeshConstant::Private* Private = MeshConstantVar->GetPrivate();
+		for (int32 LayoutIndex = 0; LayoutIndex < Private->m_layouts.Num(); LayoutIndex++)
+		{
+			AddChildFunc(Private->m_layouts[LayoutIndex].get(), FString::Printf(TEXT("LAYOUT [%d]"), LayoutIndex));
+		}
+	}
+
+	else if (ParentNode->GetType() == mu::NodeImageFormat::GetStaticType())
+	{
+		mu::NodeImageFormat* ImageFormatVar = StaticCast<mu::NodeImageFormat*>(ParentNode);
+		mu::NodeImageFormat::Private* Private = ImageFormatVar->GetPrivate();
+		AddChildFunc(Private->m_source.get(), FString::Printf(TEXT("SOURCE IMAGE")));
+	}
+
+	else if (ParentNode->GetType() == mu::NodeMeshFormat::GetStaticType())
+	{
+		mu::NodeMeshFormat* MeshFormatVar = StaticCast<mu::NodeMeshFormat*>(ParentNode);
+		mu::NodeMeshFormat::Private* Private = MeshFormatVar->GetPrivate();
+		AddChildFunc(Private->m_pSource.get(), FString::Printf(TEXT("SOURCE MESH")));
+	}
+
+	else if (ParentNode->GetType() == mu::NodePatchImage::GetStaticType())
+	{
+		mu::NodePatchImage* PatchImageVar = StaticCast<mu::NodePatchImage*>(ParentNode);
+		mu::NodePatchImage::Private* Private = PatchImageVar->GetPrivate();
+		AddChildFunc(Private->m_pImage.get(), FString::Printf(TEXT("IMAGE")));
+		AddChildFunc(Private->m_pMask.get(), FString::Printf(TEXT("MASK")));
+	}
+
+	else if (ParentNode->GetType() == mu::NodeModifierMeshClipMorphPlane::GetStaticType())
+	{
+		// Nothing to show
+	}
+
+	else if (ParentNode->GetType() == mu::NodeModifierMeshClipWithMesh::GetStaticType())
+	{
+		mu::NodeModifierMeshClipWithMesh* ModifierMeshClipWithMeshVar = StaticCast<mu::NodeModifierMeshClipWithMesh*>(ParentNode);
+		mu::NodeModifierMeshClipWithMesh::Private* Private = ModifierMeshClipWithMeshVar->GetPrivate();
+		AddChildFunc(Private->ClipMesh.get(), FString::Printf(TEXT("CLIP MESH")));
+	}
+
+	else if (ParentNode->GetType() == mu::NodeModifierMeshClipDeform::GetStaticType())
+	{
+		mu::NodeModifierMeshClipDeform* ModifierMeshClipDeformVar = StaticCast<mu::NodeModifierMeshClipDeform*>(ParentNode);
+		mu::NodeModifierMeshClipDeform::Private* Private = ModifierMeshClipDeformVar->GetPrivate();
+		AddChildFunc(Private->ClipMesh.get(), FString::Printf(TEXT("CLIP MESH")));
+	}
+
+	else if (ParentNode->GetType() == mu::NodeModifierMeshClipWithUVMask::GetStaticType())
+	{
+		mu::NodeModifierMeshClipWithUVMask* ModifierMeshClipWithUVMaskVar = StaticCast<mu::NodeModifierMeshClipWithUVMask*>(ParentNode);
+		mu::NodeModifierMeshClipWithUVMask::Private* Private = ModifierMeshClipWithUVMaskVar->GetPrivate();
+		AddChildFunc(Private->ClipMask.get(), FString::Printf(TEXT("CLIP MASK")));
+	}
+	
+	else if (ParentNode->GetType() == mu::NodePatchMesh::GetStaticType())
+	{
+		mu::NodePatchMesh* PatchMeshVar = StaticCast<mu::NodePatchMesh*>(ParentNode);
+		mu::NodePatchMesh::Private* Private = PatchMeshVar->GetPrivate();
+		AddChildFunc(Private->m_pAdd.get(), FString::Printf(TEXT("ADD")));
+		AddChildFunc(Private->m_pRemove.get(), FString::Printf(TEXT("REMOVE")));
+	}
+
+	else if (ParentNode->GetType() == mu::NodeImageSwitch::GetStaticType())
+	{
+		mu::NodeImageSwitch* ImageSwitchVar = StaticCast<mu::NodeImageSwitch*>(ParentNode);
+		mu::NodeImageSwitch::Private* Private = ImageSwitchVar->GetPrivate();
+		AddChildFunc(Private->m_pParameter.get(), FString::Printf(TEXT("PARAM")));
+		for (int32 OptionIndex = 0; OptionIndex < Private->m_options.Num(); OptionIndex++)
+		{
+			AddChildFunc(Private->m_options[OptionIndex].get(), FString::Printf(TEXT("OPTION [%d]"), OptionIndex));
+		}
+	}
+
+	else if (ParentNode->GetType() == mu::NodeImageMipmap::GetStaticType())
+	{
+		mu::NodeImageMipmap* ImageMipMapVar = StaticCast<mu::NodeImageMipmap*>(ParentNode);
+		mu::NodeImageMipmap::Private* Private = ImageMipMapVar->GetPrivate();
+		AddChildFunc(Private->m_pSource.get(), FString::Printf(TEXT("SOURCE")));
+		AddChildFunc(Private->m_pFactor.get(), FString::Printf(TEXT("FACTOR")));
+	}
+
+	else if (ParentNode->GetType() == mu::NodeImageLayer::GetStaticType())
+	{
+		mu::NodeImageLayer* ImageLayerVar = StaticCast<mu::NodeImageLayer*>(ParentNode);
+		mu::NodeImageLayer::Private* Private = ImageLayerVar->GetPrivate();
+		AddChildFunc(Private->m_pBase.get(), FString::Printf(TEXT("BASE")));
+		AddChildFunc(Private->m_pMask.get(), FString::Printf(TEXT("MASK")));
+		AddChildFunc(Private->m_pBlended.get(), FString::Printf(TEXT("BLEND")));
+	}
+	
+	else if (ParentNode->GetType() == mu::NodeImageLayerColour::GetStaticType())
+	{
+		mu::NodeImageLayerColour* ImageLayerColourVar = StaticCast<mu::NodeImageLayerColour*>(ParentNode);
+		mu::NodeImageLayerColour::Private* Private = ImageLayerColourVar->GetPrivate();
+		AddChildFunc(Private->m_pBase.get(), FString::Printf(TEXT("BASE")));
+		AddChildFunc(Private->m_pMask.get(), FString::Printf(TEXT("MASK")));
+		AddChildFunc(Private->m_pColour.get(), FString::Printf(TEXT("COLOR")));
+	}
+
+	else if (ParentNode->GetType() == mu::NodeImageResize::GetStaticType())
+	{
+		mu::NodeImageResize* ImageResizeVar = StaticCast<mu::NodeImageResize*>(ParentNode);
+		mu::NodeImageResize::Private* Private = ImageResizeVar->GetPrivate();
+		AddChildFunc(Private->m_pBase.get(), FString::Printf(TEXT("BASE")));
+	}
+
+	else if (ParentNode->GetType() == mu::NodeMeshMorph::GetStaticType())
+	{
+		mu::NodeMeshMorph* MeshMorphVar = StaticCast<mu::NodeMeshMorph*>(ParentNode);
+		mu::NodeMeshMorph::Private* Private = MeshMorphVar->GetPrivate();
+		AddChildFunc(Private->Base.get(), FString::Printf(TEXT("BASE")));
+		AddChildFunc(Private->Morph.get(), FString::Printf(TEXT("MORPH")));
+		AddChildFunc(Private->Factor.get(), FString::Printf(TEXT("FACTOR")));
+	}
+
+	else if (ParentNode->GetType() == mu::NodeImageProject::GetStaticType())
+	{
+		mu::NodeImageProject* ImageProjectVar = StaticCast<mu::NodeImageProject*>(ParentNode);
+		mu::NodeImageProject::Private* Private = ImageProjectVar->GetPrivate();
+		AddChildFunc(Private->m_pProjector.get(), FString::Printf(TEXT("PROJECTOR")));
+		AddChildFunc(Private->m_pMesh.get(), FString::Printf(TEXT("MESH")));
+		AddChildFunc(Private->m_pImage.get(), FString::Printf(TEXT("IMAGE")));
+		AddChildFunc(Private->m_pMask.get(), FString::Printf(TEXT("MASK")));
+		AddChildFunc(Private->m_pAngleFadeStart.get(), FString::Printf(TEXT("FADE START ANGLE")));
+		AddChildFunc(Private->m_pAngleFadeEnd.get(), FString::Printf(TEXT("FADE END ANGLE")));
+	}
+
+	else if (ParentNode->GetType() == mu::NodeImagePlainColour::GetStaticType())
+	{
+		mu::NodeImagePlainColour* ImagePlainColourVar = StaticCast<mu::NodeImagePlainColour*>(ParentNode);
+        mu::NodeImagePlainColour::Private* Private = ImagePlainColourVar->GetPrivate();
+		AddChildFunc(Private->m_pColour.get(), FString::Printf(TEXT("COLOR")));
+	}
+
+	else if (ParentNode->GetType() == mu::NodeLayoutBlocks::GetStaticType())
+	{
+		// Nothing to show
+	}
+
+	else if (ParentNode->GetType() == mu::NodeScalarEnumParameter::GetStaticType())
+	{
+		mu::NodeScalarEnumParameter* ScalarEnumParameterVar = StaticCast<mu::NodeScalarEnumParameter*>(ParentNode);
+		mu::NodeScalarEnumParameter::Private* Private = ScalarEnumParameterVar->GetPrivate();
+		for (int32 RangeIndex = 0; RangeIndex < Private->m_ranges.Num(); RangeIndex++)
+		{
+			AddChildFunc(Private->m_ranges[RangeIndex].get(), FString::Printf(TEXT("RANGE [%d]"), RangeIndex));
+		}
+	}
+
+	else if (ParentNode->GetType() == mu::NodeMeshFragment::GetStaticType())
+	{
+		mu::NodeMeshFragment* MeshFragmentVar = StaticCast<mu::NodeMeshFragment*>(ParentNode);
+		mu::NodeMeshFragment::Private* Private = MeshFragmentVar->GetPrivate();
+		AddChildFunc(Private->m_pMesh.get(), FString::Printf(TEXT("MESH")));
+	}
+
+	else if (ParentNode->GetType() == mu::NodeColourSampleImage::GetStaticType())
+	{
+		mu::NodeColourSampleImage* ColorSampleImageVar = StaticCast<mu::NodeColourSampleImage*>(ParentNode);
+		mu::NodeColourSampleImage::Private* Private = ColorSampleImageVar->GetPrivate();
+		AddChildFunc(Private->m_pImage.get(), FString::Printf(TEXT("IMAGE")));
+		AddChildFunc(Private->m_pX.get(), FString::Printf(TEXT("X")));
+		AddChildFunc(Private->m_pY.get(), FString::Printf(TEXT("Y")));
+	}
 
+	else if (ParentNode->GetType() == mu::NodeImageInterpolate::GetStaticType())
+	{
+		mu::NodeImageInterpolate* ImageInterpolateVar = StaticCast<mu::NodeImageInterpolate*>(ParentNode);
+		mu::NodeImageInterpolate::Private* Private = ImageInterpolateVar->GetPrivate();
+		AddChildFunc(Private->m_pFactor.get(), FString::Printf(TEXT("FACTOR")));
+		for (int32 TargetIndex = 0; TargetIndex < Private->m_targets.Num(); TargetIndex++)
+		{
+			AddChildFunc(Private->m_targets[TargetIndex].get(), FString::Printf(TEXT("TARGET [%d]"), TargetIndex));
+		}
+	}
+	
+	else if (ParentNode->GetType() == mu::NodeScalarConstant::GetStaticType())
+	{
+		// Nothing to show
+	}
+
+	else if (ParentNode->GetType() == mu::NodeScalarParameter::GetStaticType())
+	{
+		mu::NodeScalarParameter* ScalarParameterVar = StaticCast<mu::NodeScalarParameter*>(ParentNode);
+		mu::NodeScalarParameter::Private* Private = ScalarParameterVar->GetPrivate();
+		for (int32 RangeIndex = 0; RangeIndex < Private->m_ranges.Num(); RangeIndex++)
+		{
+			AddChildFunc(Private->m_ranges[RangeIndex].get(), FString::Printf(TEXT("RANGE [%d]"), RangeIndex));
+		}
+	}
+
+	else if (ParentNode->GetType() == mu::NodeColourParameter::GetStaticType())
+	{
+		mu::NodeColourParameter* ColorParameterVar = StaticCast<mu::NodeColourParameter*>(ParentNode);
+		mu::NodeColourParameter::Private* Private = ColorParameterVar->GetPrivate();
+		for (int32 RangeIndex = 0; RangeIndex < Private->m_ranges.Num(); RangeIndex++)
+		{
+			AddChildFunc(Private->m_ranges[RangeIndex].get(), FString::Printf(TEXT("RANGE [%d]"), RangeIndex));
+		}
+	}
+
+	else if (ParentNode->GetType() == mu::NodeColourConstant::GetStaticType())
+	{
+		// Nothing to show
+	}
+
+
+	else if (ParentNode->GetType() == mu::NodeImageConstant::GetStaticType())
+	{
+		// Nothing to show
+	}
+
+	
+	else if (ParentNode->GetType() == mu::NodeScalarCurve::GetStaticType())
+	{
+		mu::NodeScalarCurve* ScalarCurveVar = StaticCast<mu::NodeScalarCurve*>(ParentNode);
+		mu::NodeScalarCurve::Private* Private = ScalarCurveVar->GetPrivate();
+		AddChildFunc(Private->m_input_scalar.get(), FString::Printf(TEXT("INPUT")));
+	}
+
+	else if (ParentNode->GetType() == mu::NodeMeshMakeMorph::GetStaticType())
+	{
+		mu::NodeMeshMakeMorph* MeshMakeMorphVar = StaticCast<mu::NodeMeshMakeMorph*>(ParentNode);
+		mu::NodeMeshMakeMorph::Private* Private = MeshMakeMorphVar->GetPrivate();
+		AddChildFunc(Private->m_pBase.get(), FString::Printf(TEXT("BASE")));
+		AddChildFunc(Private->m_pTarget.get(), FString::Printf(TEXT("TARGET")));
+	}
+	
+	else if (ParentNode->GetType() == mu::NodeProjectorParameter::GetStaticType())
+	{
+		mu::NodeProjectorParameter* ProjectorParameterVar = StaticCast<mu::NodeProjectorParameter*>(ParentNode);
+		mu::NodeProjectorParameter::Private* Private = ProjectorParameterVar->GetPrivate();
+		for (int32 RangeIndex = 0; RangeIndex < Private->m_ranges.Num(); RangeIndex++)
+		{
+			AddChildFunc(Private->m_ranges[RangeIndex].get(), FString::Printf(TEXT("RANGE [%d]"), RangeIndex));
+		}
+	}
+
+	else if (ParentNode->GetType() == mu::NodeProjectorConstant::GetStaticType())
+	{
+		// Nothing to show
+	}
+
+	else if (ParentNode->GetType() == mu::NodeColourSwitch::GetStaticType())
+	{
+		mu::NodeColourSwitch* ColorSwitchVar = StaticCast<mu::NodeColourSwitch*>(ParentNode);
+		mu::NodeColourSwitch::Private* Private = ColorSwitchVar->GetPrivate();
+		AddChildFunc(Private->m_pParameter.get(), TEXT("PARAM"));
+		for (int32 OptionIndex = 0; OptionIndex < Private->m_options.Num(); ++OptionIndex)
+		{
+			AddChildFunc(Private->m_options[OptionIndex].get(), FString::Printf(TEXT("OPTION [%d]"), OptionIndex));
+		}
+	}
+
+	else if (ParentNode->GetType() == mu::NodeImageSwizzle::GetStaticType())
+	{
+		mu::NodeImageSwizzle* ImageSwizzleVar = StaticCast<mu::NodeImageSwizzle*>(ParentNode);
+		mu::NodeImageSwizzle::Private* Private = ImageSwizzleVar->GetPrivate();
+		for (int32 SourceIndex = 0; SourceIndex < Private->m_sources.Num(); ++SourceIndex)
+		{
+			AddChildFunc(Private->m_sources[SourceIndex].get(), FString::Printf(TEXT("SOURCE [%d]"), SourceIndex));
+		}
+	}
+	
+	else if (ParentNode->GetType() == mu::NodeImageInvert::GetStaticType())
+	{
+		mu::NodeImageInvert* ImageInvertVar = StaticCast<mu::NodeImageInvert*>(ParentNode);
+		mu::NodeImageInvert::Private* Private = ImageInvertVar->GetPrivate();
+		AddChildFunc(Private->m_pBase.get(), TEXT("BASE"));
+	}
+
+	else if (ParentNode->GetType() == mu::NodeImageMultiLayer::GetStaticType())
+	{
+		mu::NodeImageMultiLayer* ImageMultilayerVar = StaticCast<mu::NodeImageMultiLayer*>(ParentNode);
+		mu::NodeImageMultiLayer::Private* Private = ImageMultilayerVar->GetPrivate();
+		AddChildFunc(Private->m_pBase.get(), FString::Printf(TEXT("BASE")));
+		AddChildFunc(Private->m_pMask.get(), FString::Printf(TEXT("MASK")));
+		AddChildFunc(Private->m_pBlended.get(), FString::Printf(TEXT("BLEND")));
+		AddChildFunc(Private->m_pRange.get(), FString::Printf(TEXT("RANGE")));
+	}
+	
+	else if (ParentNode->GetType() == mu::NodeImageTable::GetStaticType())
+	{
+		// No nodes to show
+	}
+
+	else if (ParentNode->GetType() == mu::NodeMeshTable::GetStaticType())
+	{
+		mu::NodeMeshTable* MeshTableVar = StaticCast<mu::NodeMeshTable*>(ParentNode);
+		mu::NodeMeshTable::Private* Private = MeshTableVar->GetPrivate();
+		for (int32 LayoutIndex = 0; LayoutIndex < Private->Layouts.Num(); ++LayoutIndex)
+		{
+			AddChildFunc(Private->Layouts[LayoutIndex].get(), FString::Printf(TEXT("LAYOUT [%d]"), LayoutIndex));
+		}
+	}
+
+	else if (ParentNode->GetType() == mu::NodeScalarTable::GetStaticType())
+	{
+		// Nothing to show
+	}
+
+	else if (ParentNode->GetType() == mu::NodeScalarSwitch::GetStaticType())
+	{
+		mu::NodeScalarSwitch* ScalarSwitchVar = StaticCast<mu::NodeScalarSwitch*>(ParentNode);
+		mu::NodeScalarSwitch::Private* Private = ScalarSwitchVar->GetPrivate();
+		AddChildFunc(Private->m_pParameter.get(), TEXT("PARAM"));
+		for (int32 OptionIndex = 0; OptionIndex < Private->m_options.Num(); ++OptionIndex)
+		{
+			AddChildFunc(Private->m_options[OptionIndex].get(), FString::Printf(TEXT("OPTION [%d]"), OptionIndex));
+		}
+	}
+	
+	else if (ParentNode->GetType() == mu::NodeColourFromScalars::GetStaticType())
+	{
+		mu::NodeColourFromScalars* ScalarTableVar = StaticCast<mu::NodeColourFromScalars*>(ParentNode);
+		mu::NodeColourFromScalars::Private* Private = ScalarTableVar->GetPrivate();
+		AddChildFunc(Private->m_pX.get(), TEXT("X"));
+		AddChildFunc(Private->m_pY.get(), TEXT("Y"));
+		AddChildFunc(Private->m_pZ.get(), TEXT("Z"));
+		AddChildFunc(Private->m_pW.get(), TEXT("W"));
+	}
+	
 	else
 	{
-		// This node type has not been implemented, so its children won't be added to the tree.
-		ensure(false);
+		const FString ParentNodeTypeString = ANSI_TO_TCHAR(ParentNode->GetType()->m_strName);
+		UE_LOG(LogMutable,Error,TEXT("The node of type %s has not been implemented, so its children won't be added to the tree."), *ParentNodeTypeString);
+
+		// Add a placeholder to the tree
+		const FString Prefix =  FString::Printf(TEXT("[%s] NODE TYPE NOT IMPLEMENTED"), *ParentNodeTypeString);
+		AddChildFunc(nullptr, Prefix);
 	}
 #endif
 }
@@ -529,3 +981,5 @@ FReply SMutableGraphViewer::OnDrop(const FGeometry& MyGeometry, const FDragDropE
 }
 
 #undef LOCTEXT_NAMESPACE 
+
+
diff --git a/Engine/Plugins/Experimental/Mutable/Source/MutableTools/Internal/MuT/AST.h b/Engine/Plugins/Experimental/Mutable/Source/MutableTools/Internal/MuT/AST.h
new file mode 100644
--- /dev/null
+++ b/Engine/Plugins/Experimental/Mutable/Source/MutableTools/Internal/MuT/AST.h
@@ -0,0 +1,1199 @@
+// Copyright Epic Games, Inc. All Rights Reserved.
+
+#pragma once
+
+#include "MuT/Platform.h"
+#include "MuT/Compiler.h"
+#include "MuR/Image.h"
+#include "MuR/Mesh.h"
+#include "MuR/ModelPrivate.h"
+#include "MuR/MutableMath.h"
+#include "MuR/MutableMemory.h"
+#include "MuR/Operations.h"
+#include "MuR/ParametersPrivate.h"
+#include "MuR/Ptr.h"
+#include "MuR/RefCounted.h"
+
+#include "Containers/Array.h"
+#include "Containers/ContainerAllocationPolicies.h"
+#include "Containers/Map.h"
+#include "HAL/PlatformCrt.h"
+#include "HAL/PlatformMath.h"
+#include "Misc/AssertionMacros.h"
+#include "Templates/TypeHash.h"
+#include "Templates/Function.h"
+#include "Hash/CityHash.h"
+
+#include <array>
+#include <atomic>
+#include <unordered_map>
+#include <unordered_set>
+#include <set>
+
+
+namespace std
+{
+
+  template<typename T>
+  struct hash<mu::Ptr<T>>
+  {
+    uint64 operator()(const mu::Ptr<T>& k) const
+    {
+      return hash<const void*>()(k.get());
+    }
+  };
+
+}
+
+
+namespace mu
+{
+	class ASTOp;
+	class ASTOpFixed;
+
+	template<typename T>
+	inline uint32 GetTypeHash(const Ptr<const T>& p)
+	{
+		return ::GetTypeHash(p.get());
+	}
+
+
+	template<typename T>
+	inline uint32 GetTypeHash(const Ptr<T>& p)
+	{
+		return ::GetTypeHash(p.get());
+	}
+
+
+    //---------------------------------------------------------------------------------------------
+    //! This class stores the expression that defines the size of an image.
+    //--------------------------------------------------------------------------------------------
+    class ImageSizeExpression : public RefCounted
+    {
+    public:
+
+        enum
+        {
+            ISET_UNKNOWN,
+            ISET_CONSTANT,
+            ISET_LAYOUTFACTOR,
+            ISET_CONDITIONAL
+        } type;
+
+        // For constant sizes
+        FImageSize size = FImageSize(0, 0);
+
+        // For layout factor sizes
+        Ptr<class ASTOp> layout;
+        uint16 factor[2];
+
+        // For conditionals
+        Ptr<class ASTOp> condition;
+        Ptr<ImageSizeExpression> yes;
+        Ptr<ImageSizeExpression> no;
+
+        //!
+        ImageSizeExpression()
+        {
+            type = ISET_UNKNOWN;
+            size[0] = 0;
+            size[1] = 0;
+            layout = 0;
+            factor[0] = 0;
+            factor[1] = 0;
+            condition = 0;
+        }
+
+        //!
+        void CopyFrom(const ImageSizeExpression& o)
+        {
+            type = o.type;
+            size = o.size;
+            layout = o.layout;
+            factor[0] = o.factor[0];
+            factor[1] = o.factor[1];
+            condition = o.condition;
+            yes = o.yes;
+            no = o.no;
+        }
+
+        //!
+        bool operator==( const ImageSizeExpression& other ) const
+        {
+            if ( type == other.type )
+            {
+                switch (type)
+                {
+                case ISET_CONSTANT:
+                    return size[0]==other.size[0]
+                            && size[1] == other.size[1];
+
+                case ISET_LAYOUTFACTOR:
+                    return layout==other.layout
+                            && factor[0]==other.factor[0]
+                            && factor[1]==other.factor[1];
+
+                case ISET_CONDITIONAL:
+                    return condition==other.condition
+                            && *yes==*(other.yes)
+                            && *no==*(other.no);
+
+                default:
+                    return false;
+                }
+            }
+
+            return false;
+        }
+
+        void Optimise()
+        {
+            switch (type)
+            {
+            case ISET_CONSTANT:
+                break;
+
+            case ISET_LAYOUTFACTOR:
+                // TODO: See if the layout is constant and so is this expression.
+                break;
+
+            case ISET_CONDITIONAL:
+                yes->Optimise();
+                no->Optimise();
+
+                // TODO: See if the condition is constant
+                if ( *yes==*no )
+                {
+                    CopyFrom( *yes );
+                }
+
+            default:
+                break;
+            }
+        }
+
+    };
+
+
+    typedef TArray<Ptr<ASTOp>> ASTOpList;
+    typedef TSet<Ptr<ASTOp>> ASTOpSet;
+
+
+    //! Detailed optimization flags
+    struct FModelOptimizationOptions
+    {
+        bool bEnabled = true;
+        bool bOptimiseOverlappedMasks = false;
+        bool bConstReduction = true;
+
+        //! Preprocess all mesh fragments so that they use the same skeleton, even if not all bones
+        //! are relevant for all fragments.
+        bool bUniformizeSkeleton = true;
+
+        //! Maximum number of iterations when optimising models. If 0 as many as necessary will be performed.
+        int32 MaxOptimisationLoopCount = 8;
+
+        //! Store resource data in disk instead of memory
+        bool bUseDiskCache = false;
+
+		/** Compile optimizing for the generation of smaller mipmaps of every image. */
+		bool bEnableProgressiveImages = false;
+
+        // Additional advanced fine-tuning parameters
+        //---------------------------------------------------------------------
+
+        //! Ratio used to decide if it is worth to generate a crop operation
+        float AcceptableCropRatio = 0.5f;
+
+        //! Ratio used to decide if it is worth to generate a crop operation
+        float MinRLECompressionGain = 1.2f;
+
+		//
+		FReferencedResourceFunc ReferencedResourceProvider;
+    };
+
+
+	struct FLinkerOptions
+	{
+		FLinkerOptions(FImageOperator& InImOp)
+			: ImageOperator(InImOp)
+		{
+		}
+
+		// TODO: Unused?
+		int32 MinTextureResidentMipCount = 0;
+
+		/** This flag controls the splitting of image data into mips to store separately. It is usually necessary to
+		* be able to generate progressive textures (for texture streaming).
+		*/
+		bool bSeparateImageMips = true;
+
+		/** Structure used to speedup mesh constant comparison. */
+		struct FDeduplicationMeshFuncs : TDefaultMapHashableKeyFuncs<Ptr<const Mesh>, int32, false>
+		{
+			static FORCEINLINE bool Matches(KeyInitType A, KeyInitType B)
+			{
+				return *A == *B;
+			}
+
+			static FORCEINLINE uint32 GetKeyHash(KeyInitType Key)
+			{
+				const mu::Mesh* Data = Key.get();
+				return HashCombineFast(
+					::GetTypeHash(Data->m_VertexBuffers.GetElementCount()),
+					::GetTypeHash(Data->m_IndexBuffers.GetElementCount())
+				);
+			}
+		};
+
+		TMap<Ptr<const Mesh>, int32, FDefaultSetAllocator, FDeduplicationMeshFuncs> MeshConstantMap;
+
+		/** Structure used to speedup image mip comparison. */
+		struct FDeduplicationImageFuncs : TDefaultMapHashableKeyFuncs<Ptr<const Image>, int32, false>
+		{
+			static FORCEINLINE bool Matches(KeyInitType A, KeyInitType B)
+			{
+				return *A == *B;
+			}
+
+			static FORCEINLINE uint32 GetKeyHash(KeyInitType Key)
+			{
+				const mu::Image* Data = Key.get();
+				uint32 Hash = HashCombineFast(
+					::GetTypeHash(Data->m_format),
+					HashCombineFast(::GetTypeHash(Data->m_size[0]),
+								::GetTypeHash(Data->m_size[1]))
+				);
+
+				uint64 DataHash = CityHash64(reinterpret_cast<const char*>(Data->m_data.GetData()), Data->m_data.Num());
+				Hash = HashCombineFast( Hash, ::GetTypeHash(DataHash) );
+
+				return Hash;
+			}
+		};
+
+		TMap<Ptr<const Image>, int32, FDefaultSetAllocator, FDeduplicationImageFuncs> ImageConstantMipMap;
+
+		/** Image operation functions, so that they can be overriden. */
+		FImageOperator& ImageOperator;
+	};
+
+
+	//! For each operation we sink, the map from old instructions to new instructions.
+	struct FSinkerOldToNewKey
+	{
+		Ptr<const ASTOp> Op;
+		Ptr<const ASTOp> SinkingOp;
+
+		friend inline uint32 GetTypeHash(const FSinkerOldToNewKey& Key)
+		{
+			return HashCombineFast(::GetTypeHash(Key.Op.get()), ::GetTypeHash(Key.SinkingOp.get()));
+		}
+
+		friend inline bool operator==(const FSinkerOldToNewKey& A, const FSinkerOldToNewKey& B)
+		{
+			return A.Op==B.Op && A.SinkingOp==B.SinkingOp;
+		}
+	};
+
+
+	class Sink_ImageCropAST
+	{
+	public:
+
+		Ptr<ASTOp> Apply(const class ASTOpImageCrop* root);
+
+	protected:
+
+		const class ASTOpImageCrop* m_root = nullptr;
+		Ptr<ASTOp> m_initialSource;
+		//! For each operation we sink, the map from old instructions to new instructions.
+		TMap<FSinkerOldToNewKey, Ptr<ASTOp>> OldToNew;
+
+		Ptr<ASTOp> Visit(Ptr<ASTOp> at, const class ASTOpImageCrop* currentCropOp);
+	};
+
+	//---------------------------------------------------------------------------------------------
+	//!
+	//---------------------------------------------------------------------------------------------
+	class Sink_ImagePixelFormatAST
+	{
+	public:
+
+		Ptr<ASTOp> Apply(const class ASTOpImagePixelFormat* root);
+
+	protected:
+
+		const class ASTOpImagePixelFormat* m_root = nullptr;
+		Ptr<ASTOp> m_initialSource;
+		TMap<FSinkerOldToNewKey, Ptr<ASTOp>> OldToNew;
+
+		Ptr<ASTOp> Visit(Ptr<ASTOp> at, const class ASTOpImagePixelFormat* currentFormatOp);
+	};
+
+
+	//---------------------------------------------------------------------------------------------
+	//!
+	//---------------------------------------------------------------------------------------------
+	class Sink_MeshFormatAST
+	{
+	public:
+
+		// \TODO This is recursive and may cause stack overflows in big models.
+		Ptr<ASTOp> Apply(const class ASTOpMeshFormat* Root);
+
+	protected:
+
+		const class ASTOpMeshFormat* m_root = nullptr;
+		Ptr<ASTOp> m_initialSource;
+		TMap<FSinkerOldToNewKey, Ptr<ASTOp>> OldToNew;
+
+		Ptr<ASTOp> Visit(const Ptr<ASTOp>& at, const class ASTOpMeshFormat* currentFormatOp);
+	};
+
+
+	//---------------------------------------------------------------------------------------------
+	//!
+	//---------------------------------------------------------------------------------------------
+	class Sink_ImageMipmapAST
+	{
+	public:
+
+		Ptr<ASTOp> Apply(const class ASTOpImageMipmap* Root);
+
+	protected:
+
+		const class ASTOpImageMipmap* m_root = nullptr;
+		Ptr<ASTOp> m_initialSource;
+
+		//! For each operation we sink, the map from old instructions to new instructions.
+		TMap<FSinkerOldToNewKey, Ptr<ASTOp>> OldToNew;
+
+		Ptr<ASTOp> Visit(Ptr<ASTOp> at, const class ASTOpImageMipmap* currentMipmapOp);
+
+	};
+
+	struct FOptimizeSinkContext
+	{
+		Sink_ImageCropAST ImageCropSinker;
+		Sink_ImagePixelFormatAST ImagePixelFormatSinker;
+		Sink_ImageMipmapAST ImageMipmapSinker;
+		Sink_MeshFormatAST MeshFormatSinker;
+	};
+
+
+    //!
+    class ASTChild
+    {
+    public:
+        explicit ASTChild(ASTOp* parent, const Ptr<ASTOp>& child=Ptr<ASTOp>());
+        ASTChild(const Ptr<ASTOp>& parent, const Ptr<ASTOp>& child);
+        ~ASTChild();
+
+        ASTChild(const ASTChild&) = delete;
+        ASTChild& operator=(const ASTChild&) = delete;
+
+        // move constructor
+        ASTChild(ASTChild&& rhs)
+             : m_parent(rhs.m_parent)
+             , m_child(rhs.m_child)
+             , m_parentIndexInChild(rhs.m_parentIndexInChild)
+        {
+            rhs.m_parent=nullptr;
+            rhs.m_child.reset();
+        }
+
+        // Move assignment
+        ASTChild& operator=( ASTChild&& rhs );
+
+        ASTChild& operator=( const Ptr<ASTOp>& c );
+
+        inline explicit operator bool() const
+        {
+            return m_child.get()!=nullptr;
+        }
+
+        inline Ptr<class ASTOp>& child()
+        {
+            return m_child;
+        }
+
+        inline const Ptr<class ASTOp>& child() const
+        {
+            return m_child;
+        }
+
+        inline const Ptr<class ASTOp>& operator->() const
+        {
+            return m_child;
+        }
+
+        inline bool operator==(const ASTChild& o) const
+        {
+            return m_child==o.m_child;
+        }
+
+        class ASTOp* m_parent;
+        Ptr<class ASTOp> m_child;
+        int32 m_parentIndexInChild = 0;
+
+    private:
+
+        inline void AddParent();
+
+        inline void ClearParent();
+    };
+
+
+    //---------------------------------------------------------------------------------------------
+    //! Abstract Syntax Tree of operations in the mutable virtual machine.
+    //! Avoid any kind of recursivity here, since the hierarchy can be very deep, and it will
+    //! easily cause stack overflows with production models.
+    //---------------------------------------------------------------------------------------------
+    class ASTOp : public RefCounted
+    {
+    private:
+
+        //! Operations referring to this one. They may be null: elements are never removed
+        //! from this TArray.
+		TArray<ASTOp*, TInlineAllocator<4> > m_parents;
+
+    public:
+
+        inline ASTOp()
+        {
+			bIsConstantSubgraph = false;
+			bHasSpecialOpInSubgraph = false;
+        }
+
+        virtual ~ASTOp() {}
+
+        //! Get the operation type
+        virtual OP_TYPE GetOpType() const = 0;
+
+        //! Validate that everything is fine with this tree
+        virtual void Assert();
+
+        //! Run something for each child operation, with a chance to modify it.
+        virtual void ForEachChild( const TFunctionRef<void(ASTChild&)> ) = 0;
+
+        //! Run something for each parent operation+.
+		void ForEachParent(const TFunctionRef<void(ASTOp*)>) const;
+
+        //! Run something for each child operation, with a chance to modify it.
+        virtual bool operator==( const ASTOp& other ) const;
+
+        //! Hint hash method for op sorting and containers. It's a hash of the actual operation.
+        virtual uint64 Hash() const = 0;
+
+        //! Shallow clone. New node will have no parents but reference to the same children.
+		using MapChildFunc = TFunction<Ptr<ASTOp>(const Ptr<ASTOp>&)>;
+		using MapChildFuncRef = TFunctionRef<Ptr<ASTOp>(const Ptr<ASTOp>&)>;
+		virtual Ptr<ASTOp> Clone(MapChildFuncRef mapChild) const = 0;
+
+		//
+		virtual bool IsConditional() const { return false; }
+		virtual bool IsSwitch() const { return false; }
+
+    protected:
+
+        void RemoveChildren();
+
+        virtual bool IsEqual(const ASTOp& other) const = 0;
+
+    public:
+
+        //---------------------------------------------------------------------------------------------
+        static void FullAssert( const TArray<Ptr<ASTOp>>& roots );
+
+        static int32 CountNodes( const TArray<Ptr<ASTOp>>& roots );
+
+		inline bool IsConstantOp() const
+		{
+			OP_TYPE Type = GetOpType();
+			return Type == OP_TYPE::BO_CONSTANT
+				|| Type == OP_TYPE::NU_CONSTANT
+				|| Type == OP_TYPE::SC_CONSTANT
+				|| Type == OP_TYPE::CO_CONSTANT
+				|| Type == OP_TYPE::IM_CONSTANT
+				|| Type == OP_TYPE::ME_CONSTANT
+				|| Type == OP_TYPE::LA_CONSTANT
+				|| Type == OP_TYPE::PR_CONSTANT
+				|| Type == OP_TYPE::ST_CONSTANT
+				|| Type == OP_TYPE::ED_CONSTANT
+				;
+		}
+
+        //! Deep clone. New node will have no parents and reference new children
+        static Ptr<ASTOp> DeepClone( const Ptr<ASTOp>& );
+
+		//!
+		static void LogHistogram(ASTOpList& roots);
+
+        // Code optimisation methods
+        //---------------------------------------------------------------------------------------------
+
+        //!
+        virtual Ptr<ASTOp> OptimiseSize() const { return nullptr; }
+        virtual Ptr<ASTOp> OptimiseSemantic(const FModelOptimizationOptions&, int32 Pass) const { return nullptr; }
+        virtual Ptr<ASTOp> OptimiseSink(const FModelOptimizationOptions&, FOptimizeSinkContext& ) const { return nullptr; }
+        virtual Ptr<ImageSizeExpression> GetImageSizeExpression() const
+        {
+            check( false );
+            return nullptr;
+        }
+
+
+        // Code linking
+        //---------------------------------------------------------------------------------------------
+
+        //!
+        static void FullLink( Ptr<ASTOp>& root, FProgram&, FLinkerOptions*);
+
+        //!
+        static void ClearLinkData( Ptr<ASTOp>& root );
+
+    private:
+
+        //!
+        virtual void Link( FProgram& program, FLinkerOptions* Options ) = 0;
+
+    protected:
+
+        //---------------------------------------------------------------------------------------------
+        //!
+        //---------------------------------------------------------------------------------------------
+        struct FRangeData
+        {
+            //!
+            ASTChild rangeSize;
+
+            //!
+            FString rangeName;
+
+            //!
+			FString rangeUID;
+
+            //!
+            FRangeData( ASTOp* parentOp, Ptr<ASTOp> childOp, const FString& name, const FString& uid )
+                : rangeSize( parentOp, childOp )
+                , rangeName(name)
+                , rangeUID(uid)
+            {
+            }
+
+            FRangeData(const FRangeData&) = delete;
+            FRangeData& operator=(const FRangeData&) = delete;
+            FRangeData& operator=(FRangeData&&) = delete;
+
+
+            // move constructor
+            FRangeData(FRangeData&& rhs)
+                 : rangeSize(MoveTemp(rhs.rangeSize))
+                 , rangeName(MoveTemp(rhs.rangeName))
+                 , rangeUID(MoveTemp(rhs.rangeUID))
+            {
+            }
+
+            //!
+
+            //!
+            bool operator==(const FRangeData& o) const
+            {
+                return rangeSize==o.rangeSize
+                        &&
+                        rangeName==o.rangeName
+                        &&
+                        rangeUID==o.rangeUID;
+            }
+        };
+
+        //!
+		static void LinkRange(FProgram& program,
+			const FRangeData& range,
+			OP::ADDRESS& rangeSize,
+			uint16& rangeId);
+
+    public:
+
+        // Generic traversals
+        //---------------------------------------------------------------------------------------------
+
+        //!
+        static void Traverse_TopDown_Unique( const TArray<Ptr<ASTOp>>& roots, TFunctionRef<bool(Ptr<ASTOp>&)> f );
+
+        //! \todo: it is not strictly top down.
+        static void Traverse_TopDown_Unique_Imprecise( const TArray<Ptr<ASTOp>>& roots, TFunctionRef<bool(Ptr<ASTOp>&)> f );
+
+        //! Kind of top-down, but really not.
+        //! This version is slighlty faster, but doesn't support recursive traversals so
+        //! use it only in controlled cases.
+        static void Traverse_TopRandom_Unique_NonReentrant
+            (
+                const TArray<Ptr<ASTOp>>& roots,
+				TFunctionRef<bool(Ptr<ASTOp>&)> f
+            );
+
+        //! Kind of top-down, but really not.
+        //! This version is slighlty faster, but doesn't support recursive traversals so
+        //! use it only in controlled cases.
+        template<typename STATE>
+        static inline void Traverse_TopDown_Unique_Imprecise_WithState
+            (
+                Ptr<ASTOp>& root, const STATE& initialState,
+				TFunctionRef<bool(Ptr<ASTOp>&, STATE&, TArray<TPair<Ptr<ASTOp>,STATE>>&)> f
+            )
+        {
+            if (!root) { return; }
+
+            TArray<TPair<Ptr<ASTOp>,STATE>> pending;
+			pending.Emplace( root, initialState );
+
+
+            struct custom_partial_hash
+            {
+				uint64 operator()(const std::pair<Ptr<ASTOp>,STATE>& k) const
+                {
+                    return std::hash<const void*>()(k.first.get());
+                }
+            };
+
+            std::unordered_set<std::pair<Ptr<ASTOp>,STATE>, custom_partial_hash> traversed;
+
+            while (!pending.IsEmpty())
+            {
+				TPair<Ptr<ASTOp>,STATE> current = pending.Pop();
+
+                // It could have been completed in another branch
+				auto iti = traversed.insert({ current.Key,current.Value });
+                if (iti.second)
+                {
+                    // Process. State in current.Value may change
+                    bool recurse = f(current.Key,current.Value,pending);
+
+                    // Recurse children
+                    if (recurse)
+                    {
+                        current.Key->ForEachChild([&]( ASTChild& c )
+                        {
+							std::pair<Ptr<ASTOp>, STATE> TraverseKey(c.m_child, current.Value);
+							if (c.m_child && !traversed.count(TraverseKey))
+                            {
+								pending.Emplace( c.m_child, current.Value );
+                            }
+                        });
+                    }
+                }
+            }
+        }
+
+
+        //! This version is slighlty faster, but doesn't support recursive traversals so
+        //! use it only in controlled cases.
+        static void Traverse_BottomUp_Unique_NonReentrant( ASTOpList& roots, TFunctionRef<void(Ptr<ASTOp>&)> f );
+
+        //! This version is slighlty faster, but doesn't support recursive traversals so
+        //! use it only in controlled cases.
+        static void Traverse_BottomUp_Unique_NonReentrant
+        (
+            ASTOpList& roots,
+			TFunctionRef<void(Ptr<ASTOp>&)> f,
+			TFunctionRef<bool(const ASTOp*)> accept
+        );
+
+        //!
+        static void Traverse_BottomUp_Unique
+        (
+            ASTOpList& roots,
+			TFunctionRef<void(Ptr<ASTOp>&)> f,
+			TFunctionRef<bool(const ASTOp*)> accept = [](const ASTOp*){return true;}
+        );
+
+        //!
+        static void Traverse_BottomUp_Unique
+        (
+            Ptr<ASTOp>& root,
+			TFunctionRef<void(Ptr<ASTOp>&)> f,
+			TFunctionRef<bool(const ASTOp*)> accept = [](const ASTOp*){return true;}
+        );
+
+        //!
+        OP::ADDRESS linkedAddress = 0;
+
+        //! Generic traverse control counter. It should always be left to 0 after any process for all
+        //! nodes in the hierarchy.
+        uint32 m_traverseIndex = 0;
+        static std::atomic<uint32> s_lastTraverseIndex;
+
+        //!
+        int8 linkedRange = -1;
+
+        /** Embedded node data for the constant subtree detection.This flag is only valid if the
+         * constant detection process has been executed and no relevant AST transformations have
+         * happened. */
+        uint8 bIsConstantSubgraph : 1;
+        uint8 bHasSpecialOpInSubgraph : 1;
+
+    private:
+        friend class ASTChild;
+
+        //! Remove a node from the parent list if it is there.
+        //void RemoveParent(const ASTOp* parent);
+
+    public:
+
+		/** Get the number of ASTOps that have this one as child, in any existing AST graph. */
+        int32 GetParentCount() const;
+
+        //! Make all parents of this node point at the other node instead.
+        static void Replace( const Ptr<ASTOp>& node, const Ptr<ASTOp>& other );
+
+        // Other code generation utilities
+        //---------------------------------------------------------------------------------------------
+
+        //! This class contains the support data to accelerate the GetImageDesc recursive function.
+        //! If none is provided in the call, one will be created at that level and used from there on.
+        class FGetImageDescContext
+        {
+        public:
+            TMap<const ASTOp*, FImageDesc> m_results;
+        };
+
+        //!
+        virtual FImageDesc GetImageDesc( bool returnBestOption=false, FGetImageDescContext* context=nullptr ) const;
+
+        //! Optional cache struct to use int he method below.
+        using FBlockLayoutSizeCache=TMap< const TPair<ASTOp*,int>, TPair<int,int>>;
+
+        //! Return the size in layout blocks of a particular block given by absolute index
+        virtual void GetBlockLayoutSize( int blockIndex, int* pBlockX, int* pBlockY, FBlockLayoutSizeCache* cache );
+        void GetBlockLayoutSizeCached( int blockIndex, int* pBlockX, int* pBlockY, FBlockLayoutSizeCache* cache );
+
+
+        //! Return the size in pixels of the layout grid block for the image operation
+        virtual void GetLayoutBlockSize( int* pBlockX, int* pBlockY );
+
+        virtual bool IsImagePlainConstant(FVector4f& colour ) const;
+        virtual bool IsColourConstant(FVector4f& colour ) const;
+
+        //!
+        virtual bool GetNonBlackRect( FImageRect& maskUsage ) const;
+
+        // Logic expression evaluation
+        //---------------------------------------------------------------------------------------------
+        typedef enum
+        {
+            BET_UNKNOWN,
+            BET_TRUE,
+            BET_FALSE,
+        } FBoolEvalResult;
+
+
+        using FEvaluateBoolCache = std::unordered_map<const ASTOp*,FBoolEvalResult>;
+
+        virtual FBoolEvalResult EvaluateBool( ASTOpList& /*facts*/, FEvaluateBoolCache* = nullptr ) const
+        {
+            check(false);
+            return BET_UNKNOWN;
+        }
+
+        virtual int EvaluateInt( ASTOpList& /*facts*/, bool &unknown ) const
+        {
+            check(false);
+            unknown = true;
+            return 0;
+        }
+
+    };
+
+    template<class DERIVED>
+    inline Ptr<DERIVED> Clone( const ASTOp* s )
+    {
+		ASTOp::MapChildFunc Identity = [](const Ptr<ASTOp>& o) {return o; };
+		Ptr<ASTOp> c = s->Clone(Identity);
+		Ptr<DERIVED> t = static_cast<DERIVED*>(c.get());
+        return t;
+    }
+
+    template<class DERIVED>
+    inline Ptr<DERIVED> Clone( const Ptr<const ASTOp>& s )
+    {
+		ASTOp::MapChildFunc Identity = [](const Ptr<ASTOp>& o) {return o; };
+		Ptr<ASTOp> c = s->Clone(Identity);
+        Ptr<DERIVED> t = static_cast<DERIVED*>(c.get());
+		return t;
+    }
+
+    //---------------------------------------------------------------------------------------------
+    //!
+    //---------------------------------------------------------------------------------------------
+    template<typename STATE>
+    class Visitor_TopDown_Unique_Const
+    {
+    private:
+
+        //!
+        virtual bool Visit( const Ptr<ASTOp>& node ) = 0;
+
+    private:
+
+        //! States found so far
+        TArray<STATE> m_states;
+
+        //! Index of the current state, from the m_states TArray.
+        int m_currentState;
+
+        struct PENDING
+        {
+            PENDING()
+            {
+                stateIndex = 0;
+            }
+            PENDING(Ptr<ASTOp> _at, int _stateIndex)
+            {
+                at = _at;
+                stateIndex = _stateIndex;
+            }
+
+            Ptr<ASTOp> at;
+            int stateIndex;
+        };
+
+        TArray<PENDING> m_pending;
+
+        //! List of traversed nodes with the state in which they were traversed.
+        std::multimap<Ptr<ASTOp>,int> m_traversed;
+
+    protected:
+
+        //!
+        const STATE& GetCurrentState() const
+        {
+            return m_states[m_currentState];
+        }
+
+        //!
+        STATE GetDefaultState() const
+        {
+            return m_states[0];
+        }
+
+        //! For manual recursion that changes the state for a specific path.
+        void RecurseWithState(const Ptr<ASTOp>& at, const STATE& newState)
+        {
+            if(at)
+            {
+                auto it = m_states.Find(newState);
+                if (it==INDEX_NONE)
+                {
+                    m_states.Add(newState);
+                }
+                int stateIndex = m_states.Find(newState);
+
+                m_pending.Add( PENDING(at,stateIndex) );
+            }
+        }
+
+        //! For manual recursion that doesn't change the state for a specific path.
+        void RecurseWithCurrentState(const Ptr<ASTOp>& at)
+        {
+            if(at)
+            {
+                m_pending.Emplace( at, m_currentState );
+            }
+        }
+
+        //! Can be called from visit to set the state to visit all children ops
+        void SetCurrentState(const STATE& newState)
+        {
+            auto it = m_states.Find(newState);
+            if (it==INDEX_NONE)
+            {
+                m_states.Add(newState);
+            }
+            m_currentState = m_states.Find(newState);
+        }
+
+    public:
+
+        //! Ensure virtual destruction
+        virtual ~Visitor_TopDown_Unique_Const() {}
+
+        //!
+        void Traverse(const ASTOpList& roots, const STATE& initialState)
+        {
+            m_pending.Empty();
+            m_states.Empty();
+            m_states.Add(initialState);
+            m_currentState = 0;
+
+            for( const auto& r: roots)
+            {
+                if (r)
+                {
+                    m_pending.Add( PENDING(r,m_currentState) );
+                }
+            }
+
+            while (m_pending.Num())
+            {
+                PENDING item = m_pending.Pop();
+
+                auto thisNodeRange = m_traversed.equal_range(item.at);
+                bool visitedInThisState = false;
+                for (auto it=thisNodeRange.first; it!=thisNodeRange.second; ++it)
+                {
+                    if (it->second==item.stateIndex)
+                    {
+                        visitedInThisState = true;
+                        break;
+                    }
+                }
+
+                // It could have been completed in another branch
+                if (!visitedInThisState)
+                {
+                    m_traversed.insert( std::make_pair(item.at,item.stateIndex) );
+
+                    // Process
+                    m_currentState = item.stateIndex;
+                    bool recurse = Visit(item.at);
+
+                    // Recurse children
+                    if (recurse)
+                    {
+                        item.at->ForEachChild([&]( ASTChild& c )
+                        {
+                            if (c)
+                            {
+                                m_pending.Add( PENDING(c.child(),m_currentState) );
+                            }
+                        });
+                    }
+                }
+            }
+        }
+    };
+
+
+    //---------------------------------------------------------------------------------------------
+    //! Stateless top-down code visitor that can change the instructions. Iterative version.
+    //! Once an instruction has changed, all the chain of instructions up to the root will be
+    //! cloned, referencing the new instruction.
+    //---------------------------------------------------------------------------------------------
+    class Visitor_TopDown_Unique_Cloning
+    {
+    public:
+
+        //! Ensure virtual destruction
+        virtual ~Visitor_TopDown_Unique_Cloning() {}
+
+    protected:
+
+        //! Do the actual work by overriding this in the derived classes.
+        virtual Ptr<ASTOp> Visit( Ptr<ASTOp> at, bool& processChildren ) = 0;
+
+
+        void Traverse( Ptr<ASTOp>& root );
+
+    private:
+
+        //! Operations to be processed
+        TArray< TPair<bool,Ptr<ASTOp>> > m_pending;
+
+        //! Map for visited operations
+        std::unordered_map<Ptr<ASTOp>,Ptr<ASTOp>> m_oldToNew;
+
+        //!
+        inline Ptr<ASTOp> GetOldToNew( const Ptr<ASTOp>& old ) const
+        {
+            Ptr<ASTOp> n;
+            auto it = m_oldToNew.find(old);
+            if (it!=m_oldToNew.end())
+            {
+                n = it->second;
+            }
+
+            it = m_oldToNew.find(n);
+            while (it!=m_oldToNew.end() && it->second.get()!=nullptr && it->second!=n)
+            {
+                n = it->second;
+                it = m_oldToNew.find(n);
+            }
+
+            return n;
+        }
+
+        //! Process all the pending operations and visit all children if necessary
+        void Process();
+
+    };
+
+
+    //---------------------------------------------------------------------------------------------
+    //! This class of operation is used to ease the transition from legacy code generation.
+    //---------------------------------------------------------------------------------------------
+    class ASTOpFixed : public ASTOp
+    {
+    public:
+
+        //! Fixed size op structure. Any code ADDRESS in it, refers to the children TArray below.
+        OP op;
+
+        //! Operations used by this one. To be indexed with the ADDRESSes in the m_op.
+        //TArray<ASTChild> children;
+        std::array<ASTChild,8> children;
+        uint8 childCount = 1;
+
+    public:
+        ASTOpFixed();
+
+        ~ASTOpFixed() override;
+
+        OP_TYPE GetOpType() const override { return op.type; }
+
+        Ptr<ASTOp> Clone( MapChildFuncRef mapChild ) const override;
+        void ForEachChild( const TFunctionRef<void(ASTChild&)> f ) override;
+        void Link( FProgram& program, FLinkerOptions* Options) override;
+        bool IsEqual(const ASTOp& otherUntyped) const override;
+		uint64 Hash() const override;
+		FImageDesc GetImageDesc( bool returnBestOption, FGetImageDescContext* context ) const override;
+        void GetLayoutBlockSize( int* pBlockX, int* pBlockY ) override;
+        Ptr<ASTOp> OptimiseSize() const override;
+        Ptr<ASTOp> OptimiseSemantic(const FModelOptimizationOptions&, int32 Pass) const override;
+        Ptr<ASTOp> OptimiseSink(const FModelOptimizationOptions&, FOptimizeSinkContext&) const override;
+        FBoolEvalResult EvaluateBool( ASTOpList& facts, FEvaluateBoolCache* cache ) const override;
+        int EvaluateInt( ASTOpList& facts, bool &unknown ) const override;
+        bool IsImagePlainConstant(FVector4f& colour ) const override;
+        bool IsColourConstant(FVector4f& colour ) const override;
+        Ptr<ImageSizeExpression> GetImageSizeExpression() const override;
+
+        //---------------------------------------------------------------------------------------------
+        //! Add child for OP fixed nodes.
+        //---------------------------------------------------------------------------------------------
+        inline void SetChild( OP::ADDRESS& at, const Ptr<ASTOp>& child )
+        {
+            // hack check
+            check( ((uint8*)&at)>=((uint8*)&op.args)
+                            &&
+                            ((uint8*)&at)<=((uint8*)&op.args)+sizeof(op.args));
+
+            if (!at)
+            {
+                if (child)
+                {
+                    check(childCount<8);
+                    at = (OP::ADDRESS)childCount++;
+                    children[at] = child;
+                }
+            }
+            else
+            {
+                children[at]=child;
+                if (!child)
+                {
+                    at=0;
+                }
+            }
+        }
+
+        inline void SetChild( OP::ADDRESS& at, const ASTChild& child )
+        {
+            SetChild(at,child.child());
+        }
+
+    private:
+
+        //
+        inline FImageDesc GetImageDesc( OP::ADDRESS at, bool returnBestOption=false, FGetImageDescContext* context=nullptr ) const
+        {
+			FImageDesc res;
+
+            if (children[at])
+            {
+                res = children[at]->GetImageDesc( returnBestOption, context );
+            }
+
+            return res;
+        }
+
+        inline void GetLayoutBlockSize( OP::ADDRESS at, int* pBlockX, int* pBlockY )
+        {
+            if (children[at])
+            {
+                children[at]->GetLayoutBlockSize(pBlockX,pBlockY);
+            }
+        }
+
+    };
+
+
+    //---------------------------------------------------------------------------------------------
+    //---------------------------------------------------------------------------------------------
+    //---------------------------------------------------------------------------------------------
+    inline void ASTChild::AddParent()
+    {
+        m_parentIndexInChild = m_child->m_parents.Num();
+        m_child->m_parents.Add(m_parent);
+    }
+
+    //---------------------------------------------------------------------------------------------
+    inline void ASTChild::ClearParent()
+    {
+        check( m_parentIndexInChild<m_child->m_parents.Num() );
+		// Can't do this, because the indices are stored in children.
+		//m_child->m_parents.RemoveAtSwap(m_parentIndexInChild);
+		m_child->m_parents[m_parentIndexInChild] = nullptr;
+	}
+
+
+    //---------------------------------------------------------------------------------------------
+    //---------------------------------------------------------------------------------------------
+    //---------------------------------------------------------------------------------------------
+    class UniqueOpPool
+    {
+    private:
+        struct op_pointer_hash
+        {
+			uint64 operator() (const Ptr<ASTOp>& n) const
+            {
+                return n->Hash();
+            }
+        };
+
+        struct op_pointer_equal
+        {
+            bool operator() (const Ptr<ASTOp>& a, const Ptr<ASTOp>& b) const
+            {
+                return *a==*b;
+            }
+        };
+
+        // Existing ops, per type
+        std::unordered_set<Ptr<ASTOp>,op_pointer_hash,op_pointer_equal> visited[(int)OP_TYPE::COUNT];
+
+    public:
+
+        bool m_disabled = false;
+
+        Ptr<ASTOp> Add( const Ptr<ASTOp>& n )
+        {
+            if (m_disabled) return n;
+
+            if (!n) return nullptr;
+
+            auto& container = visited[(int)n->GetOpType()];
+            auto it = container.insert(n);
+            if( !it.second )
+            {
+                return *it.first;
+            }
+            else
+            {
+                return n;
+            }
+        }
+    };
+
+}
+
diff --git a/Engine/Plugins/Experimental/Mutable/Source/MutableTools/Internal/MuT/ErrorLogPrivate.h b/Engine/Plugins/Experimental/Mutable/Source/MutableTools/Internal/MuT/ErrorLogPrivate.h
new file mode 100644
--- /dev/null
+++ b/Engine/Plugins/Experimental/Mutable/Source/MutableTools/Internal/MuT/ErrorLogPrivate.h
@@ -0,0 +1,44 @@
+// Copyright Epic Games, Inc. All Rights Reserved.
+
+#pragma once
+
+#include "MuT/ErrorLog.h"
+
+#include "MuR/ModelPrivate.h"
+#include "MuR/Operations.h"
+
+namespace mu
+{
+	class ErrorLog::Private
+	{
+	public:
+
+        struct FErrorData
+        {
+            TArray< float > m_unassignedUVs;
+        };
+
+		struct FMessage
+		{
+			ErrorLogMessageType m_type = ELMT_NONE;
+			ErrorLogMessageSpamBin m_spam = ELMSB_ALL;
+			FString m_text;
+            TSharedPtr<FErrorData> m_data;
+			const void* m_context = nullptr;
+		};
+
+		TArray<FMessage> m_messages;
+
+
+		//!
+		void Add(const FString& Message, ErrorLogMessageType Type, const void* Context, ErrorLogMessageSpamBin SpamBin = ELMSB_ALL);
+
+        //!
+        void Add(const FString& Message, const ErrorLogMessageAttachedDataView& Data, ErrorLogMessageType Type, const void* Context, ErrorLogMessageSpamBin SpamBin = ELMSB_ALL);
+	};
+
+
+    extern const TCHAR* GetOpName( OP_TYPE type );
+
+}
+
diff --git a/Engine/Plugins/Experimental/Mutable/Source/MutableTools/Internal/MuT/NodeColourFromScalarsPrivate.h b/Engine/Plugins/Experimental/Mutable/Source/MutableTools/Internal/MuT/NodeColourFromScalarsPrivate.h
new file mode 100644
--- /dev/null
+++ b/Engine/Plugins/Experimental/Mutable/Source/MutableTools/Internal/MuT/NodeColourFromScalarsPrivate.h
@@ -0,0 +1,54 @@
+// Copyright Epic Games, Inc. All Rights Reserved.
+
+#pragma once
+
+#include "MuT/NodeColourPrivate.h"
+#include "MuT/NodeColourFromScalars.h"
+#include "MuT/NodeScalar.h"
+#include "MuT/NodeImage.h"
+#include "MuT/AST.h"
+
+
+namespace mu
+{
+
+
+	class NodeColourFromScalars::Private : public NodeColour::Private
+	{
+	public:
+
+		static FNodeType s_type;
+
+		NodeScalarPtr m_pX;
+		NodeScalarPtr m_pY;
+		NodeScalarPtr m_pZ;
+		NodeScalarPtr m_pW;
+
+		//!
+		void Serialise( OutputArchive& arch ) const
+		{
+            uint32_t ver = 0;
+			arch << ver;
+
+			arch << m_pX;
+			arch << m_pY;
+			arch << m_pZ;
+			arch << m_pW;
+		}
+
+		//!
+		void Unserialise( InputArchive& arch )
+		{
+            uint32_t ver;
+			arch >> ver;
+			check(ver==0);
+
+			arch >> m_pX;
+			arch >> m_pY;
+			arch >> m_pZ;
+			arch >> m_pW;
+		}
+	};
+
+
+}
diff --git a/Engine/Plugins/Experimental/Mutable/Source/MutableTools/Internal/MuT/NodeColourParameterPrivate.h b/Engine/Plugins/Experimental/Mutable/Source/MutableTools/Internal/MuT/NodeColourParameterPrivate.h
new file mode 100644
--- /dev/null
+++ b/Engine/Plugins/Experimental/Mutable/Source/MutableTools/Internal/MuT/NodeColourParameterPrivate.h
@@ -0,0 +1,77 @@
+// Copyright Epic Games, Inc. All Rights Reserved.
+
+#pragma once
+
+#include "MuT/NodeColourParameter.h"
+#include "MuT/NodeColourPrivate.h"
+
+#include "MuT/NodeRange.h"
+#include "MuR/MutableMath.h"
+
+
+namespace mu
+{
+
+	class NodeColourParameter::Private : public NodeColour::Private
+	{
+	public:
+
+		static FNodeType s_type;
+
+		vec4<float> m_defaultValue;
+		FString m_name;
+		FString m_uid;
+
+        TArray<Ptr<NodeRange>> m_ranges;
+
+		//!
+		void Serialise( OutputArchive& arch ) const
+		{
+            uint32_t ver = 4;
+			arch << ver;
+
+			arch << m_defaultValue;
+			arch << m_name;
+            arch << m_uid;
+            arch << m_ranges;
+        }
+
+		//!
+		void Unserialise( InputArchive& arch )
+		{
+            uint32_t ver;
+			arch >> ver;
+            check(ver>=2&&ver<=4);
+
+			if (ver <= 3)
+			{
+				vec3<float> aux_defaultValue;
+				arch >> aux_defaultValue;
+
+				m_defaultValue = aux_defaultValue;
+				m_defaultValue.w() = 1.f;
+			}
+			else
+			{
+				arch >> m_defaultValue;
+			}
+			
+			if (ver <= 2)
+			{
+				std::string Temp;
+				arch >> Temp;
+				m_name = Temp.c_str();
+				arch >> Temp;
+				m_uid = Temp.c_str();
+			}
+			else
+			{
+				arch >> m_name;
+				arch >> m_uid;
+			}
+			
+			arch >> m_ranges;
+        }
+	};
+
+}
diff --git a/Engine/Plugins/Experimental/Mutable/Source/MutableTools/Internal/MuT/NodeColourPrivate.h b/Engine/Plugins/Experimental/Mutable/Source/MutableTools/Internal/MuT/NodeColourPrivate.h
new file mode 100644
--- /dev/null
+++ b/Engine/Plugins/Experimental/Mutable/Source/MutableTools/Internal/MuT/NodeColourPrivate.h
@@ -0,0 +1,18 @@
+// Copyright Epic Games, Inc. All Rights Reserved.
+
+#pragma once
+
+#include "MuT/NodePrivate.h"
+#include "MuT/NodeColour.h"
+
+
+namespace mu
+{
+
+
+	class NodeColour::Private : public Node::Private
+	{
+	};
+
+
+}
diff --git a/Engine/Plugins/Experimental/Mutable/Source/MutableTools/Internal/MuT/NodeColourSampleImagePrivate.h b/Engine/Plugins/Experimental/Mutable/Source/MutableTools/Internal/MuT/NodeColourSampleImagePrivate.h
new file mode 100644
--- /dev/null
+++ b/Engine/Plugins/Experimental/Mutable/Source/MutableTools/Internal/MuT/NodeColourSampleImagePrivate.h
@@ -0,0 +1,51 @@
+// Copyright Epic Games, Inc. All Rights Reserved.
+
+#pragma once
+
+#include "MuT/NodeColourPrivate.h"
+#include "MuT/NodeColourSampleImage.h"
+#include "MuT/NodeScalar.h"
+#include "MuT/NodeImage.h"
+#include "MuT/AST.h"
+
+
+namespace mu
+{
+
+
+	class NodeColourSampleImage::Private : public NodeColour::Private
+	{
+	public:
+
+		static FNodeType s_type;
+
+		NodeImagePtr m_pImage;
+		NodeScalarPtr m_pX;
+		NodeScalarPtr m_pY;
+
+		//!
+		void Serialise( OutputArchive& arch ) const
+		{
+            uint32_t ver = 0;
+			arch << ver;
+
+			arch << m_pImage;
+			arch << m_pX;
+			arch << m_pY;
+		}
+
+		//!
+		void Unserialise( InputArchive& arch )
+		{
+            uint32_t ver;
+			arch >> ver;
+			check(ver==0);
+
+			arch >> m_pImage;
+			arch >> m_pX;
+			arch >> m_pY;
+		}
+	};
+
+
+}
diff --git a/Engine/Plugins/Experimental/Mutable/Source/MutableTools/Internal/MuT/NodeColourSwitchPrivate.h b/Engine/Plugins/Experimental/Mutable/Source/MutableTools/Internal/MuT/NodeColourSwitchPrivate.h
new file mode 100644
--- /dev/null
+++ b/Engine/Plugins/Experimental/Mutable/Source/MutableTools/Internal/MuT/NodeColourSwitchPrivate.h
@@ -0,0 +1,48 @@
+// Copyright Epic Games, Inc. All Rights Reserved.
+
+#pragma once
+
+#include "MuT/NodeColourSwitch.h"
+#include "MuT/NodeColourPrivate.h"
+#include "MuT/NodeScalar.h"
+#include "MuT/AST.h"
+
+
+namespace mu
+{
+
+
+    class NodeColourSwitch::Private : public NodeColour::Private
+	{
+	public:
+
+		static FNodeType s_type;
+
+		NodeScalarPtr m_pParameter;
+        TArray<NodeColourPtr> m_options;
+
+		//!
+		void Serialise( OutputArchive& arch ) const
+		{
+            uint32_t ver = 0;
+			arch << ver;
+
+			arch << m_pParameter;
+			arch << m_options;
+		}
+
+		//!
+		void Unserialise( InputArchive& arch )
+		{
+            uint32_t ver;
+			arch >> ver;
+			check(ver==0);
+
+			arch >> m_pParameter;
+			arch >> m_options;
+		}
+	};
+
+
+}
+
diff --git a/Engine/Plugins/Experimental/Mutable/Source/MutableTools/Internal/MuT/NodeComponentEditPrivate.h b/Engine/Plugins/Experimental/Mutable/Source/MutableTools/Internal/MuT/NodeComponentEditPrivate.h
new file mode 100644
--- /dev/null
+++ b/Engine/Plugins/Experimental/Mutable/Source/MutableTools/Internal/MuT/NodeComponentEditPrivate.h
@@ -0,0 +1,59 @@
+// Copyright Epic Games, Inc. All Rights Reserved.
+
+
+#pragma once
+
+#include "MuT/NodeComponentPrivate.h"
+#include "MuT/NodeComponentEdit.h"
+#include "MuT/NodePatchImagePrivate.h"
+#include "MuT/NodeMesh.h"
+#include "MuT/NodePatchMesh.h"
+#include "MuT/NodeImage.h"
+#include "MuT/NodeSurface.h"
+
+
+namespace mu
+{
+
+
+	//---------------------------------------------------------------------------------------------
+	//!
+	//---------------------------------------------------------------------------------------------
+	class NodeComponentEdit::Private : public NodeComponent::Private
+	{
+	public:
+
+		static FNodeType s_type;
+
+        NodeComponentPtr m_pParent;
+
+        TArray<NodeSurfacePtr> m_surfaces;
+
+		//!
+		void Serialise( OutputArchive& arch ) const
+		{
+            uint32_t ver = 2;
+			arch << ver;
+
+			arch << m_pParent;
+            arch << m_surfaces;
+		}
+
+		//!
+		void Unserialise( InputArchive& arch )
+		{
+            uint32_t ver;
+			arch >> ver;
+			check(ver == 2);
+
+            arch >> m_pParent;
+            arch >> m_surfaces;
+		}
+
+
+		// NodeComponent::Private interface
+		const NodeComponentNew::Private* GetParentComponentNew() const override;
+	};
+
+}
+
diff --git a/Engine/Plugins/Experimental/Mutable/Source/MutableTools/Internal/MuT/NodeComponentNewPrivate.h b/Engine/Plugins/Experimental/Mutable/Source/MutableTools/Internal/MuT/NodeComponentNewPrivate.h
new file mode 100644
--- /dev/null
+++ b/Engine/Plugins/Experimental/Mutable/Source/MutableTools/Internal/MuT/NodeComponentNewPrivate.h
@@ -0,0 +1,82 @@
+// Copyright Epic Games, Inc. All Rights Reserved.
+
+
+#pragma once
+
+
+#include "MuT/NodeComponentPrivate.h"
+#include "MuT/NodeComponentNew.h"
+#include "MuT/NodePatchImagePrivate.h"
+
+#include "MuT/NodeMesh.h"
+#include "MuT/NodePatchMesh.h"
+#include "MuT/NodeImage.h"
+#include "MuT/NodeSurface.h"
+
+
+namespace mu
+{
+
+	//---------------------------------------------------------------------------------------------
+	//!
+	//---------------------------------------------------------------------------------------------
+	class NodeComponentNew::Private : public NodeComponent::Private
+	{
+	public:
+
+		static FNodeType s_type;
+
+		FString m_name;
+
+		uint16 m_id = 0;
+
+        TArray<NodeSurfacePtr> m_surfaces;
+
+		//!
+		void Serialise( OutputArchive& arch ) const
+		{
+            uint32_t ver = 4;
+			
+			arch << ver;
+
+			arch << m_name;
+			arch << m_id;
+            arch << m_surfaces;
+        }
+
+		//!
+		void Unserialise( InputArchive& arch )
+		{
+            uint32_t ver;
+			arch >> ver;
+			check(ver>=2 && ver<=4);
+
+			if (ver <= 3)
+			{
+				std::string Temp;
+				arch>>Temp;
+				m_name = Temp.c_str();
+			}
+			else
+			{
+				arch >> m_name;
+			}
+
+			if (ver >= 3)
+			{
+				arch >> m_id;
+			}
+
+            arch >> m_surfaces;
+		}
+
+		// NodeComponent::Private interface
+        const NodeComponentNew::Private* GetParentComponentNew() const override
+		{
+			return this;
+		}
+
+	};
+
+}
+
diff --git a/Engine/Plugins/Experimental/Mutable/Source/MutableTools/Internal/MuT/NodeComponentPrivate.h b/Engine/Plugins/Experimental/Mutable/Source/MutableTools/Internal/MuT/NodeComponentPrivate.h
new file mode 100644
--- /dev/null
+++ b/Engine/Plugins/Experimental/Mutable/Source/MutableTools/Internal/MuT/NodeComponentPrivate.h
@@ -0,0 +1,26 @@
+// Copyright Epic Games, Inc. All Rights Reserved.
+
+
+#pragma once
+
+#include "MuT/NodePrivate.h"
+#include "MuT/NodeComponent.h"
+#include "MuT/NodeComponentNew.h"
+
+
+namespace mu
+{
+	//---------------------------------------------------------------------------------------------
+	//!
+	//---------------------------------------------------------------------------------------------
+	class NodeComponent::Private : public Node::Private
+	{
+	public:
+
+		//! Look for the parent component that this node is editing
+		virtual const NodeComponentNew::Private* GetParentComponentNew() const = 0;
+	};
+
+
+}
+
diff --git a/Engine/Plugins/Experimental/Mutable/Source/MutableTools/Internal/MuT/NodeImageFormatPrivate.h b/Engine/Plugins/Experimental/Mutable/Source/MutableTools/Internal/MuT/NodeImageFormatPrivate.h
new file mode 100644
--- /dev/null
+++ b/Engine/Plugins/Experimental/Mutable/Source/MutableTools/Internal/MuT/NodeImageFormatPrivate.h
@@ -0,0 +1,55 @@
+// Copyright Epic Games, Inc. All Rights Reserved.
+
+#pragma once
+
+#include "MuT/NodeImageFormat.h"
+#include "MuT/NodeImagePrivate.h"
+#include "MuT/AST.h"
+#include "MuR/ImagePrivate.h"
+
+
+namespace mu
+{
+
+
+	class NodeImageFormat::Private : public NodeImage::Private
+	{
+	public:
+
+		static FNodeType s_type;
+
+		EImageFormat m_format = EImageFormat::IF_NONE;
+		EImageFormat m_formatIfAlpha = EImageFormat::IF_NONE;
+        NodeImagePtr m_source;
+
+		//!
+		void Serialise( OutputArchive& arch ) const
+		{
+            uint32_t ver = 1;
+			arch << ver;
+
+            arch << (uint32)m_format;
+            arch << (uint32)m_formatIfAlpha;
+            arch << m_source;
+		}
+
+		//!
+		void Unserialise( InputArchive& arch )
+		{
+            uint32_t ver;
+			arch >> ver;
+            check(ver==1);
+
+			uint32 format;
+			arch >> format;
+			m_format = EImageFormat(format);
+
+			arch >> format;
+			m_formatIfAlpha = EImageFormat(format);
+
+			arch >> m_source;
+		}
+	};
+
+
+}
diff --git a/Engine/Plugins/Experimental/Mutable/Source/MutableTools/Internal/MuT/NodeImageInterpolatePrivate.h b/Engine/Plugins/Experimental/Mutable/Source/MutableTools/Internal/MuT/NodeImageInterpolatePrivate.h
new file mode 100644
--- /dev/null
+++ b/Engine/Plugins/Experimental/Mutable/Source/MutableTools/Internal/MuT/NodeImageInterpolatePrivate.h
@@ -0,0 +1,45 @@
+// Copyright Epic Games, Inc. All Rights Reserved.
+
+#pragma once
+
+#include "MuT/NodeImageInterpolate.h"
+#include "MuT/NodeImagePrivate.h"
+#include "MuT/NodeScalar.h"
+
+namespace mu
+{
+
+
+	class NodeImageInterpolate::Private : public NodeImage::Private
+	{
+	public:
+
+		static FNodeType s_type;
+
+		NodeScalarPtr m_pFactor;
+		TArray<NodeImagePtr> m_targets;
+
+		//!
+		void Serialise( OutputArchive& arch ) const
+		{
+            uint32_t ver = 0;
+			arch << ver;
+
+			arch << m_pFactor;
+			arch << m_targets;
+		}
+
+		//!
+		void Unserialise( InputArchive& arch )
+		{
+            uint32_t ver;
+			arch >> ver;
+			check(ver==0);
+
+			arch >> m_pFactor;
+			arch >> m_targets;
+		}
+	};
+
+
+}
diff --git a/Engine/Plugins/Experimental/Mutable/Source/MutableTools/Internal/MuT/NodeImageInvertPrivate.h b/Engine/Plugins/Experimental/Mutable/Source/MutableTools/Internal/MuT/NodeImageInvertPrivate.h
new file mode 100644
--- /dev/null
+++ b/Engine/Plugins/Experimental/Mutable/Source/MutableTools/Internal/MuT/NodeImageInvertPrivate.h
@@ -0,0 +1,41 @@
+// Copyright Epic Games, Inc. All Rights Reserved.
+
+#pragma once
+
+#include "MuT/NodeImagePrivate.h"
+#include "MuT/NodeImageInvert.h"
+#include "MuT/AST.h"
+
+
+namespace mu
+{
+
+
+	class NodeImageInvert::Private : public NodeImage::Private
+	{
+	public:
+
+		static FNodeType s_type;
+
+		NodeImagePtr m_pBase;
+		
+		//!
+		void Serialise(OutputArchive& arch) const
+		{
+			uint32_t ver = 0;
+			arch << ver;
+
+			arch << m_pBase;
+		}
+
+		//!
+		void Unserialise(InputArchive& arch)
+		{
+			uint32_t ver;
+			arch >> ver;
+			check(ver == 0);
+
+			arch >> m_pBase;
+		}
+	};
+}
diff --git a/Engine/Plugins/Experimental/Mutable/Source/MutableTools/Internal/MuT/NodeImageLayerColourPrivate.h b/Engine/Plugins/Experimental/Mutable/Source/MutableTools/Internal/MuT/NodeImageLayerColourPrivate.h
new file mode 100644
--- /dev/null
+++ b/Engine/Plugins/Experimental/Mutable/Source/MutableTools/Internal/MuT/NodeImageLayerColourPrivate.h
@@ -0,0 +1,57 @@
+// Copyright Epic Games, Inc. All Rights Reserved.
+
+#pragma once
+
+#include "MuT/NodeImagePrivate.h"
+#include "MuT/NodeImageLayerPrivate.h"
+#include "MuT/NodeImageLayerColour.h"
+#include "MuT/NodeColour.h"
+#include "MuR/Image.h"
+
+
+namespace mu
+{
+
+	class NodeImageLayerColour::Private : public NodeImage::Private
+	{
+	public:
+
+		static FNodeType s_type;
+
+		NodeImagePtr m_pBase;
+		NodeImagePtr m_pMask;
+		NodeColourPtr m_pColour;
+		EBlendType m_type;
+
+		//!
+		void Serialise( OutputArchive& arch ) const
+		{
+            uint32_t ver = 1;
+			arch << ver;
+
+			arch << m_pBase;
+			arch << m_pMask;
+			arch << m_pColour;
+			arch << m_type;
+		}
+
+		//!
+		void Unserialise( InputArchive& arch )
+		{
+            uint32_t ver;
+			arch >> ver;
+			check(ver<=1);
+
+			arch >> m_pBase;
+			arch >> m_pMask;
+			arch >> m_pColour;
+			arch >> m_type;
+			if (ver < 1)
+			{
+				m_type = EBlendType( uint32(m_type)+1 );
+			}
+
+		}
+	};
+
+}
diff --git a/Engine/Plugins/Experimental/Mutable/Source/MutableTools/Internal/MuT/NodeImageLayerPrivate.h b/Engine/Plugins/Experimental/Mutable/Source/MutableTools/Internal/MuT/NodeImageLayerPrivate.h
new file mode 100644
--- /dev/null
+++ b/Engine/Plugins/Experimental/Mutable/Source/MutableTools/Internal/MuT/NodeImageLayerPrivate.h
@@ -0,0 +1,58 @@
+// Copyright Epic Games, Inc. All Rights Reserved.
+
+#pragma once
+
+#include "MuT/NodeImagePrivate.h"
+#include "MuT/NodeImageLayer.h"
+#include "MuT/NodeColour.h"
+#include "MuT/AST.h"
+
+
+namespace mu
+{
+
+	class NodeImageLayer::Private : public NodeImage::Private
+	{
+	public:
+
+		static FNodeType s_type;
+
+		NodeImagePtr m_pBase;
+		NodeImagePtr m_pMask;
+		NodeImagePtr m_pBlended;
+		EBlendType m_type;
+
+		//!
+		void Serialise( OutputArchive& arch ) const
+		{
+            uint32_t ver = 1;
+			arch << ver;
+
+			arch << m_pBase;
+			arch << m_pMask;
+			arch << m_pBlended;
+            arch << (uint32_t)m_type;
+		}
+
+		//!
+		void Unserialise( InputArchive& arch )
+		{
+            uint32_t ver;
+			arch >> ver;
+			check(ver<=1);
+
+			arch >> m_pBase;
+			arch >> m_pMask;
+			arch >> m_pBlended;
+
+            uint32_t t;
+            arch >> t;
+			if (ver < 1)
+			{
+				++t;
+			}
+            m_type=(EBlendType)t;
+		}
+	};
+
+}
diff --git a/Engine/Plugins/Experimental/Mutable/Source/MutableTools/Internal/MuT/NodeImageMipmapPrivate.h b/Engine/Plugins/Experimental/Mutable/Source/MutableTools/Internal/MuT/NodeImageMipmapPrivate.h
new file mode 100644
--- /dev/null
+++ b/Engine/Plugins/Experimental/Mutable/Source/MutableTools/Internal/MuT/NodeImageMipmapPrivate.h
@@ -0,0 +1,53 @@
+// Copyright Epic Games, Inc. All Rights Reserved.
+
+#pragma once
+
+#include "MuT/NodeImagePrivate.h"
+#include "MuT/NodeImageMipmap.h"
+#include "MuT/NodeScalar.h"
+#include "MuT/AST.h"
+
+#include "MuR/ImagePrivate.h"
+
+
+namespace mu
+{	
+    class NodeImageMipmap::Private : public NodeImage::Private
+	{
+	public:
+
+		static FNodeType s_type;
+
+		NodeImagePtr m_pSource;
+		NodeScalarPtr m_pFactor;
+
+		FMipmapGenerationSettings m_settings;
+
+		//!
+		void Serialise( OutputArchive& arch ) const
+		{
+            uint32_t ver = 1;
+			arch << ver;
+
+			arch << m_pSource;
+			arch << m_pFactor;
+
+			arch << m_settings;
+		}
+
+		//!
+		void Unserialise( InputArchive& arch )
+		{
+            uint32_t ver;
+			arch >> ver;
+			check(ver==1);
+
+			arch >> m_pSource;
+			arch >> m_pFactor;
+			
+			arch >> m_settings;
+		}	
+	};
+
+
+}
diff --git a/Engine/Plugins/Experimental/Mutable/Source/MutableTools/Internal/MuT/NodeImageMultiLayerPrivate.h b/Engine/Plugins/Experimental/Mutable/Source/MutableTools/Internal/MuT/NodeImageMultiLayerPrivate.h
new file mode 100644
--- /dev/null
+++ b/Engine/Plugins/Experimental/Mutable/Source/MutableTools/Internal/MuT/NodeImageMultiLayerPrivate.h
@@ -0,0 +1,62 @@
+// Copyright Epic Games, Inc. All Rights Reserved.
+
+#pragma once
+
+#include "MuT/NodeImagePrivate.h"
+#include "MuT/NodeImageMultiLayer.h"
+#include "MuT/NodeColour.h"
+#include "MuT/AST.h"
+
+
+namespace mu
+{
+
+    class NodeImageMultiLayer::Private : public NodeImage::Private
+	{
+	public:
+
+		static FNodeType s_type;
+
+		NodeImagePtr m_pBase;
+		NodeImagePtr m_pMask;
+        NodeImagePtr m_pBlended;
+        NodeRangePtr m_pRange;
+		EBlendType m_type;
+
+		//!
+		void Serialise( OutputArchive& arch ) const
+		{
+            uint32_t ver = 1;
+			arch << ver;
+            arch << uint32_t(m_type);
+
+			arch << m_pBase;
+			arch << m_pMask;
+            arch << m_pBlended;
+            arch << m_pRange;
+		}
+
+		//!
+		void Unserialise( InputArchive& arch )
+		{
+            uint32_t ver;
+			arch >> ver;
+			check(ver<=1);
+
+            uint32_t t;
+            arch >> t;
+			if (ver < 1)
+			{
+				++t;
+			}
+            m_type=(EBlendType)t;
+
+            arch >> m_pBase;
+			arch >> m_pMask;
+            arch >> m_pBlended;
+            arch >> m_pRange;
+
+		}
+	};
+
+}
diff --git a/Engine/Plugins/Experimental/Mutable/Source/MutableTools/Internal/MuT/NodeImagePlainColourPrivate.h b/Engine/Plugins/Experimental/Mutable/Source/MutableTools/Internal/MuT/NodeImagePlainColourPrivate.h
new file mode 100644
--- /dev/null
+++ b/Engine/Plugins/Experimental/Mutable/Source/MutableTools/Internal/MuT/NodeImagePlainColourPrivate.h
@@ -0,0 +1,69 @@
+// Copyright Epic Games, Inc. All Rights Reserved.
+
+#pragma once
+
+#include "MuT/NodeColour.h"
+#include "MuT/NodeImagePrivate.h"
+#include "MuT/NodeImagePlainColour.h"
+#include "MuT/AST.h"
+
+
+namespace mu
+{
+
+
+	class NodeImagePlainColour::Private : public NodeImage::Private
+	{
+	public:
+
+		Private()
+		{
+            m_sizeX = 4;
+            m_sizeY = 4;
+		}
+
+		static FNodeType s_type;
+
+		Ptr<NodeColour> m_pColour;
+		int32 m_sizeX, m_sizeY;
+		EImageFormat Format = EImageFormat::IF_RGB_UBYTE;
+
+
+		//!
+		void Serialise( OutputArchive& arch ) const
+		{
+            uint32 ver = 1;
+			arch << ver;
+
+			arch << m_pColour;
+			arch << m_sizeX;
+			arch << m_sizeY;
+			arch << (uint32)Format;
+		}
+
+		//!
+		void Unserialise( InputArchive& arch )
+		{
+            uint32 ver;
+			arch >> ver;
+			check(ver<=1);
+
+			arch >> m_pColour;
+			arch >> m_sizeX;
+			arch >> m_sizeY;
+
+			if (ver >= 1)
+			{
+				uint32 Temp = 0;
+				arch >> Temp;
+				Format = EImageFormat(Temp);
+			}
+			else
+			{
+				Format = EImageFormat::IF_RGB_UBYTE;
+			}
+		}
+	};
+
+
+}
diff --git a/Engine/Plugins/Experimental/Mutable/Source/MutableTools/Internal/MuT/NodeImagePrivate.h b/Engine/Plugins/Experimental/Mutable/Source/MutableTools/Internal/MuT/NodeImagePrivate.h
new file mode 100644
--- /dev/null
+++ b/Engine/Plugins/Experimental/Mutable/Source/MutableTools/Internal/MuT/NodeImagePrivate.h
@@ -0,0 +1,17 @@
+// Copyright Epic Games, Inc. All Rights Reserved.
+
+#pragma once
+
+#include "MuT/NodePrivate.h"
+#include "MuT/NodeImage.h"
+
+
+namespace mu
+{
+
+
+	class NodeImage::Private : public Node::Private
+	{
+	};
+
+}
diff --git a/Engine/Plugins/Experimental/Mutable/Source/MutableTools/Internal/MuT/NodeImageProjectPrivate.h b/Engine/Plugins/Experimental/Mutable/Source/MutableTools/Internal/MuT/NodeImageProjectPrivate.h
new file mode 100644
--- /dev/null
+++ b/Engine/Plugins/Experimental/Mutable/Source/MutableTools/Internal/MuT/NodeImageProjectPrivate.h
@@ -0,0 +1,93 @@
+// Copyright Epic Games, Inc. All Rights Reserved.
+
+#pragma once
+
+#include "MuT/NodeImageProject.h"
+
+#include "MuT/NodeImagePrivate.h"
+#include "MuT/NodeScalar.h"
+#include "MuT/NodeColour.h"
+#include "MuT/NodeMesh.h"
+#include "MuT/NodeProjector.h"
+#include "MuR/MutableMath.h"
+
+namespace mu
+{
+
+	class NodeImageProject::Private : public NodeImage::Private
+	{
+	public:
+
+		static FNodeType s_type;
+
+		Ptr<NodeProjector> m_pProjector;
+		Ptr<NodeMesh> m_pMesh;
+		Ptr<NodeScalar> m_pAngleFadeStart;
+		Ptr<NodeScalar> m_pAngleFadeEnd;
+		Ptr<NodeImage> m_pImage;
+		Ptr<NodeImage> m_pMask;
+		FUintVector2 m_imageSize;
+		uint8 m_layout = 0;
+		bool bIsRGBFadingEnabled = true;
+		bool bIsAlphaFadingEnabled = true;
+		bool bEnableTextureSeamCorrection = true;
+		ESamplingMethod SamplingMethod = ESamplingMethod::Point;
+		EMinFilterMethod MinFilterMethod = EMinFilterMethod::None;
+
+		//!
+		void Serialise( OutputArchive& arch ) const
+		{
+            uint32 ver = 5;
+			arch << ver;
+
+			arch << m_pProjector;
+			arch << m_pMesh;
+			arch << m_pAngleFadeStart;
+			arch << m_pAngleFadeEnd;
+			arch << m_pImage;
+            arch << m_pMask;
+            arch << m_layout;
+			arch << m_imageSize;
+			arch << bIsRGBFadingEnabled;
+			arch << bIsAlphaFadingEnabled;
+			arch << bEnableTextureSeamCorrection;
+			arch << SamplingMethod;
+			arch << MinFilterMethod;
+		}
+
+		//!
+		void Unserialise( InputArchive& arch )
+		{
+            uint32 ver;
+			arch >> ver;
+            check(ver>=2 && ver<=5);
+
+			arch >> m_pProjector;
+			arch >> m_pMesh;
+			arch >> m_pAngleFadeStart;
+			arch >> m_pAngleFadeEnd;
+			arch >> m_pImage;
+			arch >> m_pMask;
+            arch >> m_layout;
+			arch >> m_imageSize;
+			if (ver >= 3)
+			{
+				arch >> bIsRGBFadingEnabled;
+				arch >> bIsAlphaFadingEnabled;
+			}
+
+			if (ver >= 5)
+			{
+				arch >> bEnableTextureSeamCorrection;
+			}
+
+			if (ver >= 4)
+			{
+				arch >> SamplingMethod;
+				arch >> MinFilterMethod;
+			}
+		}
+	};
+
+
+}
diff --git a/Engine/Plugins/Experimental/Mutable/Source/MutableTools/Internal/MuT/NodeImageResizePrivate.h b/Engine/Plugins/Experimental/Mutable/Source/MutableTools/Internal/MuT/NodeImageResizePrivate.h
new file mode 100644
--- /dev/null
+++ b/Engine/Plugins/Experimental/Mutable/Source/MutableTools/Internal/MuT/NodeImageResizePrivate.h
@@ -0,0 +1,51 @@
+// Copyright Epic Games, Inc. All Rights Reserved.
+
+#pragma once
+
+#include "MuT/NodeImagePrivate.h"
+#include "MuT/NodeImageResize.h"
+#include "MuT/AST.h"
+
+
+namespace mu
+{
+
+
+	class NodeImageResize::Private : public NodeImage::Private
+	{
+	public:
+
+		static FNodeType s_type;
+
+		NodeImagePtr m_pBase;
+		bool m_relative = true;
+		float m_sizeX = 0.5f, m_sizeY = 0.5f;
+
+		//!
+		void Serialise( OutputArchive& arch ) const
+		{
+            uint32_t ver = 0;
+			arch << ver;
+
+			arch << m_pBase;
+			arch << m_relative;
+			arch << m_sizeX;
+			arch << m_sizeY;
+		}
+
+		//!
+		void Unserialise( InputArchive& arch )
+		{
+            uint32_t ver;
+			arch >> ver;
+			check(ver==0);
+
+			arch >> m_pBase;
+			arch >> m_relative;
+			arch >> m_sizeX;
+			arch >> m_sizeY;
+		}
+	};
+
+
+}
diff --git a/Engine/Plugins/Experimental/Mutable/Source/MutableTools/Internal/MuT/NodeImageSwitchPrivate.h b/Engine/Plugins/Experimental/Mutable/Source/MutableTools/Internal/MuT/NodeImageSwitchPrivate.h
new file mode 100644
--- /dev/null
+++ b/Engine/Plugins/Experimental/Mutable/Source/MutableTools/Internal/MuT/NodeImageSwitchPrivate.h
@@ -0,0 +1,46 @@
+// Copyright Epic Games, Inc. All Rights Reserved.
+
+#pragma once
+
+#include "MuT/NodeImageSwitch.h"
+#include "MuT/NodeImagePrivate.h"
+#include "MuT/NodeScalar.h"
+
+
+namespace mu
+{
+
+
+	class NodeImageSwitch::Private : public NodeImage::Private
+	{
+	public:
+
+		static FNodeType s_type;
+
+		NodeScalarPtr m_pParameter;
+		TArray<NodeImagePtr> m_options;
+
+		//!
+		void Serialise( OutputArchive& arch ) const
+		{
+            uint32_t ver = 0;
+			arch << ver;
+
+			arch << m_pParameter;
+			arch << m_options;
+		}
+
+		//!
+		void Unserialise( InputArchive& arch )
+		{
+            uint32_t ver;
+			arch >> ver;
+			check(ver==0);
+
+			arch >> m_pParameter;
+			arch >> m_options;
+		}
+	};
+
+
+}
diff --git a/Engine/Plugins/Experimental/Mutable/Source/MutableTools/Internal/MuT/NodeImageSwizzlePrivate.h b/Engine/Plugins/Experimental/Mutable/Source/MutableTools/Internal/MuT/NodeImageSwizzlePrivate.h
new file mode 100644
--- /dev/null
+++ b/Engine/Plugins/Experimental/Mutable/Source/MutableTools/Internal/MuT/NodeImageSwizzlePrivate.h
@@ -0,0 +1,52 @@
+// Copyright Epic Games, Inc. All Rights Reserved.
+
+#pragma once
+
+#include "MuT/NodeImageSwizzle.h"
+#include "MuT/NodeImagePrivate.h"
+#include "MuR/ImagePrivate.h"
+
+
+namespace mu
+{
+
+
+	class NodeImageSwizzle::Private : public NodeImage::Private
+	{
+	public:
+
+		static FNodeType s_type;
+
+		EImageFormat m_format;
+		TArray<NodeImagePtr> m_sources;
+		TArray<int> m_sourceChannels;
+
+		//!
+		void Serialise( OutputArchive& arch ) const
+		{
+            uint32_t ver = 0;
+			arch << ver;
+
+			arch << (uint32)m_format;
+			arch << m_sources;
+			arch << m_sourceChannels;
+		}
+
+		//!
+		void Unserialise( InputArchive& arch )
+		{
+            uint32_t ver;
+			arch >> ver;
+			check(ver==0);
+
+			uint32 format;
+			arch >> format;
+			m_format = (EImageFormat)format;
+
+			arch >> m_sources;
+			arch >> m_sourceChannels;
+		}
+	};
+
+
+}
diff --git a/Engine/Plugins/Experimental/Mutable/Source/MutableTools/Internal/MuT/NodeLODPrivate.h b/Engine/Plugins/Experimental/Mutable/Source/MutableTools/Internal/MuT/NodeLODPrivate.h
new file mode 100644
--- /dev/null
+++ b/Engine/Plugins/Experimental/Mutable/Source/MutableTools/Internal/MuT/NodeLODPrivate.h
@@ -0,0 +1,47 @@
+// Copyright Epic Games, Inc. All Rights Reserved.
+
+#pragma once
+
+#include "MuT/NodePrivate.h"
+#include "MuT/NodeLOD.h"
+#include "MuT/NodeComponentPrivate.h"
+#include "MuT/NodeModifierPrivate.h"
+#include "MuT/AST.h"
+
+
+namespace mu
+{
+
+
+	class NodeLOD::Private : public Node::Private
+	{
+	public:
+
+		static FNodeType s_type;
+
+		TArray<NodeComponentPtr> m_components;
+		TArray<NodeModifierPtr> m_modifiers;
+
+		//!
+		void Serialise( OutputArchive& arch ) const
+		{
+            uint32_t ver = 1;
+			arch << ver;
+
+			arch << m_components;
+			arch << m_modifiers;
+		}
+
+		//!
+		void Unserialise( InputArchive& arch )
+		{
+            uint32_t ver;
+			arch >> ver;
+            check(ver == 1);
+
+			arch >> m_components;
+			arch >> m_modifiers;
+		}
+	};
+
+}
diff --git a/Engine/Plugins/Experimental/Mutable/Source/MutableTools/Internal/MuT/NodeMeshConstantPrivate.h b/Engine/Plugins/Experimental/Mutable/Source/MutableTools/Internal/MuT/NodeMeshConstantPrivate.h
new file mode 100644
--- /dev/null
+++ b/Engine/Plugins/Experimental/Mutable/Source/MutableTools/Internal/MuT/NodeMeshConstantPrivate.h
@@ -0,0 +1,56 @@
+// Copyright Epic Games, Inc. All Rights Reserved.
+
+#pragma once
+
+#include "MuT/NodeMeshPrivate.h"
+
+#include "MuT/NodeMeshConstant.h"
+#include "MuT/NodeLayout.h"
+#include "MuT/AST.h"
+
+
+namespace mu
+{
+
+	class NodeMeshConstant::Private : public NodeMesh::Private
+	{
+	public:
+
+		static FNodeType s_type;
+
+		MeshPtr m_pValue;
+
+		TArray<NodeLayoutPtr> m_layouts;
+
+		//!
+		void Serialise( OutputArchive& arch ) const
+		{
+            uint32_t ver = 0;
+			arch << ver;
+
+			arch << m_pValue;
+			arch << m_layouts;
+		}
+
+		//!
+		void Unserialise( InputArchive& arch )
+		{
+            uint32_t ver;
+			arch >> ver;
+			check(ver==0);
+
+			arch >> m_pValue;
+			arch >> m_layouts;
+
+			//if (m_pValue)
+			//{
+			//	m_pValue->CheckIntegrity();
+			//}
+        }
+
+		// NodeMesh::Private interface
+        NodeLayoutPtr GetLayout( int index ) const override;
+
+	};
+
+}
diff --git a/Engine/Plugins/Experimental/Mutable/Source/MutableTools/Internal/MuT/NodeMeshFormatPrivate.h b/Engine/Plugins/Experimental/Mutable/Source/MutableTools/Internal/MuT/NodeMeshFormatPrivate.h
new file mode 100644
--- /dev/null
+++ b/Engine/Plugins/Experimental/Mutable/Source/MutableTools/Internal/MuT/NodeMeshFormatPrivate.h
@@ -0,0 +1,64 @@
+// Copyright Epic Games, Inc. All Rights Reserved.
+
+#pragma once
+
+#include "MuT/NodeMeshPrivate.h"
+#include "MuT/NodeMeshFormat.h"
+#include "MuT/AST.h"
+
+
+namespace mu
+{
+
+
+	class NodeMeshFormat::Private : public NodeMesh::Private
+	{
+	public:
+
+		static FNodeType s_type;
+
+		//! Source mesh to transform
+		NodeMeshPtr m_pSource;
+
+		//! New mesh format. The buffers in the sets have no elements, but they define the formats.
+		//! If they are null it means that they are left with the original format.
+		FMeshBufferSet m_VertexBuffers;
+		FMeshBufferSet m_IndexBuffers;
+		FMeshBufferSet m_FaceBuffers;
+
+		//!
+		void Serialise( OutputArchive& arch ) const
+		{
+            uint32_t ver = 4;
+			arch << ver;
+
+			arch << m_pSource;
+			arch << m_VertexBuffers;
+			arch << m_IndexBuffers;
+			arch << m_FaceBuffers;
+		}
+
+		//!
+		void Unserialise( InputArchive& arch )
+		{
+            uint32_t ver;
+			arch >> ver;
+            check(ver>=3 && ver<=4);
+
+			arch >> m_pSource;
+			arch >> m_VertexBuffers;
+			arch >> m_IndexBuffers;
+			arch >> m_FaceBuffers;
+			if (ver == 3)
+			{
+				bool bDummy;
+				arch >> bDummy;
+			}
+		}
+
+		// NodeMesh::Private interface
+        NodeLayoutPtr GetLayout( int index ) const override;
+	};
+
+
+}
diff --git a/Engine/Plugins/Experimental/Mutable/Source/MutableTools/Internal/MuT/NodeMeshFragmentPrivate.h b/Engine/Plugins/Experimental/Mutable/Source/MutableTools/Internal/MuT/NodeMeshFragmentPrivate.h
new file mode 100644
--- /dev/null
+++ b/Engine/Plugins/Experimental/Mutable/Source/MutableTools/Internal/MuT/NodeMeshFragmentPrivate.h
@@ -0,0 +1,68 @@
+// Copyright Epic Games, Inc. All Rights Reserved.
+
+#pragma once
+
+#include "MuT/NodeMeshFragment.h"
+#include "MuT/NodeScalar.h"
+#include "MuT/NodeMeshPrivate.h"
+#include "MuT/AST.h"
+
+
+namespace mu
+{
+
+
+	class NodeMeshFragment::Private : public NodeMesh::Private
+	{
+	public:
+
+		Private()
+		{
+            m_layoutOrGroup = -1;
+		}
+
+		static FNodeType s_type;
+
+		NodeMeshPtr m_pMesh;
+        int m_layoutOrGroup;
+		TArray<int> m_blocks;
+
+        FRAGMENT_TYPE m_fragmentType;
+
+		//!
+		void Serialise( OutputArchive& arch ) const
+		{
+            uint32_t ver = 1;
+			arch << ver;
+
+			arch << m_pMesh;
+            arch << (int32_t)m_fragmentType;
+            arch << m_layoutOrGroup;
+			arch << m_blocks;
+
+		}
+
+		//!
+		void Unserialise( InputArchive& arch )
+		{
+            uint32_t ver;
+			arch >> ver;
+            check(ver==1);
+
+			arch >> m_pMesh;
+
+            int32_t temp;
+            arch >> temp;
+            m_fragmentType = (FRAGMENT_TYPE)temp;
+
+            arch >> m_layoutOrGroup;
+			arch >> m_blocks;
+		}
+
+
+		// NodeMesh::Private interface
+        NodeLayoutPtr GetLayout( int index ) const override;
+	};
+
+
+}
diff --git a/Engine/Plugins/Experimental/Mutable/Source/MutableTools/Internal/MuT/NodeMeshMakeMorphPrivate.h b/Engine/Plugins/Experimental/Mutable/Source/MutableTools/Internal/MuT/NodeMeshMakeMorphPrivate.h
new file mode 100644
--- /dev/null
+++ b/Engine/Plugins/Experimental/Mutable/Source/MutableTools/Internal/MuT/NodeMeshMakeMorphPrivate.h
@@ -0,0 +1,58 @@
+// Copyright Epic Games, Inc. All Rights Reserved.
+
+#pragma once
+
+#include "MuT/NodeMeshPrivate.h"
+#include "MuT/NodeMeshMakeMorph.h"
+#include "MuT/AST.h"
+
+
+namespace mu
+{
+
+
+    class NodeMeshMakeMorph::Private : public NodeMesh::Private
+	{
+	public:
+
+		static FNodeType s_type;
+
+		NodeMeshPtr m_pBase;
+		NodeMeshPtr m_pTarget;
+		
+		bool bOnlyPositionAndNormal = false;
+
+		//!
+		void Serialise( OutputArchive& arch ) const
+		{
+            uint32_t ver = 1;
+			arch << ver;
+
+			arch << m_pBase;
+			arch << m_pTarget;
+			arch << bOnlyPositionAndNormal;
+		}
+
+		//!
+		void Unserialise( InputArchive& arch )
+		{
+            uint32_t ver;
+			arch >> ver;
+			check(ver <= 1);
+
+			arch >> m_pBase;
+			arch >> m_pTarget;
+
+			if (ver == 1)
+			{
+				arch >> bOnlyPositionAndNormal;
+			}
+		}
+
+		// NodeMesh::Private interface
+        NodeLayoutPtr GetLayout( int index ) const override;
+
+	};
+
+
+}
diff --git a/Engine/Plugins/Experimental/Mutable/Source/MutableTools/Internal/MuT/NodeMeshMorphPrivate.h b/Engine/Plugins/Experimental/Mutable/Source/MutableTools/Internal/MuT/NodeMeshMorphPrivate.h
new file mode 100644
--- /dev/null
+++ b/Engine/Plugins/Experimental/Mutable/Source/MutableTools/Internal/MuT/NodeMeshMorphPrivate.h
@@ -0,0 +1,154 @@
+// Copyright Epic Games, Inc. All Rights Reserved.
+
+#pragma once
+
+#include "MuT/NodeMeshMorph.h"
+#include "MuT/NodeScalar.h"
+#include "MuT/NodeMeshPrivate.h"
+#include "MuT/AST.h"
+
+namespace mu
+{
+
+
+	class NodeMeshMorph::Private : public NodeMesh::Private
+	{
+	public:
+
+		static FNodeType s_type;
+
+		Ptr<NodeScalar> Factor;
+		Ptr<NodeMesh> Base;
+		Ptr<NodeMesh> Morph;
+		
+		bool bReshapeSkeleton = false;
+		bool bReshapePhysicsVolumes = false;
+		
+		TArray<uint16> BonesToDeform;
+		TArray<uint16> PhysicsToDeform;
+
+        //!
+		void Serialise( OutputArchive& arch ) const
+		{
+            uint32_t ver = 7;
+			arch << ver;
+
+			arch << Factor;
+			arch << Base;
+			arch << Morph;
+
+			arch << bReshapeSkeleton;
+			arch << bReshapePhysicsVolumes;
+			arch << BonesToDeform;
+			arch << PhysicsToDeform;
+        }
+
+		//!
+		void Unserialise( InputArchive& arch )
+		{
+            uint32_t ver;
+			arch >> ver;
+			check(ver <= 7);
+
+			arch >> Factor;
+			arch >> Base;
+
+			if (ver < 6)
+			{
+				TArray<Ptr<NodeMesh>> Morphs;
+				arch >> Morphs;
+
+				// The result will not be the same.
+				if (!Morphs.IsEmpty())
+				{
+					Morph = Morphs.Last();
+				}
+			}
+			else
+			{
+				arch >> Morph;
+			}
+
+			if (ver < 6)
+			{
+				bool bDeprecated;
+				arch >> bDeprecated;
+			}
+	
+			if (ver >= 2)
+			{
+				arch >> bReshapeSkeleton;
+				arch >> bReshapePhysicsVolumes;
+	
+				// This repetition is needed, there was a bug where m_reshapePhysicsVolumes was serialized twice in ver 2.
+				if (ver == 2) 
+				{
+					arch >> bReshapePhysicsVolumes;
+				}
+
+				if (ver >= 7)
+				{
+					arch >> BonesToDeform;
+				}
+				else
+				{
+					TArray<string> BonesToDeform_DEPRECATED;
+					arch >> BonesToDeform_DEPRECATED;
+
+					const int32 NumBonesToDeform = BonesToDeform.Num();
+					BonesToDeform.SetNumUninitialized(NumBonesToDeform);
+					for (int32 Index = 0; Index < NumBonesToDeform; ++Index)
+					{
+						BonesToDeform[Index] = Index;
+					}
+				}
+			}
+			else
+			{
+				bReshapeSkeleton = false;
+				bReshapePhysicsVolumes = false;
+				BonesToDeform.Empty();
+			}
+
+			if (ver == 3)
+			{
+				bool bDeformAllBones_DEPRECATED;
+				arch >> bDeformAllBones_DEPRECATED;
+			}
+
+			if (ver >= 3 && ver < 5)
+			{
+				bool bDeformAllPhysics_DEPRECATED;
+				arch >> bDeformAllPhysics_DEPRECATED;
+			}
+
+			if (ver >= 7)
+			{
+				arch >> PhysicsToDeform;
+			}
+			else if (ver >= 3)
+			{
+				TArray<string> PhysicsToDeform_DEPRECATED;
+				arch >> PhysicsToDeform_DEPRECATED;
+
+				const int32 NumPhysicsToDeform = PhysicsToDeform.Num();
+				PhysicsToDeform.SetNumUninitialized(NumPhysicsToDeform);
+				for (int32 Index = 0; Index < NumPhysicsToDeform; ++Index)
+				{
+					PhysicsToDeform[Index] = Index;
+				}
+			}
+			else
+			{
+				PhysicsToDeform.Empty();
+			}
+
+		}
+
+		// NodeMesh::Private interface
+        Ptr<NodeLayout> GetLayout( int index ) const override;
+
+	};
+
+
+}
diff --git a/Engine/Plugins/Experimental/Mutable/Source/MutableTools/Internal/MuT/NodeMeshPrivate.h b/Engine/Plugins/Experimental/Mutable/Source/MutableTools/Internal/MuT/NodeMeshPrivate.h
new file mode 100644
--- /dev/null
+++ b/Engine/Plugins/Experimental/Mutable/Source/MutableTools/Internal/MuT/NodeMeshPrivate.h
@@ -0,0 +1,23 @@
+// Copyright Epic Games, Inc. All Rights Reserved.
+
+#pragma once
+
+#include "MuT/NodePrivate.h"
+
+#include "MuT/NodeMesh.h"
+#include "MuT/NodeLayout.h"
+
+
+namespace mu
+{
+
+
+	class NodeMesh::Private : public Node::Private
+	{
+	public:
+
+		virtual NodeLayoutPtr GetLayout( int index ) const = 0;
+
+	};
+
+}
diff --git a/Engine/Plugins/Experimental/Mutable/Source/MutableTools/Internal/MuT/NodeMeshTablePrivate.h b/Engine/Plugins/Experimental/Mutable/Source/MutableTools/Internal/MuT/NodeMeshTablePrivate.h
new file mode 100644
--- /dev/null
+++ b/Engine/Plugins/Experimental/Mutable/Source/MutableTools/Internal/MuT/NodeMeshTablePrivate.h
@@ -0,0 +1,90 @@
+// Copyright Epic Games, Inc. All Rights Reserved.
+
+#pragma once
+
+#include "MuT/NodeMeshPrivate.h"
+#include "MuT/NodeMeshTable.h"
+#include "MuT/NodeLayout.h"
+#include "MuT/TablePrivate.h"
+#include "MuT/AST.h"
+
+
+namespace mu
+{
+
+
+	class NodeMeshTable::Private : public NodeMesh::Private
+	{
+	public:
+
+		static FNodeType s_type;
+
+		FString ParameterName;
+		TablePtr Table;
+		FString ColumnName;
+		bool bNoneOption = false;
+
+		TArray<NodeLayoutPtr> Layouts;
+
+		//!
+		void Serialise( OutputArchive& arch ) const
+		{
+            uint32 ver = 3;
+			arch << ver;
+
+			arch << ParameterName;
+			arch << Table;
+			arch << ColumnName;
+			arch << Layouts;
+			arch << bNoneOption;
+		}
+
+		//!
+		void Unserialise( InputArchive& arch )
+		{
+            uint32 ver;
+			arch >> ver;
+            check(ver>=1 && ver<=3);
+
+			if (ver == 1)
+			{
+				std::string Temp;
+				arch >> Temp;
+				ParameterName = Temp.c_str();
+			}
+			else
+			{
+				arch >> ParameterName;
+			}
+
+			arch >> Table;
+
+			if (ver == 1)
+			{
+				std::string Temp;
+				arch >> Temp;
+				ColumnName = Temp.c_str();
+			}
+			else
+			{
+				arch >> ColumnName;
+			}
+			
+			arch >> Layouts;
+
+			if(ver >= 3)
+			{
+				arch >> bNoneOption;
+			}
+			else
+			{
+				bNoneOption = Table->GetPrivate()->bNoneOption_DEPRECATED;
+			}
+		}
+
+		// NodeMesh::Private interface
+        NodeLayoutPtr GetLayout( int index ) const override;
+
+	};
+
+}
diff --git a/Engine/Plugins/Experimental/Mutable/Source/MutableTools/Internal/MuT/NodeModifierMeshClipDeformPrivate.h b/Engine/Plugins/Experimental/Mutable/Source/MutableTools/Internal/MuT/NodeModifierMeshClipDeformPrivate.h
new file mode 100644
--- /dev/null
+++ b/Engine/Plugins/Experimental/Mutable/Source/MutableTools/Internal/MuT/NodeModifierMeshClipDeformPrivate.h
@@ -0,0 +1,62 @@
+// Copyright Epic Games, Inc. All Rights Reserved.
+
+#pragma once
+
+#include "MuT/NodeModifierPrivate.h"
+#include "MuT/NodeModifierMeshClipDeform.h"
+#include "MuT/NodeMesh.h"
+#include "MuT/AST.h"
+
+namespace mu
+{
+
+
+    class NodeModifierMeshClipDeform::Private : public NodeModifier::Private
+    {
+    public:
+
+		Private()
+		{
+		}
+
+		static FNodeType s_type;
+
+		//! 
+		NodeMeshPtr ClipMesh;
+		EShapeBindingMethod BindingMethod = EShapeBindingMethod::ClipDeformClosestProject;
+
+		//!
+		void Serialise( OutputArchive& arch ) const
+		{
+			NodeModifier::Private::Serialise(arch);
+
+            uint32_t ver = 1;
+			arch << ver;
+
+			arch << ClipMesh;
+			arch << BindingMethod;
+		}
+
+		//!
+		void Unserialise( InputArchive& arch )
+		{
+			NodeModifier::Private::Unserialise( arch );
+			
+            uint32_t ver;
+			arch >> ver;
+            check(ver<=1);
+
+			arch >> ClipMesh;
+			if (ver >= 1)
+			{
+				arch >> BindingMethod;
+			}
+			else
+			{
+				BindingMethod = EShapeBindingMethod::ClipDeformClosestProject;
+			}
+	
+		}
+    };
+
+}
diff --git a/Engine/Plugins/Experimental/Mutable/Source/MutableTools/Internal/MuT/NodeModifierMeshClipMorphPlanePrivate.h b/Engine/Plugins/Experimental/Mutable/Source/MutableTools/Internal/MuT/NodeModifierMeshClipMorphPlanePrivate.h
new file mode 100644
--- /dev/null
+++ b/Engine/Plugins/Experimental/Mutable/Source/MutableTools/Internal/MuT/NodeModifierMeshClipMorphPlanePrivate.h
@@ -0,0 +1,120 @@
+// Copyright Epic Games, Inc. All Rights Reserved.
+
+#pragma once
+
+#include "MuT/NodeModifierPrivate.h"
+#include "MuT/NodeModifierMeshClipMorphPlane.h"
+#include "MuT/AST.h"
+
+#include "MuR/MutableMath.h"
+
+
+namespace mu
+{
+    class NodeModifierMeshClipMorphPlane::Private : public NodeModifier::Private
+    {
+    public:
+
+        Private()
+                : m_dist(0.0f), m_factor(0.0f)
+                , m_radius1(0.0f), m_radius2(0.0f), m_rotation(0.0f)
+                , m_vertexSelectionType(VS_ALL)
+				, m_maxEffectRadius(-1.f)
+        {
+        }
+
+        static FNodeType s_type;
+
+        // Morph field parameters
+
+        //! Distance to the plane of last affected vertex
+        float m_dist;
+
+        //! "Linearity" factor of the influence.
+        float m_factor;
+
+        // Ellipse location
+        vec3f m_origin;
+        vec3f m_normal;
+        float m_radius1, m_radius2, m_rotation;
+
+        //! Typed of vertex selection
+        typedef enum
+        {
+            //! All vertices, so no extra info is needed
+            VS_ALL,
+
+            //! Select vertices inside a shape
+            VS_SHAPE,
+
+            //! Select all vertices affected by any bone in a sub hierarchy
+            VS_BONE_HIERARCHY,
+        } VERTEX_SELECTION;
+
+        // Vertex selection box
+        uint8_t m_vertexSelectionType;
+        vec3f m_selectionBoxOrigin;
+        vec3f m_selectionBoxRadius;
+        uint16 m_vertexSelectionBone;
+
+		// Max distance a vertex can have to the bone in order to be affected. A negative value
+		// means no limit.
+		float m_maxEffectRadius;
+
+        //!
+        void Serialise( OutputArchive& arch ) const
+        {
+            NodeModifier::Private::Serialise(arch);
+
+			uint32_t ver = 3;
+            arch << ver;
+
+            arch << m_origin;
+            arch << m_normal;
+            arch << m_dist;
+            arch << m_factor;
+            arch << m_radius1;
+            arch << m_radius2;
+            arch << m_rotation;
+            arch << m_vertexSelectionType;
+            arch << m_selectionBoxOrigin;
+            arch << m_selectionBoxRadius;
+            arch << m_vertexSelectionBone;
+			arch << m_maxEffectRadius;
+        }
+
+        //!
+        void Unserialise( InputArchive& arch )
+        {
+            NodeModifier::Private::Unserialise( arch );
+
+            uint32_t ver;
+            arch >> ver;
+            check(ver>=2);
+
+            arch >> m_origin;
+            arch >> m_normal;
+            arch >> m_dist;
+            arch >> m_factor;
+            arch >> m_radius1;
+            arch >> m_radius2;
+            arch >> m_rotation;
+            arch >> m_vertexSelectionType;
+            arch >> m_selectionBoxOrigin;
+            arch >> m_selectionBoxRadius;
+			if (ver >= 3)
+			{
+				arch >> m_vertexSelectionBone;
+			}
+			else
+			{
+				string OldVertexSelectionBone;
+				arch >> OldVertexSelectionBone;
+				m_vertexSelectionBone = 0;
+			}
+			arch >> m_maxEffectRadius;
+        }
+
+    };
+
+}
diff --git a/Engine/Plugins/Experimental/Mutable/Source/MutableTools/Internal/MuT/NodeModifierMeshClipWithMeshPrivate.h b/Engine/Plugins/Experimental/Mutable/Source/MutableTools/Internal/MuT/NodeModifierMeshClipWithMeshPrivate.h
new file mode 100644
--- /dev/null
+++ b/Engine/Plugins/Experimental/Mutable/Source/MutableTools/Internal/MuT/NodeModifierMeshClipWithMeshPrivate.h
@@ -0,0 +1,55 @@
+// Copyright Epic Games, Inc. All Rights Reserved.
+
+#pragma once
+
+#include "MuT/NodeModifierPrivate.h"
+#include "MuT/NodeModifierMeshClipWithMesh.h"
+#include "MuT/NodeMesh.h"
+#include "MuT/AST.h"
+
+#include "MuR/MutableMath.h"
+
+
+namespace mu
+{
+
+    class NodeModifierMeshClipWithMesh::Private : public NodeModifier::Private
+	{
+	public:
+
+		Private()
+		{
+		}
+
+		static FNodeType s_type;
+
+		//! 
+		Ptr<NodeMesh> ClipMesh;
+
+		//!
+		void Serialise( OutputArchive& arch ) const
+		{
+			NodeModifier::Private::Serialise(arch);
+
+            uint32_t ver = 0;
+			arch << ver;
+
+			arch << ClipMesh;
+		}
+
+		//!
+		void Unserialise( InputArchive& arch )
+		{
+			NodeModifier::Private::Unserialise( arch );
+			
+            uint32_t ver;
+			arch >> ver;
+            check(ver<=0);
+
+			arch >> ClipMesh;
+		}
+
+	};
+
+
+}
diff --git a/Engine/Plugins/Experimental/Mutable/Source/MutableTools/Internal/MuT/NodeModifierMeshClipWithUVMaskPrivate.h b/Engine/Plugins/Experimental/Mutable/Source/MutableTools/Internal/MuT/NodeModifierMeshClipWithUVMaskPrivate.h
new file mode 100644
--- /dev/null
+++ b/Engine/Plugins/Experimental/Mutable/Source/MutableTools/Internal/MuT/NodeModifierMeshClipWithUVMaskPrivate.h
@@ -0,0 +1,56 @@
+// Copyright Epic Games, Inc. All Rights Reserved.
+
+#pragma once
+
+#include "MuT/NodeModifierPrivate.h"
+#include "MuT/NodeModifierMeshClipWithUVMask.h"
+#include "MuT/NodeImage.h"
+#include "MuT/AST.h"
+
+#include "MuR/MutableMath.h"
+
+
+namespace mu
+{
+
+    class NodeModifierMeshClipWithUVMask::Private : public NodeModifier::Private
+	{
+	public:
+
+		static FNodeType s_type;
+
+		/** Image with the regions to remove. It will be interpreted as a bitmap. */
+		Ptr<NodeImage> ClipMask;
+
+		/** Layout index of the UVs to use inthe source mesh to ben clipped with the mask. */
+		uint8 LayoutIndex = 0;
+
+		//!
+		void Serialise( OutputArchive& arch ) const
+		{
+			NodeModifier::Private::Serialise(arch);
+
+            uint32 ver = 0;
+			arch << ver;
+
+			arch << ClipMask;
+			arch << LayoutIndex;
+		}
+
+		//!
+		void Unserialise( InputArchive& arch )
+		{
+			NodeModifier::Private::Unserialise( arch );
+			
+            uint32 ver;
+			arch >> ver;
+            check(ver<=0);
+
+			arch >> ClipMask;
+			arch >> LayoutIndex;
+		}
+
+	};
+
+
+}
diff --git a/Engine/Plugins/Experimental/Mutable/Source/MutableTools/Internal/MuT/NodeModifierPrivate.h b/Engine/Plugins/Experimental/Mutable/Source/MutableTools/Internal/MuT/NodeModifierPrivate.h
new file mode 100644
--- /dev/null
+++ b/Engine/Plugins/Experimental/Mutable/Source/MutableTools/Internal/MuT/NodeModifierPrivate.h
@@ -0,0 +1,72 @@
+// Copyright Epic Games, Inc. All Rights Reserved.
+
+#pragma once
+
+#include "MuT/NodePrivate.h"
+#include "MuT/NodeModifier.h"
+
+#include <string>
+
+
+namespace mu
+{
+
+	MUTABLE_DEFINE_ENUM_SERIALISABLE(EMutableMultipleTagPolicy);
+
+	class NodeModifier::Private : public Node::Private
+	{
+	public:
+
+		/** Tags that target surface need to have enabled to receive this modifier. */
+		TArray<FString> RequiredTags;
+
+		/** In case of multiple tags in RequiredTags: are they all required, or one is enough? */
+		EMutableMultipleTagPolicy MultipleTagsPolicy = EMutableMultipleTagPolicy::OnlyOneRequired;
+
+		// Wether the modifier has to be applied after the normal node operations or before
+		bool bApplyBeforeNormalOperations = true;
+
+		//!
+		void Serialise(OutputArchive& arch) const
+		{
+            uint32_t ver = 4;
+			arch << ver;
+
+			arch << RequiredTags;
+			arch << MultipleTagsPolicy;
+			arch << bApplyBeforeNormalOperations;
+		}
+
+		//!
+		void Unserialise(InputArchive& arch)
+		{
+            uint32_t ver;
+			arch >> ver;
+			check(ver>=2 && ver<=3);
+
+			if (ver <= 2)
+			{
+				TArray<std::string> Temp;
+				arch >> Temp;
+				RequiredTags.SetNum(Temp.Num());
+				for( int32 i=0; i<Temp.Num(); ++i)
+				{
+					RequiredTags[i] = Temp[i].c_str();
+				}
+			}
+			else
+			{
+				arch >> RequiredTags;
+			}
+
+			if (ver >= 4)
+			{
+				arch >> MultipleTagsPolicy;
+			}
+
+			arch >> bApplyBeforeNormalOperations;
+		}
+
+	};
+
+}
diff --git a/Engine/Plugins/Experimental/Mutable/Source/MutableTools/Internal/MuT/NodePatchImagePrivate.h b/Engine/Plugins/Experimental/Mutable/Source/MutableTools/Internal/MuT/NodePatchImagePrivate.h
new file mode 100644
--- /dev/null
+++ b/Engine/Plugins/Experimental/Mutable/Source/MutableTools/Internal/MuT/NodePatchImagePrivate.h
@@ -0,0 +1,67 @@
+// Copyright Epic Games, Inc. All Rights Reserved.
+
+
+#pragma once
+
+#include "MuT/NodePrivate.h"
+#include "MuT/NodePatchImage.h"
+#include "MuT/NodeImage.h"
+
+
+namespace mu
+{
+
+    class NodePatchImage::Private : public Node::Private
+    {
+    public:
+
+        static FNodeType s_type;
+
+        NodeImagePtr m_pImage;
+        NodeImagePtr m_pMask;
+
+        TArray<int> m_blocks;
+
+        EBlendType m_blendType = EBlendType::BT_BLEND;
+
+        // Patch alpha channel as well?
+        bool m_applyToAlpha = false;
+
+        //!
+        void Serialise( OutputArchive& arch ) const
+        {
+            uint32_t ver = 3;
+            arch << ver;
+
+            arch << m_pImage;
+            arch << m_pMask;
+            arch << m_blocks;
+            arch << (uint32_t)m_blendType;
+            arch << m_applyToAlpha;
+        }
+
+        //!
+        void Unserialise( InputArchive& arch )
+        {
+            uint32_t ver;
+            arch >> ver;
+            check(ver>=2 && ver<=3);
+
+            arch >> m_pImage;
+            arch >> m_pMask;
+            arch >> m_blocks;
+
+			uint32_t t;
+            arch >> t;
+			if (ver <= 2)
+			{
+				++t;
+			}
+            m_blendType=(EBlendType)t;
+
+            arch >> m_applyToAlpha;
+        }
+    };
+
+}
+
diff --git a/Engine/Plugins/Experimental/Mutable/Source/MutableTools/Internal/MuT/NodePatchMeshPrivate.h b/Engine/Plugins/Experimental/Mutable/Source/MutableTools/Internal/MuT/NodePatchMeshPrivate.h
new file mode 100644
--- /dev/null
+++ b/Engine/Plugins/Experimental/Mutable/Source/MutableTools/Internal/MuT/NodePatchMeshPrivate.h
@@ -0,0 +1,47 @@
+// Copyright Epic Games, Inc. All Rights Reserved.
+
+#pragma once
+
+#include "MuT/NodeMeshPrivate.h"
+#include "MuT/NodePatchMesh.h"
+#include "MuT/AST.h"
+
+
+namespace mu
+{
+
+
+	class NodePatchMesh::Private : public Node::Private
+	{
+	public:
+
+		static FNodeType s_type;
+
+		NodeMeshPtr m_pRemove;
+		NodeMeshPtr m_pAdd;
+
+		//!
+		void Serialise( OutputArchive& arch ) const
+		{
+            uint32_t ver = 1;
+			arch << ver;
+
+			arch << m_pRemove;
+			arch << m_pAdd;
+		}
+
+		//!
+		void Unserialise( InputArchive& arch )
+		{
+            uint32_t ver;
+			arch >> ver;
+            check( ver==1 );
+
+			arch >> m_pRemove;
+			arch >> m_pAdd;
+		}
+
+	};
+
+
+}
diff --git a/Engine/Plugins/Experimental/Mutable/Source/MutableTools/Internal/MuT/NodeProjectorPrivate.h b/Engine/Plugins/Experimental/Mutable/Source/MutableTools/Internal/MuT/NodeProjectorPrivate.h
new file mode 100644
--- /dev/null
+++ b/Engine/Plugins/Experimental/Mutable/Source/MutableTools/Internal/MuT/NodeProjectorPrivate.h
@@ -0,0 +1,125 @@
+// Copyright Epic Games, Inc. All Rights Reserved.
+
+#pragma once
+
+#include "MuT/NodePrivate.h"
+#include "MuT/NodeProjector.h"
+
+#include "MuT/NodeRange.h"
+#include "MuR/MutableMath.h"
+#include "MuR/ParametersPrivate.h"
+
+
+namespace mu
+{
+
+
+	class NodeProjector::Private : public Node::Private
+	{
+	};
+
+
+	class NodeProjectorConstant::Private : public NodeProjector::Private
+	{
+	public:
+
+		static FNodeType s_type;
+
+        PROJECTOR_TYPE m_type = PROJECTOR_TYPE::PLANAR;
+		vec3<float> m_position;
+		vec3<float> m_direction;
+		vec3<float> m_up;
+        vec3<float> m_scale;
+        float m_projectionAngle = 0.0f;
+
+		//!
+		void Serialise( OutputArchive& arch ) const
+		{
+            uint32_t ver = 2;
+			arch << ver;
+
+            arch << m_type;
+            arch << m_position;
+            arch << m_direction;
+			arch << m_up;
+			arch << m_scale;
+            arch << m_projectionAngle;
+		}
+
+		//!
+		void Unserialise( InputArchive& arch )
+		{
+            uint32_t ver;
+			arch >> ver;
+            check(ver==2);
+
+            arch >> m_type;
+			arch >> m_position;
+			arch >> m_direction;
+			arch >> m_up;
+            arch >> m_scale;
+            arch >> m_projectionAngle;
+        }
+
+	};
+
+
+	class NodeProjectorParameter::Private : public NodeProjectorConstant::Private
+	{
+	public:
+
+		static FNodeType s_type;
+
+		FString m_name;
+		FString m_uid;
+
+		TArray<Ptr<NodeRange>> m_ranges;
+
+		//!
+		void Serialise( OutputArchive& arch ) const
+		{
+			NodeProjectorConstant::Private::Serialise( arch );
+
+            uint32_t ver = 3;
+			arch << ver;
+
+			arch << m_name;
+            arch << m_uid;
+            arch << m_ranges;
+        }
+
+		//!
+		void Unserialise( InputArchive& arch )
+		{
+			NodeProjectorConstant::Private::Unserialise( arch );
+
+            uint32_t ver;
+			arch >> ver;
+            check(ver<=3);
+
+			if (ver <= 2)
+			{
+				std::string Temp;
+				arch >> Temp;
+				m_name = Temp.c_str();
+				if (ver >= 1)
+				{
+					arch >> Temp;
+					m_uid = Temp.c_str();
+				}
+			}
+			else
+			{
+				arch >> m_name;
+				arch >> m_uid;
+			}
+
+            if (ver>=2)
+            {
+                arch >> m_ranges;
+            }
+        }
+	};
+
+
+}
diff --git a/Engine/Plugins/Experimental/Mutable/Source/MutableTools/Internal/MuT/NodeScalarCurvePrivate.h b/Engine/Plugins/Experimental/Mutable/Source/MutableTools/Internal/MuT/NodeScalarCurvePrivate.h
new file mode 100644
--- /dev/null
+++ b/Engine/Plugins/Experimental/Mutable/Source/MutableTools/Internal/MuT/NodeScalarCurvePrivate.h
@@ -0,0 +1,51 @@
+// Copyright Epic Games, Inc. All Rights Reserved.
+
+#pragma once
+
+#include "MuT/NodePrivate.h"
+#include "MuT/NodeScalarCurve.h"
+#include "MuT/NodeScalar.h"
+#include "MuT/AST.h"
+
+#include "MuR/SerialisationPrivate.h"
+#include "MuR/ParametersPrivate.h"
+
+
+namespace mu
+{
+	//---------------------------------------------------------------------------------------------
+	//!
+	//---------------------------------------------------------------------------------------------
+    class NodeScalarCurve::Private : public Node::Private
+	{
+	public:
+
+		static FNodeType s_type;
+
+		NodeScalarPtr m_input_scalar;
+		Curve m_curve;
+
+		//!
+		void Serialise(OutputArchive& arch) const
+		{
+            uint32_t ver = 1;
+			arch << ver;
+
+			arch << m_input_scalar;
+			arch << m_curve;
+		}
+
+		//!
+		void Unserialise(InputArchive& arch)
+		{
+            uint32_t ver;
+			arch >> ver;
+            check(ver == 1);
+
+			arch >> m_input_scalar;
+			arch >> m_curve;
+		}
+	};
+
+
+}
diff --git a/Engine/Plugins/Experimental/Mutable/Source/MutableTools/Internal/MuT/NodeScalarEnumParameterPrivate.h b/Engine/Plugins/Experimental/Mutable/Source/MutableTools/Internal/MuT/NodeScalarEnumParameterPrivate.h
new file mode 100644
--- /dev/null
+++ b/Engine/Plugins/Experimental/Mutable/Source/MutableTools/Internal/MuT/NodeScalarEnumParameterPrivate.h
@@ -0,0 +1,112 @@
+// Copyright Epic Games, Inc. All Rights Reserved.
+
+#pragma once
+
+#include "MuT/NodeScalarPrivate.h"
+
+#include "MuT/NodeScalarEnumParameter.h"
+#include "MuT/NodeImage.h"
+#include "MuR/ParametersPrivate.h"
+
+
+namespace mu
+{
+
+
+	class NodeScalarEnumParameter::Private : public NodeScalar::Private
+	{
+	public:
+
+		static FNodeType s_type;
+
+		FString m_name;
+		FString m_uid;
+		int32 m_defaultValue = 0;
+
+		struct OPTION
+		{
+			FString name;
+			float value;
+
+			//!
+			void Serialise( OutputArchive& arch ) const
+			{
+                uint32 ver = 1;
+				arch << ver;
+
+				arch << name;
+				arch << value;
+			}
+
+			//!
+			void Unserialise( InputArchive& arch )
+			{
+                uint32 ver;
+				arch >> ver;
+                check(ver<=1);
+
+				if (ver == 0)
+				{
+					std::string Temp;
+					arch >> Temp;
+					name = Temp.c_str();
+				}
+				else
+				{
+					arch >> name;
+				}
+				arch >> value;
+			}
+		};
+
+		TArray<OPTION> m_options;
+
+		TArray<Ptr<NodeRange>> m_ranges;
+
+		//!
+		void Serialise( OutputArchive& arch ) const
+		{
+            uint32_t ver = 4;
+			arch << ver;
+
+			arch << m_name;
+			arch << m_uid;
+			arch << m_defaultValue;
+            arch << m_options;
+            arch << m_ranges;
+        }
+
+		//!
+		void Unserialise( InputArchive& arch )
+		{
+            uint32_t ver;
+			arch >> ver;
+            check(ver>=2 && ver<=4);
+
+			if (ver <= 3)
+			{
+				std::string Temp;
+				arch >> Temp;
+				m_name = Temp.c_str();
+				arch >> Temp;
+				m_uid = Temp.c_str();
+			}
+			else
+			{
+				arch >> m_name;
+				arch >> m_uid;
+			}
+
+			arch >> m_defaultValue;
+			if (ver <= 2)
+			{
+				int32 Dummy = 0;
+				arch >> Dummy;
+			}
+			arch >> m_options;
+            arch >> m_ranges;
+        }
+	};
+
+}
+
diff --git a/Engine/Plugins/Experimental/Mutable/Source/MutableTools/Internal/MuT/NodeScalarParameterPrivate.h b/Engine/Plugins/Experimental/Mutable/Source/MutableTools/Internal/MuT/NodeScalarParameterPrivate.h
new file mode 100644
--- /dev/null
+++ b/Engine/Plugins/Experimental/Mutable/Source/MutableTools/Internal/MuT/NodeScalarParameterPrivate.h
@@ -0,0 +1,79 @@
+// Copyright Epic Games, Inc. All Rights Reserved.
+
+#pragma once
+
+#include "MuT/NodeScalarPrivate.h"
+
+#include "MuT/NodeScalarParameter.h"
+#include "MuT/NodeImage.h"
+#include "MuT/NodeRange.h"
+#include "MuR/ParametersPrivate.h"
+
+
+namespace mu
+{
+
+
+	class NodeScalarParameter::Private : public NodeScalar::Private
+	{
+	public:
+
+		static FNodeType s_type;
+
+		float m_defaultValue = 0.0f;
+		FString m_name;
+		FString m_uid;
+
+		TArray<Ptr<NodeRange>> m_ranges;
+
+		//!
+		void Serialise( OutputArchive& arch ) const
+		{
+            uint32_t ver = 7;
+			arch << ver;
+
+			arch << m_defaultValue;
+			arch << m_name;
+			arch << m_uid;
+            arch << m_ranges;
+        }
+
+		//!
+		void Unserialise( InputArchive& arch )
+		{
+            uint32_t ver;
+			arch >> ver;
+            check(ver>=4 && ver<=7);
+
+			arch >> m_defaultValue;
+			if (ver <= 6)
+			{
+				std::string Temp;
+				arch >> Temp;
+				m_name = Temp.c_str();
+				arch >> Temp;
+				m_uid = Temp.c_str();
+			}
+			else
+			{
+				arch >> m_name;
+				arch >> m_uid;
+			}
+
+			if (ver < 5)
+			{
+				TArray<Ptr<NodeImage>> TempAdditionalImages;
+				arch >> TempAdditionalImages;
+			}
+
+			if (ver <= 5)
+			{
+				int32 Dummy=0;
+				arch >> Dummy;
+			}
+
+            arch >> m_ranges;
+        }
+	};
+
+}
diff --git a/Engine/Plugins/Experimental/Mutable/Source/MutableTools/Internal/MuT/NodeScalarPrivate.h b/Engine/Plugins/Experimental/Mutable/Source/MutableTools/Internal/MuT/NodeScalarPrivate.h
new file mode 100644
--- /dev/null
+++ b/Engine/Plugins/Experimental/Mutable/Source/MutableTools/Internal/MuT/NodeScalarPrivate.h
@@ -0,0 +1,20 @@
+// Copyright Epic Games, Inc. All Rights Reserved.
+
+#pragma once
+
+#include "MuT/NodePrivate.h"
+
+#include "MuT/NodeScalar.h"
+#include "MuR/ParametersPrivate.h"
+
+
+namespace mu
+{
+
+
+	class NodeScalar::Private : public Node::Private
+	{
+	};
+
+
+}
diff --git a/Engine/Plugins/Experimental/Mutable/Source/MutableTools/Internal/MuT/NodeScalarSwitchPrivate.h b/Engine/Plugins/Experimental/Mutable/Source/MutableTools/Internal/MuT/NodeScalarSwitchPrivate.h
new file mode 100644
--- /dev/null
+++ b/Engine/Plugins/Experimental/Mutable/Source/MutableTools/Internal/MuT/NodeScalarSwitchPrivate.h
@@ -0,0 +1,47 @@
+// Copyright Epic Games, Inc. All Rights Reserved.
+
+#pragma once
+
+#include "MuT/NodeScalarSwitch.h"
+#include "MuT/NodeScalarPrivate.h"
+#include "MuT/NodeScalar.h"
+#include "MuT/AST.h"
+
+
+namespace mu
+{
+
+
+    class NodeScalarSwitch::Private : public NodeScalar::Private
+	{
+	public:
+
+		static FNodeType s_type;
+
+		NodeScalarPtr m_pParameter;
+        TArray<NodeScalarPtr> m_options;
+
+		//!
+		void Serialise( OutputArchive& arch ) const
+		{
+            uint32_t ver = 0;
+			arch << ver;
+
+			arch << m_pParameter;
+			arch << m_options;
+		}
+
+		//!
+		void Unserialise( InputArchive& arch )
+		{
+            uint32_t ver;
+			arch >> ver;
+			check(ver==0);
+
+			arch >> m_pParameter;
+			arch >> m_options;
+		}
+	};
+
+
+}
diff --git a/Engine/Plugins/Experimental/Mutable/Source/MutableTools/Internal/MuT/Platform.h b/Engine/Plugins/Experimental/Mutable/Source/MutableTools/Internal/MuT/Platform.h
new file mode 100644
--- /dev/null
+++ b/Engine/Plugins/Experimental/Mutable/Source/MutableTools/Internal/MuT/Platform.h
@@ -0,0 +1,37 @@
+// Copyright Epic Games, Inc. All Rights Reserved.
+
+#pragma once
+
+#include <vector>
+#include <map>
+#include <set>
+#include <string>
+#include <unordered_map>
+
+
+namespace mu
+{
+
+	//! STL-like containers using this allocator
+	template< typename T >
+	using vector = std::vector<T>;
+
+	template< typename T >
+	using basic_string = std::basic_string<T, std::char_traits<T>>;
+
+	using string = std::basic_string<char, std::char_traits<char> >;
+
+	template< typename K, typename T >
+	using map = std::map< K, T, std::less<K> >;
+
+	template< typename T >
+	using set = std::set< T, std::less<T> >;
+
+	template< typename T >
+	using multiset = std::multiset< T, std::less<T> >;
+
+	template< typename K, typename T >
+	using pair = std::pair<K, T>;
+
+}
+
diff --git a/Engine/Plugins/Experimental/Mutable/Source/MutableTools/Internal/MuT/TablePrivate.h b/Engine/Plugins/Experimental/Mutable/Source/MutableTools/Internal/MuT/TablePrivate.h
new file mode 100644
--- /dev/null
+++ b/Engine/Plugins/Experimental/Mutable/Source/MutableTools/Internal/MuT/TablePrivate.h
@@ -0,0 +1,224 @@
+// Copyright Epic Games, Inc. All Rights Reserved.
+
+#pragma once
+
+#include "MuT/Table.h"
+
+#include "MuT/ErrorLogPrivate.h"
+#include "MuR/MutableMath.h"
+#include "MuR/SerialisationPrivate.h"
+
+#include <string>
+
+namespace mu
+{
+	MUTABLE_DEFINE_ENUM_SERIALISABLE(ETableColumnType)
+
+
+	struct FTableColumn
+	{
+		FString Name;
+
+		ETableColumnType Type;
+	};
+
+
+	struct FTableValue
+	{
+		// TODO: Union
+		float Scalar;
+		FVector4f Color;
+		Ptr<ResourceProxy<Image>> ProxyImage;
+		Ptr<Mesh> Mesh;
+		FString String;
+
+		const void* ErrorContext;
+	};
+
+
+	struct FTableRow
+	{
+        uint32 Id;
+		TArray<FTableValue> Values;
+	};
+
+
+	//!
+	class Table::Private
+	{
+	public:
+
+		FString Name;
+		TArray<FTableColumn> Columns;
+		TArray<FTableRow> Rows;
+
+		// Transient value for serialization compatibility
+		bool bNoneOption_DEPRECATED = false;
+
+
+		//!
+		void Serialise( OutputArchive& arch ) const
+		{
+            uint32 ver = 4;
+			arch << ver;
+
+            arch << Columns.Num();
+			for ( int32 c=0; c< Columns.Num(); ++c )
+			{
+				arch << Columns[c].Name;
+				arch << Columns[c].Type;
+			}
+
+            arch << Rows.Num();
+			for (int32 r=0; r< Rows.Num(); ++r )
+			{
+				arch << Rows[r].Id;
+
+				for (int32 c=0; c<Columns.Num(); ++c )
+				{
+					const FTableValue& v = Rows[r].Values[c];
+
+					switch (Columns[c].Type)
+					{
+					case ETableColumnType::Scalar: 	arch << v.Scalar; break;
+					case ETableColumnType::Color: 	arch << v.Color; break;
+					case ETableColumnType::Mesh: 	arch << v.Mesh; break;
+					case ETableColumnType::Image:
+					{
+						Ptr<const Image> Value;
+						if (v.ProxyImage)
+						{
+							Value = v.ProxyImage->Get();
+						}
+						arch << Value;
+						break;
+					}
+					case ETableColumnType::String:	arch << v.String; break;
+					default: check( false);
+					}
+				}
+			}
+		}
+
+		//!
+		void Unserialise( InputArchive& arch )
+		{
+            uint32 ver;
+			arch >> ver;
+			check(ver<=4);
+
+            uint32 ColumnCount;
+			arch >> ColumnCount;
+			Columns.SetNum(ColumnCount);
+            for (uint32 c=0; c<ColumnCount; ++c )
+			{
+				if (ver <= 2)
+				{
+					std::string LegacyString;
+					arch >> LegacyString;
+					Columns[c].Name = LegacyString.c_str();
+				}
+				else
+				{
+					arch >> Columns[c].Name;
+				}
+				arch >> Columns[c].Type;
+			}
+
+            uint32 RowCount;
+			arch >> RowCount;
+			Rows.SetNum( RowCount );
+            for ( uint32 r=0; r<RowCount; ++r )
+			{
+				arch >> Rows[r].Id;
+				Rows[r].Values.SetNum( ColumnCount );
+
+                for (uint32_t c=0; c<ColumnCount; ++c )
+				{
+					FTableValue& v = Rows[r].Values[c];
+
+					switch (Columns[c].Type)
+					{
+					case ETableColumnType::Scalar:	
+					{
+						arch >> v.Scalar;
+						break;
+					}
+					case ETableColumnType::Color:
+					{
+						if (ver <= 1)
+						{
+							vec3<float> Value;
+							arch >> Value;
+
+							v.Color = FVector4f(Value[0], Value[1], Value[2], 1.0f);
+						}
+						else
+						{
+							arch >> v.Color;
+						}
+						break;
+					}
+					case ETableColumnType::Mesh:		
+					{
+						arch >> v.Mesh;
+						break;
+					}
+					case ETableColumnType::Image:
+					{
+						// Are we using proxies?
+						v.ProxyImage = arch.NewImageProxy();
+						if (!v.ProxyImage)
+						{
+							// Normal serialisation
+							Ptr<Image> Value;
+							arch >> Value;
+							v.ProxyImage = new ResourceProxyMemory<Image>(Value.get());
+						}
+						break;
+					}
+					case ETableColumnType::String:
+					{
+						if (ver <= 2)
+						{
+							std::string LegacyString;
+							arch >> LegacyString;
+							v.String = LegacyString.c_str();
+						}
+						else
+						{
+							arch >> v.String;
+						}
+						break;
+					}
+					default: check( false);
+					}
+				}
+			}
+
+			if (ver >= 1 && ver <= 3)
+			{
+				arch >> bNoneOption_DEPRECATED;
+			}
+		}
+
+
+		//! Find a row in the table by id. Return -1 if not found.
+        int32 FindRow( uint32 id ) const
+		{
+			int32 res = -1;
+
+			for ( int32 r=0; res<0 && r<Rows.Num(); ++r )
+			{
+				if ( Rows[r].Id==id )
+				{
+					res = (int32)r;
+				}
+			}
+
+			return res;
+		}
+
+	};
+
+}
diff --git a/Engine/Plugins/Experimental/Mutable/Source/MutableTools/Private/MuT/AST.h b/Engine/Plugins/Experimental/Mutable/Source/MutableTools/Private/MuT/AST.h
deleted file mode 100644
diff --git a/Engine/Plugins/Experimental/Mutable/Source/MutableTools/Private/MuT/CodeOptimiser_Semantic.cpp b/Engine/Plugins/Experimental/Mutable/Source/MutableTools/Private/MuT/CodeOptimiser_Semantic.cpp
--- a/Engine/Plugins/Experimental/Mutable/Source/MutableTools/Private/MuT/CodeOptimiser_Semantic.cpp
+++ b/Engine/Plugins/Experimental/Mutable/Source/MutableTools/Private/MuT/CodeOptimiser_Semantic.cpp
@@ -1,6 +1,6 @@
 // Copyright Epic Games, Inc. All Rights Reserved.
 
-#include "NodeImageFormatPrivate.h"
+#include "MuT/NodeImageFormatPrivate.h"
 #include "HAL/PlatformMath.h"
 #include "Misc/AssertionMacros.h"
 #include "MuR/Image.h"
diff --git a/Engine/Plugins/Experimental/Mutable/Source/MutableTools/Private/MuT/ErrorLogPrivate.h b/Engine/Plugins/Experimental/Mutable/Source/MutableTools/Private/MuT/ErrorLogPrivate.h
deleted file mode 100644
diff --git a/Engine/Plugins/Experimental/Mutable/Source/MutableTools/Private/MuT/NodeColourFromScalarsPrivate.h b/Engine/Plugins/Experimental/Mutable/Source/MutableTools/Private/MuT/NodeColourFromScalarsPrivate.h
deleted file mode 100644
diff --git a/Engine/Plugins/Experimental/Mutable/Source/MutableTools/Private/MuT/NodeColourParameterPrivate.h b/Engine/Plugins/Experimental/Mutable/Source/MutableTools/Private/MuT/NodeColourParameterPrivate.h
deleted file mode 100644
diff --git a/Engine/Plugins/Experimental/Mutable/Source/MutableTools/Private/MuT/NodeColourPrivate.h b/Engine/Plugins/Experimental/Mutable/Source/MutableTools/Private/MuT/NodeColourPrivate.h
deleted file mode 100644
diff --git a/Engine/Plugins/Experimental/Mutable/Source/MutableTools/Private/MuT/NodeColourSampleImagePrivate.h b/Engine/Plugins/Experimental/Mutable/Source/MutableTools/Private/MuT/NodeColourSampleImagePrivate.h
deleted file mode 100644
diff --git a/Engine/Plugins/Experimental/Mutable/Source/MutableTools/Private/MuT/NodeColourSwitchPrivate.h b/Engine/Plugins/Experimental/Mutable/Source/MutableTools/Private/MuT/NodeColourSwitchPrivate.h
deleted file mode 100644
diff --git a/Engine/Plugins/Experimental/Mutable/Source/MutableTools/Private/MuT/NodeComponentEditPrivate.h b/Engine/Plugins/Experimental/Mutable/Source/MutableTools/Private/MuT/NodeComponentEditPrivate.h
deleted file mode 100644
diff --git a/Engine/Plugins/Experimental/Mutable/Source/MutableTools/Private/MuT/NodeComponentNewPrivate.h b/Engine/Plugins/Experimental/Mutable/Source/MutableTools/Private/MuT/NodeComponentNewPrivate.h
deleted file mode 100644
diff --git a/Engine/Plugins/Experimental/Mutable/Source/MutableTools/Private/MuT/NodeComponentPrivate.h b/Engine/Plugins/Experimental/Mutable/Source/MutableTools/Private/MuT/NodeComponentPrivate.h
deleted file mode 100644
diff --git a/Engine/Plugins/Experimental/Mutable/Source/MutableTools/Private/MuT/NodeImageFormatPrivate.h b/Engine/Plugins/Experimental/Mutable/Source/MutableTools/Private/MuT/NodeImageFormatPrivate.h
deleted file mode 100644
diff --git a/Engine/Plugins/Experimental/Mutable/Source/MutableTools/Private/MuT/NodeImageInterpolatePrivate.h b/Engine/Plugins/Experimental/Mutable/Source/MutableTools/Private/MuT/NodeImageInterpolatePrivate.h
deleted file mode 100644
diff --git a/Engine/Plugins/Experimental/Mutable/Source/MutableTools/Private/MuT/NodeImageInvertPrivate.h b/Engine/Plugins/Experimental/Mutable/Source/MutableTools/Private/MuT/NodeImageInvertPrivate.h
deleted file mode 100644
diff --git a/Engine/Plugins/Experimental/Mutable/Source/MutableTools/Private/MuT/NodeImageLayerColourPrivate.h b/Engine/Plugins/Experimental/Mutable/Source/MutableTools/Private/MuT/NodeImageLayerColourPrivate.h
deleted file mode 100644
diff --git a/Engine/Plugins/Experimental/Mutable/Source/MutableTools/Private/MuT/NodeImageLayerPrivate.h b/Engine/Plugins/Experimental/Mutable/Source/MutableTools/Private/MuT/NodeImageLayerPrivate.h
deleted file mode 100644
diff --git a/Engine/Plugins/Experimental/Mutable/Source/MutableTools/Private/MuT/NodeImageMipmapPrivate.h b/Engine/Plugins/Experimental/Mutable/Source/MutableTools/Private/MuT/NodeImageMipmapPrivate.h
deleted file mode 100644
diff --git a/Engine/Plugins/Experimental/Mutable/Source/MutableTools/Private/MuT/NodeImageMultiLayerPrivate.h b/Engine/Plugins/Experimental/Mutable/Source/MutableTools/Private/MuT/NodeImageMultiLayerPrivate.h
deleted file mode 100644
diff --git a/Engine/Plugins/Experimental/Mutable/Source/MutableTools/Private/MuT/NodeImagePlainColourPrivate.h b/Engine/Plugins/Experimental/Mutable/Source/MutableTools/Private/MuT/NodeImagePlainColourPrivate.h
deleted file mode 100644
diff --git a/Engine/Plugins/Experimental/Mutable/Source/MutableTools/Private/MuT/NodeImagePrivate.h b/Engine/Plugins/Experimental/Mutable/Source/MutableTools/Private/MuT/NodeImagePrivate.h
deleted file mode 100644
diff --git a/Engine/Plugins/Experimental/Mutable/Source/MutableTools/Private/MuT/NodeImageProjectPrivate.h b/Engine/Plugins/Experimental/Mutable/Source/MutableTools/Private/MuT/NodeImageProjectPrivate.h
deleted file mode 100644
diff --git a/Engine/Plugins/Experimental/Mutable/Source/MutableTools/Private/MuT/NodeImageResizePrivate.h b/Engine/Plugins/Experimental/Mutable/Source/MutableTools/Private/MuT/NodeImageResizePrivate.h
deleted file mode 100644
diff --git a/Engine/Plugins/Experimental/Mutable/Source/MutableTools/Private/MuT/NodeImageSwitchPrivate.h b/Engine/Plugins/Experimental/Mutable/Source/MutableTools/Private/MuT/NodeImageSwitchPrivate.h
deleted file mode 100644
diff --git a/Engine/Plugins/Experimental/Mutable/Source/MutableTools/Private/MuT/NodeImageSwizzlePrivate.h b/Engine/Plugins/Experimental/Mutable/Source/MutableTools/Private/MuT/NodeImageSwizzlePrivate.h
deleted file mode 100644
diff --git a/Engine/Plugins/Experimental/Mutable/Source/MutableTools/Private/MuT/NodeLODPrivate.h b/Engine/Plugins/Experimental/Mutable/Source/MutableTools/Private/MuT/NodeLODPrivate.h
deleted file mode 100644
diff --git a/Engine/Plugins/Experimental/Mutable/Source/MutableTools/Private/MuT/NodeMeshConstantPrivate.h b/Engine/Plugins/Experimental/Mutable/Source/MutableTools/Private/MuT/NodeMeshConstantPrivate.h
deleted file mode 100644
diff --git a/Engine/Plugins/Experimental/Mutable/Source/MutableTools/Private/MuT/NodeMeshFormatPrivate.h b/Engine/Plugins/Experimental/Mutable/Source/MutableTools/Private/MuT/NodeMeshFormatPrivate.h
deleted file mode 100644
diff --git a/Engine/Plugins/Experimental/Mutable/Source/MutableTools/Private/MuT/NodeMeshFragmentPrivate.h b/Engine/Plugins/Experimental/Mutable/Source/MutableTools/Private/MuT/NodeMeshFragmentPrivate.h
deleted file mode 100644
diff --git a/Engine/Plugins/Experimental/Mutable/Source/MutableTools/Private/MuT/NodeMeshMakeMorphPrivate.h b/Engine/Plugins/Experimental/Mutable/Source/MutableTools/Private/MuT/NodeMeshMakeMorphPrivate.h
deleted file mode 100644
diff --git a/Engine/Plugins/Experimental/Mutable/Source/MutableTools/Private/MuT/NodeMeshMorphPrivate.h b/Engine/Plugins/Experimental/Mutable/Source/MutableTools/Private/MuT/NodeMeshMorphPrivate.h
deleted file mode 100644
diff --git a/Engine/Plugins/Experimental/Mutable/Source/MutableTools/Private/MuT/NodeMeshPrivate.h b/Engine/Plugins/Experimental/Mutable/Source/MutableTools/Private/MuT/NodeMeshPrivate.h
deleted file mode 100644
diff --git a/Engine/Plugins/Experimental/Mutable/Source/MutableTools/Private/MuT/NodeMeshTablePrivate.h b/Engine/Plugins/Experimental/Mutable/Source/MutableTools/Private/MuT/NodeMeshTablePrivate.h
deleted file mode 100644
diff --git a/Engine/Plugins/Experimental/Mutable/Source/MutableTools/Private/MuT/NodeModifierMeshClipDeformPrivate.h b/Engine/Plugins/Experimental/Mutable/Source/MutableTools/Private/MuT/NodeModifierMeshClipDeformPrivate.h
deleted file mode 100644
diff --git a/Engine/Plugins/Experimental/Mutable/Source/MutableTools/Private/MuT/NodeModifierMeshClipMorphPlanePrivate.h b/Engine/Plugins/Experimental/Mutable/Source/MutableTools/Private/MuT/NodeModifierMeshClipMorphPlanePrivate.h
deleted file mode 100644
diff --git a/Engine/Plugins/Experimental/Mutable/Source/MutableTools/Private/MuT/NodeModifierMeshClipWithMeshPrivate.h b/Engine/Plugins/Experimental/Mutable/Source/MutableTools/Private/MuT/NodeModifierMeshClipWithMeshPrivate.h
deleted file mode 100644
diff --git a/Engine/Plugins/Experimental/Mutable/Source/MutableTools/Private/MuT/NodeModifierMeshClipWithUVMaskPrivate.h b/Engine/Plugins/Experimental/Mutable/Source/MutableTools/Private/MuT/NodeModifierMeshClipWithUVMaskPrivate.h
deleted file mode 100644
diff --git a/Engine/Plugins/Experimental/Mutable/Source/MutableTools/Private/MuT/NodeModifierPrivate.h b/Engine/Plugins/Experimental/Mutable/Source/MutableTools/Private/MuT/NodeModifierPrivate.h
deleted file mode 100644
diff --git a/Engine/Plugins/Experimental/Mutable/Source/MutableTools/Private/MuT/NodePatchImagePrivate.h b/Engine/Plugins/Experimental/Mutable/Source/MutableTools/Private/MuT/NodePatchImagePrivate.h
deleted file mode 100644
diff --git a/Engine/Plugins/Experimental/Mutable/Source/MutableTools/Private/MuT/NodePatchMeshPrivate.h b/Engine/Plugins/Experimental/Mutable/Source/MutableTools/Private/MuT/NodePatchMeshPrivate.h
deleted file mode 100644
diff --git a/Engine/Plugins/Experimental/Mutable/Source/MutableTools/Private/MuT/NodeProjectorPrivate.h b/Engine/Plugins/Experimental/Mutable/Source/MutableTools/Private/MuT/NodeProjectorPrivate.h
deleted file mode 100644
diff --git a/Engine/Plugins/Experimental/Mutable/Source/MutableTools/Private/MuT/NodeScalarCurvePrivate.h b/Engine/Plugins/Experimental/Mutable/Source/MutableTools/Private/MuT/NodeScalarCurvePrivate.h
deleted file mode 100644
diff --git a/Engine/Plugins/Experimental/Mutable/Source/MutableTools/Private/MuT/NodeScalarEnumParameterPrivate.h b/Engine/Plugins/Experimental/Mutable/Source/MutableTools/Private/MuT/NodeScalarEnumParameterPrivate.h
deleted file mode 100644
diff --git a/Engine/Plugins/Experimental/Mutable/Source/MutableTools/Private/MuT/NodeScalarParameterPrivate.h b/Engine/Plugins/Experimental/Mutable/Source/MutableTools/Private/MuT/NodeScalarParameterPrivate.h
deleted file mode 100644
diff --git a/Engine/Plugins/Experimental/Mutable/Source/MutableTools/Private/MuT/NodeScalarPrivate.h b/Engine/Plugins/Experimental/Mutable/Source/MutableTools/Private/MuT/NodeScalarPrivate.h
deleted file mode 100644
diff --git a/Engine/Plugins/Experimental/Mutable/Source/MutableTools/Private/MuT/NodeScalarSwitchPrivate.h b/Engine/Plugins/Experimental/Mutable/Source/MutableTools/Private/MuT/NodeScalarSwitchPrivate.h
deleted file mode 100644
diff --git a/Engine/Plugins/Experimental/Mutable/Source/MutableTools/Private/MuT/Platform.h b/Engine/Plugins/Experimental/Mutable/Source/MutableTools/Private/MuT/Platform.h
deleted file mode 100644
diff --git a/Engine/Plugins/Experimental/Mutable/Source/MutableTools/Private/MuT/TablePrivate.h b/Engine/Plugins/Experimental/Mutable/Source/MutableTools/Private/MuT/TablePrivate.h
deleted file mode 100644
