[Mutable] Change RequestedLODs from bitmask of LODs to generate to first LOD to generate.

#rnw

diff --git a/Engine/Plugins/Experimental/Mutable/Source/CustomizableObject/Private/MuCO/CustomizableInstanceLODManagement.cpp b/Engine/Plugins/Experimental/Mutable/Source/CustomizableObject/Private/MuCO/CustomizableInstanceLODManagement.cpp
--- a/Engine/Plugins/Experimental/Mutable/Source/CustomizableObject/Private/MuCO/CustomizableInstanceLODManagement.cpp
+++ b/Engine/Plugins/Experimental/Mutable/Source/CustomizableObject/Private/MuCO/CustomizableInstanceLODManagement.cpp
@@ -241,12 +241,12 @@ void UCustomizableInstanceLODManagement::UpdateInstanceDistsAndLODs(FMutableInst
 				if (!bAlreadyReachedFixedLOD && SortedInstances[i]->GetMinSquareDistToPlayer() < DistanceForFixedLODSquared)
 				{
 					RequestedLODs.Init(MAX_uint16, SortedInstances[i]->GetNumComponents());
-					SortedInstances[i]->SetRequestedLODs(0, MAX_int32, RequestedLODs, InOutRequestedUpdates);
+					SortedInstances[i]->SetRequestedLODs(0, 0, RequestedLODs, InOutRequestedUpdates);
 				}
 				else
 				{
 					RequestedLODs.Init(MAX_uint16, SortedInstances[i]->GetNumComponents());
-					SortedInstances[i]->SetRequestedLODs(2, 2, RequestedLODs, InOutRequestedUpdates);
+					SortedInstances[i]->SetRequestedLODs(2, 0, RequestedLODs, InOutRequestedUpdates);
 					bAlreadyReachedFixedLOD = true;
 				}
 			}
@@ -256,12 +256,12 @@ void UCustomizableInstanceLODManagement::UpdateInstanceDistsAndLODs(FMutableInst
 				if (!bAlreadyReachedFixedLOD && SortedInstances[i]->GetMinSquareDistToPlayer() < DistanceForFixedLODSquared)
 				{
 					RequestedLODs.Init(MAX_uint16, SortedInstances[i]->GetNumComponents());
-					SortedInstances[i]->SetRequestedLODs(1, MAX_int32, RequestedLODs, InOutRequestedUpdates);
+					SortedInstances[i]->SetRequestedLODs(1, 0, RequestedLODs, InOutRequestedUpdates);
 				}
 				else
 				{
 					RequestedLODs.Init(MAX_uint16, SortedInstances[i]->GetNumComponents());
-					SortedInstances[i]->SetRequestedLODs(2, MAX_int32, RequestedLODs, InOutRequestedUpdates);
+					SortedInstances[i]->SetRequestedLODs(2, 0, RequestedLODs, InOutRequestedUpdates);
 					bAlreadyReachedFixedLOD = true;
 				}
 			}
@@ -269,7 +269,7 @@ void UCustomizableInstanceLODManagement::UpdateInstanceDistsAndLODs(FMutableInst
 			for (int32 i = NumGeneratedInstancesLimitLODs + NumGeneratedInstancesLimitLOD1; i < NumGeneratedInstancesLimitLODs + NumGeneratedInstancesLimitLOD1 + NumGeneratedInstancesLimitLOD2 && i < SortedInstances.Num(); ++i)
 			{
 				RequestedLODs.Init(MAX_uint16, SortedInstances[i]->GetNumComponents());
-				SortedInstances[i]->SetRequestedLODs(2, MAX_int32, RequestedLODs, InOutRequestedUpdates);
+				SortedInstances[i]->SetRequestedLODs(2, 0, RequestedLODs, InOutRequestedUpdates);
 			}
 
 			for (int32 i = NumGeneratedInstancesLimitLODs + NumGeneratedInstancesLimitLOD1 + NumGeneratedInstancesLimitLOD2; i < SortedInstances.Num(); ++i)
@@ -283,7 +283,7 @@ void UCustomizableInstanceLODManagement::UpdateInstanceDistsAndLODs(FMutableInst
 			for (int32 i = 0; i < SortedInstances.Num(); ++i)
 			{
 				RequestedLODs.Init(MAX_uint16, SortedInstances[i]->GetNumComponents());
-				SortedInstances[i]->SetRequestedLODs(2, MAX_int32, RequestedLODs, InOutRequestedUpdates);
+				SortedInstances[i]->SetRequestedLODs(2, 0, RequestedLODs, InOutRequestedUpdates);
 			}
 		}
 	}
@@ -294,7 +294,7 @@ void UCustomizableInstanceLODManagement::UpdateInstanceDistsAndLODs(FMutableInst
 			int32 MinLOD = MAX_int32;
 
 			bool bInitialized = false;
-			TArray<uint16> RequestedLODsPerComponent;
+			TArray<uint16> RequestedLODPerComponent;
 		};
 
 		TMap<TObjectPtr<UCustomizableObjectInstance>, FLODTracker> InstancesMinLOD;
@@ -351,7 +351,7 @@ void UCustomizableInstanceLODManagement::UpdateInstanceDistsAndLODs(FMutableInst
 
 				if (!LODTracker.bInitialized)
 				{
-					LODTracker.RequestedLODsPerComponent.AddZeroed(COI->GetCustomizableObject()->GetComponentCount());
+					LODTracker.RequestedLODPerComponent.Init(MAX_uint16, COI->GetCustomizableObject()->GetComponentCount());
 					LODTracker.bInitialized = true;
 				}
 
@@ -386,9 +386,10 @@ void UCustomizableInstanceLODManagement::UpdateInstanceDistsAndLODs(FMutableInst
 					LODTracker.MinLOD = FMath::Min(LODTracker.MinLOD, Parent->bOverrideMinLod ? Parent->MinLodModel : 0);
 
 					// If the parent component have a SkeletalMesh use the RequestedLODLevel of the component as reference to know which LODs mutable should generate.
-					if (UE_MUTABLE_GETSKELETALMESHASSET(Parent) && LODTracker.RequestedLODsPerComponent.IsValidIndex(ComponentIndex))
+					if (UE_MUTABLE_GETSKELETALMESHASSET(Parent) && LODTracker.RequestedLODPerComponent.IsValidIndex(ComponentIndex))
 					{
-						LODTracker.RequestedLODsPerComponent[ComponentIndex] |= 1 << Parent->GetPredictedLODLevel();
+						uint16& RequestedLOD = LODTracker.RequestedLODPerComponent[ComponentIndex];
+						RequestedLOD = FMath::Min((int32)RequestedLOD, Parent->GetPredictedLODLevel());
 					}
 				}
 			}
@@ -412,7 +413,7 @@ void UCustomizableInstanceLODManagement::UpdateInstanceDistsAndLODs(FMutableInst
 				// Limit MinLOD
 				It.Value.MinLOD = FMath::Min(It.Value.MinLOD, CustomizableObject->GetNumLODs() - 1);
 
-				It.Key->SetRequestedLODs(It.Value.MinLOD, 0, It.Value.RequestedLODsPerComponent, InOutRequestedUpdates);
+				It.Key->SetRequestedLODs(It.Value.MinLOD, 0, It.Value.RequestedLODPerComponent, InOutRequestedUpdates);
 			}
 		}
 	}
diff --git a/Engine/Plugins/Experimental/Mutable/Source/CustomizableObject/Private/MuCO/CustomizableObjectInstance.cpp b/Engine/Plugins/Experimental/Mutable/Source/CustomizableObject/Private/MuCO/CustomizableObjectInstance.cpp
--- a/Engine/Plugins/Experimental/Mutable/Source/CustomizableObject/Private/MuCO/CustomizableObjectInstance.cpp
+++ b/Engine/Plugins/Experimental/Mutable/Source/CustomizableObject/Private/MuCO/CustomizableObjectInstance.cpp
@@ -3491,7 +3491,7 @@ void UCustomizableInstancePrivate::BuildOrCopyElementData(const TSharedRef<FUpda
 {
 	MUTABLE_CPUPROFILER_SCOPE(UCustomizableInstancePrivate::BuildOrCopyElementData);
 
-	for (int32 LODIndex = OperationData->NumLODsAvailable - 1; LODIndex >= FirstLODAvailable; --LODIndex)
+	for (int32 LODIndex = FirstLODAvailable; LODIndex < OperationData->NumLODsAvailable; ++LODIndex)
 	{
 		const FInstanceUpdateData::FLOD& LOD = OperationData->InstanceUpdateData.LODs[LODIndex];
 
@@ -6351,12 +6351,12 @@ void UCustomizableObjectInstance::SetRequestedLODs(int32 InMinLOD, int32 , const
 
 	// Clamp Min LOD
 	const int32 MinLODIdx = GetCustomizableObject()->GetPrivate()->GetMinLODIndex();
-	InMinLOD = FMath::Min(FMath::Max(InMinLOD, MinLODIdx), (int32)PrivateData->NumMaxLODsToStream);
+	const int32 MaxLODIdx = PrivateData->NumLODsAvailable - 1;
+	InMinLOD = FMath::Clamp(InMinLOD, MinLODIdx, MaxLODIdx);
 
 	const bool bMinLODChanged = Descriptor.MinLOD != InMinLOD;
 
-	const bool bIsDowngradeLODUpdate = GetCurrentMinLOD() >= 0 && InMinLOD > GetCurrentMinLOD();
-	PrivateData->SetCOInstanceFlags(bIsDowngradeLODUpdate ? PendingLODsDowngrade : ECONone);
+	PrivateData->SetCOInstanceFlags(InMinLOD > GetCurrentMinLOD() ? PendingLODsDowngrade : ECONone);
 
 	// Save the new LODs
 	MutableUpdateCandidate.MinLOD = InMinLOD;
@@ -6365,8 +6365,13 @@ void UCustomizableObjectInstance::SetRequestedLODs(int32 InMinLOD, int32 , const
 	bool bUpdateRequestedLODs = false;
 	if (UCustomizableObjectSystem::GetInstance()->IsOnlyGenerateRequestedLODsEnabled())
 	{
+		const uint16 FirstNonStreamedLODIndex = FMath::Clamp(PrivateData->NumMaxLODsToStream, 0, MaxLODIdx);
+
 		const int32 ComponentCount = GetNumComponents();
-		MutableUpdateCandidate.RequestedLODLevels.SetNumZeroed(ComponentCount);
+		if (ComponentCount != MutableUpdateCandidate.RequestedLODLevels.Num())
+		{
+			MutableUpdateCandidate.RequestedLODLevels.Init(FirstNonStreamedLODIndex, ComponentCount);
+		}
 
 		const TArray<uint16>& GeneratedLODsPerComponent = GetPrivate()->CommittedDescriptorHash.RequestedLODsPerComponent;
 
@@ -6379,32 +6384,17 @@ void UCustomizableObjectInstance::SetRequestedLODs(int32 InMinLOD, int32 , const
 			bUpdateRequestedLODs = bIgnoreGeneratedLODs;
 			for (int32 ComponentIndex = 0; ComponentIndex < ComponentCount; ++ComponentIndex)
 			{
-				// Find the first requested LOD. We'll generate [FirstRequestedLOD ... MaxLOD].
-				int32 FirstRequestedLOD = FMath::Max(((int32)PrivateData->NumMaxLODsToStream) - 1, 0);
-				for (int32 LODIndex = 0; LODIndex < FirstRequestedLOD; ++LODIndex)
-				{
-					if ((!bIgnoreGeneratedLODs && GeneratedLODsPerComponent[ComponentIndex] & (1 << LODIndex))
-						||
-						InRequestedLODsPerComponent[ComponentIndex] & (1 << LODIndex))
-					{
-						// First RequestedLOD that fall within the range
-						FirstRequestedLOD = LODIndex >= MutableUpdateCandidate.MinLOD ? LODIndex : MutableUpdateCandidate.MinLOD;
-					}
-				}
-
-				// Generate at least the MaxLOD
-				int32 RequestedLODs = 1 << FirstRequestedLOD;
+				uint16 PredictedLOD = FMath::Min(InRequestedLODsPerComponent[ComponentIndex], FirstNonStreamedLODIndex);
 
-				// Mark all LODs up until MAX_MESH_LOD_COUNT since MaxLOD can be set to Max_int32
-				for (int32 LODIndex = FirstRequestedLOD; LODIndex < MAX_MESH_LOD_COUNT; ++LODIndex)
+				if (!bIgnoreGeneratedLODs)
 				{
-					RequestedLODs |= (1 << LODIndex);
+					PredictedLOD = FMath::Min(PredictedLOD, GeneratedLODsPerComponent[ComponentIndex]);
 				}
-				
-				bUpdateRequestedLODs |= (RequestedLODs != MutableUpdateCandidate.RequestedLODLevels[ComponentIndex]);
+
+				bUpdateRequestedLODs |= (PredictedLOD != MutableUpdateCandidate.RequestedLODLevels[ComponentIndex]);
 
 				// Save new RequestedLODs
-				MutableUpdateCandidate.RequestedLODLevels[ComponentIndex] = RequestedLODs;
+				MutableUpdateCandidate.RequestedLODLevels[ComponentIndex] = PredictedLOD;
 			}
 		}
 	}
diff --git a/Engine/Plugins/Experimental/Mutable/Source/CustomizableObject/Private/MuCO/CustomizableObjectInstanceDescriptor.cpp b/Engine/Plugins/Experimental/Mutable/Source/CustomizableObject/Private/MuCO/CustomizableObjectInstanceDescriptor.cpp
--- a/Engine/Plugins/Experimental/Mutable/Source/CustomizableObject/Private/MuCO/CustomizableObjectInstanceDescriptor.cpp
+++ b/Engine/Plugins/Experimental/Mutable/Source/CustomizableObject/Private/MuCO/CustomizableObjectInstanceDescriptor.cpp
@@ -682,7 +682,7 @@ void FCustomizableObjectInstanceDescriptor::ReloadParameters()
 
 	SetState(FMath::Clamp(GetState(), 0, CustomizableObject->GetStateCount() - 1));
 	
-	RequestedLODLevels.Init(0, CustomizableObject->GetComponentCount());
+	RequestedLODLevels.Init(MAX_uint16, CustomizableObject->GetComponentCount());
 
 	TArray<FCustomizableObjectBoolParameterValue> OldBoolParameters = BoolParameters;
 	TArray<FCustomizableObjectIntParameterValue> OldIntParameters = IntParameters;
diff --git a/Engine/Plugins/Experimental/Mutable/Source/CustomizableObject/Private/MuCO/CustomizableObjectSystem.cpp b/Engine/Plugins/Experimental/Mutable/Source/CustomizableObject/Private/MuCO/CustomizableObjectSystem.cpp
--- a/Engine/Plugins/Experimental/Mutable/Source/CustomizableObject/Private/MuCO/CustomizableObjectSystem.cpp
+++ b/Engine/Plugins/Experimental/Mutable/Source/CustomizableObject/Private/MuCO/CustomizableObjectSystem.cpp
@@ -1637,11 +1637,8 @@ namespace impl
 
 		for (int32 ComponentIndex = 0; ComponentIndex < Operation->NumComponents; ++ComponentIndex)
 		{
-			// Ensure we're generating at least one LOD
-			for (int32 LODIndex = Operation->NumLODsAvailable - 1; LODIndex < MAX_MESH_LOD_COUNT; ++LODIndex)
-			{
-				RequestedLODs[ComponentIndex] |= (1 << LODIndex);
-			}
+			// Clamp value to the valid LOD range.
+			RequestedLODs[ComponentIndex] = FMath::Min(RequestedLODs[ComponentIndex], (uint16)(Operation->NumLODsAvailable - 1));
 		}
 
 		Operation->SetRequestedLODs(RequestedLODs);
@@ -1688,6 +1685,8 @@ namespace impl
 			return;
 		}
 
+		const TArray<uint16>& RequestedLODs = OperationData->GetRequestedLODs();
+
 		// Map SharedSurfaceId to surface index
 		TArray<int32> SurfacesSharedId;
 
@@ -1713,8 +1712,7 @@ namespace impl
 				Component.FirstSurface = OperationData->InstanceUpdateData.Surfaces.Num();
 				Component.SurfaceCount = 0;
 
-				const TArray<uint16>& RequestedLODs = OperationData->GetRequestedLODs();
-				const bool bGenerateLOD = RequestedLODs.IsValidIndex(Component.Id) ? (RequestedLODs[Component.Id] & (1 << MutableLODIndex)) != 0 : true;
+				const bool bGenerateLOD = RequestedLODs.IsValidIndex(Component.Id) ? RequestedLODs[Component.Id] <= MutableLODIndex : true;
 
 				// Mesh
 				if (Instance->GetMeshCount(MutableLODIndex, ComponentIndex) > 0)
@@ -2601,16 +2599,16 @@ namespace impl
 
 		Operation->MeshDescriptors.SetNum(NumComponents);
 
+		const TArray<uint16>& RequestedLODs = Operation->GetRequestedLODs();
 		for (int32 ComponentIndex = 0; ComponentIndex < NumComponents; ++ComponentIndex)
 		{
 			TArray<mu::FResourceID>& MeshId = Operation->MeshDescriptors[ComponentIndex];
 			MeshId.Init(MAX_uint64, MAX_MESH_LOD_COUNT);
 
-			const TArray<uint16>& RequestedLODs = Operation->GetRequestedLODs();
 			
 			for (int32 LODIndex = Operation->GetMinLOD(); LODIndex < Operation->NumLODsAvailable; ++LODIndex)
 			{
-				const bool bGenerateLOD = RequestedLODs.IsValidIndex(ComponentIndex) ? (RequestedLODs[ComponentIndex] & (1 << LODIndex)) != 0 : true;
+				const bool bGenerateLOD = RequestedLODs.IsValidIndex(ComponentIndex) ? RequestedLODs[ComponentIndex] <= LODIndex : true;
 				if (bGenerateLOD)
 				{
 					MeshId[LODIndex] = Operation->MutableInstance->GetMeshId(LODIndex, ComponentIndex, 0);
@@ -2809,7 +2807,7 @@ namespace impl
 			bIsInEditorViewport)
 		{
 			TArray<uint16> RequestedLODs = Operation->GetRequestedLODs();
-			RequestedLODs.Init(MAX_uint8, Operation->NumComponents);
+			RequestedLODs.Init(0, Operation->NumComponents);
 
 			Operation->SetRequestedLODs(RequestedLODs);
 		}
diff --git a/Engine/Plugins/Experimental/Mutable/Source/CustomizableObject/Private/MuCO/DescriptorHash.cpp b/Engine/Plugins/Experimental/Mutable/Source/CustomizableObject/Private/MuCO/DescriptorHash.cpp
--- a/Engine/Plugins/Experimental/Mutable/Source/CustomizableObject/Private/MuCO/DescriptorHash.cpp
+++ b/Engine/Plugins/Experimental/Mutable/Source/CustomizableObject/Private/MuCO/DescriptorHash.cpp
@@ -92,15 +92,10 @@ bool FDescriptorHash::IsSubset(const FDescriptorHash& Other) const
 	
 	for (int32 ComponentIndex = 0; ComponentIndex < RequestedLODsPerComponent.Num(); ++ComponentIndex)
 	{
-		const int32 RequestedLODs = RequestedLODsPerComponent[ComponentIndex];
-		const int32 OtherRequestedLODs = Other.RequestedLODsPerComponent[ComponentIndex];
-		for (uint16 LODIndex = MinLOD; LODIndex <= MAX_MESH_LOD_COUNT; ++LODIndex)
+		// It is a subset if the Requested LOD is the same or greater.
+		if (RequestedLODsPerComponent[ComponentIndex] < Other.RequestedLODsPerComponent[ComponentIndex])
 		{
-			// To be a subset all bits set in RequestedLODs must be set in OtherRequestedLODs. OtherRequestedLODs can have additional bits set
-			if (RequestedLODs & (1 << LODIndex) && !(OtherRequestedLODs & (1 << LODIndex)))
-			{
-				return false;
-			}
+			return false;
 		}
 	}
 
diff --git a/Engine/Plugins/Experimental/Mutable/Source/CustomizableObjectEditor/Private/MuCOE/GenerateMutableSource/GenerateMutableSource.cpp b/Engine/Plugins/Experimental/Mutable/Source/CustomizableObjectEditor/Private/MuCOE/GenerateMutableSource/GenerateMutableSource.cpp
--- a/Engine/Plugins/Experimental/Mutable/Source/CustomizableObjectEditor/Private/MuCOE/GenerateMutableSource/GenerateMutableSource.cpp
+++ b/Engine/Plugins/Experimental/Mutable/Source/CustomizableObjectEditor/Private/MuCOE/GenerateMutableSource/GenerateMutableSource.cpp
@@ -842,6 +842,8 @@ mu::NodeObjectPtr GenerateMutableSource(const UEdGraphPin * Pin, FMutableGraphGe
 						RefSkeletalMesh->GetMaxNumStreamedLODs(GenerationContext.Options.TargetPlatform));
 				}
 			}
+
+			GenerationContext.NumMaxLODsToStream = FMath::Clamp(GenerationContext.NumMaxLODsToStream, 0, GenerationContext.NumLODsInRoot - 1);
 		}
 
 		// States
