[Mutable] Moved GetImages and GetMesh to a new file in preparation for BusyWait changes.

#rnx

diff --git a/Engine/Plugins/Experimental/Mutable/Source/CustomizableObject/Internal/MuCO/CustomizableObjectSystemPrivate.h b/Engine/Plugins/Experimental/Mutable/Source/CustomizableObject/Internal/MuCO/CustomizableObjectSystemPrivate.h
--- a/Engine/Plugins/Experimental/Mutable/Source/CustomizableObject/Internal/MuCO/CustomizableObjectSystemPrivate.h
+++ b/Engine/Plugins/Experimental/Mutable/Source/CustomizableObject/Internal/MuCO/CustomizableObjectSystemPrivate.h
@@ -815,5 +815,19 @@ public:
 };
 
 
+namespace impl
+{
+	void CreateMutableInstance(const TSharedRef<FUpdateContextPrivate>& Operation);
+
+	void FixLODs(const TSharedRef<FUpdateContextPrivate>& Operation);
+	
+	void Subtask_Mutable_PrepareSkeletonData(const TSharedRef<FUpdateContextPrivate>& OperationData);
+
+	void Subtask_Mutable_UpdateParameterRelevancy(const TSharedRef<FUpdateContextPrivate>& OperationData);
+	
+	void Subtask_Mutable_PrepareTextures(const TSharedRef<FUpdateContextPrivate>& OperationData);
+}
+
+
 /** Set OnlyLOD to -1 to generate all mips */
 CUSTOMIZABLEOBJECT_API FTexturePlatformData* MutableCreateImagePlatformData(mu::Ptr<const mu::Image> MutableImage, int32 OnlyLOD, uint16 FullSizeX, uint16 FullSizeY);
diff --git a/Engine/Plugins/Experimental/Mutable/Source/CustomizableObject/Private/MuCO/BusyWaits_Deprecated.cpp b/Engine/Plugins/Experimental/Mutable/Source/CustomizableObject/Private/MuCO/BusyWaits_Deprecated.cpp
new file mode 100644
--- /dev/null
+++ b/Engine/Plugins/Experimental/Mutable/Source/CustomizableObject/Private/MuCO/BusyWaits_Deprecated.cpp
@@ -0,0 +1,416 @@
+// Copyright Epic Games, Inc. All Rights Reserved.
+
+#include "MuCO/BusyWaits_Deprecated.h"
+#include "MuCO/CustomizableObjectInstancePrivate.h"
+#include "MuCO/CustomizableObjectSystem.h"
+#include "MuCO/CustomizableObjectSystemPrivate.h"
+#include "MuR/MutableTrace.h"
+
+
+namespace impl_deprecated
+{
+	void Subtask_Mutable_GetImages(const TSharedRef<FUpdateContextPrivate>& OperationData)
+	{
+		MUTABLE_CPUPROFILER_SCOPE(Subtask_Mutable_GetImages)
+
+		const UCustomizableObjectSystemPrivate* CustomizableObjectSystemPrivateData = UCustomizableObjectSystem::GetInstanceChecked()->GetPrivate();
+		mu::System* System = CustomizableObjectSystemPrivateData->MutableSystem.get();
+		check(System != nullptr);
+
+		// Generate all the required resources, that are not cached
+		TArray<mu::FResourceID> ImagesInThisInstance;
+		for (FInstanceUpdateData::FImage& Image : OperationData->InstanceUpdateData.Images)
+		{
+			MUTABLE_CPUPROFILER_SCOPE(GetImage);
+
+			// If the image is a reference to an engine texture, we are done.
+			if (Image.bIsPassThrough)
+			{
+				continue;
+			}
+
+			mu::FImageDesc ImageDesc;
+
+			// This should only be done when using progressive images, since GetImageDesc does some actual processing.
+			{
+				System->GetImageDesc(OperationData->InstanceID, Image.ImageID, ImageDesc);
+
+				const uint16 MaxTextureSizeToGenerate = static_cast<uint16>(CustomizableObjectSystemPrivateData->MaxTextureSizeToGenerate);
+				const uint16 MaxSize = FMath::Max(ImageDesc.m_size[0], ImageDesc.m_size[1]);
+				uint16 Reduction = 1;
+
+				if (MaxTextureSizeToGenerate > 0 && MaxSize > MaxTextureSizeToGenerate)
+				{
+					// Find the reduction factor, and the BaseMip of the texture.
+					const uint32 NextPowerOfTwo = FMath::RoundUpToPowerOfTwo(FMath::DivideAndRoundUp(MaxSize, MaxTextureSizeToGenerate));
+					Reduction = FMath::Max(NextPowerOfTwo, 2U); // At least divide the texture by a factor of two
+					Image.BaseMip = FMath::FloorLog2(Reduction);
+				}
+
+				Image.FullImageSizeX = ImageDesc.m_size[0] / Reduction;
+				Image.FullImageSizeY = ImageDesc.m_size[1] / Reduction;
+			}
+
+			const bool bCached = ImagesInThisInstance.Contains(Image.ImageID) || // See if it is cached from this same instance (can happen with LODs)
+				(CVarReuseImagesBetweenInstances.GetValueOnAnyThread() && CustomizableObjectSystemPrivateData->ProtectedObjectCachedImages.Contains(Image.ImageID)); // See if it is cached from another instance
+
+			if (bCached)
+			{
+				UE_LOG(LogMutable, VeryVerbose, TEXT("Texture resource with id [%llu] is cached."), Image.ImageID);
+			}
+			else
+			{
+				const int32 MaxSize = FMath::Max(Image.FullImageSizeX, Image.FullImageSizeY);
+				const int32 FullLODCount = FMath::CeilLogTwo(MaxSize) + 1;
+				const int32 MinMipsInImage = FMath::Min(FullLODCount, UTexture::GetStaticMinTextureResidentMipCount());
+				const int32 MaxMipsToSkip = FullLODCount - MinMipsInImage;
+				int32 MipsToSkip = FMath::Min(MaxMipsToSkip, OperationData->MipsToSkip);
+
+				if (!FMath::IsPowerOfTwo(Image.FullImageSizeX) || !FMath::IsPowerOfTwo(Image.FullImageSizeY))
+				{
+					// It doesn't make sense to skip mips as non-power-of-two size textures cannot be streamed anyway
+					MipsToSkip = 0;
+				}
+
+				const int32 MipSizeX = FMath::Max(Image.FullImageSizeX >> MipsToSkip, 1);
+				const int32 MipSizeY = FMath::Max(Image.FullImageSizeY >> MipsToSkip, 1);
+				if (MipsToSkip > 0 && CustomizableObjectSystemPrivateData->EnableSkipGenerateResidentMips != 0 && OperationData->LowPriorityTextures.Find(Image.Name.ToString()) != INDEX_NONE)
+				{
+					Image.Image = new mu::Image(MipSizeX, MipSizeY, FullLODCount - MipsToSkip, ImageDesc.m_format, mu::EInitializationType::Black);
+				}
+				else
+				{
+					Image.Image = System->GetImage(OperationData->InstanceID, Image.ImageID, Image.BaseMip + MipsToSkip, Image.BaseLOD);
+				}
+
+				check(Image.Image);
+
+				// We should have generated exactly this size.
+				const bool bSizeMissmatch = Image.Image->GetSizeX() != MipSizeX || Image.Image->GetSizeY() != MipSizeY;
+				if (bSizeMissmatch)
+				{
+					// Generate a correctly-sized but empty image instead, to avoid crashes.
+					UE_LOG(LogMutable, Warning, TEXT("Mutable generated a wrongly-sized image %llu."), Image.ImageID);
+					Image.Image = new mu::Image(MipSizeX, MipSizeY, FullLODCount - MipsToSkip, Image.Image->GetFormat(), mu::EInitializationType::Black);
+				}
+
+				// We need one mip or the complete chain. Otherwise there was a bug.
+				const int32 FullMipCount = Image.Image->GetMipmapCount(Image.Image->GetSizeX(), Image.Image->GetSizeY());
+				const int32 RealMipCount = Image.Image->GetLODCount();
+
+				bool bForceMipchain = 
+					// Did we fail to generate the entire mipchain (if we have mips at all)?
+					(RealMipCount != 1) && (RealMipCount != FullMipCount);
+
+				if (bForceMipchain)
+				{
+					MUTABLE_CPUPROFILER_SCOPE(GetImage_MipFix);
+
+					UE_LOG(LogMutable, Warning, TEXT("Mutable generated an incomplete mip chain for image %llu."), Image.ImageID);
+
+					// Force the right number of mips. The missing data will be black.
+					const mu::Ptr<mu::Image> NewImage = new mu::Image(Image.Image->GetSizeX(), Image.Image->GetSizeY(), FullMipCount, Image.Image->GetFormat(), mu::EInitializationType::Black);
+					check(NewImage);
+					if (NewImage->GetDataSize() >= Image.Image->GetDataSize())
+					{
+						FMemory::Memcpy(NewImage->GetData(), Image.Image->GetData(), Image.Image->GetDataSize());
+					}
+					Image.Image = NewImage;
+				}
+
+				ImagesInThisInstance.Add(Image.ImageID);
+			}
+		}
+	}
+
+
+	void Task_Mutable_Update_GetImages(const TSharedRef<FUpdateContextPrivate>& OperationData)
+	{
+		MUTABLE_CPUPROFILER_SCOPE(Task_Mutable_GetImages)
+		FMutableScopeTimer Timer(OperationData->TaskGetImagesTime);
+
+#if WITH_EDITOR
+		uint32 StartCycles = FPlatformTime::Cycles();
+#endif
+
+		Subtask_Mutable_GetImages(OperationData);
+
+		// TODO: Not strictly mutable: move to another worker thread task to free mutable access?
+		impl::Subtask_Mutable_PrepareTextures(OperationData);
+
+#if WITH_EDITOR
+		const uint32 EndCycles = FPlatformTime::Cycles();
+		OperationData->MutableRuntimeCycles += EndCycles - StartCycles;
+#endif
+	}
+
+
+	void Subtask_Mutable_BeginUpdate_GetMesh(const TSharedRef<FUpdateContextPrivate>& OperationData, TSharedPtr<mu::Model> Model)
+	{
+		MUTABLE_CPUPROFILER_SCOPE(Subtask_Mutable_BeginUpdate_GetMesh)
+
+		check(OperationData->Parameters);
+		OperationData->InstanceUpdateData.Clear();
+
+		check(UCustomizableObjectSystem::GetInstance() != nullptr);
+		check(UCustomizableObjectSystem::GetInstance()->GetPrivate() != nullptr);
+		mu::System* System = UCustomizableObjectSystem::GetInstance()->GetPrivate()->MutableSystem.get();
+		check(System != nullptr);
+
+		const UCustomizableObject* CustomizableObject = OperationData->Instance->GetCustomizableObject();
+		const FModelResources& ModelResources = CustomizableObject->GetPrivate()->GetModelResources();
+
+		UCustomizableInstancePrivate* CustomizableObjectInstancePrivateData = OperationData->Instance->GetPrivate();
+
+		CustomizableObjectInstancePrivateData->PassThroughTexturesToLoad.Empty();
+
+		if (OperationData->PixelFormatOverride)
+		{
+			System->SetImagePixelConversionOverride( OperationData->PixelFormatOverride );
+		}
+
+		if (!OperationData->bUseMeshCache)
+		{
+			impl::CreateMutableInstance(OperationData);
+			impl::FixLODs(OperationData);
+		}
+
+		// Main instance generation step
+		// LOD mask, set to all ones to build  all LODs
+		const mu::Instance* Instance = OperationData->MutableInstance; // TODO GMTFuture remove
+		if (!Instance)
+		{
+			UE_LOG(LogMutable, Warning, TEXT("An Instace update has failed."));
+			return;
+		}
+
+		const TArray<uint16>& RequestedLODs = OperationData->GetRequestedLODs();
+
+		// Map SharedSurfaceId to surface index
+		TArray<int32> SurfacesSharedId;
+
+		// Generate the mesh and gather all the required resource Ids
+		OperationData->InstanceUpdateData.LODs.SetNum(OperationData->NumLODsAvailable);
+		for (int32 MutableLODIndex = 0; MutableLODIndex < OperationData->NumLODsAvailable; ++MutableLODIndex)
+		{
+			// Skip LODs outside the range we want to generate
+			if (MutableLODIndex < OperationData->GetMinLOD())
+			{
+				continue;
+			}
+
+			FInstanceUpdateData::FLOD& LOD = OperationData->InstanceUpdateData.LODs[MutableLODIndex];
+			LOD.FirstComponent = OperationData->InstanceUpdateData.Components.Num();
+			LOD.ComponentCount = Instance->GetComponentCount(MutableLODIndex);
+
+			for (int32 ComponentIndex = 0; ComponentIndex < LOD.ComponentCount; ++ComponentIndex)
+			{
+				OperationData->InstanceUpdateData.Components.Push(FInstanceUpdateData::FComponent());
+				FInstanceUpdateData::FComponent& Component = OperationData->InstanceUpdateData.Components.Last();
+				Component.Id = Instance->GetComponentId(MutableLODIndex, ComponentIndex);
+				Component.FirstSurface = OperationData->InstanceUpdateData.Surfaces.Num();
+				Component.SurfaceCount = 0;
+
+				const bool bGenerateLOD = RequestedLODs.IsValidIndex(Component.Id) ? RequestedLODs[Component.Id] <= MutableLODIndex : true;
+
+				// Mesh
+				if (Instance->GetMeshCount(MutableLODIndex, ComponentIndex) > 0)
+				{
+					MUTABLE_CPUPROFILER_SCOPE(GetMesh);
+
+					Component.MeshID = Instance->GetMeshId(MutableLODIndex, ComponentIndex, 0);
+
+					if(bGenerateLOD)
+					{
+						Component.Mesh = System->GetMesh(OperationData->InstanceID, Component.MeshID);
+					}
+				}
+
+				if (!Component.Mesh)
+				{
+					continue;
+				}
+
+				Component.bGenerated = true;
+
+
+				// Materials and images
+				const int32 SurfaceCount = Component.Mesh->GetSurfaceCount();
+				for (int32 MeshSurfaceIndex = 0; MeshSurfaceIndex < SurfaceCount; ++MeshSurfaceIndex)
+				{
+					const uint32 SurfaceId = Component.Mesh->GetSurfaceId(MeshSurfaceIndex);
+					const int32 InstanceSurfaceIndex = Instance->FindSurfaceById(MutableLODIndex, ComponentIndex, SurfaceId);
+					check(Component.Mesh->GetVertexCount() > 0 || InstanceSurfaceIndex >= 0);
+
+					int32 BaseSurfaceIndex = InstanceSurfaceIndex;
+					int32 BaseLODIndex = MutableLODIndex;
+
+					if (InstanceSurfaceIndex >= 0)
+					{
+						OperationData->InstanceUpdateData.Surfaces.Push({});
+						FInstanceUpdateData::FSurface& Surface = OperationData->InstanceUpdateData.Surfaces.Last();
+						++Component.SurfaceCount;
+
+						// Now Surface.MaterialIndex is decoded from a parameter at the end of this if()
+						Surface.SurfaceId = SurfaceId;
+
+						const int32 SharedSurfaceId = Instance->GetSharedSurfaceId(MutableLODIndex, ComponentIndex, InstanceSurfaceIndex);
+						const int32 SharedSurfaceIndex = SurfacesSharedId.Find(SharedSurfaceId);
+
+						SurfacesSharedId.Add(SharedSurfaceId);
+
+						if (SharedSurfaceId != INDEX_NONE)
+						{
+							if (SharedSurfaceIndex >= 0)
+							{
+								Surface = OperationData->InstanceUpdateData.Surfaces[SharedSurfaceIndex];
+								continue;
+							}
+
+							// Find the first LOD where this surface can be found
+							Instance->FindBaseSurfaceBySharedId(ComponentIndex, SharedSurfaceId, BaseSurfaceIndex, BaseLODIndex);
+
+							Surface.SurfaceId = Instance->GetSurfaceId(BaseLODIndex, ComponentIndex, BaseSurfaceIndex);
+						}
+
+						// Images
+						Surface.FirstImage = OperationData->InstanceUpdateData.Images.Num();
+						Surface.ImageCount = Instance->GetImageCount(BaseLODIndex, ComponentIndex, InstanceSurfaceIndex);
+						for (int32 ImageIndex = 0; ImageIndex < Surface.ImageCount; ++ImageIndex)
+						{
+							MUTABLE_CPUPROFILER_SCOPE(GetImageId);
+
+							OperationData->InstanceUpdateData.Images.Push({});
+							FInstanceUpdateData::FImage& Image = OperationData->InstanceUpdateData.Images.Last();
+							Image.Name = Instance->GetImageName(BaseLODIndex, ComponentIndex, InstanceSurfaceIndex, ImageIndex);
+							Image.ImageID = Instance->GetImageId(BaseLODIndex, ComponentIndex, BaseSurfaceIndex, ImageIndex);
+							Image.FullImageSizeX = 0;
+							Image.FullImageSizeY = 0;
+							Image.BaseLOD = BaseLODIndex;
+							Image.BaseMip = 0;
+
+							FString KeyName = Image.Name.ToString();
+							int32 ImageKey = FCString::Atoi(*KeyName);
+
+							if (ImageKey >= 0 && ImageKey < ModelResources.ImageProperties.Num())
+							{
+								const FMutableModelImageProperties& Props = ModelResources.ImageProperties[ImageKey];
+
+								if (Props.IsPassThrough)
+								{
+									Image.bIsPassThrough = true;
+
+									// Since it's known it's a pass-through texture there is no need to cache or convert it so we can generate it here already.
+									Image.Image = System->GetImage(OperationData->InstanceID, Image.ImageID, 0, 0);
+									check(Image.Image->IsReference());
+
+									uint32 ReferenceID = Image.Image->GetReferencedTexture();
+
+									if (ModelResources.PassThroughTextures.IsValidIndex(ReferenceID))
+									{
+										TSoftObjectPtr<UTexture> Ref = ModelResources.PassThroughTextures[ReferenceID];
+										CustomizableObjectInstancePrivateData->PassThroughTexturesToLoad.Add(Ref);
+									}
+									else
+									{
+										// internal error.
+										UE_LOG(LogMutable, Error, TEXT("Referenced image [%d] was not stored in the resource array."), ReferenceID);
+									}
+								}
+							}
+							else
+							{
+								// This means the compiled model (maybe coming from derived data) has images that the asset doesn't know about.
+								UE_LOG(LogMutable, Error, TEXT("CustomizableObject derived data out of sync with asset for [%s]. Try recompiling it."), *CustomizableObject->GetName());
+							}
+						}
+
+						// Vectors
+						Surface.FirstVector = OperationData->InstanceUpdateData.Vectors.Num();
+						Surface.VectorCount = Instance->GetVectorCount(BaseLODIndex, ComponentIndex, InstanceSurfaceIndex);
+						for (int32 VectorIndex = 0; VectorIndex < Surface.VectorCount; ++VectorIndex)
+						{
+							MUTABLE_CPUPROFILER_SCOPE(GetVector);
+							OperationData->InstanceUpdateData.Vectors.Push({});
+							FInstanceUpdateData::FVector& Vector = OperationData->InstanceUpdateData.Vectors.Last();
+							Vector.Name = Instance->GetVectorName(BaseLODIndex, ComponentIndex, InstanceSurfaceIndex, VectorIndex);
+							Vector.Vector = Instance->GetVector(BaseLODIndex, ComponentIndex, InstanceSurfaceIndex, VectorIndex);
+						}
+
+						// Scalars
+						Surface.FirstScalar = OperationData->InstanceUpdateData.Scalars.Num();
+						Surface.ScalarCount = Instance->GetScalarCount(BaseLODIndex, ComponentIndex, InstanceSurfaceIndex);
+						for (int32 ScalarIndex = 0; ScalarIndex < Surface.ScalarCount; ++ScalarIndex)
+						{
+							MUTABLE_CPUPROFILER_SCOPE(GetScalar)
+
+							const FName ScalarName = Instance->GetScalarName(BaseLODIndex, ComponentIndex, InstanceSurfaceIndex, ScalarIndex);
+							const float ScalarValue = Instance->GetScalar(BaseLODIndex, ComponentIndex, InstanceSurfaceIndex, ScalarIndex);
+							
+							FString EncodingMaterialIdString = "__MutableMaterialId";
+							
+							// Decoding Material Switch from Mutable parameter name
+							if (ScalarName.ToString().Equals(EncodingMaterialIdString))
+							{
+								Surface.MaterialIndex = static_cast<uint32>(ScalarValue);
+							
+								// This parameter is not needed in the final material instance
+								Surface.ScalarCount -= 1;
+							}
+							else
+							{
+								OperationData->InstanceUpdateData.Scalars.Push({ ScalarName, ScalarValue });
+							}
+						}
+					}
+				}
+			}
+		}
+
+		// Copy ExtensionData Object node input from the Instance to the InstanceUpdateData
+		for (int32 ExtensionDataIndex = 0; ExtensionDataIndex < Instance->GetExtensionDataCount(); ExtensionDataIndex++)
+		{
+			mu::Ptr<const mu::ExtensionData> ExtensionData;
+			FName Name;
+			Instance->GetExtensionData(ExtensionDataIndex, ExtensionData, Name);
+
+			check(ExtensionData);
+
+			FInstanceUpdateData::FNamedExtensionData& NewEntry = OperationData->InstanceUpdateData.ExtendedInputPins.AddDefaulted_GetRef();
+			NewEntry.Data = ExtensionData;
+			NewEntry.Name = Name;
+			check(NewEntry.Name != NAME_None);
+		}
+	}
+
+
+	void Task_Mutable_Update_GetMesh(const TSharedRef<FUpdateContextPrivate>& OperationData,
+		const TSharedPtr<mu::Model>& Model)
+	{
+		MUTABLE_CPUPROFILER_SCOPE(Task_Mutable_Update_GetMesh)
+		FMutableScopeTimer Timer(OperationData->TaskGetMeshTime);
+			
+#if WITH_EDITOR
+		const uint32 StartCycles = FPlatformTime::Cycles();
+#endif
+
+		Subtask_Mutable_BeginUpdate_GetMesh(OperationData, Model);
+
+		// TODO: Not strictly mutable: move to another worker thread task to free mutable access?
+		impl::Subtask_Mutable_PrepareSkeletonData(OperationData);
+
+		if (OperationData->GetCapturedDescriptor().GetBuildParameterRelevancy())
+		{
+			impl::Subtask_Mutable_UpdateParameterRelevancy(OperationData);
+		}
+		else
+		{
+			OperationData->RelevantParametersInProgress.Reset();
+		}
+
+#if WITH_EDITOR
+		const uint32 EndCycles = FPlatformTime::Cycles();
+		OperationData->MutableRuntimeCycles = EndCycles - StartCycles;
+#endif
+	}
+}
diff --git a/Engine/Plugins/Experimental/Mutable/Source/CustomizableObject/Private/MuCO/BusyWaits_Deprecated.h b/Engine/Plugins/Experimental/Mutable/Source/CustomizableObject/Private/MuCO/BusyWaits_Deprecated.h
new file mode 100644
--- /dev/null
+++ b/Engine/Plugins/Experimental/Mutable/Source/CustomizableObject/Private/MuCO/BusyWaits_Deprecated.h
@@ -0,0 +1,29 @@
+// Copyright Epic Games, Inc. All Rights Reserved.
+
+// ---------------------------------------------------------------
+// THIS FILE SHOULD ONLY CONTAIN CODE BEING DEPRECATED BY UE-207754
+// ---------------------------------------------------------------
+
+#include "Templates/SharedPointer.h"
+
+class FUpdateContextPrivate;
+namespace mu
+{
+	class Model;
+}
+
+
+namespace impl_deprecated
+{
+	// This runs in the mutable thread.
+	void Subtask_Mutable_GetImages(const TSharedRef<FUpdateContextPrivate>& OperationData);
+
+	// This runs in a worker thread.
+	void Task_Mutable_Update_GetImages(const TSharedRef<FUpdateContextPrivate>& OperationData);
+
+	// This runs in the mutable thread.
+	void Subtask_Mutable_BeginUpdate_GetMesh(const TSharedRef<FUpdateContextPrivate>& OperationData, TSharedPtr<mu::Model> Model);
+
+	// This runs in a worker thread.
+	void Task_Mutable_Update_GetMesh(const TSharedRef<FUpdateContextPrivate>& OperationData, const TSharedPtr<mu::Model>& Model);
+}
diff --git a/Engine/Plugins/Experimental/Mutable/Source/CustomizableObject/Private/MuCO/CustomizableObjectSystem.cpp b/Engine/Plugins/Experimental/Mutable/Source/CustomizableObject/Private/MuCO/CustomizableObjectSystem.cpp
--- a/Engine/Plugins/Experimental/Mutable/Source/CustomizableObject/Private/MuCO/CustomizableObjectSystem.cpp
+++ b/Engine/Plugins/Experimental/Mutable/Source/CustomizableObject/Private/MuCO/CustomizableObjectSystem.cpp
@@ -45,6 +45,8 @@
 #endif
 
 
+#include "Busywaits_Deprecated.h"
+
 #include UE_INLINE_GENERATED_CPP_BY_NAME(CustomizableObjectSystem)
 
 class AActor;
@@ -1644,362 +1646,6 @@ namespace impl
 		Operation->SetRequestedLODs(RequestedLODs);
 	}
 	
-	
-	// This runs in the mutable thread.
-	void Subtask_Mutable_BeginUpdate_GetMesh(const TSharedRef<FUpdateContextPrivate>& OperationData, TSharedPtr<mu::Model> Model)
-	{
-		MUTABLE_CPUPROFILER_SCOPE(Subtask_Mutable_BeginUpdate_GetMesh)
-
-		check(OperationData->Parameters);
-		OperationData->InstanceUpdateData.Clear();
-
-		check(UCustomizableObjectSystem::GetInstance() != nullptr);
-		check(UCustomizableObjectSystem::GetInstance()->GetPrivate() != nullptr);
-		mu::System* System = UCustomizableObjectSystem::GetInstance()->GetPrivate()->MutableSystem.get();
-		check(System != nullptr);
-
-		const UCustomizableObject* CustomizableObject = OperationData->Instance->GetCustomizableObject();
-		const FModelResources& ModelResources = CustomizableObject->GetPrivate()->GetModelResources();
-
-		UCustomizableInstancePrivate* CustomizableObjectInstancePrivateData = OperationData->Instance->GetPrivate();
-
-		CustomizableObjectInstancePrivateData->PassThroughTexturesToLoad.Empty();
-
-		if (OperationData->PixelFormatOverride)
-		{
-			System->SetImagePixelConversionOverride( OperationData->PixelFormatOverride );
-		}
-
-		if (!OperationData->bUseMeshCache)
-		{
-			CreateMutableInstance(OperationData);
-			FixLODs(OperationData);
-		}
-
-		// Main instance generation step
-		// LOD mask, set to all ones to build  all LODs
-		const mu::Instance* Instance = OperationData->MutableInstance; // TODO GMTFuture remove
-		if (!Instance)
-		{
-			UE_LOG(LogMutable, Warning, TEXT("An Instace update has failed."));
-			return;
-		}
-
-		const TArray<uint16>& RequestedLODs = OperationData->GetRequestedLODs();
-
-		// Map SharedSurfaceId to surface index
-		TArray<int32> SurfacesSharedId;
-
-		// Generate the mesh and gather all the required resource Ids
-		OperationData->InstanceUpdateData.LODs.SetNum(OperationData->NumLODsAvailable);
-		for (int32 MutableLODIndex = 0; MutableLODIndex < OperationData->NumLODsAvailable; ++MutableLODIndex)
-		{
-			// Skip LODs outside the range we want to generate
-			if (MutableLODIndex < OperationData->GetMinLOD())
-			{
-				continue;
-			}
-
-			FInstanceUpdateData::FLOD& LOD = OperationData->InstanceUpdateData.LODs[MutableLODIndex];
-			LOD.FirstComponent = OperationData->InstanceUpdateData.Components.Num();
-			LOD.ComponentCount = Instance->GetComponentCount(MutableLODIndex);
-
-			for (int32 ComponentIndex = 0; ComponentIndex < LOD.ComponentCount; ++ComponentIndex)
-			{
-				OperationData->InstanceUpdateData.Components.Push(FInstanceUpdateData::FComponent());
-				FInstanceUpdateData::FComponent& Component = OperationData->InstanceUpdateData.Components.Last();
-				Component.Id = Instance->GetComponentId(MutableLODIndex, ComponentIndex);
-				Component.FirstSurface = OperationData->InstanceUpdateData.Surfaces.Num();
-				Component.SurfaceCount = 0;
-
-				const bool bGenerateLOD = RequestedLODs.IsValidIndex(Component.Id) ? RequestedLODs[Component.Id] <= MutableLODIndex : true;
-
-				// Mesh
-				if (Instance->GetMeshCount(MutableLODIndex, ComponentIndex) > 0)
-				{
-					MUTABLE_CPUPROFILER_SCOPE(GetMesh);
-
-					Component.MeshID = Instance->GetMeshId(MutableLODIndex, ComponentIndex, 0);
-
-					if(bGenerateLOD)
-					{
-						Component.Mesh = System->GetMesh(OperationData->InstanceID, Component.MeshID);
-					}
-				}
-
-				if (!Component.Mesh)
-				{
-					continue;
-				}
-
-				Component.bGenerated = true;
-
-
-				// Materials and images
-				const int32 SurfaceCount = Component.Mesh->GetSurfaceCount();
-				for (int32 MeshSurfaceIndex = 0; MeshSurfaceIndex < SurfaceCount; ++MeshSurfaceIndex)
-				{
-					const uint32 SurfaceId = Component.Mesh->GetSurfaceId(MeshSurfaceIndex);
-					const int32 InstanceSurfaceIndex = Instance->FindSurfaceById(MutableLODIndex, ComponentIndex, SurfaceId);
-					check(Component.Mesh->GetVertexCount() > 0 || InstanceSurfaceIndex >= 0);
-
-					int32 BaseSurfaceIndex = InstanceSurfaceIndex;
-					int32 BaseLODIndex = MutableLODIndex;
-
-					if (InstanceSurfaceIndex >= 0)
-					{
-						OperationData->InstanceUpdateData.Surfaces.Push({});
-						FInstanceUpdateData::FSurface& Surface = OperationData->InstanceUpdateData.Surfaces.Last();
-						++Component.SurfaceCount;
-
-						// Now Surface.MaterialIndex is decoded from a parameter at the end of this if()
-						Surface.SurfaceId = SurfaceId;
-
-						const int32 SharedSurfaceId = Instance->GetSharedSurfaceId(MutableLODIndex, ComponentIndex, InstanceSurfaceIndex);
-						const int32 SharedSurfaceIndex = SurfacesSharedId.Find(SharedSurfaceId);
-
-						SurfacesSharedId.Add(SharedSurfaceId);
-
-						if (SharedSurfaceId != INDEX_NONE)
-						{
-							if (SharedSurfaceIndex >= 0)
-							{
-								Surface = OperationData->InstanceUpdateData.Surfaces[SharedSurfaceIndex];
-								continue;
-							}
-
-							// Find the first LOD where this surface can be found
-							Instance->FindBaseSurfaceBySharedId(ComponentIndex, SharedSurfaceId, BaseSurfaceIndex, BaseLODIndex);
-
-							Surface.SurfaceId = Instance->GetSurfaceId(BaseLODIndex, ComponentIndex, BaseSurfaceIndex);
-						}
-
-						// Images
-						Surface.FirstImage = OperationData->InstanceUpdateData.Images.Num();
-						Surface.ImageCount = Instance->GetImageCount(BaseLODIndex, ComponentIndex, InstanceSurfaceIndex);
-						for (int32 ImageIndex = 0; ImageIndex < Surface.ImageCount; ++ImageIndex)
-						{
-							MUTABLE_CPUPROFILER_SCOPE(GetImageId);
-
-							OperationData->InstanceUpdateData.Images.Push({});
-							FInstanceUpdateData::FImage& Image = OperationData->InstanceUpdateData.Images.Last();
-							Image.Name = Instance->GetImageName(BaseLODIndex, ComponentIndex, InstanceSurfaceIndex, ImageIndex);
-							Image.ImageID = Instance->GetImageId(BaseLODIndex, ComponentIndex, BaseSurfaceIndex, ImageIndex);
-							Image.FullImageSizeX = 0;
-							Image.FullImageSizeY = 0;
-							Image.BaseLOD = BaseLODIndex;
-							Image.BaseMip = 0;
-
-							FString KeyName = Image.Name.ToString();
-							int32 ImageKey = FCString::Atoi(*KeyName);
-
-							if (ImageKey >= 0 && ImageKey < ModelResources.ImageProperties.Num())
-							{
-								const FMutableModelImageProperties& Props = ModelResources.ImageProperties[ImageKey];
-
-								if (Props.IsPassThrough)
-								{
-									Image.bIsPassThrough = true;
-
-									// Since it's known it's a pass-through texture there is no need to cache or convert it so we can generate it here already.
-									Image.Image = System->GetImage(OperationData->InstanceID, Image.ImageID, 0, 0);
-									check(Image.Image->IsReference());
-
-									uint32 ReferenceID = Image.Image->GetReferencedTexture();
-
-									if (ModelResources.PassThroughTextures.IsValidIndex(ReferenceID))
-									{
-										TSoftObjectPtr<UTexture> Ref = ModelResources.PassThroughTextures[ReferenceID];
-										CustomizableObjectInstancePrivateData->PassThroughTexturesToLoad.Add(Ref);
-									}
-									else
-									{
-										// internal error.
-										UE_LOG(LogMutable, Error, TEXT("Referenced image [%d] was not stored in the resource array."), ReferenceID);
-									}
-								}
-							}
-							else
-							{
-								// This means the compiled model (maybe coming from derived data) has images that the asset doesn't know about.
-								UE_LOG(LogMutable, Error, TEXT("CustomizableObject derived data out of sync with asset for [%s]. Try recompiling it."), *CustomizableObject->GetName());
-							}
-						}
-
-						// Vectors
-						Surface.FirstVector = OperationData->InstanceUpdateData.Vectors.Num();
-						Surface.VectorCount = Instance->GetVectorCount(BaseLODIndex, ComponentIndex, InstanceSurfaceIndex);
-						for (int32 VectorIndex = 0; VectorIndex < Surface.VectorCount; ++VectorIndex)
-						{
-							MUTABLE_CPUPROFILER_SCOPE(GetVector);
-							OperationData->InstanceUpdateData.Vectors.Push({});
-							FInstanceUpdateData::FVector& Vector = OperationData->InstanceUpdateData.Vectors.Last();
-							Vector.Name = Instance->GetVectorName(BaseLODIndex, ComponentIndex, InstanceSurfaceIndex, VectorIndex);
-							Vector.Vector = Instance->GetVector(BaseLODIndex, ComponentIndex, InstanceSurfaceIndex, VectorIndex);
-						}
-
-						// Scalars
-						Surface.FirstScalar = OperationData->InstanceUpdateData.Scalars.Num();
-						Surface.ScalarCount = Instance->GetScalarCount(BaseLODIndex, ComponentIndex, InstanceSurfaceIndex);
-						for (int32 ScalarIndex = 0; ScalarIndex < Surface.ScalarCount; ++ScalarIndex)
-						{
-							MUTABLE_CPUPROFILER_SCOPE(GetScalar)
-
-							const FName ScalarName = Instance->GetScalarName(BaseLODIndex, ComponentIndex, InstanceSurfaceIndex, ScalarIndex);
-							const float ScalarValue = Instance->GetScalar(BaseLODIndex, ComponentIndex, InstanceSurfaceIndex, ScalarIndex);
-							
-							FString EncodingMaterialIdString = "__MutableMaterialId";
-							
-							// Decoding Material Switch from Mutable parameter name
-							if (ScalarName.ToString().Equals(EncodingMaterialIdString))
-							{
-								Surface.MaterialIndex = static_cast<uint32>(ScalarValue);
-							
-								// This parameter is not needed in the final material instance
-								Surface.ScalarCount -= 1;
-							}
-							else
-							{
-								OperationData->InstanceUpdateData.Scalars.Push({ ScalarName, ScalarValue });
-							}
-						}
-					}
-				}
-			}
-		}
-
-		// Copy ExtensionData Object node input from the Instance to the InstanceUpdateData
-		for (int32 ExtensionDataIndex = 0; ExtensionDataIndex < Instance->GetExtensionDataCount(); ExtensionDataIndex++)
-		{
-			mu::Ptr<const mu::ExtensionData> ExtensionData;
-			FName Name;
-			Instance->GetExtensionData(ExtensionDataIndex, ExtensionData, Name);
-
-			check(ExtensionData);
-
-			FInstanceUpdateData::FNamedExtensionData& NewEntry = OperationData->InstanceUpdateData.ExtendedInputPins.AddDefaulted_GetRef();
-			NewEntry.Data = ExtensionData;
-			NewEntry.Name = Name;
-			check(NewEntry.Name != NAME_None);
-		}
-	}
-
-
-	// This runs in the mutable thread.
-	void Subtask_Mutable_GetImages(const TSharedRef<FUpdateContextPrivate>& OperationData)
-	{
-		MUTABLE_CPUPROFILER_SCOPE(Subtask_Mutable_GetImages)
-
-		const UCustomizableObjectSystemPrivate* CustomizableObjectSystemPrivateData = UCustomizableObjectSystem::GetInstanceChecked()->GetPrivate();
-		mu::System* System = CustomizableObjectSystemPrivateData->MutableSystem.get();
-		check(System != nullptr);
-
-		// Generate all the required resources, that are not cached
-		TArray<mu::FResourceID> ImagesInThisInstance;
-		for (FInstanceUpdateData::FImage& Image : OperationData->InstanceUpdateData.Images)
-		{
-			MUTABLE_CPUPROFILER_SCOPE(GetImage);
-
-			// If the image is a reference to an engine texture, we are done.
-			if (Image.bIsPassThrough)
-			{
-				continue;
-			}
-
-			mu::FImageDesc ImageDesc;
-
-			// This should only be done when using progressive images, since GetImageDesc does some actual processing.
-			{
-				System->GetImageDesc(OperationData->InstanceID, Image.ImageID, ImageDesc);
-
-				const uint16 MaxTextureSizeToGenerate = static_cast<uint16>(CustomizableObjectSystemPrivateData->MaxTextureSizeToGenerate);
-				const uint16 MaxSize = FMath::Max(ImageDesc.m_size[0], ImageDesc.m_size[1]);
-				uint16 Reduction = 1;
-
-				if (MaxTextureSizeToGenerate > 0 && MaxSize > MaxTextureSizeToGenerate)
-				{
-					// Find the reduction factor, and the BaseMip of the texture.
-					const uint32 NextPowerOfTwo = FMath::RoundUpToPowerOfTwo(FMath::DivideAndRoundUp(MaxSize, MaxTextureSizeToGenerate));
-					Reduction = FMath::Max(NextPowerOfTwo, 2U); // At least divide the texture by a factor of two
-					Image.BaseMip = FMath::FloorLog2(Reduction);
-				}
-
-				Image.FullImageSizeX = ImageDesc.m_size[0] / Reduction;
-				Image.FullImageSizeY = ImageDesc.m_size[1] / Reduction;
-			}
-
-			const bool bCached = ImagesInThisInstance.Contains(Image.ImageID) || // See if it is cached from this same instance (can happen with LODs)
-				(CVarReuseImagesBetweenInstances.GetValueOnAnyThread() && CustomizableObjectSystemPrivateData->ProtectedObjectCachedImages.Contains(Image.ImageID)); // See if it is cached from another instance
-
-			if (bCached)
-			{
-				UE_LOG(LogMutable, VeryVerbose, TEXT("Texture resource with id [%llu] is cached."), Image.ImageID);
-			}
-			else
-			{
-				const int32 MaxSize = FMath::Max(Image.FullImageSizeX, Image.FullImageSizeY);
-				const int32 FullLODCount = FMath::CeilLogTwo(MaxSize) + 1;
-				const int32 MinMipsInImage = FMath::Min(FullLODCount, UTexture::GetStaticMinTextureResidentMipCount());
-				const int32 MaxMipsToSkip = FullLODCount - MinMipsInImage;
-				int32 MipsToSkip = FMath::Min(MaxMipsToSkip, OperationData->MipsToSkip);
-
-				if (!FMath::IsPowerOfTwo(Image.FullImageSizeX) || !FMath::IsPowerOfTwo(Image.FullImageSizeY))
-				{
-					// It doesn't make sense to skip mips as non-power-of-two size textures cannot be streamed anyway
-					MipsToSkip = 0;
-				}
-
-				const int32 MipSizeX = FMath::Max(Image.FullImageSizeX >> MipsToSkip, 1);
-				const int32 MipSizeY = FMath::Max(Image.FullImageSizeY >> MipsToSkip, 1);
-				if (MipsToSkip > 0 && CustomizableObjectSystemPrivateData->EnableSkipGenerateResidentMips != 0 && OperationData->LowPriorityTextures.Find(Image.Name.ToString()) != INDEX_NONE)
-				{
-					Image.Image = new mu::Image(MipSizeX, MipSizeY, FullLODCount - MipsToSkip, ImageDesc.m_format, mu::EInitializationType::Black);
-				}
-				else
-				{
-					Image.Image = System->GetImage(OperationData->InstanceID, Image.ImageID, Image.BaseMip + MipsToSkip, Image.BaseLOD);
-				}
-
-				check(Image.Image);
-
-				// We should have generated exactly this size.
-				const bool bSizeMissmatch = Image.Image->GetSizeX() != MipSizeX || Image.Image->GetSizeY() != MipSizeY;
-				if (bSizeMissmatch)
-				{
-					// Generate a correctly-sized but empty image instead, to avoid crashes.
-					UE_LOG(LogMutable, Warning, TEXT("Mutable generated a wrongly-sized image %llu."), Image.ImageID);
-					Image.Image = new mu::Image(MipSizeX, MipSizeY, FullLODCount - MipsToSkip, Image.Image->GetFormat(), mu::EInitializationType::Black);
-				}
-
-				// We need one mip or the complete chain. Otherwise there was a bug.
-				const int32 FullMipCount = Image.Image->GetMipmapCount(Image.Image->GetSizeX(), Image.Image->GetSizeY());
-				const int32 RealMipCount = Image.Image->GetLODCount();
-
-				bool bForceMipchain = 
-					// Did we fail to generate the entire mipchain (if we have mips at all)?
-					(RealMipCount != 1) && (RealMipCount != FullMipCount);
-
-				if (bForceMipchain)
-				{
-					MUTABLE_CPUPROFILER_SCOPE(GetImage_MipFix);
-
-					UE_LOG(LogMutable, Warning, TEXT("Mutable generated an incomplete mip chain for image %llu."), Image.ImageID);
-
-					// Force the right number of mips. The missing data will be black.
-					const mu::Ptr<mu::Image> NewImage = new mu::Image(Image.Image->GetSizeX(), Image.Image->GetSizeY(), FullMipCount, Image.Image->GetFormat(), mu::EInitializationType::Black);
-					check(NewImage);
-					if (NewImage->GetDataSize() >= Image.Image->GetDataSize())
-					{
-						FMemory::Memcpy(NewImage->GetData(), Image.Image->GetData(), Image.Image->GetDataSize());
-					}
-					Image.Image = NewImage;
-				}
-
-				ImagesInThisInstance.Add(Image.ImageID);
-			}
-		}
-	}
-	
 
 	// This runs in a worker thread
 	void Subtask_Mutable_PrepareTextures(const TSharedRef<FUpdateContextPrivate>& OperationData)
@@ -2094,57 +1740,9 @@ namespace impl
 	}
 
 
-	// This runs in a worker thread.
-	void Task_Mutable_Update_GetMesh(const TSharedRef<FUpdateContextPrivate>& OperationData, const TSharedPtr<mu::Model>& Model)
-	{
-		MUTABLE_CPUPROFILER_SCOPE(Task_Mutable_Update_GetMesh)
-		FMutableScopeTimer Timer(OperationData->TaskGetMeshTime);
-		
-#if WITH_EDITOR
-		const uint32 StartCycles = FPlatformTime::Cycles();
-#endif
-
-		Subtask_Mutable_BeginUpdate_GetMesh(OperationData, Model);
-
-		// TODO: Not strictly mutable: move to another worker thread task to free mutable access?
-		Subtask_Mutable_PrepareSkeletonData(OperationData);
-
-		if (OperationData->GetCapturedDescriptor().GetBuildParameterRelevancy())
-		{
-			Subtask_Mutable_UpdateParameterRelevancy(OperationData);
-		}
-		else
-		{
-			OperationData->RelevantParametersInProgress.Reset();
-		}
-
-#if WITH_EDITOR
-		const uint32 EndCycles = FPlatformTime::Cycles();
-		OperationData->MutableRuntimeCycles = EndCycles - StartCycles;
-#endif
-	}
 
 	
-	// This runs in a worker thread.
-	void Task_Mutable_Update_GetImages(const TSharedRef<FUpdateContextPrivate>& OperationData)
-	{
-		MUTABLE_CPUPROFILER_SCOPE(Task_Mutable_GetImages)
-		FMutableScopeTimer Timer(OperationData->TaskGetImagesTime);
-
-#if WITH_EDITOR
-		uint32 StartCycles = FPlatformTime::Cycles();
-#endif
-
-		Subtask_Mutable_GetImages(OperationData);
-
-		// TODO: Not strictly mutable: move to another worker thread task to free mutable access?
-		Subtask_Mutable_PrepareTextures(OperationData);
-
-#if WITH_EDITOR
-		const uint32 EndCycles = FPlatformTime::Cycles();
-		OperationData->MutableRuntimeCycles += EndCycles - StartCycles;
-#endif
-	}
+	
 
 
 	// This runs in a worker thread.
@@ -2486,7 +2084,7 @@ namespace impl
 				TEXT("Task_Mutable_GetImages"),
 				[OperationData]()
 				{
-					impl::Task_Mutable_Update_GetImages(OperationData);
+					impl_deprecated::Task_Mutable_Update_GetImages(OperationData);
 				});
 		}
 
@@ -2569,7 +2167,7 @@ namespace impl
 			TEXT("Task_Mutable_Update_GetMesh"),
 			[Operation, Model]()
 			{
-				impl::Task_Mutable_Update_GetMesh(Operation, Model);
+				impl_deprecated::Task_Mutable_Update_GetMesh(Operation, Model);
 			});
 
 		SystemPrivate->AddGameThreadTask(
@@ -2846,7 +2444,7 @@ namespace impl
 				TEXT("Task_Mutable_Update_GetMesh"),
 				[Operation, Model]()
 				{
-					impl::Task_Mutable_Update_GetMesh(Operation, Model);
+					impl_deprecated::Task_Mutable_Update_GetMesh(Operation, Model);
 				});
 
 			// Task: Lock cache
