[Mutable] Remove old task graph api code.

#rnx

diff --git a/Engine/Plugins/Experimental/Mutable/Source/CustomizableObject/Internal/MuCO/CustomizableObjectSystemPrivate.h b/Engine/Plugins/Experimental/Mutable/Source/CustomizableObject/Internal/MuCO/CustomizableObjectSystemPrivate.h
--- a/Engine/Plugins/Experimental/Mutable/Source/CustomizableObject/Internal/MuCO/CustomizableObjectSystemPrivate.h
+++ b/Engine/Plugins/Experimental/Mutable/Source/CustomizableObject/Internal/MuCO/CustomizableObjectSystemPrivate.h
@@ -25,12 +25,7 @@
 #include "AssetRegistry/AssetData.h"
 #include "ContentStreaming.h"
 
-// This define could come from MuR/System.h
-#ifdef MUTABLE_USE_NEW_TASKGRAPH
-	#include "Tasks/Task.h"
-#else
-	#include "Async/TaskGraphInterfaces.h"
-#endif
+#include "Tasks/Task.h"
 
 #include "CustomizableObjectSystemPrivate.generated.h"
 
@@ -290,50 +285,26 @@ struct FMutableTask
 	/** Actual function to perform in this task. */
 	FMutableTaskDelegate Function;
 
-	/** We can have 2 types of depencies:
-		From the new task system: */
-#ifdef MUTABLE_USE_NEW_TASKGRAPH
-	UE::Tasks::FTask Dependency;
-#endif
-
-	UE::Tasks::FTask GraphDependency0;
+	TArray<UE::Tasks::FTask, TFixedAllocator<2>> Dependencies; 
 	
-#ifdef MUTABLE_USE_NEW_TASKGRAPH
-	UE::Tasks::FTask GraphDependency1;
-#else
-	FGraphEventRef GraphDependency1;
-#endif
-
 	/** Check if all the dependencies of this task have been completed. */
-	inline bool AreDependenciesComplete() const
+	FORCEINLINE bool AreDependenciesComplete() const
 	{
-#ifdef MUTABLE_USE_NEW_TASKGRAPH
-		return (!Dependency.IsValid() || Dependency.IsCompleted())
-			&&
-			(!GraphDependency0.IsValid() || GraphDependency0.IsCompleted())
-			&&
-			(GraphDependency1.IsCompleted());
-#else
-		return 
-			(!GraphDependency0.IsValid() || GraphDependency0.IsCompleted())
-			&&
-			(!GraphDependency1 || GraphDependency1->IsComplete());
-#endif
+		for (const UE::Tasks::FTask& Dependency : Dependencies)
+		{
+			if (!Dependency.IsCompleted())
+			{
+				return false;
+			}
+		}
+
+		return true;
 	}
 
 	/** Free the handles for any dependency of this task. */
-	inline void ClearDependencies()
+	FORCEINLINE void ClearDependencies()
 	{
-#ifdef MUTABLE_USE_NEW_TASKGRAPH
-		Dependency = {};
-#endif
-		GraphDependency0 = {};
-
-#ifdef MUTABLE_USE_NEW_TASKGRAPH
-		GraphDependency1 = {};
-#else
-		GraphDependency1 = nullptr;
-#endif
+		Dependencies.Empty();
 	}
 };
 
diff --git a/Engine/Plugins/Experimental/Mutable/Source/CustomizableObject/Internal/MuCO/FMutableTaskGraph.h b/Engine/Plugins/Experimental/Mutable/Source/CustomizableObject/Internal/MuCO/FMutableTaskGraph.h
--- a/Engine/Plugins/Experimental/Mutable/Source/CustomizableObject/Internal/MuCO/FMutableTaskGraph.h
+++ b/Engine/Plugins/Experimental/Mutable/Source/CustomizableObject/Internal/MuCO/FMutableTaskGraph.h
@@ -40,15 +40,9 @@ class FMutableTaskGraph
 	};
 	
 public:
-	using TaskType = 
-#ifdef MUTABLE_USE_NEW_TASKGRAPH
-		UE::Tasks::FTask;
-#else
-		FGraphEventRef;
-#endif
 
 	/** Create and launch a task on the Mutable Thread with Normal priority. */
-	TaskType AddMutableThreadTask(const TCHAR* DebugName, TUniqueFunction<void()>&& TaskBody);
+	UE::Tasks::FTask AddMutableThreadTask(const TCHAR* DebugName, TUniqueFunction<void()>&& TaskBody);
 
 	/** Create and launch a task on the Mutable Thread with Low priority. */
 	uint32 AddMutableThreadTaskLowPriority(const TCHAR* DebugName, TFunction<void()>&& TaskBody);
@@ -78,7 +72,7 @@ private:
 	void TryLaunchMutableTaskLowPriority(bool bFromMutableTask);
 
 	/** Return true if the task is completed (or is no longer valid). */
-	bool IsTaskCompleted(const TaskType& Task) const;
+	bool IsTaskCompleted(const UE::Tasks::FTask& Task) const;
 
 	mutable FCriticalSection MutableTaskLock;
 
@@ -96,9 +90,9 @@ private:
 	uint32 TaskIdGenerator = INVALID_ID;
 
 	/** Last Mutable Task Low Priority launched to the TaskGraph system. */
-	TaskType LastMutableTaskLowPriority = {};
+	UE::Tasks::FTask LastMutableTaskLowPriority = {};
 
 	/** Last Mutable Task launched to the TaskGraph system. Low and normal priority. */
-	TaskType LastMutableTask = {};
+	UE::Tasks::FTask LastMutableTask = {};
 };
 
diff --git a/Engine/Plugins/Experimental/Mutable/Source/CustomizableObject/Private/MuCO/CustomizableObjectInstance.cpp b/Engine/Plugins/Experimental/Mutable/Source/CustomizableObject/Private/MuCO/CustomizableObjectInstance.cpp
--- a/Engine/Plugins/Experimental/Mutable/Source/CustomizableObject/Private/MuCO/CustomizableObjectInstance.cpp
+++ b/Engine/Plugins/Experimental/Mutable/Source/CustomizableObject/Private/MuCO/CustomizableObjectInstance.cpp
@@ -18,6 +18,7 @@
 #include "Materials/Material.h"
 #include "Materials/MaterialInstanceDynamic.h"
 #include "Modules/ModuleManager.h"
+#include "Tasks/Task.h"
 
 #include "MuCO/CustomizableObjectSystemPrivate.h"
 #include "MuCO/CustomizableInstanceLODManagement.h"
@@ -4801,8 +4802,6 @@ UE::Tasks::FTask UCustomizableInstancePrivate::LoadAdditionalAssetsAndDataAsync(
 
 	const FModelResources& ModelResources = CustomizableObject->GetPrivate()->GetModelResources();
 
-	FGraphEventRef Result = nullptr;
-
 	TArray<FSoftObjectPath> AssetsToStream;
 	TArray<uint32> RealTimeMorphStreamableBlocksToStream;
 
@@ -5237,7 +5236,7 @@ UE::Tasks::FTask UCustomizableInstancePrivate::LoadAdditionalAssetsAndDataAsync(
 	else
 	{
 		check(ReadRequestTasks.Num() == 0);
-		return UE::Tasks::MakeCompletedTask<void>(TEXT("GatherStreamingRequestsCompletionTaskComplete"));
+		return UE::Tasks::MakeCompletedTask<void>();
 	}
 }
 
diff --git a/Engine/Plugins/Experimental/Mutable/Source/CustomizableObject/Private/MuCO/CustomizableObjectSystem.cpp b/Engine/Plugins/Experimental/Mutable/Source/CustomizableObject/Private/MuCO/CustomizableObjectSystem.cpp
--- a/Engine/Plugins/Experimental/Mutable/Source/CustomizableObject/Private/MuCO/CustomizableObjectSystem.cpp
+++ b/Engine/Plugins/Experimental/Mutable/Source/CustomizableObject/Private/MuCO/CustomizableObjectSystem.cpp
@@ -1389,7 +1389,7 @@ bool UCustomizableObjectSystem::LockObject(const class UCustomizableObject* InOb
 		check(GetPrivate() != nullptr);
 		if (GetPrivate()->Streamer)
 		{
-			FMutableTaskGraph::TaskType Task = Private->MutableTaskGraph.AddMutableThreadTask(TEXT("EndStream"), [InObject, Streamer = GetPrivate()->Streamer]()
+			UE::Tasks::FTask Task = Private->MutableTaskGraph.AddMutableThreadTask(TEXT("EndStream"), [InObject, Streamer = GetPrivate()->Streamer]()
 				{
 					Streamer->CancelStreamingForObject(InObject);
 				});
@@ -2073,11 +2073,7 @@ namespace impl
 
 		// Task: Mutable GetImages
 		//-------------------------------------------------------------
-#ifdef MUTABLE_USE_NEW_TASKGRAPH
 		UE::Tasks::FTask Mutable_GetImagesTask;
-#else
-		FGraphEventRef Mutable_GetImagesTask;
-#endif
 		{
 			// Task inputs
 			Mutable_GetImagesTask = SystemPrivateData->MutableTaskGraph.AddMutableThreadTask(
@@ -2091,23 +2087,19 @@ namespace impl
 
 		// Next Task: Load Unreal Assets
 		//-------------------------------------------------------------
-		UE::Tasks::FTask Game_LoadUnrealAssets = ObjectInstancePrivateData->LoadAdditionalAssetsAndDataAsync(
-				OperationData, UCustomizableObjectSystem::GetInstance()->GetPrivate()->StreamableManager);
-		
+		UE::Tasks::FTask Game_LoadUnrealAssets = ObjectInstancePrivateData->LoadAdditionalAssetsAndDataAsync(OperationData, UCustomizableObjectSystem::GetInstance()->GetPrivate()->StreamableManager);
+
 		// Next-next Task: Convert Resources
 		//-------------------------------------------------------------
 		SystemPrivateData->AddGameThreadTask(
+			FMutableTask 
 			{
-			FMutableTaskDelegate::CreateLambda(
+				FMutableTaskDelegate::CreateLambda(
 				[OperationData]()
 				{
 					Task_Game_ConvertResources(OperationData);
 				}),
-#ifdef MUTABLE_USE_NEW_TASKGRAPH
-				{},
-#endif
-				Game_LoadUnrealAssets,
-				Mutable_GetImagesTask
+				{ Game_LoadUnrealAssets, Mutable_GetImagesTask }
 			});
 	}
 
@@ -2156,14 +2148,8 @@ namespace impl
 
 		// Task inputs
 		TSharedPtr<mu::Model> Model = CustomizableObject->GetPrivate()->GetModel();
-
-#ifdef MUTABLE_USE_NEW_TASKGRAPH
-		UE::Tasks::FTask Dependency;
-#else
-		FGraphEventRef Dependency;
-#endif
-		
-		Dependency = SystemPrivate->MutableTaskGraph.AddMutableThreadTask(
+	
+		UE::Tasks::FTask Dependency = SystemPrivate->MutableTaskGraph.AddMutableThreadTask(
 			TEXT("Task_Mutable_Update_GetMesh"),
 			[Operation, Model]()
 			{
@@ -2171,13 +2157,14 @@ namespace impl
 			});
 
 		SystemPrivate->AddGameThreadTask(
+			FMutableTask
 			{
-			FMutableTaskDelegate::CreateLambda(
+				FMutableTaskDelegate::CreateLambda(
 				[Operation]()
 				{
 					impl::Task_Game_LockCache(Operation);
 				}),
-				Dependency,
+				{ Dependency }
 			});
 	}
 	
@@ -2410,11 +2397,7 @@ namespace impl
 			Operation->SetRequestedLODs(RequestedLODs);
 		}
 
-#ifdef MUTABLE_USE_NEW_TASKGRAPH
 		UE::Tasks::FTask Mutable_GetMeshTask;
-#else
-		FGraphEventRef Mutable_GetMeshTask;
-#endif
 		
 		if (Operation->bUseMeshCache)
 		{
@@ -2426,13 +2409,14 @@ namespace impl
 				});
 
 			SystemPrivateData->AddGameThreadTask(
+				FMutableTask 
 				{
-				FMutableTaskDelegate::CreateLambda(
+					FMutableTaskDelegate::CreateLambda(
 					[Operation]()
 					{
 						impl::Task_Game_LockMeshCache(Operation);
 					}),
-					Mutable_GetMeshTask
+					{ Mutable_GetMeshTask }
 				});
 		}
 		else
@@ -2452,13 +2436,14 @@ namespace impl
 			{
 				// Task inputs
 				SystemPrivateData->AddGameThreadTask(
+					FMutableTask 
 					{
-					FMutableTaskDelegate::CreateLambda(
+						FMutableTaskDelegate::CreateLambda(
 						[Operation]()
 						{
 							impl::Task_Game_LockCache(Operation);
 						}),
-					Mutable_GetMeshTask
+						{ Mutable_GetMeshTask }
 					});
 			}
 		}
diff --git a/Engine/Plugins/Experimental/Mutable/Source/CustomizableObject/Private/MuCO/FMutableTaskGraph.cpp b/Engine/Plugins/Experimental/Mutable/Source/CustomizableObject/Private/MuCO/FMutableTaskGraph.cpp
--- a/Engine/Plugins/Experimental/Mutable/Source/CustomizableObject/Private/MuCO/FMutableTaskGraph.cpp
+++ b/Engine/Plugins/Experimental/Mutable/Source/CustomizableObject/Private/MuCO/FMutableTaskGraph.cpp
@@ -6,12 +6,7 @@
 #include "MuR/MutableTrace.h"
 
 
-#ifdef MUTABLE_USE_NEW_TASKGRAPH
-const LowLevelTasks::ETaskPriority TASKGRAPH_PRIORITY = LowLevelTasks::ETaskPriority::BackgroundHigh;
-#else
-const ENamedThreads::Type TASKGRAPH_PRIORITY = ENamedThreads::AnyBackgroundHiPriTask;
-#endif
-
+constexpr LowLevelTasks::ETaskPriority TASKGRAPH_PRIORITY = LowLevelTasks::ETaskPriority::BackgroundHigh;
 
 static TAutoConsoleVariable<float> CVarMutableTaskLowPriorityMaxWaitTime(
 	TEXT("mutable.MutableTaskLowPriorityMaxWaitTime"),
@@ -26,34 +21,15 @@ static TAutoConsoleVariable<bool> CVarEnableMutableTaskLowPriority(
 	ECVF_Scalability);
 
 
-FMutableTaskGraph::TaskType FMutableTaskGraph::AddMutableThreadTask(const TCHAR* DebugName, TUniqueFunction<void()>&& TaskBody)
+UE::Tasks::FTask FMutableTaskGraph::AddMutableThreadTask(const TCHAR* DebugName, TUniqueFunction<void()>&& TaskBody)
 {
 	FScopeLock Lock(&MutableTaskLock);
 
-#ifdef MUTABLE_USE_NEW_TASKGRAPH
-	if (LastMutableTask.IsValid())
-	{
-		LastMutableTask = UE::Tasks::Launch(DebugName, MoveTemp(TaskBody), LastMutableTask, TASKGRAPH_PRIORITY);
-	}
-	else
-	{
-		LastMutableTask = UE::Tasks::Launch(DebugName, MoveTemp(TaskBody), TASKGRAPH_PRIORITY);
-	}
-#else
-	FGraphEventArray Prerequisites;
-	if (LastMutableTask)
-	{
-		Prerequisites.Add(LastMutableTask);
-	}
+	TArray<UE::Tasks::FTask, TFixedAllocator<1>> Prerequisites = LastMutableTask.IsValid() 
+			? TArray<UE::Tasks::FTask, TFixedAllocator<1>>{LastMutableTask} 
+			: TArray<UE::Tasks::FTask, TFixedAllocator<1>>{};
 
-	LastMutableTask = FFunctionGraphTask::CreateAndDispatchWhenReady(
-		MoveTemp(TaskBody),
-		TStatId{},
-		&Prerequisites, 
-		TASKGRAPH_PRIORITY);
-	
-	LastMutableTask->SetDebugName(DebugName);
-#endif
+	LastMutableTask = UE::Tasks::Launch(DebugName, MoveTemp(TaskBody), Prerequisites, TASKGRAPH_PRIORITY);	
 	
 	return LastMutableTask;
 }
@@ -98,17 +74,7 @@ void FMutableTaskGraph::CancelMutableThreadTaskLowPriority(uint32 Id)
 
 void FMutableTaskGraph::AddAnyThreadTask(const TCHAR* DebugName, TUniqueFunction<void()>&& TaskBody) const
 {
-#ifdef MUTABLE_USE_NEW_TASKGRAPH
 	UE::Tasks::Launch(DebugName, MoveTemp(TaskBody));
-#else
-	FGraphEventRef Task = FFunctionGraphTask::CreateAndDispatchWhenReady(
-		MoveTemp(TaskBody),
-		TStatId{},
-		nullptr,
-		TASKGRAPH_PRIORITY);
-
-	Task->SetDebugName(DebugName);
-#endif
 }
 	
 
@@ -118,7 +84,6 @@ void FMutableTaskGraph::WaitForMutableTasks()
 
 	if (LastMutableTask.IsValid())
 	{
-#ifdef MUTABLE_USE_NEW_TASKGRAPH
 		if (CVarTaskGraphBusyWait->GetBool())
 		{
 			LastMutableTask.BusyWait();
@@ -127,21 +92,7 @@ void FMutableTaskGraph::WaitForMutableTasks()
 		{
 			LastMutableTask.Wait();
 		}
-#else
-		if (CVarTaskGraphBusyWait->GetBool())
-		{
-			LowLevelTasks::BusyWaitUntil(
-				[Task = LastMutableTask]()
-				{
-					return Task->IsComplete();
-				}
-			);
-		}
-		else 
-		{
-			LastMutableTask->Wait();
-		}
-#endif
+
 		LastMutableTask = {};
 	}
 }
@@ -214,13 +165,9 @@ void FMutableTaskGraph::TryLaunchMutableTaskLowPriority(bool bFromMutableTask)
 }
 
 
-bool FMutableTaskGraph::IsTaskCompleted(const TaskType& Task) const
+bool FMutableTaskGraph::IsTaskCompleted(const UE::Tasks::FTask& Task) const
 {
-#ifdef MUTABLE_USE_NEW_TASKGRAPH
 	return Task.IsCompleted();
-#else
-	return !Task.IsValid() || Task->IsComplete();
-#endif
 }
 
 
diff --git a/Engine/Plugins/Experimental/Mutable/Source/CustomizableObject/Private/MuCO/UnrealMutableImageProvider.cpp b/Engine/Plugins/Experimental/Mutable/Source/CustomizableObject/Private/MuCO/UnrealMutableImageProvider.cpp
--- a/Engine/Plugins/Experimental/Mutable/Source/CustomizableObject/Private/MuCO/UnrealMutableImageProvider.cpp
+++ b/Engine/Plugins/Experimental/Mutable/Source/CustomizableObject/Private/MuCO/UnrealMutableImageProvider.cpp
@@ -149,11 +149,7 @@ bool FUnrealMutableImageProvider::Tick()
 
 
 //-------------------------------------------------------------------------------------------------
-#ifdef MUTABLE_USE_NEW_TASKGRAPH
-	TTuple<UE::Tasks::FTask, TFunction<void()>> FUnrealMutableImageProvider::GetImageAsync(FName Id, uint8 MipmapsToSkip, TFunction<void(mu::Ptr<mu::Image>)>& ResultCallback)
-#else
-	TTuple<FGraphEventRef, TFunction<void()>> FUnrealMutableImageProvider::GetImageAsync(FName Id, uint8 MipmapsToSkip, TFunction<void(mu::Ptr<mu::Image>)>& ResultCallback)
-#endif
+TTuple<UE::Tasks::FTask, TFunction<void()>> FUnrealMutableImageProvider::GetImageAsync(FName Id, uint8 MipmapsToSkip, TFunction<void(mu::Ptr<mu::Image>)>& ResultCallback)
 {
 	// Thread: worker
 	MUTABLE_CPUPROFILER_SCOPE(FUnrealMutableImageProvider::GetImage);
@@ -168,23 +164,10 @@ bool FUnrealMutableImageProvider::Tick()
 	mu::EImageFormat MutImageFormat = mu::EImageFormat::IF_NONE;
 	int32 MutImageDataSize = 0;
 
-#ifdef MUTABLE_USE_NEW_TASKGRAPH
 	auto TrivialReturn = []() -> TTuple<UE::Tasks::FTask, TFunction<void()>>
 	{
-		UE::Tasks::FTaskEvent CompletionEvent(TEXT("GetImageAsyncCompleted"));
-		CompletionEvent.Trigger();
-
-		return MakeTuple(CompletionEvent, []() -> void {});
-	};
-#else
-	auto TrivialReturn = []() -> TTuple<FGraphEventRef, TFunction<void()>>
-	{
-		FGraphEventRef CompletionEvent = FGraphEvent::CreateGraphEvent();
-		CompletionEvent->DispatchSubsequents();
-
-		return MakeTuple(CompletionEvent, []() -> void {});
+		return MakeTuple(UE::Tasks::MakeCompletedTask<void>(), []() -> void {});
 	};
-#endif
 
 	{
 		FScopeLock Lock(&ExternalImagesLock);
@@ -283,12 +266,7 @@ bool FUnrealMutableImageProvider::Tick()
 			// Create a streaming request if the data is not loaded or copy the mip data
 			if (!BulkData.IsBulkDataLoaded())
 			{
-#ifdef MUTABLE_USE_NEW_TASKGRAPH
 				UE::Tasks::FTaskEvent IORequestCompletionEvent(TEXT("Mutable_IORequestCompletionEvent"));
-#else
-				FGraphEventRef IORequestCompletionEvent = FGraphEvent::CreateGraphEvent();
-#endif
-
 
 				TFunction<void(bool, IBulkDataIORequest*)> IOCallback =
 					[
@@ -303,15 +281,8 @@ bool FUnrealMutableImageProvider::Tick()
 				{
 					ON_SCOPE_EXIT
 					{
-#ifdef MUTABLE_USE_NEW_TASKGRAPH
 						UE::Tasks::FTaskEvent EventCopy = IORequestCompletionEvent;
 						EventCopy.Trigger();
-#else
-						if (IORequestCompletionEvent.IsValid())
-						{
-							IORequestCompletionEvent->DispatchSubsequents();
-						}
-#endif
 					};
 					
 					// Should we do someting different than returning a dummy image if cancelled?
@@ -384,14 +355,7 @@ bool FUnrealMutableImageProvider::Tick()
 				{
 					UE_LOG(LogMutable, Warning, TEXT("Failed to create an IORequest for a UTexture2D BulkData for an application-specific image parameter."));
 
-#ifdef MUTABLE_USE_NEW_TASKGRAPH
 					IORequestCompletionEvent.Trigger();
-#else
-					if (IORequestCompletionEvent.IsValid())
-					{
-						IORequestCompletionEvent->DispatchSubsequents();
-					}
-#endif
 					
 					ResultCallback(CreateDummy());
 					return Invoke(TrivialReturn);
@@ -434,32 +398,15 @@ bool FUnrealMutableImageProvider::Tick()
 
 
 //-------------------------------------------------------------------------------------------------
-#ifdef MUTABLE_USE_NEW_TASKGRAPH
 TTuple<UE::Tasks::FTask, TFunction<void()>> FUnrealMutableImageProvider::GetReferencedImageAsync(const void* ModelPtr, int32 Id, uint8 MipmapsToSkip, TFunction<void(mu::Ptr<mu::Image>)>& ResultCallback)
-#else
-TTuple<FGraphEventRef, TFunction<void()>> FUnrealMutableImageProvider::GetReferencedImageAsync(const void* ModelPtr, int32 Id, uint8 MipmapsToSkip, TFunction<void(mu::Ptr<mu::Image>)>& ResultCallback)
-#endif
 {
 	// Thread: worker
 	MUTABLE_CPUPROFILER_SCOPE(FUnrealMutableImageProvider::GetReferencedImageAsync);
 
-#ifdef MUTABLE_USE_NEW_TASKGRAPH
 	auto TrivialReturn = []() -> TTuple<UE::Tasks::FTask, TFunction<void()>>
 	{
-		UE::Tasks::FTaskEvent CompletionEvent(TEXT("GetImageAsyncCompleted"));
-		CompletionEvent.Trigger();
-
-		return MakeTuple(CompletionEvent, []() -> void {});
-	};
-#else
-	auto TrivialReturn = []() -> TTuple<FGraphEventRef, TFunction<void()>>
-	{
-		FGraphEventRef CompletionEvent = FGraphEvent::CreateGraphEvent();
-		CompletionEvent->DispatchSubsequents();
-
-		return MakeTuple(CompletionEvent, []() -> void {});
+		return MakeTuple(UE::Tasks::MakeCompletedTask<void>(), []() -> void {});
 	};
-#endif
 
 
 #if WITH_EDITOR
diff --git a/Engine/Plugins/Experimental/Mutable/Source/CustomizableObject/Private/MuCO/UnrealMutableImageProvider.h b/Engine/Plugins/Experimental/Mutable/Source/CustomizableObject/Private/MuCO/UnrealMutableImageProvider.h
--- a/Engine/Plugins/Experimental/Mutable/Source/CustomizableObject/Private/MuCO/UnrealMutableImageProvider.h
+++ b/Engine/Plugins/Experimental/Mutable/Source/CustomizableObject/Private/MuCO/UnrealMutableImageProvider.h
@@ -9,10 +9,7 @@
 #include "Containers/Map.h"
 #include "Containers/Queue.h"
 
-#ifdef MUTABLE_USE_NEW_TASKGRAPH
 #include "Tasks/Task.h"
-#else
-#endif
 
 class UTexture2D;
 
@@ -24,13 +21,8 @@ class FUnrealMutableImageProvider : public mu::ImageParameterGenerator, public F
 public:
 	// mu::ImageParameterGenerator interface
 	// Thread: worker
-#ifdef MUTABLE_USE_NEW_TASKGRAPH
 	virtual TTuple<UE::Tasks::FTask, TFunction<void()>> GetImageAsync(FName Id, uint8 MipmapsToSkip, TFunction<void(mu::Ptr<mu::Image>)>& ResultCallback) override;
 	virtual TTuple<UE::Tasks::FTask, TFunction<void()>> GetReferencedImageAsync(const void* ModelPtr, int32 Id, uint8 MipmapsToSkip, TFunction<void(mu::Ptr<mu::Image>)>& ResultCallback) override;
-#else
-	virtual TTuple<FGraphEventRef, TFunction<void()>> GetImageAsync(FName Id, uint8 MipmapsToSkip, TFunction<void(mu::Ptr<mu::Image>)>& ResultCallback) override;
-	virtual TTuple<FGraphEventRef, TFunction<void()>> GetReferencedImageAsync(const void* ModelPtr, int32 Id, uint8 MipmapsToSkip, TFunction<void(mu::Ptr<mu::Image>)>& ResultCallback) override;
-#endif
 
 	virtual mu::FImageDesc GetImageDesc(FName Id, uint8 MipmapsToSkip) override;
 
diff --git a/Engine/Plugins/Experimental/Mutable/Source/CustomizableObjectEditor/Private/MuCOE/SMutableCodeViewer.cpp b/Engine/Plugins/Experimental/Mutable/Source/CustomizableObjectEditor/Private/MuCOE/SMutableCodeViewer.cpp
--- a/Engine/Plugins/Experimental/Mutable/Source/CustomizableObjectEditor/Private/MuCOE/SMutableCodeViewer.cpp
+++ b/Engine/Plugins/Experimental/Mutable/Source/CustomizableObjectEditor/Private/MuCOE/SMutableCodeViewer.cpp
@@ -2499,11 +2499,7 @@ namespace
 	public:
 
 
-#ifdef MUTABLE_USE_NEW_TASKGRAPH
 		TTuple<UE::Tasks::FTask, TFunction<void()>> GetImageAsync(FName Id, uint8 MipmapsToSkip, TFunction<void(mu::Ptr<mu::Image>)>& ResultCallback) override
-#else
-		TTuple<FGraphEventRef, TFunction<void()>> GetImageAsync(FName Id, uint8 MipmapsToSkip, TFunction<void(mu::Ptr<mu::Image>)>& ResultCallback) override
-#endif
 		{
 			MUTABLE_CPUPROFILER_SCOPE(TestImageProvider_GetImage);
 
@@ -2542,15 +2538,7 @@ namespace
 
 			ResultCallback(Image);
 
-#ifdef MUTABLE_USE_NEW_TASKGRAPH
-			UE::Tasks::FTaskEvent CompletionEvent(TEXT("TestImageProvider_GetImageAsunc_Completed"));
-			CompletionEvent.Trigger();
-#else
-			FGraphEventRef CompletionEvent = FGraphEvent::CreateGraphEvent();
-			CompletionEvent->DispatchSubsequents();
-#endif
-
-			return MakeTuple(CompletionEvent, []() -> void {});
+			return MakeTuple(UE::Tasks::MakeCompletedTask<void>(), []() -> void {});
 		}
 
 		mu::FImageDesc GetImageDesc(FName Id, uint8 MipmapsToSkip) override
@@ -2559,12 +2547,7 @@ namespace
 		}
 
 
-		//-------------------------------------------------------------------------------------------------
-#ifdef MUTABLE_USE_NEW_TASKGRAPH
 		TTuple<UE::Tasks::FTask, TFunction<void()>> GetReferencedImageAsync(const void* ModelPtr, int32 Id, uint8 MipmapsToSkip, TFunction<void(mu::Ptr<mu::Image>)>& ResultCallback)
-#else
-		TTuple<FGraphEventRef, TFunction<void()>> GetReferencedImageAsync(const void* ModelPtr, int32 Id, uint8 MipmapsToSkip, TFunction<void(mu::Ptr<mu::Image>)>& ResultCallback)
-#endif
 		{
 			check(ReferencedTextures.IsValidIndex(Id));
 
@@ -2583,23 +2566,10 @@ namespace
 
 			ResultCallback(ResultImage);
 
-#ifdef MUTABLE_USE_NEW_TASKGRAPH
 			auto TrivialReturn = []() -> TTuple<UE::Tasks::FTask, TFunction<void()>>
 			{
-				UE::Tasks::FTaskEvent CompletionEvent(TEXT("GetImageAsyncCompleted"));
-				CompletionEvent.Trigger();
-
-				return MakeTuple(CompletionEvent, []() -> void {});
+				return MakeTuple(UE::Tasks::MakeCompletedTask<void>(), []() -> void {});
 			};
-#else
-			auto TrivialReturn = []() -> TTuple<FGraphEventRef, TFunction<void()>>
-			{
-				FGraphEventRef CompletionEvent = FGraphEvent::CreateGraphEvent();
-				CompletionEvent->DispatchSubsequents();
-
-				return MakeTuple(CompletionEvent, []() -> void {});
-			};
-#endif
 
 			return Invoke(TrivialReturn);
 		}
diff --git a/Engine/Plugins/Experimental/Mutable/Source/MutableRuntime/Private/MuR/CodeRunner.cpp b/Engine/Plugins/Experimental/Mutable/Source/MutableRuntime/Private/MuR/CodeRunner.cpp
--- a/Engine/Plugins/Experimental/Mutable/Source/MutableRuntime/Private/MuR/CodeRunner.cpp
+++ b/Engine/Plugins/Experimental/Mutable/Source/MutableRuntime/Private/MuR/CodeRunner.cpp
@@ -149,12 +149,7 @@ namespace mu
 	}
 
 
-    //---------------------------------------------------------------------------------------------
-#ifdef MUTABLE_USE_NEW_TASKGRAPH
 	TTuple<UE::Tasks::FTask, TFunction<void()>> CodeRunner::LoadExternalImageAsync(FExternalImageId Id, uint8 MipmapsToSkip, TFunction<void(Ptr<Image>)>& ResultCallback)
-#else
-	TTuple<FGraphEventRef, TFunction<void()>> CodeRunner::LoadExternalImageAsync(FExternalImageId Id, uint8 MipmapsToSkip, TFunction<void(Ptr<Image>)>& ResultCallback)
-#endif
     {
 		MUTABLE_CPUPROFILER_SCOPE(LoadExternalImageAsync);
 
@@ -179,16 +174,7 @@ namespace mu
 			check(false);
 		}
 
-		// Not needed as it should never reach this point, but added for correctness.
-#ifdef MUTABLE_USE_NEW_TASKGRAPH
-		UE::Tasks::FTaskEvent CompletionEvent(TEXT("LoadExternalImageAsyncCompletion"));
-		CompletionEvent.Trigger();
-#else
-		FGraphEventRef CompletionEvent = FGraphEvent::CreateGraphEvent();
-		CompletionEvent->DispatchSubsequents();
-#endif
-
-		return MakeTuple(CompletionEvent, []() -> void {});
+		return MakeTuple(UE::Tasks::MakeCompletedTask<void>(), []() -> void {});
 	}
 
 	
diff --git a/Engine/Plugins/Experimental/Mutable/Source/MutableRuntime/Private/MuR/CodeRunner.h b/Engine/Plugins/Experimental/Mutable/Source/MutableRuntime/Private/MuR/CodeRunner.h
--- a/Engine/Plugins/Experimental/Mutable/Source/MutableRuntime/Private/MuR/CodeRunner.h
+++ b/Engine/Plugins/Experimental/Mutable/Source/MutableRuntime/Private/MuR/CodeRunner.h
@@ -2,7 +2,6 @@
 
 #pragma once
 
-#include "Async/TaskGraphInterfaces.h"
 #include "Containers/Array.h"
 #include "Containers/ContainerAllocationPolicies.h"
 #include "HAL/PlatformCrt.h"
@@ -19,17 +18,11 @@
 #include "MuR/System.h"
 #include "MuR/SystemPrivate.h"
 #include "MuR/Types.h"
+#include "Tasks/Task.h"
 #include "Templates/RefCounting.h"
 #include "Templates/SharedPointer.h"
 #include "Templates/Tuple.h"
 
-// This define could come from MuR/System.h
-#ifdef MUTABLE_USE_NEW_TASKGRAPH
-	#include "Tasks/Task.h"
-#else
-	#include "Async/TaskGraphFwd.h"
-#endif
-
 namespace mu::MemoryCounters
 {
 	struct FStreamingMemoryCounterTag {};
@@ -42,20 +35,6 @@ namespace  mu
 	class Parameters;
 	class RangeIndex;
 
-	using EventType = 
-#ifdef MUTABLE_USE_NEW_TASKGRAPH
-		UE::Tasks::FTaskEvent;
-#else
-		FGraphEventRef;
-#endif
-
-	using TaskType = 
-#ifdef MUTABLE_USE_NEW_TASKGRAPH
-		UE::Tasks::FTask;
-#else
-		FGraphEventRef;
-#endif
-
     /** Code execution of the mutable virtual machine. */
     class CodeRunner
     {
@@ -147,11 +126,7 @@ namespace  mu
 
 		//! Load an external image asynchronously, retuns an event to wait for complition and a cleanup function 
 		//! that must be called once the event has completed.
-#ifdef MUTABLE_USE_NEW_TASKGRAPH
 		TTuple<UE::Tasks::FTask, TFunction<void()>> LoadExternalImageAsync(FExternalImageId Id, uint8 MipmapsToSkip, TFunction<void(Ptr<Image>)>& ResultCallback);
-#else
-		TTuple<FGraphEventRef, TFunction<void()>> LoadExternalImageAsync(FExternalImageId Id, uint8 MipmapsToSkip, TFunction<void(Ptr<Image>)>& ResultCallback);
-#endif
  	    mu::FImageDesc GetExternalImageDesc(FName Id, uint8 MipmapsToSkip);
 
 		/** Settings that may affect the execution of some operations, like image conversion quality. */
@@ -221,11 +196,7 @@ namespace  mu
 		public:
 			const FScheduledOp Op;
 
-#ifdef MUTABLE_USE_NEW_TASKGRAPH
 			UE::Tasks::FTask Event = {};
-#else
-			FGraphEventRef Event = nullptr;
-#endif
 
 			FIssuedTask(const FScheduledOp& InOp) : Op(InOp) {}
 			virtual ~FIssuedTask() {}
@@ -236,12 +207,7 @@ namespace  mu
 			virtual void Complete(CodeRunner*) = 0;
 			virtual bool IsComplete(CodeRunner*)
 			{ 
-				// Event can be null if we forced single-threaded execution.
-#ifdef MUTABLE_USE_NEW_TASKGRAPH
 				return !Event.IsValid() || Event.IsCompleted(); 
-#else
-				return !Event.IsValid() || Event->IsComplete();
-#endif
 			}
 		};
 
@@ -252,7 +218,7 @@ namespace  mu
 			int32 RomIndex = -1;
 			ModelReader::OPERATION_ID m_streamID = -1;
 			StreamingDataContainerType m_streamBuffer;
-			TaskType Event;
+			UE::Tasks::FTask Event;
 		};
 
 		class FLoadMeshRomTask : public CodeRunner::FIssuedTask
diff --git a/Engine/Plugins/Experimental/Mutable/Source/MutableRuntime/Private/MuR/CodeRunnerTasked.cpp b/Engine/Plugins/Experimental/Mutable/Source/MutableRuntime/Private/MuR/CodeRunnerTasked.cpp
--- a/Engine/Plugins/Experimental/Mutable/Source/MutableRuntime/Private/MuR/CodeRunnerTasked.cpp
+++ b/Engine/Plugins/Experimental/Mutable/Source/MutableRuntime/Private/MuR/CodeRunnerTasked.cpp
@@ -263,45 +263,9 @@ namespace mu
 			}
 			else
 			{
-#ifdef MUTABLE_USE_NEW_TASKGRAPH
 				TaskToIssue->Event = UE::Tasks::Launch(TEXT("MutableCore_Task"),
 					[TaskToIssue]() { TaskToIssue->DoWork(); },
 					UE::Tasks::ETaskPriority::Inherit);
-#else
-
-				ENamedThreads::Type Priority;
-				switch (CoreRunnerTaskPriority)
-				{
-				default:
-				case 0:
-					Priority = ENamedThreads::AnyThread;
-					break;
-				case 1:
-					Priority = ENamedThreads::AnyHiPriThreadHiPriTask;
-					break;
-				case 2:
-					Priority = ENamedThreads::AnyHiPriThreadNormalTask;
-					break;
-				case 3:
-					Priority = ENamedThreads::AnyNormalThreadHiPriTask;
-					break;
-				case 4:
-					Priority = ENamedThreads::AnyNormalThreadNormalTask;
-					break;
-				case 5:
-					Priority = ENamedThreads::AnyBackgroundHiPriTask;
-					break;
-				case 6:
-					Priority = ENamedThreads::AnyBackgroundThreadNormalTask;
-					break;
-				}
-
-				TaskToIssue->Event = FFunctionGraphTask::CreateAndDispatchWhenReady(
-					[TaskToIssue]() { TaskToIssue->DoWork(); },
-					TStatId{},
-					nullptr,
-					Priority);
-#endif
 			}
 		}
 
@@ -549,7 +513,6 @@ namespace mu
 				{
 					if (IssuedTasks[IssuedIndex]->Event.IsValid())
 					{
-#ifdef MUTABLE_USE_NEW_TASKGRAPH
 						if (CVarTaskGraphBusyWait->GetBool())
 						{
 							IssuedTasks[IssuedIndex]->Event.BusyWait();							
@@ -558,21 +521,7 @@ namespace mu
 						{
 							IssuedTasks[IssuedIndex]->Event.Wait();
 						}
-#else
-						if (CVarTaskGraphBusyWait->GetBool())
-						{
-							LowLevelTasks::BusyWaitUntil(
-								[Task = IssuedTasks[IssuedIndex]->Event]()
-								{
-									return Task->IsComplete();
-								}
-							);
-						}
-						else 
-						{
-							IssuedTasks[IssuedIndex]->Event->Wait();
-						}
-#endif
+
 						break;
 					}
 				}
@@ -2090,22 +2039,12 @@ namespace mu
 		RomLoadOp.m_streamBuffer.SetNumUninitialized(RomSize);
 		const int32 SizeAfter = RomLoadOp.m_streamBuffer.GetAllocatedSize();
 
-		EventType EventTyped =
-#ifdef MUTABLE_USE_NEW_TASKGRAPH
-			UE::Tasks::FTaskEvent(TEXT("FLoadMeshRomTask"));
-#else
-			FGraphEvent::CreateGraphEvent();		
-#endif
-		
-		RomLoadOp.Event = EventTyped;
+	 	UE::Tasks::FTaskEvent ReadCompletionEvent = UE::Tasks::FTaskEvent(TEXT("FLoadMeshRomTask"));	
+		RomLoadOp.Event = ReadCompletionEvent;
 		
-		TFunction<void(bool)> Callback = [EventTyped](bool bSuccess) mutable
+		TFunction<void(bool)> Callback = [ReadCompletionEvent](bool bSuccess) mutable
 		{
-#ifdef MUTABLE_USE_NEW_TASKGRAPH
-			EventTyped.Trigger();
-#else
-			EventTyped->DispatchSubsequents();
-#endif				
+			ReadCompletionEvent.Trigger();
 		};
 		
 		const uint32 RomId = Program.m_roms[RomIndex].Id;
@@ -2116,7 +2055,7 @@ namespace mu
 			return false;
 		}
 
-		Event = EventTyped; // Wait for read operation to end
+		Event = ReadCompletionEvent; // Wait for read operation to end
 		return false; // No worker thread work
 	}
 	
@@ -2265,11 +2204,7 @@ namespace mu
 
 		FWorkingMemoryManager::FModelCacheEntry* ModelCache = Runner->m_pSystem->WorkingMemoryManager.FindModelCache(Runner->m_pModel.Get());
 
-#ifdef MUTABLE_USE_NEW_TASKGRAPH
 		TArray<UE::Tasks::FTask> ReadCompleteEvents;
-#else
-		FGraphEventArray ReadCompleteEvents;
-#endif
 
 		ReadCompleteEvents.Reserve(LODIndexCount); 
 
@@ -2322,23 +2257,13 @@ namespace mu
 			RomLoadOp.m_streamBuffer.SetNumUninitialized(RomSize);
 			const int32 SizeAfter = RomLoadOp.m_streamBuffer.GetAllocatedSize();
 
-			EventType EventTyped = 
-#ifdef MUTABLE_USE_NEW_TASKGRAPH
-				UE::Tasks::FTaskEvent(TEXT("FLoadImageRomsTaskRom"));
-#else
-				FGraphEvent::CreateGraphEvent();
-#endif
+			UE::Tasks::FTaskEvent ReadCompletionEvent(TEXT("FLoadImageRomsTaskRom"));
+			ReadCompleteEvents.Add(ReadCompletionEvent);
+			RomLoadOp.Event = ReadCompletionEvent;
 
-			RomLoadOp.Event = EventTyped;
-			ReadCompleteEvents.Add(EventTyped);
-
-			TFunction<void(bool)> Callback = [EventTyped](bool bSuccess) mutable // Mutable due Trigger not being const
+			TFunction<void(bool)> Callback = [ReadCompletionEvent](bool bSuccess) mutable // Mutable due Trigger not being const
 			{
-#ifdef MUTABLE_USE_NEW_TASKGRAPH
-				EventTyped.Trigger();
-#else
-				EventTyped->DispatchSubsequents();
-#endif				
+				ReadCompletionEvent.Trigger();
 			};
 			
 			const uint32 RomId = program.m_roms[RomIndex].Id;
@@ -2351,19 +2276,11 @@ namespace mu
 		}
 
 		// Wait for all read operations to end
-#ifdef MUTABLE_USE_NEW_TASKGRAPH
-		UE::Tasks::FTaskEvent EventTyped(TEXT("FLoadImageRomsTask"));
-		EventTyped.AddPrerequisites(ReadCompleteEvents);
-		EventTyped.Trigger();
-		Event = EventTyped;
-#else
-		Event = FGraphEvent::CreateGraphEvent();
-		for (const FGraphEventRef& ReadCompleteEvent : ReadCompleteEvents)
-		{
-			Event->DontCompleteUntil(ReadCompleteEvent);
-		}
-		Event->DispatchSubsequents();
-#endif
+		UE::Tasks::FTaskEvent GatherReadsCompletionEvent(TEXT("FLoadImageRomsTask"));
+		GatherReadsCompletionEvent.AddPrerequisites(ReadCompleteEvents);
+		GatherReadsCompletionEvent.Trigger();
+
+		Event = GatherReadsCompletionEvent;
 			
 		return false; // No worker thread work
 	}
diff --git a/Engine/Plugins/Experimental/Mutable/Source/MutableRuntime/Public/MuR/System.h b/Engine/Plugins/Experimental/Mutable/Source/MutableRuntime/Public/MuR/System.h
--- a/Engine/Plugins/Experimental/Mutable/Source/MutableRuntime/Public/MuR/System.h
+++ b/Engine/Plugins/Experimental/Mutable/Source/MutableRuntime/Public/MuR/System.h
@@ -2,7 +2,6 @@
 
 #pragma once
 
-#include "Async/TaskGraphInterfaces.h"
 #include "HAL/Platform.h"
 #include "MuR/Image.h"
 #include "MuR/Instance.h"
@@ -11,20 +10,15 @@
 #include "MuR/RefCounted.h"
 #include "MuR/Settings.h"
 #include "MuR/Types.h"
+#include "Tasks/Task.h"
 #include "Templates/Tuple.h"
 
-/** This define will use the newer task graph interface to manage mutable concurrency. It has not been fully tested. */
-#define MUTABLE_USE_NEW_TASKGRAPH
 
 /** If set to 1, this enables some expensive Unreal Insights traces, but can lead to 5x slower mutable operation. 
 * Other cheaper traces are enabled at all times.
 */
 #define UE_MUTABLE_ENABLE_SLOW_TRACES	0
 
-#ifdef MUTABLE_USE_NEW_TASKGRAPH
-#include "Tasks/Task.h"
-#else
-#endif
 
 #include "System.generated.h"
 
@@ -92,13 +86,8 @@ namespace mu
         virtual ~ImageParameterGenerator() = default;
 
         //! Returns the completion event and a cleanup function that must be called once event is completed.
-#ifdef MUTABLE_USE_NEW_TASKGRAPH
 		virtual TTuple<UE::Tasks::FTask, TFunction<void()>> GetImageAsync(FName Id, uint8 MipmapsToSkip, TFunction<void(Ptr<Image>)>& ResultCallback) = 0;
 		virtual TTuple<UE::Tasks::FTask, TFunction<void()>> GetReferencedImageAsync(const void* ModelPtr, int32 Id, uint8 MipmapsToSkip, TFunction<void(Ptr<Image>)>& ResultCallback) { check(false); return {}; }
-#else
-		virtual TTuple<FGraphEventRef, TFunction<void()>> GetImageAsync(FName Id, uint8 MipmapsToSkip, TFunction<void(Ptr<Image>)>& ResultCallback) = 0;
-		virtual TTuple<FGraphEventRef, TFunction<void()>> GetReferencedImageAsync(const void* ModelPtr, int32 Id, uint8 MipmapsToSkip, TFunction<void(Ptr<Image>)>& ResultCallback) { check(false); return {}; };
-#endif
 
         virtual mu::FImageDesc GetImageDesc(FName Id, uint8 MipmapsToSkip) = 0;
     };
