[Mutable] Remove SkeletalMesh dependency in some of the conversion functions from the UnrealConversionUtils.
It must be removed to allow calling them from a non-game thread.

#rnx

diff --git a/Engine/Plugins/Experimental/Mutable/Source/CustomizableObject/Internal/MuCO/UnrealConversionUtils.h b/Engine/Plugins/Experimental/Mutable/Source/CustomizableObject/Internal/MuCO/UnrealConversionUtils.h
--- a/Engine/Plugins/Experimental/Mutable/Source/CustomizableObject/Internal/MuCO/UnrealConversionUtils.h
+++ b/Engine/Plugins/Experimental/Mutable/Source/CustomizableObject/Internal/MuCO/UnrealConversionUtils.h
@@ -38,7 +38,7 @@ namespace UnrealConversionUtils
 	 * @param InBoneMap - The bones to be set as part of the sections.
 	 */
 	CUSTOMIZABLEOBJECT_API void SetupRenderSections(
-		const USkeletalMesh* OutSkeletalMesh,
+		FSkeletalMeshLODRenderData& LODResource,
 		const mu::Ptr<const mu::Mesh> InMutableMesh,
 		const int32 InMeshLODIndex,
 		const TArray<uint16>& InBoneMap,
@@ -51,9 +51,9 @@ namespace UnrealConversionUtils
 	 * @param MeshLODIndex - The LOD index we are working with.
 	 */
 	CUSTOMIZABLEOBJECT_API void CopyMutableVertexBuffers(
-		USkeletalMesh* OutSkeletalMesh,
+		FSkeletalMeshLODRenderData& LODResource,
 		const mu::Ptr<const mu::Mesh> InMutableMesh,
-		const int32 InMeshLODIndex);
+		const bool bAllowCPUAccess);
 
 	
 	/**
@@ -62,8 +62,9 @@ namespace UnrealConversionUtils
 	 * @param OutLODModel - The LOD model to be updated.
 	 * @return True if the operation could be performed successfully, false if not.
 	 */
-	CUSTOMIZABLEOBJECT_API bool CopyMutableIndexBuffers(mu::Ptr<const mu::Mesh> InMutableMesh,
-	                                                    FSkeletalMeshLODRenderData& OutLODModel);
+	CUSTOMIZABLEOBJECT_API bool CopyMutableIndexBuffers(
+		FSkeletalMeshLODRenderData& LODResource,
+		const mu::Ptr<const mu::Mesh> InMutableMesh);
 	
 
 	/**
@@ -75,7 +76,7 @@ namespace UnrealConversionUtils
 	 * @return True if the operation could be performed successfully,	 false if not.
 	 */
 	CUSTOMIZABLEOBJECT_API void CopyMutableSkinWeightProfilesBuffers(
-		FSkeletalMeshLODRenderData& OutLODModel,
+		FSkeletalMeshLODRenderData& LODResource,
 		const FName InProfileName,
 		const mu::FMeshBufferSet& InMutableMeshVertexBuffers,
 		const int32 InBoneIndexBuffer);
diff --git a/Engine/Plugins/Experimental/Mutable/Source/CustomizableObject/Private/MuCO/CustomizableObjectInstance.cpp b/Engine/Plugins/Experimental/Mutable/Source/CustomizableObject/Private/MuCO/CustomizableObjectInstance.cpp
--- a/Engine/Plugins/Experimental/Mutable/Source/CustomizableObject/Private/MuCO/CustomizableObjectInstance.cpp
+++ b/Engine/Plugins/Experimental/Mutable/Source/CustomizableObject/Private/MuCO/CustomizableObjectInstance.cpp
@@ -4683,21 +4683,27 @@ bool UCustomizableInstancePrivate::BuildOrCopyRenderData(const TSharedRef<FUpdat
 
 		SetLastMeshId(Component->Id, LODIndex, Component->MeshID);
 
+		FSkeletalMeshLODRenderData& LODResource = SkeletalMesh->GetResourceForRendering()->LODRenderData[LODIndex];
+
 		UnrealConversionUtils::SetupRenderSections(
-			SkeletalMesh,
+			LODResource,
 			Component->Mesh,
 			LODIndex,
 			OperationData->InstanceUpdateData.BoneMaps,
 			Component->FirstBoneMap);
 
 		UnrealConversionUtils::CopyMutableVertexBuffers(
-			SkeletalMesh,
+			LODResource,
 			Component->Mesh,
-			LODIndex);
+			SkeletalMesh->GetLODInfo(LODIndex)->bAllowCPUAccess);
+
 
-		FSkeletalMeshLODRenderData& LODModel = SkeletalMesh->GetResourceForRendering()->LODRenderData[LODIndex];
+		if (LODResource.StaticVertexBuffers.ColorVertexBuffer.GetNumVertices())
+		{
+			SkeletalMesh->SetHasVertexColors(true);
+		}
 
-		if (LODModel.DoesVertexBufferUse16BitBoneIndex() && !UCustomizableObjectSystem::GetInstance()->IsSupport16BitBoneIndexEnabled())
+		if (LODResource.DoesVertexBufferUse16BitBoneIndex() && !UCustomizableObjectSystem::GetInstance()->IsSupport16BitBoneIndexEnabled())
 		{
 			OperationData->UpdateResult = EUpdateResult::Error16BitBoneIndex;
 
@@ -4715,8 +4721,8 @@ bool UCustomizableInstancePrivate::BuildOrCopyRenderData(const TSharedRef<FUpdat
 		}
 
 		// Update active and required bones
-		LODModel.ActiveBoneIndices.Append(Component->ActiveBones);
-		LODModel.RequiredBones.Append(Component->ActiveBones);
+		LODResource.ActiveBoneIndices.Append(Component->ActiveBones);
+		LODResource.RequiredBones.Append(Component->ActiveBones);
 
 		if (!ModelResources.SkinWeightProfilesInfo.IsEmpty())
 		{
@@ -4741,7 +4747,7 @@ bool UCustomizableInstancePrivate::BuildOrCopyRenderData(const TSharedRef<FUpdat
 
 				if (!bHasSkinWeightProfiles)
 				{
-					LODModel.SkinWeightProfilesData.Init(&LODModel.SkinWeightVertexBuffer);
+					LODResource.SkinWeightProfilesData.Init(&LODResource.SkinWeightVertexBuffer);
 					bHasSkinWeightProfiles = true;
 				}
 
@@ -4756,7 +4762,7 @@ bool UCustomizableInstancePrivate::BuildOrCopyRenderData(const TSharedRef<FUpdat
 				}
 
 				UnrealConversionUtils::CopyMutableSkinWeightProfilesBuffers(
-					LODModel,
+					LODResource,
 					Profile.Name,
 					MutableMeshVertexBuffers,
 					BoneIndicesBufferIndex);
@@ -4764,7 +4770,7 @@ bool UCustomizableInstancePrivate::BuildOrCopyRenderData(const TSharedRef<FUpdat
 		}
 
 		// Copy indices.
-		if (!UnrealConversionUtils::CopyMutableIndexBuffers(Component->Mesh, LODModel))
+		if (!UnrealConversionUtils::CopyMutableIndexBuffers(LODResource, Component->Mesh))
 		{
 			// End with failure
 			return false;
@@ -4772,13 +4778,13 @@ bool UCustomizableInstancePrivate::BuildOrCopyRenderData(const TSharedRef<FUpdat
 
 		// Update LOD and streaming data
 		const FMutableRefLODRenderData& RefLODRenderData = ModelResources.ReferenceSkeletalMeshesData[Component->Id].LODData[LODIndex].RenderData;
-		LODModel.bIsLODOptional = RefLODRenderData.bIsLODOptional;
-		LODModel.bStreamedDataInlined = RefLODRenderData.bStreamedDataInlined;
+		LODResource.bIsLODOptional = RefLODRenderData.bIsLODOptional;
+		LODResource.bStreamedDataInlined = RefLODRenderData.bStreamedDataInlined;
 
 		// WARNING! BufferSize must be > 0 or all texture mips in this LOD will be requested at once. 
 		// USkeletalMesh::IsMaterialUsed checks this size to see if a material is being used. If it fails the textures used by it won't be included 
 		// in the map of textures to stream.
-		LODModel.BuffersSize = 1;
+		LODResource.BuffersSize = 1;
 
 		LastValidLODIndex = LODIndex;
 	}
diff --git a/Engine/Plugins/Experimental/Mutable/Source/CustomizableObject/Private/MuCO/UnrealConversionUtils.cpp b/Engine/Plugins/Experimental/Mutable/Source/CustomizableObject/Private/MuCO/UnrealConversionUtils.cpp
--- a/Engine/Plugins/Experimental/Mutable/Source/CustomizableObject/Private/MuCO/UnrealConversionUtils.cpp
+++ b/Engine/Plugins/Experimental/Mutable/Source/CustomizableObject/Private/MuCO/UnrealConversionUtils.cpp
@@ -44,7 +44,7 @@ namespace UnrealConversionUtils
 			FMemory::Memcpy(OutVertexBuffers.PositionVertexBuffer.GetVertexData(), InMutablePositionData, NumVertices * OutVertexBuffers.PositionVertexBuffer.GetStride());
 
 			// tangent and texture coords
-			OutVertexBuffers.StaticMeshVertexBuffer.SetUseFullPrecisionUVs(true);
+			OutVertexBuffers.StaticMeshVertexBuffer.SetUseFullPrecisionUVs(bUseFullPrecisionUVs);
 			OutVertexBuffers.StaticMeshVertexBuffer.SetUseHighPrecisionTangentBasis(false);
 			OutVertexBuffers.StaticMeshVertexBuffer.Init(NumVertices, NumTexCoords, bNeedCPUAccess);
 			FMemory::Memcpy(OutVertexBuffers.StaticMeshVertexBuffer.GetTangentData(), InMutableTangentData, OutVertexBuffers.StaticMeshVertexBuffer.GetTangentSize());
@@ -128,7 +128,7 @@ namespace UnrealConversionUtils
 
 
 	void SetupRenderSections(
-		const USkeletalMesh* OutSkeletalMesh,
+		FSkeletalMeshLODRenderData& LODResource,
 		const mu::MeshPtrConst InMutableMesh,
 		const int32 MeshLODIndex,
 		const TArray<uint16>& InBoneMap,
@@ -162,7 +162,7 @@ namespace UnrealConversionUtils
 			int32 BoneCount;
 			bool bCastShadow;
 			InMutableMesh->GetSurface(SurfaceIndex, &FirstVertex, &VertexCount, &FirstIndex, &IndexCount, &FirstBone, &BoneCount, &bCastShadow);
-			FSkelMeshRenderSection& Section = OutSkeletalMesh->GetResourceForRendering()->LODRenderData[MeshLODIndex].RenderSections[SurfaceIndex];
+			FSkelMeshRenderSection& Section = LODResource.RenderSections[SurfaceIndex];
 
 			Section.DuplicatedVerticesBuffer.Init(1, TMap<int, TArray<int32>>());
 
@@ -194,32 +194,28 @@ namespace UnrealConversionUtils
 
 
 	void CopyMutableVertexBuffers(
-		USkeletalMesh* SkeletalMesh,
+		FSkeletalMeshLODRenderData& LODResource,
 		const mu::MeshPtrConst MutableMesh,
-		const int32 MeshLODIndex)
+		const bool bAllowCPUAccess)
 
 	{
 		MUTABLE_CPUPROFILER_SCOPE(CopyMutableVertexBuffers);
 
-		FSkeletalMeshLODRenderData& LODModel = SkeletalMesh->GetResourceForRendering()->LODRenderData[MeshLODIndex];
 		const mu::FMeshBufferSet& MutableMeshVertexBuffers = MutableMesh->GetVertexBuffers();
 		const int32 NumVertices = MutableMeshVertexBuffers.GetElementCount();
 
-		MUTABLE_CPUPROFILER_SCOPE(UpdateSkeletalMesh_SurfaceLoop_MemCpy);
+		//const ESkeletalMeshVertexFlags BuildFlags = SkeletalMesh->GetVertexBufferFlags();
+		//const bool bUseFullPrecisionUVs = EnumHasAllFlags(BuildFlags, ESkeletalMeshVertexFlags::UseFullPrecisionUVs);
+		const bool bUseFullPrecisionUVs = true;
 
-		const ESkeletalMeshVertexFlags BuildFlags = SkeletalMesh->GetVertexBufferFlags();
-		const bool bUseFullPrecisionUVs = EnumHasAllFlags(BuildFlags, ESkeletalMeshVertexFlags::UseFullPrecisionUVs);
-		const bool bHasVertexColors = EnumHasAllFlags(BuildFlags, ESkeletalMeshVertexFlags::HasVertexColors);
 		const int NumTexCoords = MutableMeshVertexBuffers.GetBufferChannelCount(MUTABLE_VERTEXBUFFER_TEXCOORDS);
 
-		const bool bNeedsCPUAccess = SkeletalMesh->GetLODInfoArray()[MeshLODIndex].bAllowCPUAccess;
-
 		FStaticMeshVertexBuffers_InitWithMutableData(
-			LODModel.StaticVertexBuffers,
+			LODResource.StaticVertexBuffers,
 			NumVertices,
 			NumTexCoords,
 			bUseFullPrecisionUVs,
-			bNeedsCPUAccess,
+			bAllowCPUAccess,
 			MutableMeshVertexBuffers.GetBufferData(MUTABLE_VERTEXBUFFER_POSITION),
 			MutableMeshVertexBuffers.GetBufferData(MUTABLE_VERTEXBUFFER_TANGENT),
 			MutableMeshVertexBuffers.GetBufferData(MUTABLE_VERTEXBUFFER_TEXCOORDS)
@@ -248,21 +244,21 @@ namespace UnrealConversionUtils
 
 		if (BoneIndexFormat == mu::MBF_UINT16)
 		{
-			LODModel.SkinWeightVertexBuffer.SetUse16BitBoneIndex(true);
+			LODResource.SkinWeightVertexBuffer.SetUse16BitBoneIndex(true);
 		}
 
 		if (BoneWeightFormat == mu::MBF_NUINT16)
 		{
-			LODModel.SkinWeightVertexBuffer.SetUse16BitBoneWeight(true);
+			LODResource.SkinWeightVertexBuffer.SetUse16BitBoneWeight(true);
 		}
 
 		// Init skin weight buffer
 		FSkinWeightVertexBuffer_InitWithMutableData(
-			LODModel.SkinWeightVertexBuffer,
+			LODResource.SkinWeightVertexBuffer,
 			NumVertices,
 			NumBoneInfluences * NumVertices,
 			NumBoneInfluences,
-			bNeedsCPUAccess,
+			bAllowCPUAccess,
 			MutableMeshVertexBuffers.GetBufferData(BoneIndexBuffer),
 			MutableMeshVertexBuffers.GetBufferDataSize(BoneIndexBuffer)
 		);
@@ -285,13 +281,13 @@ namespace UnrealConversionUtils
 					const int32 NumBones = BonesPerVertex * NumVertices;
 
 					check(FGPUBaseSkinVertexFactory::UseUnlimitedBoneInfluences(BonesPerVertex) ||
-						!LODModel.SkinWeightVertexBuffer.GetVariableBonesPerVertex());
+						!LODResource.SkinWeightVertexBuffer.GetVariableBonesPerVertex());
 					FSkinWeightVertexBuffer_InitWithMutableData(
-						LODModel.SkinWeightVertexBuffer,
+						LODResource.SkinWeightVertexBuffer,
 						NumVertices,
 						NumBones,
 						NumBoneInfluences,
-						bNeedsCPUAccess,
+						bAllowCPUAccess,
 						MutableMeshVertexBuffers.GetBufferData(Buffer),
 						MutableMeshVertexBuffers.GetBufferDataSize(Buffer)
 					);
@@ -300,17 +296,18 @@ namespace UnrealConversionUtils
 				// colour buffer?
 				else if (Semantic == mu::MBS_COLOUR)
 				{
-					SkeletalMesh->SetHasVertexColors(true);
 					const void* DataPtr = MutableMeshVertexBuffers.GetBufferData(Buffer);
-					FColorVertexBuffers_InitWithMutableData(LODModel.StaticVertexBuffers, NumVertices, DataPtr);
-					check(LODModel.StaticVertexBuffers.ColorVertexBuffer.GetStride() == MutableMeshVertexBuffers.GetElementSize(Buffer));
+					FColorVertexBuffers_InitWithMutableData(LODResource.StaticVertexBuffers, NumVertices, DataPtr);
+					check(LODResource.StaticVertexBuffers.ColorVertexBuffer.GetStride() == MutableMeshVertexBuffers.GetElementSize(Buffer));
 				}
 			}
 		}
 	}
 
 
-	bool CopyMutableIndexBuffers(mu::MeshPtrConst InMutableMesh, FSkeletalMeshLODRenderData& LODModel)
+	bool CopyMutableIndexBuffers(
+		FSkeletalMeshLODRenderData& LODResource,
+		const mu::Ptr<const mu::Mesh> InMutableMesh)
 	{
 		MUTABLE_CPUPROFILER_SCOPE(CopyMutableIndexBuffers);
 
@@ -326,9 +323,15 @@ namespace UnrealConversionUtils
 		const uint8* DataPtr = InMutableMesh->GetIndexBuffers().GetBufferData(0);
 		const int32 ElementSize = InMutableMesh->GetIndexBuffers().GetElementSize(0);
 
-		LODModel.MultiSizeIndexContainer.CreateIndexBuffer(ElementSize);
-		LODModel.MultiSizeIndexContainer.GetIndexBuffer()->Insert(0, IndexCount);
-		FMemory::Memcpy(LODModel.MultiSizeIndexContainer.GetIndexBuffer()->GetPointerTo(0), DataPtr, IndexCount * ElementSize);
+		if (!LODResource.MultiSizeIndexContainer.IsIndexBufferValid())
+		{
+			LODResource.MultiSizeIndexContainer.CreateIndexBuffer(ElementSize);
+		}
+		
+		check(LODResource.MultiSizeIndexContainer.GetDataTypeSize() == ElementSize)
+		
+		LODResource.MultiSizeIndexContainer.GetIndexBuffer()->Insert(0, IndexCount);
+		FMemory::Memcpy(LODResource.MultiSizeIndexContainer.GetIndexBuffer()->GetPointerTo(0), DataPtr, IndexCount * ElementSize);
 
 		return true;
 	}
@@ -358,7 +361,7 @@ namespace UnrealConversionUtils
 
 	
 	void CopyMutableSkinWeightProfilesBuffers(
-		FSkeletalMeshLODRenderData& LODModel,
+		FSkeletalMeshLODRenderData& LODResource,
 		const FName ProfileName,
 		const mu::FMeshBufferSet& MutableMeshVertexBuffers,
 		const int32 BoneIndexBuffer)
@@ -366,12 +369,12 @@ namespace UnrealConversionUtils
 		MUTABLE_CPUPROFILER_SCOPE(CopyMutableSkinWeightProfilesBuffers);
 
 		// Basic Buffer override settings
-		FRuntimeSkinWeightProfileData& Override = LODModel.SkinWeightProfilesData.AddOverrideData(ProfileName);
+		FRuntimeSkinWeightProfileData& Override = LODResource.SkinWeightProfilesData.AddOverrideData(ProfileName);
 
-		const uint8 NumInfluences = LODModel.GetVertexBufferMaxBoneInfluences();
+		const uint8 NumInfluences = LODResource.GetVertexBufferMaxBoneInfluences();
 		Override.NumWeightsPerVertex = NumInfluences;
 
-		const bool b16BitBoneIndices = LODModel.DoesVertexBufferUse16BitBoneIndex();
+		const bool b16BitBoneIndices = LODResource.DoesVertexBufferUse16BitBoneIndex();
 		Override.b16BitBoneIndices = b16BitBoneIndices;
 		const int32 BoneIndexSize = b16BitBoneIndices ? 2 : 1;
 
diff --git a/Engine/Plugins/Experimental/Mutable/Source/CustomizableObjectEditor/Private/MuCOE/MutableMeshPreviewUtils.cpp b/Engine/Plugins/Experimental/Mutable/Source/CustomizableObjectEditor/Private/MuCOE/MutableMeshPreviewUtils.cpp
--- a/Engine/Plugins/Experimental/Mutable/Source/CustomizableObjectEditor/Private/MuCOE/MutableMeshPreviewUtils.cpp
+++ b/Engine/Plugins/Experimental/Mutable/Source/CustomizableObjectEditor/Private/MuCOE/MutableMeshPreviewUtils.cpp
@@ -377,24 +377,25 @@ namespace MutableMeshPreviewUtils
 			TArray<uint16> BoneMap;
 			BoneMap.SetNumZeroed(FMath::Max(NumBonesInBoneMap, 1));
 
+			FSkeletalMeshLODRenderData& LODResource = OutSkeletalMesh->GetResourceForRendering()->LODRenderData[0];
+
 			// Load buffer data found on the mutable model onto the out skeletal mesh
 			// It includes vertex and index buffers
 			UnrealConversionUtils::SetupRenderSections(
-				OutSkeletalMesh,
+				LODResource,
 				InMutableMesh,
 				0,
 				BoneMap,
 				0);
 
 			UnrealConversionUtils::CopyMutableVertexBuffers(
-				OutSkeletalMesh,
+				LODResource,
 				InMutableMesh,
-				0); 
+				false); 
 			
-			FSkeletalMeshLODRenderData& LODModel = OutSkeletalMesh->GetResourceForRendering()->LODRenderData[0];
 
 			// Ensure there's at least one bone in the bonemap of each render section, the root.
-			for (FSkelMeshRenderSection& RenderSection : LODModel.RenderSections)
+			for (FSkelMeshRenderSection& RenderSection : LODResource.RenderSections)
 			{
 				if (RenderSection.BoneMap.IsEmpty())
 				{
@@ -403,18 +404,18 @@ namespace MutableMeshPreviewUtils
 			}
 
 			// Add root as the only required bone
-			LODModel.ActiveBoneIndices.Add(0);
-			LODModel.RequiredBones.Add(0);
+			LODResource.ActiveBoneIndices.Add(0);
+			LODResource.RequiredBones.Add(0);
 			
 			// Copy index buffers or fail to generate the mesh
-			if (!UnrealConversionUtils::CopyMutableIndexBuffers(InMutableMesh, LODModel))
+			if (!UnrealConversionUtils::CopyMutableIndexBuffers(LODResource, InMutableMesh))
 			{
 				return false;
 			}
 			
 			// Update LOD and streaming data
-			LODModel.bIsLODOptional = false;
-			LODModel.bStreamedDataInlined = false;
+			LODResource.bIsLODOptional = false;
+			LODResource.bStreamedDataInlined = false;
 			
 			return true;
 		}
