[Mutable] Add WaitForLaunchedLowPriorityTasks to the MutableTaskGraph.

#rnx

diff --git a/Engine/Plugins/Experimental/Mutable/Source/CustomizableObject/Internal/MuCO/FMutableTaskGraph.h b/Engine/Plugins/Experimental/Mutable/Source/CustomizableObject/Internal/MuCO/FMutableTaskGraph.h
--- a/Engine/Plugins/Experimental/Mutable/Source/CustomizableObject/Internal/MuCO/FMutableTaskGraph.h
+++ b/Engine/Plugins/Experimental/Mutable/Source/CustomizableObject/Internal/MuCO/FMutableTaskGraph.h
@@ -47,8 +47,10 @@ public:
 	/** Create and launch a task on the Mutable Thread with Low priority. */
 	uint32 AddMutableThreadTaskLowPriority(const TCHAR* DebugName, TFunction<void()>&& TaskBody);
 
-	/** Cancel, if not already launched, a Mutable Thread with Low priority. */
-	void CancelMutableThreadTaskLowPriority(uint32 Id);
+	/** Cancel, if not already launched, a Mutable Thread with Low priority. 
+	  * Return true if the task has been canceled before launching it.
+	  * Return false if not found or running. */
+	bool CancelMutableThreadTaskLowPriority(uint32 Id);
 	
 	/** Create and launch a task on Any Thread. */
 	void AddAnyThreadTask(const TCHAR* DebugName, TUniqueFunction<void()>&& TaskBody) const;
@@ -56,6 +58,9 @@ public:
 	/** Wait for all Mutable Thread tasks. */
 	void WaitForMutableTasks();
 
+	/** Wait for the launched low-priority task if it matches the TaskID. */
+	void WaitForLaunchedLowPriorityTask(uint32 TaskID);
+
 	/** Allow or disallow launching Mutable Tasks with Low priority.
 	 @param bFromMutableTask true if called from a Mutable Task. */
 	void AllowLaunchingMutableTaskLowPriority(bool bAllow, bool bFromMutableTask);
@@ -89,6 +94,9 @@ private:
 	/** Incremental task ID generator. */
 	uint32 TaskIdGenerator = INVALID_ID;
 
+	/** The ID of the Last Mutable Task Low Priority launched to the TaskGraph system. */
+	uint32 LastMutableTaskLowPriorityID = INVALID_ID;
+
 	/** Last Mutable Task Low Priority launched to the TaskGraph system. */
 	UE::Tasks::FTask LastMutableTaskLowPriority = {};
 
diff --git a/Engine/Plugins/Experimental/Mutable/Source/CustomizableObject/Private/MuCO/FMutableTaskGraph.cpp b/Engine/Plugins/Experimental/Mutable/Source/CustomizableObject/Private/MuCO/FMutableTaskGraph.cpp
--- a/Engine/Plugins/Experimental/Mutable/Source/CustomizableObject/Private/MuCO/FMutableTaskGraph.cpp
+++ b/Engine/Plugins/Experimental/Mutable/Source/CustomizableObject/Private/MuCO/FMutableTaskGraph.cpp
@@ -57,7 +57,7 @@ uint32 FMutableTaskGraph::AddMutableThreadTaskLowPriority(const TCHAR* DebugName
 }
 
 
-void FMutableTaskGraph::CancelMutableThreadTaskLowPriority(uint32 Id)
+bool FMutableTaskGraph::CancelMutableThreadTaskLowPriority(uint32 Id)
 {
 	FScopeLock Lock(&MutableTaskLock);
 
@@ -66,9 +66,11 @@ void FMutableTaskGraph::CancelMutableThreadTaskLowPriority(uint32 Id)
 		if (QueueMutableTasksLowPriority[Index].Id == Id)
 		{
 			QueueMutableTasksLowPriority.RemoveAt(Index);
-			break;
+			return true;
 		}
 	}
+
+	return false;
 }
 
 
@@ -98,6 +100,26 @@ void FMutableTaskGraph::WaitForMutableTasks()
 }
 
 
+void FMutableTaskGraph::WaitForLaunchedLowPriorityTask(uint32 TaskID)
+{
+	UE::Tasks::FTask Task;
+
+	{
+		FScopeLock Lock(&MutableTaskLock);
+
+		if (LastMutableTaskLowPriorityID == TaskID)
+		{
+			Task = LastMutableTaskLowPriority;
+		}
+	}
+
+	if (Task.IsValid())
+	{
+		Task.Wait();
+	}
+}
+
+
 void FMutableTaskGraph::TryLaunchMutableTaskLowPriority(bool bFromMutableTask)
 {
 	MUTABLE_CPUPROFILER_SCOPE(TryLaunchMutableTaskLowPriority)
@@ -142,6 +164,7 @@ void FMutableTaskGraph::TryLaunchMutableTaskLowPriority(bool bFromMutableTask)
 			return;
 		}
 
+		LastMutableTaskLowPriorityID = NextTask.Id;
 		LastMutableTaskLowPriority = AddMutableThreadTask(*NextTask.DebugName, [this, Task = MoveTemp(NextTask)]() // Moves the task, not the pointer.
 		{
 			MUTABLE_CPUPROFILER_SCOPE(LowPriorityTaskBody)
@@ -149,6 +172,7 @@ void FMutableTaskGraph::TryLaunchMutableTaskLowPriority(bool bFromMutableTask)
 
 			{
 				FScopeLock Lock(&MutableTaskLock);
+				LastMutableTaskLowPriorityID = INVALID_ID;
 				LastMutableTaskLowPriority = {};
 				
 				TryLaunchMutableTaskLowPriority(true);
@@ -205,6 +229,7 @@ int32 FMutableTaskGraph::Tick()
 		FScopeLock Lock(&MutableTaskLock);
 		if (IsTaskCompleted(LastMutableTaskLowPriority))
 		{
+			LastMutableTaskLowPriorityID = INVALID_ID;
 			LastMutableTaskLowPriority = {};
 		}
 	}
