[mutable] Enable concurrency during cook of CustomizableObject assets from a commandlet.
#rnx

diff --git a/Engine/Plugins/Experimental/Mutable/Source/CustomizableObjectEditor/Private/MuCOE/CustomizableObjectCompileRunnable.cpp b/Engine/Plugins/Experimental/Mutable/Source/CustomizableObjectEditor/Private/MuCOE/CustomizableObjectCompileRunnable.cpp
--- a/Engine/Plugins/Experimental/Mutable/Source/CustomizableObjectEditor/Private/MuCOE/CustomizableObjectCompileRunnable.cpp
+++ b/Engine/Plugins/Experimental/Mutable/Source/CustomizableObjectEditor/Private/MuCOE/CustomizableObjectCompileRunnable.cpp
@@ -121,12 +121,37 @@ uint32 FCustomizableObjectCompileRunnable::Run()
 		CompilerOptions->SetImagePixelFormatOverride( UnrealPixelFormatFunc );
 	}
 
-	CompilerOptions->SetReferencedResourceCallback([this](int32 ID, TSharedPtr<mu::Ptr<mu::Image>> ResolvedImage)
+	auto ProviderTick = [this](float)
+		{
+			check(IsInGameThread());
+
+			constexpr double MaxSecondsPerFrame = 0.4;
+
+			double MaxTime = FPlatformTime::Seconds() + MaxSecondsPerFrame;
+
+			FReferenceResourceRequest Request;
+			while (PendingResourceReferenceRequests.Dequeue(Request))
+			{
+				*Request.ResolvedImage = LoadResourceReferenced(Request.ID);
+				Request.CompletionEvent->Trigger();
+
+				// Simple time limit enforcement to avoid blocking the game thread if there are many requests.
+				double CurrentTime = FPlatformTime::Seconds();
+				if (CurrentTime >= MaxTime)
+				{
+					break;
+				}
+			}
+
+			return true;
+		};
+
+	CompilerOptions->SetReferencedResourceCallback([this, &ProviderTick](int32 ID, TSharedPtr<mu::Ptr<mu::Image>> ResolvedImage, bool bRunImmediatlyIfPossible)
 		{
 			// This runs in a random thread
 			UE::Tasks::FTaskEvent CompletionEvent(TEXT("ReferencedResourceCallbackCompletion"));			
 
-			if (IsInGameThread())
+			if (IsInGameThread() && bRunImmediatlyIfPossible)
 			{
 				// Do everything now
 				mu::Ptr<mu::Image> Result = LoadResourceReferenced(ID);
@@ -139,34 +164,12 @@ uint32 FCustomizableObjectCompileRunnable::Run()
 			}
 
 			return CompletionEvent;
-		});
+		}, 
+		ProviderTick
+	);
 
 	// Register the tick function that will process the game-thread image reference resolve requests.
-	FTSTicker::FDelegateHandle ResolveReferenceResourcesTickerHandle = FTSTicker::GetCoreTicker().AddTicker(
-		FTickerDelegate::CreateLambda([this](float)
-			{
-				check(IsInGameThread());
-
-				constexpr double MaxSecondsPerFrame = 0.4;
-
-				double MaxTime = FPlatformTime::Seconds() + MaxSecondsPerFrame;
-
-				FReferenceResourceRequest Request;
-				while (PendingResourceReferenceRequests.Dequeue(Request))
-				{
-					*Request.ResolvedImage = LoadResourceReferenced(Request.ID);
-					Request.CompletionEvent->Trigger();
-
-					// Simple time limit enforcement to avoid blocking the game thread if there are many requests.
-					double CurrentTime = FPlatformTime::Seconds();
-					if (CurrentTime >= MaxTime)
-					{
-						break;
-					}
-				}
-
-				return true;
-			}));
+	FTSTicker::FDelegateHandle ResolveReferenceResourcesTickerHandle = FTSTicker::GetCoreTicker().AddTicker(FTickerDelegate::CreateLambda(ProviderTick));
 
 	const int32 MinResidentMips = UTexture::GetStaticMinTextureResidentMipCount();
 	CompilerOptions->SetDataPackingStrategy( MinResidentMips, Options.EmbeddedDataBytesLimit, Options.PackagedDataBytesLimit );
diff --git a/Engine/Plugins/Experimental/Mutable/Source/MutableTools/Internal/MuT/AST.h b/Engine/Plugins/Experimental/Mutable/Source/MutableTools/Internal/MuT/AST.h
--- a/Engine/Plugins/Experimental/Mutable/Source/MutableTools/Internal/MuT/AST.h
+++ b/Engine/Plugins/Experimental/Mutable/Source/MutableTools/Internal/MuT/AST.h
@@ -208,9 +208,15 @@ namespace mu
         //! Ratio used to decide if it is worth to generate a crop operation
         float MinRLECompressionGain = 1.2f;
 
-		//
+		// External resource provision functions
+		//---------------------------------------------------------------------
+
+		/** Function used to request an engine resource (image) from the compiler. */
 		FReferencedResourceFunc ReferencedResourceProvider;
-    };
+
+		/** Function used to ensure the provision of engine resources with ReferencedResourceProvider advances. */
+		FReferencedResourceGameThreadTickFunc ReferencedResourceProviderTick;
+	};
 
 
 	struct FLinkerOptions
diff --git a/Engine/Plugins/Experimental/Mutable/Source/MutableTools/Private/MuT/CodeOptimiser.cpp b/Engine/Plugins/Experimental/Mutable/Source/MutableTools/Private/MuT/CodeOptimiser.cpp
--- a/Engine/Plugins/Experimental/Mutable/Source/MutableTools/Private/MuT/CodeOptimiser.cpp
+++ b/Engine/Plugins/Experimental/Mutable/Source/MutableTools/Private/MuT/CodeOptimiser.cpp
@@ -813,7 +813,8 @@ namespace mu
 
 							TSharedPtr< Ptr<Image> > ResolveImage = MakeShared<Ptr<Image>>();
 
-							UE::Tasks::FTaskEvent ReferenceCompletionEvent = InOptions->OptimisationOptions.ReferencedResourceProvider(ImageID, ResolveImage);
+							constexpr bool bRunImmediatlyIfPossible = false;
+							UE::Tasks::FTaskEvent ReferenceCompletionEvent = InOptions->OptimisationOptions.ReferencedResourceProvider(ImageID, ResolveImage, bRunImmediatlyIfPossible);
 
 							UE::Tasks::FTask CompleteTask = UE::Tasks::Launch(TEXT("MutableResolveComplete"),
 								[SubgraphRoot, InOptions, ResolveImage, &ASTAccessLock]()
@@ -885,7 +886,20 @@ namespace mu
 			// Wait for pending tasks
 			{
 				MUTABLE_CPUPROFILER_SCOPE(ConstantGenerator_WaitPending);
-				LaunchTask.Wait();
+				if (!IsInGameThread())
+				{
+					LaunchTask.Wait();
+				}
+				else
+				{
+					const FReferencedResourceGameThreadTickFunc& ManualGameTick = InOptions->OptimisationOptions.ReferencedResourceProviderTick;
+
+					constexpr float FakeTimeStep = 0.1f;
+					while (!LaunchTask.Wait(FakeTimeStep))
+					{
+						ManualGameTick(FakeTimeStep);
+					}
+				}
 			}
 		}
 
@@ -919,7 +933,8 @@ namespace mu
 						auto ImmediateCompleteFunc = [SubgraphRoot, InOptions, ImageID]()
 							{
 								TSharedPtr< Ptr<Image> > ResolveImage = MakeShared<Ptr<Image>>();
-								UE::Tasks::FTaskEvent ReferenceCompletionEvent = InOptions->OptimisationOptions.ReferencedResourceProvider(ImageID, ResolveImage);
+								constexpr bool bRunImmediatlyIfPossible = true;
+								UE::Tasks::FTaskEvent ReferenceCompletionEvent = InOptions->OptimisationOptions.ReferencedResourceProvider(ImageID, ResolveImage, bRunImmediatlyIfPossible);
 								ReferenceCompletionEvent.Wait();
 								Ptr<ASTOpConstantResource> ConstantOp = new ASTOpConstantResource;
 								ConstantOp->type = OP_TYPE::IM_CONSTANT;
diff --git a/Engine/Plugins/Experimental/Mutable/Source/MutableTools/Private/MuT/Compiler.cpp b/Engine/Plugins/Experimental/Mutable/Source/MutableTools/Private/MuT/Compiler.cpp
--- a/Engine/Plugins/Experimental/Mutable/Source/MutableTools/Private/MuT/Compiler.cpp
+++ b/Engine/Plugins/Experimental/Mutable/Source/MutableTools/Private/MuT/Compiler.cpp
@@ -168,9 +168,10 @@ namespace mu
 
 
 	//---------------------------------------------------------------------------------------------
-	void CompilerOptions::SetReferencedResourceCallback(const FReferencedResourceFunc& Provider)
+	void CompilerOptions::SetReferencedResourceCallback(const FReferencedResourceFunc& Provider, const FReferencedResourceGameThreadTickFunc& ProviderGameThreadTick)
 	{
 		m_pD->OptimisationOptions.ReferencedResourceProvider = Provider;
+		m_pD->OptimisationOptions.ReferencedResourceProviderTick = ProviderGameThreadTick;
 	}
 
 
diff --git a/Engine/Plugins/Experimental/Mutable/Source/MutableTools/Public/MuT/Compiler.h b/Engine/Plugins/Experimental/Mutable/Source/MutableTools/Public/MuT/Compiler.h
--- a/Engine/Plugins/Experimental/Mutable/Source/MutableTools/Public/MuT/Compiler.h
+++ b/Engine/Plugins/Experimental/Mutable/Source/MutableTools/Public/MuT/Compiler.h
@@ -32,7 +32,8 @@ namespace mu
     class Node;
     class NodeTransformedObject;
 
-	typedef TFunction<UE::Tasks::FTaskEvent(int32, TSharedPtr<Ptr<Image>>)> FReferencedResourceFunc;
+	typedef TFunction<UE::Tasks::FTaskEvent(int32, TSharedPtr<Ptr<Image>>, bool)> FReferencedResourceFunc;
+	typedef TFunction<void(float)> FReferencedResourceGameThreadTickFunc;
 
     //! \brief Options used to compile the models with a compiler.
     class MUTABLETOOLS_API CompilerOptions : public RefCounted
@@ -83,7 +84,7 @@ namespace mu
 		void SetImagePixelFormatOverride(const FImageOperator::FImagePixelFormatFunc&);
 
 		/** */
-		void SetReferencedResourceCallback(const FReferencedResourceFunc&);
+		void SetReferencedResourceCallback(const FReferencedResourceFunc&, const FReferencedResourceGameThreadTickFunc&);
 
         //! Different data packing strategies
         enum class TextureLayoutStrategy : uint8
