[Mutable] Marked UCustomizableObject non-public members as private.

- Removed function deifintions that did not have implementations.
- Removed not used Relevancy member.

#rnx

diff --git a/Engine/Plugins/Experimental/Mutable/Source/CustomizableObject/Classes/MuCO/CustomizableObject.h b/Engine/Plugins/Experimental/Mutable/Source/CustomizableObject/Classes/MuCO/CustomizableObject.h
--- a/Engine/Plugins/Experimental/Mutable/Source/CustomizableObject/Classes/MuCO/CustomizableObject.h
+++ b/Engine/Plugins/Experimental/Mutable/Source/CustomizableObject/Classes/MuCO/CustomizableObject.h
@@ -69,16 +69,6 @@ struct FParameterTags
 };
 
 
-UENUM()
-enum class ECustomizableObjectRelevancy : uint8
-{
-	// 
-	All = 0 UMETA(DisplayName = "Relevant for client and server"),
-	// 
-	ClientOnly = 1 UMETA(DisplayName = "Only necessary on clients")
-};
-
-
 USTRUCT()
 struct FProfileParameterDat
 {
@@ -511,6 +501,7 @@ public:
 	UPROPERTY(Transient)
 	TArray<FCustomizableObjectMeshToMeshVertData> ClothMeshToMeshVertData;
 
+private:
 	// mu::ExtensionData::Index is an index into this array when mu::ExtensionData::Origin is ConstantAlwaysLoaded
 	UPROPERTY()
 	TArray<FCustomizableObjectResourceData> AlwaysLoadedExtensionData;
@@ -523,6 +514,7 @@ public:
 	UPROPERTY()
 	TArray<FCustomizableObjectStreamedResourceData> StreamedResourceData;
 
+public:
 	/** Use the SkeletalMesh of reference as a placeholder until the custom mesh is ready to use.
 	  * Note: If disabled, a null mesh will be used to replace the discarded mesh due to 'ReplaceDiscardedWithReferenceMesh' being enabled. */
 	UPROPERTY(EditAnywhere, Category = CustomizableObject)
@@ -538,13 +530,6 @@ public:
 	bool bEnableMeshCache = false;
 	
 #if WITH_EDITORONLY_DATA
-private:
-	// Hide this property because it is not used yet.
-	//UPROPERTY(EditAnywhere, Category = CustomizableObject)
-	UPROPERTY()
-	ECustomizableObjectRelevancy Relevancy;
-
-public:
 	// Compilation options to use in editor and for packaging for this object.
 	UPROPERTY()
 	FCompilationOptions CompileOptions;
@@ -589,7 +574,6 @@ private:
 	UPROPERTY()
 	FGuid VersionId;
 
-public:
 	UPROPERTY()
 	TArray<FProfileParameterDat> InstancePropertiesProfiles;
 #endif // WITH_EDITORONLY_DATA
@@ -597,7 +581,8 @@ public:
 	/** Amount of components in this CO. Set at the end of the model compilation process. */
 	UPROPERTY()
 	int32 NumMeshComponentsInRoot = 0; // TODO UE-205600, move to FModelResources
-	
+
+public:
 	/** Get the number of components this Customizable Object has. */
 	UFUNCTION(BlueprintCallable, Category = CustomizableObject)
 	int32 GetComponentCount() const;
@@ -708,30 +693,11 @@ public:
 	void Serialize(FArchive& Ar) override;
 #if WITH_EDITOR
 	
-	// Compile the object for a specific platform - Compile for Cook Customizable Object
-	void CompileForTargetPlatform(const ITargetPlatform* TargetPlatform);
-
-	// Unless we are packaging there is no need for keeping all the data generated during compilation, this information is stored in the derived data.
-	void ClearCompiledData(bool bIsCooking);
-
 	/** Compile the object if Automatic Compilation is enabled and has not been already compiled.
 	  * Automatic compilation can be enabled/disabled in the Mutable's Plugin Settings.
 	  * @return true if compiled */
 	bool ConditionalAutoCompile();
 	
-	// Add a profile that stores the values of the parameters used by the CustomInstance.
-	FReply AddNewParameterProfile(FString Name, class UCustomizableObjectInstance& CustomInstance);
-
-	// Create new GUID for this CO
-	void UpdateVersionId();
-	FGuid GetVersionId() const { return VersionId; }
-
-	// Compose folder name where the data is stored
-	FString GetCompiledDataFolderPath() const;
-
-	// Compose file name 
-	FString GetCompiledDataFileName(bool bIsModel, const ITargetPlatform* InTargetPlatform = nullptr, bool bIsDiskStreamer = false);
-
 	virtual void GetAssetRegistryTags(FAssetRegistryTagsContext Context) const override;
 	UE_DEPRECATED(5.4, "Implement the version that takes FAssetRegistryTagsContext instead.")
 	virtual void GetAssetRegistryTags(TArray<FAssetRegistryTag>& OutTags) const override;
@@ -785,6 +751,7 @@ public:
 	UFUNCTION(BlueprintCallable, Category = CustomizableObject)
 	FParameterUIData GetParameterUIMetadataFromIndex(int32 ParamIndex) const;
 
+private:
 	/** Textures marked as low priority will generate defaulted resident mips (if texture streaming is enabled).
 	  * Generating defaulted resident mips greatly reduce initial generation times. */
 	UPROPERTY(EditAnywhere, Category = CustomizableObject)
@@ -809,7 +776,6 @@ public:
 	// Customizable Object Population data end --------------------------------------------------------
 
 #if WITH_EDITORONLY_DATA
-private:
 	/** True if this object references a parent object. This is used basically to exclude this object
 	  * from cooking. This is actually derived from the source graph object node pointing to another
 	  * object or not, but it needs to be cached here because the source graph is not always available.
diff --git a/Engine/Plugins/Experimental/Mutable/Source/CustomizableObject/Internal/MuCO/CustomizableObjectPrivate.h b/Engine/Plugins/Experimental/Mutable/Source/CustomizableObject/Internal/MuCO/CustomizableObjectPrivate.h
--- a/Engine/Plugins/Experimental/Mutable/Source/CustomizableObject/Internal/MuCO/CustomizableObjectPrivate.h
+++ b/Engine/Plugins/Experimental/Mutable/Source/CustomizableObject/Internal/MuCO/CustomizableObjectPrivate.h
@@ -723,6 +723,26 @@ public:
 	/** Return true if the CO is not compiled or the ParticipatingObjects system has detected a change (participating objects dirty or re-saved since last compilation). */
 	bool IsCompilationOutOfDate() const;
 #endif
+
+	TMap<uint64, FMutableStreamableBlock>& GetHashToStreamableBlock();
+
+	int32& GetNumMeshComponentsInRoot();
+
+	TArray<FString>& GetCustomizableObjectClassTags();
+	
+	TArray<FString>& GetPopulationClassTags();
+
+    TMap<FString, FParameterTags>& GetCustomizableObjectParametersTags();
+
+#if WITH_EDITORONLY_DATA
+	TArray<FProfileParameterDat>& GetInstancePropertiesProfiles();
+#endif
+	
+	TArray<FCustomizableObjectResourceData>& GetAlwaysLoadedExtensionData();
+
+	TArray<FCustomizableObjectStreamedResourceData>& GetStreamedExtensionData();
+	
+	TArray<FCustomizableObjectStreamedResourceData>& GetStreamedResourceData();
 	
 	/** Cache of generated SkeletalMeshes */
 	FMeshCache MeshCache;
diff --git a/Engine/Plugins/Experimental/Mutable/Source/CustomizableObject/Private/MuCO/CustomizableObject.cpp b/Engine/Plugins/Experimental/Mutable/Source/CustomizableObject/Private/MuCO/CustomizableObject.cpp
--- a/Engine/Plugins/Experimental/Mutable/Source/CustomizableObject/Private/MuCO/CustomizableObject.cpp
+++ b/Engine/Plugins/Experimental/Mutable/Source/CustomizableObject/Private/MuCO/CustomizableObject.cpp
@@ -862,14 +862,8 @@ void UCustomizableObjectPrivate::CompileForTargetPlatform(const ITargetPlatform*
 
 	FCustomizableObjectCompilerBase* Compiler = UCustomizableObjectSystem::GetInstance()->GetNewCompiler();
 
-	bool bIsRootObject = Compiler->IsRootObject(GetPublic());
-	
-	bool bIsRelevantForThisTarget =
-		(GetPublic()->Relevancy == ECustomizableObjectRelevancy::All)
-		||
-		(GetPublic()->Relevancy == ECustomizableObjectRelevancy::ClientOnly && !TargetPlatform->IsServerOnly());
-
-	if (bIsRootObject && bIsRelevantForThisTarget)
+	const bool bIsRootObject = Compiler->IsRootObject(GetPublic());
+	if (bIsRootObject)
 	{
 		FCompilationOptions Options;
 		Options.OptimizationLevel = UE_MUTABLE_MAX_OPTIMIZATION;	// max optimization when packaging.
@@ -2006,6 +2000,62 @@ bool UCustomizableObjectPrivate::IsCompilationOutOfDate() const
 }
 #endif
 
+TMap<uint64, FMutableStreamableBlock>& UCustomizableObjectPrivate::GetHashToStreamableBlock()
+{
+	return GetPublic()->HashToStreamableBlock;
+}
+
+
+int32& UCustomizableObjectPrivate::GetNumMeshComponentsInRoot()
+{
+	return GetPublic()->NumMeshComponentsInRoot;
+}
+
+
+TArray<FString>& UCustomizableObjectPrivate::GetCustomizableObjectClassTags()
+{
+	return GetPublic()->CustomizableObjectClassTags;
+}
+
+
+TArray<FString>& UCustomizableObjectPrivate::GetPopulationClassTags()
+{
+	return GetPublic()->PopulationClassTags;
+}
+
+
+TMap<FString, FParameterTags>& UCustomizableObjectPrivate::GetCustomizableObjectParametersTags()
+{
+	return GetPublic()->CustomizableObjectParametersTags;
+}
+
+
+#if WITH_EDITORONLY_DATA
+TArray<FProfileParameterDat>& UCustomizableObjectPrivate::GetInstancePropertiesProfiles()
+{
+	return GetPublic()->InstancePropertiesProfiles;
+}
+#endif
+
+
+TArray<FCustomizableObjectResourceData>& UCustomizableObjectPrivate::GetAlwaysLoadedExtensionData()
+{
+	return GetPublic()->AlwaysLoadedExtensionData;
+}
+
+
+TArray<FCustomizableObjectStreamedResourceData>& UCustomizableObjectPrivate::GetStreamedExtensionData()
+{
+	return GetPublic()->StreamedExtensionData;
+}
+
+
+TArray<FCustomizableObjectStreamedResourceData>& UCustomizableObjectPrivate::GetStreamedResourceData()
+{
+	return GetPublic()->StreamedResourceData;
+}
+
+
 //-------------------------------------------------------------------------------------------------
 //-------------------------------------------------------------------------------------------------
 //-------------------------------------------------------------------------------------------------
@@ -2141,7 +2191,7 @@ void UCustomizableObjectBulk::PrepareBulkData(UCustomizableObject* InOuter, cons
 		for (int32 BlockIndex = 0; BlockIndex < NumBlocks; ++BlockIndex)
 		{
 			uint32 BlockId = Model->GetRomId(BlockIndex);
-			const FMutableStreamableBlock& StreamableBlock = CustomizableObject->HashToStreamableBlock[BlockId];
+			const FMutableStreamableBlock& StreamableBlock = CustomizableObject->GetPrivate()->GetHashToStreamableBlock()[BlockId];
 			const uint32 BlockSize = StreamableBlock.Size;
 
 			FBlock CurrentBlock = { EDataType::Model, BlockId, BlockSize, SourceOffset };
@@ -2268,7 +2318,7 @@ void UCustomizableObjectBulk::PrepareBulkData(UCustomizableObject* InOuter, cons
 			{
 				FBlock ThisBlock = CurrentFile.Blocks[FileBlockIndex];
 
-				FMutableStreamableBlock& StreamableBlock = CustomizableObject->HashToStreamableBlock[ThisBlock.Id];
+				FMutableStreamableBlock& StreamableBlock = CustomizableObject->GetPrivate()->GetHashToStreamableBlock()[ThisBlock.Id];
 				check(StreamableBlock.Size == ThisBlock.Size);
 				StreamableBlock.FileId = FileId;
 				StreamableBlock.Offset = OffsetInFile;
@@ -2511,13 +2561,13 @@ void FMutableRefSkeletalMeshData::InitResources(UCustomizableObject* InOuter, co
 	// Initialize AssetUserData
 	for (FMutableRefAssetUserData& Data : AssetUserData)
 	{
-		if (!InOuter->StreamedResourceData.IsValidIndex(Data.AssetUserDataIndex))
+		if (!InOuter->GetPrivate()->GetStreamedResourceData().IsValidIndex(Data.AssetUserDataIndex))
 		{
 			check(false);
 			continue;
 		}
 
-		FCustomizableObjectStreamedResourceData& StreamedResource = InOuter->StreamedResourceData[Data.AssetUserDataIndex];
+		FCustomizableObjectStreamedResourceData& StreamedResource = InOuter->GetPrivate()->GetStreamedResourceData()[Data.AssetUserDataIndex];
 		Data.AssetUserData = StreamedResource.GetPath().LoadSynchronous();
 		check(Data.AssetUserData);
 		check(Data.AssetUserData->Data.Type == ECOResourceDataType::AssetUserData);
diff --git a/Engine/Plugins/Experimental/Mutable/Source/CustomizableObject/Private/MuCO/CustomizableObjectInstance.cpp b/Engine/Plugins/Experimental/Mutable/Source/CustomizableObject/Private/MuCO/CustomizableObjectInstance.cpp
--- a/Engine/Plugins/Experimental/Mutable/Source/CustomizableObject/Private/MuCO/CustomizableObjectInstance.cpp
+++ b/Engine/Plugins/Experimental/Mutable/Source/CustomizableObject/Private/MuCO/CustomizableObjectInstance.cpp
@@ -1949,17 +1949,17 @@ bool UCustomizableInstancePrivate::UpdateSkeletalMesh_PostBeginUpdate0(UCustomiz
 				{
 					case mu::ExtensionData::EOrigin::ConstantAlwaysLoaded:
 					{
-						check(CustomizableObject->AlwaysLoadedExtensionData.IsValidIndex(ExtensionOutput.Data->Index));
-						ReferencedExtensionData = &CustomizableObject->AlwaysLoadedExtensionData[ExtensionOutput.Data->Index];
+						check(CustomizableObject->GetPrivate()->GetAlwaysLoadedExtensionData().IsValidIndex(ExtensionOutput.Data->Index));
+						ReferencedExtensionData = &CustomizableObject->GetPrivate()->GetAlwaysLoadedExtensionData()[ExtensionOutput.Data->Index];
 					}
 					break;
 
 					case mu::ExtensionData::EOrigin::ConstantStreamed:
 					{
-						check(CustomizableObject->StreamedExtensionData.IsValidIndex(ExtensionOutput.Data->Index));
+						check(CustomizableObject->GetPrivate()->GetStreamedExtensionData().IsValidIndex(ExtensionOutput.Data->Index));
 						
 						const FCustomizableObjectStreamedResourceData& StreamedData =
-							CustomizableObject->StreamedExtensionData[ExtensionOutput.Data->Index];
+							CustomizableObject->GetPrivate()->GetStreamedExtensionData()[ExtensionOutput.Data->Index];
 
 						if (!StreamedData.IsLoaded())
 						{
@@ -2462,7 +2462,7 @@ bool UCustomizableInstancePrivate::LoadParametersFromProfile(int32 ProfileIndex)
 	}
 
 #if WITH_EDITOR
-	if (ProfileIndex < 0 || ProfileIndex >= CustomizableObject->InstancePropertiesProfiles.Num() )
+	if (ProfileIndex < 0 || ProfileIndex >= CustomizableObject->GetPrivate()->GetInstancePropertiesProfiles().Num() )
 	{
 		return false;
 	}
@@ -2470,7 +2470,7 @@ bool UCustomizableInstancePrivate::LoadParametersFromProfile(int32 ProfileIndex)
 	// This could be done only when the instance changes.
 	MigrateProfileParametersToCurrentInstance(ProfileIndex);
 
-	const FProfileParameterDat& Profile = CustomizableObject->InstancePropertiesProfiles[ProfileIndex];
+	const FProfileParameterDat& Profile = CustomizableObject->GetPrivate()->GetInstancePropertiesProfiles()[ProfileIndex];
 
 	GetPublic()->Descriptor.BoolParameters = Profile.BoolParameters;
 	GetPublic()->Descriptor.IntParameters = Profile.IntParameters;
@@ -2494,12 +2494,12 @@ bool UCustomizableInstancePrivate::SaveParametersToProfile(int32 ProfileIndex)
 #if WITH_EDITOR
 	bSelectedProfileDirty = ProfileIndex != SelectedProfileIndex;
 
-	if (ProfileIndex < 0 || ProfileIndex >= CustomizableObject->InstancePropertiesProfiles.Num())
+	if (ProfileIndex < 0 || ProfileIndex >= CustomizableObject->GetPrivate()->GetInstancePropertiesProfiles().Num())
 	{
 		return false;
 	}
 
-	FProfileParameterDat& Profile = CustomizableObject->InstancePropertiesProfiles[ProfileIndex];
+	FProfileParameterDat& Profile = CustomizableObject->GetPrivate()->GetInstancePropertiesProfiles()[ProfileIndex];
 
 	Profile.BoolParameters = GetPublic()->Descriptor.BoolParameters;
 	Profile.IntParameters = GetPublic()->Descriptor.IntParameters;
@@ -2520,12 +2520,12 @@ bool UCustomizableInstancePrivate::MigrateProfileParametersToCurrentInstance(int
 	}
 
 #if WITH_EDITOR
-	if (ProfileIndex < 0 || ProfileIndex >= CustomizableObject->InstancePropertiesProfiles.Num())
+	if (ProfileIndex < 0 || ProfileIndex >= CustomizableObject->GetPrivate()->GetInstancePropertiesProfiles().Num())
 	{
 		return false;
 	}
 
-	FProfileParameterDat& Profile = CustomizableObject->InstancePropertiesProfiles[ProfileIndex];
+	FProfileParameterDat& Profile = CustomizableObject->GetPrivate()->GetInstancePropertiesProfiles()[ProfileIndex];
 	FProfileParameterDat TempProfile;
 
 	TempProfile.ProfileName = Profile.ProfileName;
@@ -4909,7 +4909,7 @@ UE::Tasks::FTask UCustomizableInstancePrivate::LoadAdditionalAssetsAndDataAsync(
 
 			for (uint32 ResourceId : StreamedResources)
 			{
-				const TArray<FCustomizableObjectStreamedResourceData>& StreamedResourcesData = CustomizableObject->StreamedResourceData;
+				const TArray<FCustomizableObjectStreamedResourceData>& StreamedResourcesData = CustomizableObject->GetPrivate()->GetStreamedResourceData();
 				FCustomizableObjectStreameableResourceId TypedResourceId = BitCast<FCustomizableObjectStreameableResourceId>(ResourceId);	
 	
 				if (TypedResourceId.Type == (uint8)FCustomizableObjectStreameableResourceId::EType::AssetUserData)
@@ -5304,7 +5304,7 @@ void UCustomizableInstancePrivate::AdditionalAssetsAsyncLoaded(UCustomizableObje
 	}
 
 	
-	TArray<FCustomizableObjectStreamedResourceData>& StreamedResources = CustomizableObject->StreamedResourceData;
+	TArray<FCustomizableObjectStreamedResourceData>& StreamedResources = CustomizableObject->GetPrivate()->GetStreamedResourceData();
 
 	for (FCustomizableInstanceComponentData& ComponentData : ComponentsData)
 	{
diff --git a/Engine/Plugins/Experimental/Mutable/Source/CustomizableObject/Private/MuCO/UnrealExtensionDataStreamer.cpp b/Engine/Plugins/Experimental/Mutable/Source/CustomizableObject/Private/MuCO/UnrealExtensionDataStreamer.cpp
--- a/Engine/Plugins/Experimental/Mutable/Source/CustomizableObject/Private/MuCO/UnrealExtensionDataStreamer.cpp
+++ b/Engine/Plugins/Experimental/Mutable/Source/CustomizableObject/Private/MuCO/UnrealExtensionDataStreamer.cpp
@@ -174,7 +174,7 @@ TSharedPtr<FStreamableHandle> FUnrealExtensionDataStreamer::StartLoadOnGameThrea
 		return nullptr;
 	}
 
-	if (!Object->StreamedExtensionData.IsValidIndex(LoadHandle->Data->Index))
+	if (!Object->GetPrivate()->GetStreamedExtensionData().IsValidIndex(LoadHandle->Data->Index))
 	{
 		// The compiled data appears to be out of sync with the CO's properties
 
@@ -185,7 +185,7 @@ TSharedPtr<FStreamableHandle> FUnrealExtensionDataStreamer::StartLoadOnGameThrea
 		return nullptr;
 	}
 
-	FCustomizableObjectStreamedResourceData& StreamedData = Object->StreamedExtensionData[LoadHandle->Data->Index];
+	FCustomizableObjectStreamedResourceData& StreamedData = Object->GetPrivate()->GetStreamedExtensionData()[LoadHandle->Data->Index];
 	if (StreamedData.IsLoaded())
 	{
 		// Already loaded
@@ -242,7 +242,7 @@ void FUnrealExtensionDataStreamer::NotifyLoadCompleted(
 	check(IsInGameThread());
 	check(Object);
 
-	if (!Object->StreamedExtensionData.IsValidIndex(LoadHandle->Data->Index))
+	if (!Object->GetPrivate()->GetStreamedExtensionData().IsValidIndex(LoadHandle->Data->Index))
 	{
 		// The compiled data appears to be out of sync with the CO's properties
 
@@ -253,7 +253,7 @@ void FUnrealExtensionDataStreamer::NotifyLoadCompleted(
 		return;
 	}
 
-	FCustomizableObjectStreamedResourceData& StreamedData = Object->StreamedExtensionData[LoadHandle->Data->Index];
+	FCustomizableObjectStreamedResourceData& StreamedData = Object->GetPrivate()->GetStreamedExtensionData()[LoadHandle->Data->Index];
 
 	// The object could have been loaded by another request, in which case we can skip updating
 	// the StreamedData.
diff --git a/Engine/Plugins/Experimental/Mutable/Source/CustomizableObject/Private/MuCO/UnrealMutableModelDiskStreamer.cpp b/Engine/Plugins/Experimental/Mutable/Source/CustomizableObject/Private/MuCO/UnrealMutableModelDiskStreamer.cpp
--- a/Engine/Plugins/Experimental/Mutable/Source/CustomizableObject/Private/MuCO/UnrealMutableModelDiskStreamer.cpp
+++ b/Engine/Plugins/Experimental/Mutable/Source/CustomizableObject/Private/MuCO/UnrealMutableModelDiskStreamer.cpp
@@ -102,7 +102,7 @@ bool FUnrealMutableModelBulkReader::PrepareStreamingForObject(UCustomizableObjec
 		NewData.BulkFilePrefix = BulkData->GetBulkFilePrefix();
 #endif
 
-		NewData.StreamableBlocks = CustomizableObject->HashToStreamableBlock;
+		NewData.StreamableBlocks = CustomizableObject->GetPrivate()->GetHashToStreamableBlock();
 		if (NewData.StreamableBlocks.IsEmpty())
 		{
 			UE_LOG(LogMutable, Warning, TEXT("Streaming: Customizable Object %s has no data to stream."), *CustomizableObject->GetName());
diff --git a/Engine/Plugins/Experimental/Mutable/Source/CustomizableObjectEditor/Private/MuCOE/CustomizableInstanceDetails.cpp b/Engine/Plugins/Experimental/Mutable/Source/CustomizableObjectEditor/Private/MuCOE/CustomizableInstanceDetails.cpp
--- a/Engine/Plugins/Experimental/Mutable/Source/CustomizableObjectEditor/Private/MuCOE/CustomizableInstanceDetails.cpp
+++ b/Engine/Plugins/Experimental/Mutable/Source/CustomizableObjectEditor/Private/MuCOE/CustomizableInstanceDetails.cpp
@@ -315,13 +315,13 @@ TSharedRef<SWidget> FCustomizableInstanceDetails::GenerateInstanceProfileSelecto
 	ParameterProfileNames.Emplace(MakeShared<FString>("None"));
 	TSharedPtr<FString> CurrentProfileName = ParameterProfileNames.Last();
 
-	for (FProfileParameterDat& Profile : CustomInstance->GetCustomizableObject()->InstancePropertiesProfiles)
+	for (FProfileParameterDat& Profile : CustomInstance->GetCustomizableObject()->GetPrivate()->GetInstancePropertiesProfiles())
 	{
 		ParameterProfileNames.Emplace(MakeShared<FString>(Profile.ProfileName));
 
 		if (ProfileIdx != INDEX_NONE)
 		{
-			const FProfileParameterDat CurrentInstanceProfile = CustomizableObject->InstancePropertiesProfiles[ProfileIdx];
+			const FProfileParameterDat CurrentInstanceProfile = CustomizableObject->GetPrivate()->GetInstancePropertiesProfiles()[ProfileIdx];
 
 			if (Profile.ProfileName == CurrentInstanceProfile.ProfileName)
 			{
@@ -433,7 +433,7 @@ FReply FCustomizableInstanceDetails::RemoveParameterProfile()
 		return FReply::Handled();
 	}
 
-	TArray<FProfileParameterDat>& Profiles = CustomizableObject->InstancePropertiesProfiles;
+	TArray<FProfileParameterDat>& Profiles = CustomizableObject->GetPrivate()->GetInstancePropertiesProfiles();
 
 	Profiles.RemoveAt(ProfileIdx);
 	CustomInstance->GetPrivate()->SelectedProfileIndex = INDEX_NONE;
@@ -465,7 +465,7 @@ void FCustomizableInstanceDetails::OnProfileSelectedChanged(TSharedPtr<FString>
 	else
 	{
 		//Set selected profile
-		TArray<FProfileParameterDat>& Profiles = CustomInstance->GetCustomizableObject()->InstancePropertiesProfiles;
+		TArray<FProfileParameterDat>& Profiles = CustomInstance->GetCustomizableObject()->GetPrivate()->GetInstancePropertiesProfiles();
 		for (int32 Idx = 0; Idx < Profiles.Num(); ++Idx)
 		{
 			if (Profiles[Idx].ProfileName == *Selection)
@@ -2131,7 +2131,7 @@ void SProfileParametersWindow::OnNameChange(const FText& NewName, ETextCommit::T
 		{
 			CustomInstance->GetPrivate()->SaveParametersToProfile(CustomInstance->GetPrivate()->SelectedProfileIndex);
 		}
-		CustomInstance->GetPrivate()->SelectedProfileIndex = CustomizableObject->InstancePropertiesProfiles.Num() - 1;
+		CustomInstance->GetPrivate()->SelectedProfileIndex = CustomizableObject->GetPrivate()->GetInstancePropertiesProfiles().Num() - 1;
 
 		if (InstanceDetails)
 		{
@@ -2156,7 +2156,7 @@ FReply SProfileParametersWindow::OnButtonClick(EAppReturnType::Type ButtonID)
 		{
 			CustomInstance->GetPrivate()->SaveParametersToProfile(CustomInstance->GetPrivate()->SelectedProfileIndex);
 		}
-		CustomInstance->GetPrivate()->SelectedProfileIndex = CustomizableObject->InstancePropertiesProfiles.Num() - 1;
+		CustomInstance->GetPrivate()->SelectedProfileIndex = CustomizableObject->GetPrivate()->GetInstancePropertiesProfiles().Num() - 1;
 
 		if (InstanceDetails)
 		{
diff --git a/Engine/Plugins/Experimental/Mutable/Source/CustomizableObjectEditor/Private/MuCOE/CustomizableObjectCompiler.cpp b/Engine/Plugins/Experimental/Mutable/Source/CustomizableObjectEditor/Private/MuCOE/CustomizableObjectCompiler.cpp
--- a/Engine/Plugins/Experimental/Mutable/Source/CustomizableObjectEditor/Private/MuCOE/CustomizableObjectCompiler.cpp
+++ b/Engine/Plugins/Experimental/Mutable/Source/CustomizableObjectEditor/Private/MuCOE/CustomizableObjectCompiler.cpp
@@ -1044,12 +1044,12 @@ void FCustomizableObjectCompiler::CompileInternal(UCustomizableObject* Object, c
 		
 		Object->GetPrivate()->bIsCompiledWithoutOptimization = GenerationContext.Options.OptimizationLevel < UE_MUTABLE_MAX_OPTIMIZATION;
 
-		Object->AlwaysLoadedExtensionData = MoveTemp(GenerationContext.AlwaysLoadedExtensionData);
+		Object->GetPrivate()->GetAlwaysLoadedExtensionData() = MoveTemp(GenerationContext.AlwaysLoadedExtensionData);
 
-		Object->StreamedExtensionData.Empty(GenerationContext.StreamedExtensionData.Num());
+		Object->GetPrivate()->GetStreamedExtensionData().Empty(GenerationContext.StreamedExtensionData.Num());
 		for (UCustomizableObjectResourceDataContainer* Container : GenerationContext.StreamedExtensionData)
 		{
-			Object->StreamedExtensionData.Emplace(Container);
+			Object->GetPrivate()->GetStreamedExtensionData().Emplace(Container);
 		}
 
 #if WITH_EDITORONLY_DATA
@@ -1058,14 +1058,14 @@ void FCustomizableObjectCompiler::CompileInternal(UCustomizableObject* Object, c
 
 		Object->LODSettings.NumLODsInRoot = GenerationContext.NumLODsInRoot;
 		
-		Object->NumMeshComponentsInRoot = GenerationContext.NumMeshComponentsInRoot;
+		Object->GetPrivate()->GetNumMeshComponentsInRoot() = GenerationContext.NumMeshComponentsInRoot;
 		
 		Object->LODSettings.FirstLODAvailable = GenerationContext.FirstLODAvailable;
 
 		Object->LODSettings.bLODStreamingEnabled = GenerationContext.bEnableLODStreaming;
 		Object->LODSettings.NumLODsToStream = GenerationContext.NumMaxLODsToStream;
 
-		Object->StreamedResourceData = MoveTemp(GenerationContext.StreamedResourceData);
+		Object->GetPrivate()->GetStreamedResourceData() = MoveTemp(GenerationContext.StreamedResourceData);
 
 		// Pass-through textures
 		TArray<TSoftObjectPtr<UTexture>> NewCompileTimeReferencedTextures;
@@ -1227,14 +1227,14 @@ void FCustomizableObjectCompiler::FinishCompilation()
 		uint64 Offset = 0;
 
 		const int32 NumStreamingFiles = Model->GetRomCount();
-		CurrentObject->HashToStreamableBlock.Empty(NumStreamingFiles);
+		CurrentObject->GetPrivate()->GetHashToStreamableBlock().Empty(NumStreamingFiles);
 
 		for (int32 FileIndex = 0; FileIndex < NumStreamingFiles; ++FileIndex)
 		{
 			const uint32 ResourceId = Model->GetRomId(FileIndex);
 			const uint32 ResourceSize = Model->GetRomSize(FileIndex);
 
-			CurrentObject->HashToStreamableBlock.Add(ResourceId, FMutableStreamableBlock{0, ResourceSize, Offset });
+			CurrentObject->GetPrivate()->GetHashToStreamableBlock().Add(ResourceId, FMutableStreamableBlock{0, ResourceSize, Offset });
 			Offset += ResourceSize;
 		}
 	}
diff --git a/Engine/Plugins/Experimental/Mutable/Source/CustomizableObjectEditor/Private/MuCOE/CustomizableObjectCookPackageSplitter.cpp b/Engine/Plugins/Experimental/Mutable/Source/CustomizableObjectEditor/Private/MuCOE/CustomizableObjectCookPackageSplitter.cpp
--- a/Engine/Plugins/Experimental/Mutable/Source/CustomizableObjectEditor/Private/MuCOE/CustomizableObjectCookPackageSplitter.cpp
+++ b/Engine/Plugins/Experimental/Mutable/Source/CustomizableObjectEditor/Private/MuCOE/CustomizableObjectCookPackageSplitter.cpp
@@ -4,6 +4,7 @@
 
 #include "Algo/Find.h"
 #include "MuCO/CustomizableObject.h"
+#include "MuCO/CustomizableObjectPrivate.h"
 #include "UObject/NameTypes.h"
 #include "UObject/Package.h"
 
@@ -152,7 +153,7 @@ bool FCustomizableObjectCookPackageSplitter::ShouldSplit(UObject* SplitData)
 {
 	const UCustomizableObject* Object = CastChecked<UCustomizableObject>(SplitData);
 
-	return Object->StreamedResourceData.Num() > 0 || Object->StreamedExtensionData.Num() > 0;
+	return Object->GetPrivate()->GetStreamedResourceData().Num() > 0 || Object->GetPrivate()->GetStreamedExtensionData().Num() > 0;
 }
 
 TArray<ICookPackageSplitter::FGeneratedPackage> FCustomizableObjectCookPackageSplitter::GetGenerateList(
@@ -167,13 +168,13 @@ TArray<ICookPackageSplitter::FGeneratedPackage> FCustomizableObjectCookPackageSp
 	TArray<ICookPackageSplitter::FGeneratedPackage> Result;
 
 	// Generate a new package for each streamed Resource Data
-	for (const FCustomizableObjectStreamedResourceData& StreamedData : Object->StreamedResourceData)
+	for (const FCustomizableObjectStreamedResourceData& StreamedData : Object->GetPrivate()->GetStreamedResourceData())
 	{
 		GenerateNewPackage(StreamedData, OwnerPackage, OwnerObject, Result);
 	}
 
 	// Generate a new package for each streamed Extension Data
-	for (const FCustomizableObjectStreamedResourceData& StreamedData : Object->StreamedExtensionData)
+	for (const FCustomizableObjectStreamedResourceData& StreamedData : Object->GetPrivate()->GetStreamedExtensionData())
 	{
 		GenerateNewPackage(StreamedData, OwnerPackage, OwnerObject, Result);
 	}
@@ -228,8 +229,8 @@ bool FCustomizableObjectCookPackageSplitter::PreSaveGeneratorPackage(
 	UCustomizableObject* Object = CastChecked<UCustomizableObject>(OwnerObject);
 
 	// There should be one generated package per streamed Resource Data
-	const int32 NumStreamedData = Object->StreamedResourceData.Num();
-	const int32 NumStreamedExtensionData = Object->StreamedExtensionData.Num();
+	const int32 NumStreamedData = Object->GetPrivate()->GetStreamedResourceData().Num();
+	const int32 NumStreamedExtensionData = Object->GetPrivate()->GetStreamedExtensionData().Num();
 	
 	check(NumStreamedData + NumStreamedExtensionData == PlaceholderPackages.Num());
 
@@ -243,7 +244,7 @@ bool FCustomizableObjectCookPackageSplitter::PreSaveGeneratorPackage(
 	for (int32 Index = 0; Index < NumStreamedData; ++Index)
 	{
 		const ICookPackageSplitter::FGeneratedPackageForPreSave& GeneratedPackage = PlaceholderPackages[Index];
-		if (!PreSavePackage(GeneratedPackage, Object->StreamedResourceData))
+		if (!PreSavePackage(GeneratedPackage, Object->GetPrivate()->GetStreamedResourceData()))
 		{
 			return false;
 		}
@@ -254,7 +255,7 @@ bool FCustomizableObjectCookPackageSplitter::PreSaveGeneratorPackage(
 	for (int32 Index = 0; Index < NumStreamedExtensionData; ++Index)
 	{
 		const ICookPackageSplitter::FGeneratedPackageForPreSave& GeneratedPackage = PlaceholderPackages[NumStreamedData + Index];
-		if (!PreSavePackage(GeneratedPackage, Object->StreamedExtensionData))
+		if (!PreSavePackage(GeneratedPackage, Object->GetPrivate()->GetStreamedExtensionData()))
 		{
 			return false;
 		}
@@ -277,11 +278,11 @@ void FCustomizableObjectCookPackageSplitter::PostSaveGeneratorPackage(UPackage*
 
 	for (const FString& ContainerName : SavedContainerNames)
 	{
-		FCustomizableObjectStreamedResourceData* ResourceData = FindStreamedResourceData(Object->StreamedResourceData, ContainerName);
+		FCustomizableObjectStreamedResourceData* ResourceData = FindStreamedResourceData(Object->GetPrivate()->GetStreamedResourceData(), ContainerName);
 		if (!ResourceData)
 		{
 			UE_LOG(LogMutable, Error, TEXT("Couldn't find streamed Resource Data container with name %s in array of %d entries"),
-				*ContainerName, Object->StreamedResourceData.Num());
+				*ContainerName, Object->GetPrivate()->GetStreamedResourceData().Num());
 
 			continue;
 		}
@@ -293,17 +294,17 @@ void FCustomizableObjectCookPackageSplitter::PostSaveGeneratorPackage(UPackage*
 		NewArray.Emplace(Container);
 	}
 
-	Object->StreamedResourceData = NewArray;
+	Object->GetPrivate()->GetStreamedResourceData() = NewArray;
 
 	NewArray.Reset(SavedExtensionContainerNames.Num());
 
 	for (const FString& ContainerName : SavedExtensionContainerNames)
 	{
-		FCustomizableObjectStreamedResourceData* ResourceData = FindStreamedResourceData(Object->StreamedExtensionData, ContainerName);
+		FCustomizableObjectStreamedResourceData* ResourceData = FindStreamedResourceData(Object->GetPrivate()->GetStreamedExtensionData(), ContainerName);
 		if (!ResourceData)
 		{
 			UE_LOG(LogMutable, Error, TEXT("Couldn't find streamed Extension Data container with name %s in array of %d entries"),
-				*ContainerName, Object->StreamedExtensionData.Num());
+				*ContainerName, Object->GetPrivate()->GetStreamedExtensionData().Num());
 
 			continue;
 		}
@@ -315,7 +316,7 @@ void FCustomizableObjectCookPackageSplitter::PostSaveGeneratorPackage(UPackage*
 		NewArray.Emplace(Container);
 	}
 
-	Object->StreamedExtensionData = NewArray;
+	Object->GetPrivate()->GetStreamedExtensionData() = NewArray;
 }
 
 bool FCustomizableObjectCookPackageSplitter::PopulateGeneratedPackage(
@@ -329,16 +330,16 @@ bool FCustomizableObjectCookPackageSplitter::PopulateGeneratedPackage(
 
 	UCustomizableObject* Object = CastChecked<UCustomizableObject>(OwnerObject);
 	
-	FCustomizableObjectStreamedResourceData* ResourceData = FindStreamedResourceData(Object->StreamedResourceData, GeneratedPackage.RelativePath);
+	FCustomizableObjectStreamedResourceData* ResourceData = FindStreamedResourceData(Object->GetPrivate()->GetStreamedResourceData(), GeneratedPackage.RelativePath);
 	if (!ResourceData)
 	{
-		ResourceData = FindStreamedResourceData(Object->StreamedExtensionData, GeneratedPackage.RelativePath);
+		ResourceData = FindStreamedResourceData(Object->GetPrivate()->GetStreamedExtensionData(), GeneratedPackage.RelativePath);
 	}
 
 	if (!ResourceData)
 	{
 		UE_LOG(LogMutable, Error, TEXT("Couldn't find streamed resource Data container with name %s in arrays of %d and %d entries"),
-			*GeneratedPackage.RelativePath, Object->StreamedResourceData.Num(), Object->StreamedExtensionData.Num());
+			*GeneratedPackage.RelativePath, Object->GetPrivate()->GetStreamedResourceData().Num(), Object->GetPrivate()->GetStreamedExtensionData().Num());
 
 		return false;
 	}
@@ -376,16 +377,16 @@ void FCustomizableObjectCookPackageSplitter::PostSaveGeneratedPackage(
 
 	UCustomizableObject* Object = CastChecked<UCustomizableObject>(OwnerObject);
 
-	FCustomizableObjectStreamedResourceData* ResourceData = FindStreamedResourceData(Object->StreamedResourceData, GeneratedPackage.RelativePath);
+	FCustomizableObjectStreamedResourceData* ResourceData = FindStreamedResourceData(Object->GetPrivate()->GetStreamedResourceData(), GeneratedPackage.RelativePath);
 	if (!ResourceData)
 	{
-		ResourceData = FindStreamedResourceData(Object->StreamedExtensionData, GeneratedPackage.RelativePath);
+		ResourceData = FindStreamedResourceData(Object->GetPrivate()->GetStreamedExtensionData(), GeneratedPackage.RelativePath);
 	}
 
 	if (!ResourceData)
 	{
 		UE_LOG(LogMutable, Error, TEXT("Couldn't find streamed resource Data container with name %s in arrays of %d and %d entries"),
-			*GeneratedPackage.RelativePath, Object->StreamedResourceData.Num(), Object->StreamedExtensionData.Num());
+			*GeneratedPackage.RelativePath, Object->GetPrivate()->GetStreamedResourceData().Num(), Object->GetPrivate()->GetStreamedExtensionData().Num());
 
 		return;
 	}
diff --git a/Engine/Plugins/Experimental/Mutable/Source/CustomizableObjectEditor/Private/MuCOE/SCustomizableInstanceProperties.cpp b/Engine/Plugins/Experimental/Mutable/Source/CustomizableObjectEditor/Private/MuCOE/SCustomizableInstanceProperties.cpp
--- a/Engine/Plugins/Experimental/Mutable/Source/CustomizableObjectEditor/Private/MuCOE/SCustomizableInstanceProperties.cpp
+++ b/Engine/Plugins/Experimental/Mutable/Source/CustomizableObjectEditor/Private/MuCOE/SCustomizableInstanceProperties.cpp
@@ -152,7 +152,7 @@ void SCustomizableInstanceProperties::Construct(const FArguments& InArgs)
 
 		if (ProfileIdx != INDEX_NONE)
 		{
-			FProfileParameterDat& Profile = CustomizableObject->InstancePropertiesProfiles[ProfileIdx];
+			FProfileParameterDat& Profile = CustomizableObject->GetPrivate()->GetInstancePropertiesProfiles()[ProfileIdx];
 			for (int i = 1; i < ParameterProfileNames.Num(); ++i)
 			{
 				if (ParameterProfileNames[i]->Equals(Profile.ProfileName))
@@ -2013,7 +2013,7 @@ FReply SCustomizableInstanceProperties::RemoveParameterProfile()
 		return FReply::Handled();
 	}
 
-	TArray<FProfileParameterDat>& Profiles = CustomizableObject->InstancePropertiesProfiles;
+	TArray<FProfileParameterDat>& Profiles = CustomizableObject->GetPrivate()->GetInstancePropertiesProfiles();
 	
 	Profiles.RemoveAt(ProfileIdx);
 	CustomInstance->GetPrivate()->SelectedProfileIndex = INDEX_NONE;
@@ -2051,7 +2051,7 @@ void SCustomizableInstanceProperties::OnProfileSelectedChanged(TSharedPtr<FStrin
 	}	
 
 	//Set selected profile
-	TArray<FProfileParameterDat>& Profiles = CustomInstance->GetCustomizableObject()->InstancePropertiesProfiles;
+	TArray<FProfileParameterDat>& Profiles = CustomInstance->GetCustomizableObject()->GetPrivate()->GetInstancePropertiesProfiles();
 	for (int32 Idx = 0; Idx < Profiles.Num(); ++Idx)
 	{
 		if (Profiles[Idx].ProfileName == *Selection)
@@ -2083,7 +2083,7 @@ void SCustomizableInstanceProperties::SetParameterProfileNamesOnEditor()
 {
 	ParameterProfileNames.Empty();
 
-	for (FProfileParameterDat& Profile : CustomInstance->GetCustomizableObject()->InstancePropertiesProfiles)
+	for (FProfileParameterDat& Profile : CustomInstance->GetCustomizableObject()->GetPrivate()->GetInstancePropertiesProfiles())
 	{
 		ParameterProfileNames.Emplace(MakeShared<FString>(Profile.ProfileName));
 	}
@@ -2332,7 +2332,7 @@ FReply SCreateProfileParameters::OnButtonClick(EAppReturnType::Type ButtonID)
 		{
 			CustomInstance->GetPrivate()->SaveParametersToProfile(CustomInstance->GetPrivate()->SelectedProfileIndex);
 		}
-		CustomInstance->GetPrivate()->SelectedProfileIndex = CustomizableObject->InstancePropertiesProfiles.Num() - 1;
+		CustomInstance->GetPrivate()->SelectedProfileIndex = CustomizableObject->GetPrivate()->GetInstancePropertiesProfiles().Num() - 1;
 
 		if (CIProperties)
 		{
diff --git a/Engine/Plugins/Experimental/Mutable/Source/CustomizableObjectPopulation/Private/MuCOP/CustomizableObjectPopulation.cpp b/Engine/Plugins/Experimental/Mutable/Source/CustomizableObjectPopulation/Private/MuCOP/CustomizableObjectPopulation.cpp
--- a/Engine/Plugins/Experimental/Mutable/Source/CustomizableObjectPopulation/Private/MuCOP/CustomizableObjectPopulation.cpp
+++ b/Engine/Plugins/Experimental/Mutable/Source/CustomizableObjectPopulation/Private/MuCOP/CustomizableObjectPopulation.cpp
@@ -3,6 +3,7 @@
 #include "MuCOP/CustomizableObjectPopulation.h"
 
 #include "MuCO/CustomizableObject.h"
+#include "MuCO/CustomizableObjectPrivate.h"
 #include "MuCOP/CustomizableObjectPopulationCharacteristic.h"
 #include "MuCOP/CustomizableObjectPopulationClass.h"
 #include "MuCOP/CustomizableObjectPopulationConstraint.h"
@@ -175,7 +176,7 @@ void CompileBoolSamplers(
 
 	for (FString& BoolParamName : InBoolParamNames)
 	{
-		FParameterTags* ParamTagsFound = PopulationClass->CustomizableObject->CustomizableObjectParametersTags.Find(BoolParamName);
+		FParameterTags* ParamTagsFound = PopulationClass->CustomizableObject->GetPrivate()->GetCustomizableObjectParametersTags().Find(BoolParamName);
 
 		const TArray<FString> EmptyTagArray;
 		const TArray<FString>& ParamTags = ParamTagsFound ? ParamTagsFound->Tags : EmptyTagArray;
@@ -531,7 +532,7 @@ void CompileIntSamplers(
 	//IntSamplers.Empty(InIntParamNames.Num());
 	for (FString& IntParamName : InIntParamNames)
 	{
-		FParameterTags* ParamTagsFound = PopulationClass->CustomizableObject->CustomizableObjectParametersTags.Find(IntParamName);
+		FParameterTags* ParamTagsFound = PopulationClass->CustomizableObject->GetPrivate()->GetCustomizableObjectParametersTags().Find(IntParamName);
 
 		TArray<FString> ParamTags = ParamTagsFound ? ParamTagsFound->Tags : TArray<FString>();
 
diff --git a/Engine/Plugins/Experimental/Mutable/Source/CustomizableObjectPopulationEditor/Private/MuCOPE/CustomizableObjectPopulationClassDetails.cpp b/Engine/Plugins/Experimental/Mutable/Source/CustomizableObjectPopulationEditor/Private/MuCOPE/CustomizableObjectPopulationClassDetails.cpp
--- a/Engine/Plugins/Experimental/Mutable/Source/CustomizableObjectPopulationEditor/Private/MuCOPE/CustomizableObjectPopulationClassDetails.cpp
+++ b/Engine/Plugins/Experimental/Mutable/Source/CustomizableObjectPopulationEditor/Private/MuCOPE/CustomizableObjectPopulationClassDetails.cpp
@@ -13,6 +13,7 @@
 #include "MuCOP/CustomizableObjectPopulationClass.h"
 #include "MuCOP/CustomizableObjectPopulationConstraint.h"
 #include "PropertyCustomizationHelpers.h"
+#include "MuCO/CustomizableObjectPrivate.h"
 #include "Widgets/Colors/SColorBlock.h"
 #include "Widgets/Colors/SColorPicker.h"
 #include "Widgets/Input/SButton.h"
@@ -377,11 +378,11 @@ void SPopulationClassTagList::Construct(const FArguments& InArgs)
 	bool bTagFound = false;
 
 	// Filling combobox options
-	for (int32 i = 0; i < PopulationClass->CustomizableObject->PopulationClassTags.Num(); ++i)
+	for (int32 i = 0; i < PopulationClass->CustomizableObject->GetPrivate()->GetPopulationClassTags().Num(); ++i)
 	{
-		AllTags.Add(MakeShareable(new FString(PopulationClass->CustomizableObject->PopulationClassTags[i])));
+		AllTags.Add(MakeShareable(new FString(PopulationClass->CustomizableObject->GetPrivate()->GetPopulationClassTags()[i])));
 
-		if (PopulationClass->CustomizableObject->PopulationClassTags[i].Equals(TagValue))
+		if (PopulationClass->CustomizableObject->GetPrivate()->GetPopulationClassTags()[i].Equals(TagValue))
 		{
 			SelectedTag = AllTags.Last();
 			bTagFound = true;
diff --git a/Engine/Plugins/Experimental/Mutable/Source/CustomizableObjectPopulationEditor/Private/MuCOPE/SCustomizableObjectPopulationClassTagManager.cpp b/Engine/Plugins/Experimental/Mutable/Source/CustomizableObjectPopulationEditor/Private/MuCOPE/SCustomizableObjectPopulationClassTagManager.cpp
--- a/Engine/Plugins/Experimental/Mutable/Source/CustomizableObjectPopulationEditor/Private/MuCOPE/SCustomizableObjectPopulationClassTagManager.cpp
+++ b/Engine/Plugins/Experimental/Mutable/Source/CustomizableObjectPopulationEditor/Private/MuCOPE/SCustomizableObjectPopulationClassTagManager.cpp
@@ -4,6 +4,7 @@
 
 #include "MuCOE/UnrealEditorPortabilityHelpers.h"
 #include "MuCO/CustomizableObject.h"
+#include "MuCO/CustomizableObjectPrivate.h"
 #include "Widgets/Input/SButton.h"
 #include "Widgets/Input/SEditableTextBox.h"
 #include "Widgets/Input/STextComboBox.h"
@@ -90,7 +91,7 @@ void SCustomizableObjectPopulationClassTagsManager::BuildTagSelector()
 		.Padding(0.0f, 3.0f, 0.0f, 0.0f)
 		[
 			SNew(STextBlock)
-			.Text(FText::FromString(FString::Printf(TEXT("%d Tags"),CustomizableObject->CustomizableObjectClassTags.Num())))
+			.Text(FText::FromString(FString::Printf(TEXT("%d Tags"),CustomizableObject->GetPrivate()->GetCustomizableObjectClassTags().Num())))
 		]
 		
 		+ SHorizontalBox::Slot()
@@ -107,21 +108,21 @@ void SCustomizableObjectPopulationClassTagsManager::BuildTagSelector()
 		]
 	];
 
-	for (int32 i = 0; i < RootObject->PopulationClassTags.Num(); ++i)
+	for (int32 i = 0; i < RootObject->GetPrivate()->GetPopulationClassTags().Num(); ++i)
 	{
-		TagOptions.Add(MakeShareable(new FString(RootObject->PopulationClassTags[i])));
+		TagOptions.Add(MakeShareable(new FString(RootObject->GetPrivate()->GetPopulationClassTags()[i])));
 	}
 
 	TagOptions.Add(MakeShareable(new FString("None")));
 
-	for (int32 i = 0; i < CustomizableObject->CustomizableObjectClassTags.Num(); ++i)
+	for (int32 i = 0; i < CustomizableObject->GetPrivate()->GetCustomizableObjectClassTags().Num(); ++i)
 	{
 		int32 TagIndex = TagOptions.Num()-1;
 		TSharedPtr<FString > SelectedOption;
 
-		for (int32 j = 0; j < RootObject->PopulationClassTags.Num(); ++j)
+		for (int32 j = 0; j < RootObject->GetPrivate()->GetPopulationClassTags().Num(); ++j)
 		{
-			if (CustomizableObject->CustomizableObjectClassTags[i] == RootObject->PopulationClassTags[j])
+			if (CustomizableObject->GetPrivate()->GetCustomizableObjectClassTags()[i] == RootObject->GetPrivate()->GetPopulationClassTags()[j])
 			{
 				TagIndex = j;
 			}
@@ -161,7 +162,7 @@ void SCustomizableObjectPopulationClassTagsManager::BuildTagSelector()
 
 FReply SCustomizableObjectPopulationClassTagsManager::OnAddTagSelectorButtonPressed()
 {
-	CustomizableObject->CustomizableObjectClassTags.Add("");
+	CustomizableObject->GetPrivate()->GetCustomizableObjectClassTags().Add("");
 	CustomizableObject->MarkPackageDirty();
 
 	RebuildWidget();
@@ -174,14 +175,14 @@ void SCustomizableObjectPopulationClassTagsManager::OnComboBoxSelectionChanged(T
 {
 	if (Selection.IsValid())
 	{
-		CustomizableObject->CustomizableObjectClassTags[Index] = *Selection;
+		CustomizableObject->GetPrivate()->GetCustomizableObjectClassTags()[Index] = *Selection;
 		CustomizableObject->MarkPackageDirty();
 	}
 }
 
 FReply SCustomizableObjectPopulationClassTagsManager::OnRemoveTagSelectorClicked(int32 Index)
 {
-	CustomizableObject->CustomizableObjectClassTags.RemoveAt(Index);
+	CustomizableObject->GetPrivate()->GetCustomizableObjectClassTags().RemoveAt(Index);
 	CustomizableObject->MarkPackageDirty();
 
 	RebuildWidget();
@@ -212,7 +213,7 @@ void SCustomizableObjectPopulationClassTagsManager::BuildTagManager()
 		.Padding(0.0f, 3.0f, 0.0f, 0.0f)
 		[
 			SNew(STextBlock)
-			.Text(FText::FromString(FString::Printf(TEXT("%d Tags"), RootObject->PopulationClassTags.Num())))
+			.Text(FText::FromString(FString::Printf(TEXT("%d Tags"), RootObject->GetPrivate()->GetPopulationClassTags().Num())))
 		]
 
 		+ SHorizontalBox::Slot()
@@ -229,7 +230,7 @@ void SCustomizableObjectPopulationClassTagsManager::BuildTagManager()
 		]
 	];
 
-	for (int32 i = 0; i < RootObject->PopulationClassTags.Num(); ++i)
+	for (int32 i = 0; i < RootObject->GetPrivate()->GetPopulationClassTags().Num(); ++i)
 	{
 		TagManager->AddSlot()
 		.AutoHeight()
@@ -249,7 +250,7 @@ void SCustomizableObjectPopulationClassTagsManager::BuildTagManager()
 			.Padding(10.0f, 0.0f, 0.0f, 0.0f)
 			[
 				SNew(SEditableTextBox)
-				.Text(FText::FromString(RootObject->PopulationClassTags[i]))
+				.Text(FText::FromString(RootObject->GetPrivate()->GetPopulationClassTags()[i]))
 				.OnTextCommitted(this, &SCustomizableObjectPopulationClassTagsManager::OnTextCommited, i)
 			]
 
@@ -272,7 +273,7 @@ void SCustomizableObjectPopulationClassTagsManager::BuildTagManager()
 
 FReply SCustomizableObjectPopulationClassTagsManager::OnAddTagButtonPressed()
 {
-	RootObject->PopulationClassTags.Add("");
+	RootObject->GetPrivate()->GetPopulationClassTags().Add("");
 	RootObject->MarkPackageDirty();
 
 	RebuildWidget();
@@ -283,7 +284,7 @@ FReply SCustomizableObjectPopulationClassTagsManager::OnAddTagButtonPressed()
 
 FReply SCustomizableObjectPopulationClassTagsManager::OnRemoveTagClicked(int32 Index)
 {
-	RootObject->PopulationClassTags.RemoveAt(Index);
+	RootObject->GetPrivate()->GetPopulationClassTags().RemoveAt(Index);
 	RootObject->MarkPackageDirty();
 
 	RebuildWidget();
@@ -296,7 +297,7 @@ void SCustomizableObjectPopulationClassTagsManager::OnTextCommited(const FText&
 {
 	if (InTextCommit == ETextCommit::OnEnter || InTextCommit == ETextCommit::OnUserMovedFocus)
 	{
-		RootObject->PopulationClassTags[Index] = NewText.ToString();
+		RootObject->GetPrivate()->GetPopulationClassTags()[Index] = NewText.ToString();
 		RootObject->MarkPackageDirty();
 
 		RebuildWidget();
diff --git a/Engine/Plugins/Experimental/Mutable/Source/CustomizableObjectPopulationEditor/Private/MuCOPE/SCustomizableObjectPopulationClassTagsTool.cpp b/Engine/Plugins/Experimental/Mutable/Source/CustomizableObjectPopulationEditor/Private/MuCOPE/SCustomizableObjectPopulationClassTagsTool.cpp
--- a/Engine/Plugins/Experimental/Mutable/Source/CustomizableObjectPopulationEditor/Private/MuCOPE/SCustomizableObjectPopulationClassTagsTool.cpp
+++ b/Engine/Plugins/Experimental/Mutable/Source/CustomizableObjectPopulationEditor/Private/MuCOPE/SCustomizableObjectPopulationClassTagsTool.cpp
@@ -3,6 +3,7 @@
 #include "MuCOPE/SCustomizableObjectPopulationClassTagsTool.h"
 
 #include "MuCO/CustomizableObject.h"
+#include "MuCO/CustomizableObjectPrivate.h"
 #include "MuCOP/CustomizableObjectPopulationClass.h"
 #include "MuCOPE/CustomizableObjectPopulationClassDetails.h"
 #include "MuCOE/UnrealEditorPortabilityHelpers.h"
@@ -247,9 +248,9 @@ void SCustomizableObjectPopulationClassTagsTool::GenerateTagsViewer()
 {
 	TagsViewer->ClearChildren();
 
-	for (int32 i = 0; i < PopulationClass->CustomizableObject->PopulationClassTags.Num(); ++i)
+	for (int32 i = 0; i < PopulationClass->CustomizableObject->GetPrivate()->GetPopulationClassTags().Num(); ++i)
 	{
-		FString TagValue = PopulationClass->CustomizableObject->PopulationClassTags[i];
+		FString TagValue = PopulationClass->CustomizableObject->GetPrivate()->GetPopulationClassTags()[i];
 
 		if (SearchItem.IsEmpty())
 		{
@@ -468,13 +469,13 @@ void SCustomizableObjectPopulationClassTagsTool::GenerateParameterOptionsManager
 							]
 						];
 
-						if (PopulationClass->CustomizableObject->CustomizableObjectParametersTags.Find(SelectedParameter))
+						if (PopulationClass->CustomizableObject->GetPrivate()->GetCustomizableObjectParametersTags().Find(SelectedParameter))
 						{
-							if (PopulationClass->CustomizableObject->CustomizableObjectParametersTags[SelectedParameter].ParameterOptions.Find(ParameterOptionName))
+							if (PopulationClass->CustomizableObject->GetPrivate()->GetCustomizableObjectParametersTags()[SelectedParameter].ParameterOptions.Find(ParameterOptionName))
 							{
 								float LineLength = 40.0f;
 
-								LineLength += 40.0f * (int)((PopulationClass->CustomizableObject->CustomizableObjectParametersTags[SelectedParameter].ParameterOptions[ParameterOptionName].Tags.Num() - 1) / 4);
+								LineLength += 40.0f * (int)((PopulationClass->CustomizableObject->GetPrivate()->GetCustomizableObjectParametersTags()[SelectedParameter].ParameterOptions[ParameterOptionName].Tags.Num() - 1) / 4);
 								
 								VerticalWidget->AddSlot()
 								.AutoWidth()
@@ -486,15 +487,15 @@ void SCustomizableObjectPopulationClassTagsTool::GenerateParameterOptionsManager
 								];
 
 								// Delete the parameter option from the parameter options map if it doesn't have any tag
-								if (PopulationClass->CustomizableObject->CustomizableObjectParametersTags[SelectedParameter].ParameterOptions[ParameterOptionName].Tags.Num() == 0)
+								if (PopulationClass->CustomizableObject->GetPrivate()->GetCustomizableObjectParametersTags()[SelectedParameter].ParameterOptions[ParameterOptionName].Tags.Num() == 0)
 								{
-									PopulationClass->CustomizableObject->CustomizableObjectParametersTags[SelectedParameter].ParameterOptions.Remove(ParameterOptionName);
+									PopulationClass->CustomizableObject->GetPrivate()->GetCustomizableObjectParametersTags()[SelectedParameter].ParameterOptions.Remove(ParameterOptionName);
 
 									// Check if we also have to delete the parameter 
-									if (PopulationClass->CustomizableObject->CustomizableObjectParametersTags[SelectedParameter].Tags.Num() == 0
-										&& PopulationClass->CustomizableObject->CustomizableObjectParametersTags[SelectedParameter].ParameterOptions.Num() == 0)
+									if (PopulationClass->CustomizableObject->GetPrivate()->GetCustomizableObjectParametersTags()[SelectedParameter].Tags.Num() == 0
+										&& PopulationClass->CustomizableObject->GetPrivate()->GetCustomizableObjectParametersTags()[SelectedParameter].ParameterOptions.Num() == 0)
 									{
-										PopulationClass->CustomizableObject->CustomizableObjectParametersTags.Remove(SelectedParameter);
+										PopulationClass->CustomizableObject->GetPrivate()->GetCustomizableObjectParametersTags().Remove(SelectedParameter);
 									}
 								}
 								else
@@ -505,7 +506,7 @@ void SCustomizableObjectPopulationClassTagsTool::GenerateParameterOptionsManager
 									[
 										SNew(SParameterTagWidget)
 										.PopulationClass(PopulationClass)
-										.TagsList(&PopulationClass->CustomizableObject->CustomizableObjectParametersTags[SelectedParameter].ParameterOptions[ParameterOptionName].Tags)
+										.TagsList(&PopulationClass->CustomizableObject->GetPrivate()->GetCustomizableObjectParametersTags()[SelectedParameter].ParameterOptions[ParameterOptionName].Tags)
 										.TagsToolPtr(this)
 									];
 								}
@@ -542,9 +543,9 @@ FReply SCustomizableObjectPopulationClassTagsTool::OnAddTagButtonPressed()
 {
 	FString NewTag = SearchBoxWidget->GetText().ToString();
 
-	if (!NewTag.IsEmpty() && PopulationClass->CustomizableObject->PopulationClassTags.Find(NewTag) == INDEX_NONE)
+	if (!NewTag.IsEmpty() && PopulationClass->CustomizableObject->GetPrivate()->GetPopulationClassTags().Find(NewTag) == INDEX_NONE)
 	{
-		PopulationClass->CustomizableObject->PopulationClassTags.Add(NewTag);
+		PopulationClass->CustomizableObject->GetPrivate()->GetPopulationClassTags().Add(NewTag);
 		PopulationClass->CustomizableObject->MarkPackageDirty();
 		SearchItem.Empty();
 		GenerateChildSlot();
@@ -557,7 +558,7 @@ FReply SCustomizableObjectPopulationClassTagsTool::OnAddTagButtonPressed()
 FReply SCustomizableObjectPopulationClassTagsTool::OnRemoveTagButtonPressed(int32 Index)
 {
 	PopulationClass->CustomizableObject->MarkPackageDirty();
-	PopulationClass->CustomizableObject->PopulationClassTags.RemoveAt(Index);
+	PopulationClass->CustomizableObject->GetPrivate()->GetPopulationClassTags().RemoveAt(Index);
 	GenerateChildSlot();
 
 	return FReply::Unhandled();
@@ -578,9 +579,9 @@ void SCustomizableObjectPopulationClassTagsTool::OnSearchBoxTextCommitted(const
 		FString NewTag = InText.ToString();
 
 		// If the written tag doesn't exist, then we create a new one
-		if (!NewTag.IsEmpty() && PopulationClass->CustomizableObject->PopulationClassTags.Find(NewTag) == INDEX_NONE)
+		if (!NewTag.IsEmpty() && PopulationClass->CustomizableObject->GetPrivate()->GetPopulationClassTags().Find(NewTag) == INDEX_NONE)
 		{
-			PopulationClass->CustomizableObject->PopulationClassTags.Add(InText.ToString());
+			PopulationClass->CustomizableObject->GetPrivate()->GetPopulationClassTags().Add(InText.ToString());
 			PopulationClass->CustomizableObject->MarkPackageDirty();
 			SearchItem.Empty();
 			GenerateChildSlot();
@@ -607,15 +608,15 @@ FReply SCustomizableObjectPopulationClassTagsTool::OnParameterAddTagButtonPresse
 	{
 		FString SelectedParameter = *ParameterComboBox->GetSelectedItem();
 
-		if (PopulationClass->CustomizableObject->CustomizableObjectParametersTags.Find(SelectedParameter))
+		if (PopulationClass->CustomizableObject->GetPrivate()->GetCustomizableObjectParametersTags().Find(SelectedParameter))
 		{
-			PopulationClass->CustomizableObject->CustomizableObjectParametersTags[SelectedParameter].Tags.Add("");
+			PopulationClass->CustomizableObject->GetPrivate()->GetCustomizableObjectParametersTags()[SelectedParameter].Tags.Add("");
 			GenerateChildSlot();
 		}
 		else
 		{
-			PopulationClass->CustomizableObject->CustomizableObjectParametersTags.Add(SelectedParameter);
-			PopulationClass->CustomizableObject->CustomizableObjectParametersTags[SelectedParameter].Tags.Add("");
+			PopulationClass->CustomizableObject->GetPrivate()->GetCustomizableObjectParametersTags().Add(SelectedParameter);
+			PopulationClass->CustomizableObject->GetPrivate()->GetCustomizableObjectParametersTags()[SelectedParameter].Tags.Add("");
 			GenerateChildSlot();
 		}
 
@@ -633,20 +634,20 @@ FReply SCustomizableObjectPopulationClassTagsTool::OnParameterOptionAddTagButton
 	{
 		FString SelectedParameter = *ParameterComboBox->GetSelectedItem();
 
-		if (!PopulationClass->CustomizableObject->CustomizableObjectParametersTags.Find(SelectedParameter))
+		if (!PopulationClass->CustomizableObject->GetPrivate()->GetCustomizableObjectParametersTags().Find(SelectedParameter))
 		{
-			PopulationClass->CustomizableObject->CustomizableObjectParametersTags.Add(SelectedParameter);
+			PopulationClass->CustomizableObject->GetPrivate()->GetCustomizableObjectParametersTags().Add(SelectedParameter);
 		}
 
-		if (PopulationClass->CustomizableObject->CustomizableObjectParametersTags[SelectedParameter].ParameterOptions.Find(ParameterOptionName))
+		if (PopulationClass->CustomizableObject->GetPrivate()->GetCustomizableObjectParametersTags()[SelectedParameter].ParameterOptions.Find(ParameterOptionName))
 		{
-			PopulationClass->CustomizableObject->CustomizableObjectParametersTags[SelectedParameter].ParameterOptions[ParameterOptionName].Tags.Add("");
+			PopulationClass->CustomizableObject->GetPrivate()->GetCustomizableObjectParametersTags()[SelectedParameter].ParameterOptions[ParameterOptionName].Tags.Add("");
 			GenerateChildSlot();
 		}
 		else
 		{
-			PopulationClass->CustomizableObject->CustomizableObjectParametersTags[SelectedParameter].ParameterOptions.Add(ParameterOptionName);
-			PopulationClass->CustomizableObject->CustomizableObjectParametersTags[SelectedParameter].ParameterOptions[ParameterOptionName].Tags.Add("");
+			PopulationClass->CustomizableObject->GetPrivate()->GetCustomizableObjectParametersTags()[SelectedParameter].ParameterOptions.Add(ParameterOptionName);
+			PopulationClass->CustomizableObject->GetPrivate()->GetCustomizableObjectParametersTags()[SelectedParameter].ParameterOptions[ParameterOptionName].Tags.Add("");
 			GenerateChildSlot();
 		}
 
@@ -760,11 +761,11 @@ void SSingleTagWidget::Construct(const FArguments& InArgs)
 	ComboBoxTagsOptions.Add(MakeShareable(new FString("None")));
 	TagsComboBox->SetSelectedItem(ComboBoxTagsOptions.Last());
 
-	for (int32 i = 0; i < PopulationClass->CustomizableObject->PopulationClassTags.Num(); ++i)
+	for (int32 i = 0; i < PopulationClass->CustomizableObject->GetPrivate()->GetPopulationClassTags().Num(); ++i)
 	{
-		ComboBoxTagsOptions.Add(MakeShareable(new FString(PopulationClass->CustomizableObject->PopulationClassTags[i])));
+		ComboBoxTagsOptions.Add(MakeShareable(new FString(PopulationClass->CustomizableObject->GetPrivate()->GetPopulationClassTags()[i])));
 
-		if ((*TagsList)[TagIndex] == PopulationClass->CustomizableObject->PopulationClassTags[i])
+		if ((*TagsList)[TagIndex] == PopulationClass->CustomizableObject->GetPrivate()->GetPopulationClassTags()[i])
 		{
 			TagsComboBox->SetSelectedItem(ComboBoxTagsOptions.Last());
 		}
