[mutable] Expand disk-cache options and collect stats. Fix concurrency problem with temp file creation.
#rnx

diff --git a/Engine/Plugins/Experimental/Mutable/Source/CustomizableObjectEditor/Private/MuCOE/CustomizableObjectCompileRunnable.cpp b/Engine/Plugins/Experimental/Mutable/Source/CustomizableObjectEditor/Private/MuCOE/CustomizableObjectCompileRunnable.cpp
--- a/Engine/Plugins/Experimental/Mutable/Source/CustomizableObjectEditor/Private/MuCOE/CustomizableObjectCompileRunnable.cpp
+++ b/Engine/Plugins/Experimental/Mutable/Source/CustomizableObjectEditor/Private/MuCOE/CustomizableObjectCompileRunnable.cpp
@@ -22,6 +22,19 @@ class ITargetPlatform;
 #define UE_MUTABLE_CORE_REGION	TEXT("Mutable Core")
 
 
+TAutoConsoleVariable<bool> CVarMutableCompilerConcurrency(
+	TEXT("mutable.ForceCompilerConcurrency"),
+	true,
+	TEXT("Force the use of multithreading when compiling CustomizableObjects both in editor and cook commandlets."),
+	ECVF_Default);
+
+TAutoConsoleVariable<bool> CVarMutableCompilerDiskCache(
+	TEXT("mutable.ForceCompilerDiskCache"),
+	false,
+	TEXT("Force the use of disk cache to reduce memory usage when compiling CustomizableObjects both in editor and cook commandlets."),
+	ECVF_Default);
+
+
 FCustomizableObjectCompileRunnable::FCustomizableObjectCompileRunnable(mu::Ptr<mu::Node> Root)
 	: MutableRoot(Root)
 	, bThreadCompleted(false)
@@ -79,8 +92,20 @@ uint32 FCustomizableObjectCompileRunnable::Run()
 	mu::Ptr<mu::CompilerOptions> CompilerOptions = new mu::CompilerOptions();
 
 	bool bUseConcurrency = !Options.bIsCooking;
+	if (CVarMutableCompilerConcurrency->GetBool())
+	{
+		bUseConcurrency = true;
+	}
+
 	CompilerOptions->SetUseConcurrency(bUseConcurrency);
-	CompilerOptions->SetUseDiskCache(Options.bUseDiskCompilation);
+
+	bool bUseDiskCache = Options.bUseDiskCompilation;
+	if (CVarMutableCompilerDiskCache->GetBool())
+	{
+		bUseDiskCache = true;
+	}
+
+	CompilerOptions->SetUseDiskCache(bUseDiskCache);
 
 	if (Options.OptimizationLevel > 2)
 	{
@@ -227,6 +252,8 @@ uint32 FCustomizableObjectCompileRunnable::Run()
 
 	UE_LOG(LogMutable, Verbose, TEXT("PROFILE: [ %16.8f ] FCustomizableObjectCompileRunnable::Run end."), FPlatformTime::Seconds());
 
+	CompilerOptions->LogStats();
+
 	TRACE_END_REGION(UE_MUTABLE_CORE_REGION);
 
 	return Result;
diff --git a/Engine/Plugins/Experimental/Mutable/Source/MutableTools/Internal/MuT/AST.h b/Engine/Plugins/Experimental/Mutable/Source/MutableTools/Internal/MuT/AST.h
--- a/Engine/Plugins/Experimental/Mutable/Source/MutableTools/Internal/MuT/AST.h
+++ b/Engine/Plugins/Experimental/Mutable/Source/MutableTools/Internal/MuT/AST.h
@@ -50,6 +50,7 @@ namespace mu
 {
 	class ASTOp;
 	class ASTOpFixed;
+	struct FProxyFileContext;
 
 	template<typename T>
 	inline uint32 GetTypeHash(const Ptr<const T>& p)
@@ -193,8 +194,8 @@ namespace mu
         //! Maximum number of iterations when optimising models. If 0 as many as necessary will be performed.
         int32 MaxOptimisationLoopCount = 8;
 
-        //! Store resource data in disk instead of memory
-        bool bUseDiskCache = false;
+        /** If valied, store resource data in disk instead of memory. */
+		FProxyFileContext* DiskCacheContext = nullptr;
 
 		/** Compile optimizing for the generation of smaller mipmaps of every image. */
 		bool bEnableProgressiveImages = false;
diff --git a/Engine/Plugins/Experimental/Mutable/Source/MutableTools/Private/MuT/ASTOpConstantResource.cpp b/Engine/Plugins/Experimental/Mutable/Source/MutableTools/Private/MuT/ASTOpConstantResource.cpp
--- a/Engine/Plugins/Experimental/Mutable/Source/MutableTools/Private/MuT/ASTOpConstantResource.cpp
+++ b/Engine/Plugins/Experimental/Mutable/Source/MutableTools/Private/MuT/ASTOpConstantResource.cpp
@@ -456,7 +456,7 @@ namespace mu
 
 
 	//-------------------------------------------------------------------------------------------------
-	void ASTOpConstantResource::SetValue(const Ptr<const RefCounted>& v, bool useDiskCache)
+	void ASTOpConstantResource::SetValue(const Ptr<const RefCounted>& v, FProxyFileContext* DiskCacheContext)
 	{
 		switch (type)
 		{
@@ -470,9 +470,9 @@ namespace mu
 
 			hash = CityHash64(static_cast<const char*>(stream.GetBuffer()), stream.GetBufferSize());
 
-			if (useDiskCache)
+			if (DiskCacheContext)
 			{
-				proxy = new ResourceProxyTempFile<Image>(r.get());
+				proxy = new ResourceProxyTempFile<Image>(r.get(), *DiskCacheContext);
 			}
 			else
 			{
diff --git a/Engine/Plugins/Experimental/Mutable/Source/MutableTools/Private/MuT/ASTOpConstantResource.h b/Engine/Plugins/Experimental/Mutable/Source/MutableTools/Private/MuT/ASTOpConstantResource.h
--- a/Engine/Plugins/Experimental/Mutable/Source/MutableTools/Private/MuT/ASTOpConstantResource.h
+++ b/Engine/Plugins/Experimental/Mutable/Source/MutableTools/Private/MuT/ASTOpConstantResource.h
@@ -13,12 +13,9 @@
 namespace mu
 {
 	struct FProgram;
-	template <class SCALAR> class vec4;
+	struct FProxyFileContext;
 
-
-	//---------------------------------------------------------------------------------------------
-	//! A constant mesh, image, volume or layout
-	//---------------------------------------------------------------------------------------------
+	/** A constant mesh, image, volume or layout. */
 	class ASTOpConstantResource final : public ASTOp
 	{
 	private:
@@ -50,8 +47,10 @@ namespace mu
 		//! Get a copy of the stored value
 		Ptr<const RefCounted> GetValue() const;
 
-		//! Set the value to store in this op
-		void SetValue(const Ptr<const RefCounted>& v, bool useDiskCache);
+		/** Set the value to store in this op.
+		* If the DiskCacheContext is not null, the disk cache will be used.
+		*/
+		void SetValue(const Ptr<const RefCounted>& v, FProxyFileContext* DiskCacheContext);
 
 
 		// ASTOp interface
diff --git a/Engine/Plugins/Experimental/Mutable/Source/MutableTools/Private/MuT/ASTOpMeshFormat.cpp b/Engine/Plugins/Experimental/Mutable/Source/MutableTools/Private/MuT/ASTOpMeshFormat.cpp
--- a/Engine/Plugins/Experimental/Mutable/Source/MutableTools/Private/MuT/ASTOpMeshFormat.cpp
+++ b/Engine/Plugins/Experimental/Mutable/Source/MutableTools/Private/MuT/ASTOpMeshFormat.cpp
@@ -282,7 +282,7 @@ mu::Ptr<ASTOp> Sink_MeshFormatAST::Visit(const mu::Ptr<ASTOp>& at, const ASTOpMe
 
 		mu::Ptr<ASTOpConstantResource> motaop = new ASTOpConstantResource();
 		motaop->type = OP_TYPE::ME_CONSTANT;
-		motaop->SetValue(pTargetMorphFormat, false /* useDiskCache */);
+		motaop->SetValue(pTargetMorphFormat, nullptr);
 		auto targetMorphFormatAt = motaop;
 
 		if (newOp->Target)
@@ -322,7 +322,7 @@ mu::Ptr<ASTOp> Sink_MeshFormatAST::Visit(const mu::Ptr<ASTOp>& at, const ASTOpMe
 
 		mu::Ptr<ASTOpConstantResource> motaop = new ASTOpConstantResource();
 		motaop->type = OP_TYPE::ME_CONSTANT;
-		motaop->SetValue(pTargetMorphFormat, false /* useDiskCache */);
+		motaop->SetValue(pTargetMorphFormat, nullptr );
 		auto targetMorphFormatAt = motaop;
 
 		for (int t = 0; t < MUTABLE_OP_MAX_INTERPOLATE_COUNT - 1; ++t)
diff --git a/Engine/Plugins/Experimental/Mutable/Source/MutableTools/Private/MuT/CodeGenerator.cpp b/Engine/Plugins/Experimental/Mutable/Source/MutableTools/Private/MuT/CodeGenerator.cpp
--- a/Engine/Plugins/Experimental/Mutable/Source/MutableTools/Private/MuT/CodeGenerator.cpp
+++ b/Engine/Plugins/Experimental/Mutable/Source/MutableTools/Private/MuT/CodeGenerator.cpp
@@ -981,7 +981,7 @@ namespace mu
 							Ptr<ASTOpConstantResource> op = new ASTOpConstantResource();
 							op->type = OP_TYPE::LA_CONSTANT;
 
-							op->SetValue(pLayout, m_compilerOptions->OptimisationOptions.bUseDiskCache);
+							op->SetValue(pLayout, m_compilerOptions->OptimisationOptions.DiskCacheContext);
 							layoutOp = op;
 						}
 
@@ -1008,7 +1008,7 @@ namespace mu
 									op->type = OP_TYPE::LA_CONSTANT;
 
 									Ptr<const Layout> pCloned = data.GeneratedLayouts[LayoutIndex];
-									op->SetValue(pCloned, m_compilerOptions->OptimisationOptions.bUseDiskCache);
+									op->SetValue(pCloned, m_compilerOptions->OptimisationOptions.DiskCacheContext);
 
 									layoutFragmentAd = op;
 								}
diff --git a/Engine/Plugins/Experimental/Mutable/Source/MutableTools/Private/MuT/CodeGenerator_Image.cpp b/Engine/Plugins/Experimental/Mutable/Source/MutableTools/Private/MuT/CodeGenerator_Image.cpp
--- a/Engine/Plugins/Experimental/Mutable/Source/MutableTools/Private/MuT/CodeGenerator_Image.cpp
+++ b/Engine/Plugins/Experimental/Mutable/Source/MutableTools/Private/MuT/CodeGenerator_Image.cpp
@@ -211,7 +211,7 @@ namespace mu
 		{
 			Ptr<ASTOpConstantResource> op = new ASTOpConstantResource();
 			op->type = OP_TYPE::IM_CONSTANT;
-			op->SetValue(pImage, m_compilerOptions->OptimisationOptions.bUseDiskCache);
+			op->SetValue(pImage, m_compilerOptions->OptimisationOptions.DiskCacheContext);
 			Result.op = op;
 		}
 
@@ -1441,7 +1441,7 @@ namespace mu
 				Ptr<Mesh> FormatMeshResult = new Mesh();
 				CreateMeshOptimisedForWrappingProjection(FormatMeshResult.get(), node.m_layout);
 
-                cop->SetValue(FormatMeshResult, m_compilerOptions->OptimisationOptions.bUseDiskCache);
+                cop->SetValue(FormatMeshResult, m_compilerOptions->OptimisationOptions.DiskCacheContext);
 
                 Ptr<ASTOpMeshFormat> FormatOp = new ASTOpMeshFormat();
 				FormatOp->Buffers = OP::MeshFormatArgs::BT_VERTEX
@@ -1473,7 +1473,7 @@ namespace mu
 				Ptr<Mesh> FormatMeshResult = new Mesh();
                 CreateMeshOptimisedForProjection(FormatMeshResult.get(), node.m_layout);
 
-                cop->SetValue(FormatMeshResult, m_compilerOptions->OptimisationOptions.bUseDiskCache);
+                cop->SetValue(FormatMeshResult, m_compilerOptions->OptimisationOptions.DiskCacheContext);
 
                 Ptr<ASTOpMeshFormat> FormatOp = new ASTOpMeshFormat();
 				FormatOp->Buffers = OP::MeshFormatArgs::BT_VERTEX
@@ -1494,7 +1494,7 @@ namespace mu
             Ptr<const Mesh> TempMesh = new Mesh();
             Ptr<ASTOpConstantResource> cop = new ASTOpConstantResource();
             cop->type = OP_TYPE::ME_CONSTANT;
-            cop->SetValue(TempMesh, m_compilerOptions->OptimisationOptions.bUseDiskCache );
+            cop->SetValue(TempMesh, m_compilerOptions->OptimisationOptions.DiskCacheContext);
 			ProjectOp->SetChild(ProjectOp->op.args.MeshProject.mesh, cop );
             m_pErrorLog->GetPrivate()->Add( "Projector mesh not set.", ELMT_ERROR, node.m_errorContext );
         }
diff --git a/Engine/Plugins/Experimental/Mutable/Source/MutableTools/Private/MuT/CodeGenerator_Mesh.cpp b/Engine/Plugins/Experimental/Mutable/Source/MutableTools/Private/MuT/CodeGenerator_Mesh.cpp
--- a/Engine/Plugins/Experimental/Mutable/Source/MutableTools/Private/MuT/CodeGenerator_Mesh.cpp
+++ b/Engine/Plugins/Experimental/Mutable/Source/MutableTools/Private/MuT/CodeGenerator_Mesh.cpp
@@ -1265,7 +1265,7 @@ namespace mu
 		{
 			// This data is required
 			MeshPtr pTempMesh = new Mesh();
-			op->SetValue(pTempMesh, m_compilerOptions->OptimisationOptions.bUseDiskCache);
+			op->SetValue(pTempMesh, m_compilerOptions->OptimisationOptions.DiskCacheContext);
 			m_constantMeshes.Add(pTempMesh);
 
 			// Log an error message
@@ -1448,7 +1448,7 @@ namespace mu
 			FinalMesh = pCloned;
 		}
 
-		op->SetValue(FinalMesh, m_compilerOptions->OptimisationOptions.bUseDiskCache);
+		op->SetValue(FinalMesh, m_compilerOptions->OptimisationOptions.DiskCacheContext);
 
 		Ptr<ASTOp> LastMeshOp = op;
 
@@ -1505,7 +1505,7 @@ namespace mu
 
             Ptr<ASTOpConstantResource> cop = new ASTOpConstantResource();
             cop->type = OP_TYPE::ME_CONSTANT;
-            cop->SetValue( pFormatMesh, m_compilerOptions->OptimisationOptions.bUseDiskCache );
+            cop->SetValue( pFormatMesh, m_compilerOptions->OptimisationOptions.DiskCacheContext );
             op->Format = cop;
 
             m_constantMeshes.Add(pFormatMesh);
diff --git a/Engine/Plugins/Experimental/Mutable/Source/MutableTools/Private/MuT/CodeOptimiser.cpp b/Engine/Plugins/Experimental/Mutable/Source/MutableTools/Private/MuT/CodeOptimiser.cpp
--- a/Engine/Plugins/Experimental/Mutable/Source/MutableTools/Private/MuT/CodeOptimiser.cpp
+++ b/Engine/Plugins/Experimental/Mutable/Source/MutableTools/Private/MuT/CodeOptimiser.cpp
@@ -454,7 +454,7 @@ namespace mu
 
 		// input
 		Ptr<ASTOp> Source;
-		bool bUseDiskCache = false;
+		FProxyFileContext* DiskCacheContext = nullptr;
 		int32 ImageCompressionQuality = 0;
 		int32 OptimizationPass = 0;
 		FReferencedResourceFunc ReferencedResourceProvider;
@@ -471,7 +471,7 @@ namespace mu
 		{
 			OptimizationPass = InOptimizationPass;
 			Source = InSource;
-			bUseDiskCache = InOptions->OptimisationOptions.bUseDiskCache;
+			DiskCacheContext = InOptions->OptimisationOptions.DiskCacheContext;
 			ImageCompressionQuality = InOptions->ImageCompressionQuality;
 			ReferencedResourceProvider = InOptions->OptimisationOptions.ReferencedResourceProvider;
 		}
@@ -521,7 +521,7 @@ namespace mu
 				{
 					mu::Ptr<ASTOpConstantResource> constantOp = new ASTOpConstantResource();
 					constantOp->type = OP_TYPE::ME_CONSTANT;
-					constantOp->SetValue( pMesh, bUseDiskCache );
+					constantOp->SetValue( pMesh, DiskCacheContext );
 					Result = constantOp;
 				  }
 				break;
@@ -537,7 +537,7 @@ namespace mu
 				{
 					mu::Ptr<ASTOpConstantResource> constantOp = new ASTOpConstantResource();
 					constantOp->type = OP_TYPE::IM_CONSTANT;
-					constantOp->SetValue( pImage, bUseDiskCache );
+					constantOp->SetValue( pImage, DiskCacheContext );
 					Result = constantOp;
 				}
 				break;
@@ -553,7 +553,7 @@ namespace mu
 				{
 					mu::Ptr<ASTOpConstantResource> constantOp = new ASTOpConstantResource();
 					constantOp->type = OP_TYPE::LA_CONSTANT;
-					constantOp->SetValue( pLayout, bUseDiskCache );
+					constantOp->SetValue( pLayout, DiskCacheContext);
 					Result = constantOp;
 				}
 				break;
@@ -823,7 +823,7 @@ namespace mu
 
 									Ptr<ASTOpConstantResource> ConstantOp = new ASTOpConstantResource;
 									ConstantOp->type = OP_TYPE::IM_CONSTANT;
-									ConstantOp->SetValue(ResolveImage->get(), InOptions->OptimisationOptions.bUseDiskCache);
+									ConstantOp->SetValue(ResolveImage->get(), InOptions->OptimisationOptions.DiskCacheContext);
 									ASTOp::Replace(SubgraphRoot, ConstantOp);
 								},
 								ReferenceCompletionEvent,
@@ -938,7 +938,7 @@ namespace mu
 								ReferenceCompletionEvent.Wait();
 								Ptr<ASTOpConstantResource> ConstantOp = new ASTOpConstantResource;
 								ConstantOp->type = OP_TYPE::IM_CONSTANT;
-								ConstantOp->SetValue(ResolveImage->get(), InOptions->OptimisationOptions.bUseDiskCache);
+								ConstantOp->SetValue(ResolveImage->get(), InOptions->OptimisationOptions.DiskCacheContext);
 								ASTOp::Replace(SubgraphRoot, ConstantOp);
 							};
 						ConstantSubgraphs[Index].NonConcurrentTask = ImmediateCompleteFunc;
@@ -1353,7 +1353,7 @@ namespace mu
 							NewMesh->CheckIntegrity();
 							mu::Ptr<ASTOpConstantResource> newOp = new ASTOpConstantResource();
 							newOp->type = OP_TYPE::ME_CONSTANT;
-							newOp->SetValue(NewMesh, options.bUseDiskCache);
+							newOp->SetValue(NewMesh, options.DiskCacheContext);
 
 							ASTOp::Replace(at, newOp);
 						}
@@ -1509,8 +1509,7 @@ namespace mu
 					ASTOpConstantResource* typed = static_cast<ASTOpConstantResource*>(n.get());
 					auto pMesh = static_cast<const Mesh*>(typed->GetValue().get());
 					pMesh->ResetStaticFormatFlags();
-					typed->SetValue( pMesh,
-									 m_options->GetPrivate()->OptimisationOptions.bUseDiskCache );
+					typed->SetValue( pMesh, m_options->GetPrivate()->OptimisationOptions.DiskCacheContext);
 				}
 			});
 
diff --git a/Engine/Plugins/Experimental/Mutable/Source/MutableTools/Private/MuT/Compiler.cpp b/Engine/Plugins/Experimental/Mutable/Source/MutableTools/Private/MuT/Compiler.cpp
--- a/Engine/Plugins/Experimental/Mutable/Source/MutableTools/Private/MuT/Compiler.cpp
+++ b/Engine/Plugins/Experimental/Mutable/Source/MutableTools/Private/MuT/Compiler.cpp
@@ -34,7 +34,8 @@
 #include "MuT/Table.h"
 #include "Trace/Detail/Channel.h"
 
-#include <string>
+#include <string>		// Required for deserialisation of old data
+#include <inttypes.h>	// Required for 64-bit printf macros
 
 
 namespace mu
@@ -103,9 +104,9 @@ namespace mu
 
 
 	//---------------------------------------------------------------------------------------------
-	void CompilerOptions::SetUseDiskCache(bool enabled)
+	void CompilerOptions::SetUseDiskCache(bool bEnabled)
 	{
-		m_pD->OptimisationOptions.bUseDiskCache = enabled;
+		m_pD->OptimisationOptions.DiskCacheContext = bEnabled ? &m_pD->DiskCacheContext : nullptr;
 	}
 
 
@@ -175,6 +176,15 @@ namespace mu
 	}
 
 
+	void CompilerOptions::LogStats() const
+	{
+		UE_LOG(LogMutableCore, Log, TEXT("   Cache Files Written : %" PRIu64), m_pD->DiskCacheContext.FilesWritten.load());
+		UE_LOG(LogMutableCore, Log, TEXT("   Cache Files Read    : %" PRIu64), m_pD->DiskCacheContext.FilesRead.load());
+		UE_LOG(LogMutableCore, Log, TEXT("   Cache MB Written    : %" PRIu64), m_pD->DiskCacheContext.BytesWritten.load() >> 20);
+		UE_LOG(LogMutableCore, Log, TEXT("   Cache MB Read       : %" PRIu64), m_pD->DiskCacheContext.BytesRead.load()>>20);
+	}
+
+
 	void FObjectState::Serialise(OutputArchive& arch) const
     {
     	const int32 ver = 6;
diff --git a/Engine/Plugins/Experimental/Mutable/Source/MutableTools/Private/MuT/CompilerPrivate.h b/Engine/Plugins/Experimental/Mutable/Source/MutableTools/Private/MuT/CompilerPrivate.h
--- a/Engine/Plugins/Experimental/Mutable/Source/MutableTools/Private/MuT/CompilerPrivate.h
+++ b/Engine/Plugins/Experimental/Mutable/Source/MutableTools/Private/MuT/CompilerPrivate.h
@@ -6,9 +6,9 @@
 
 #include "MuT/AST.h"
 #include "MuT/ErrorLogPrivate.h"
-
-#include "MuR/Operations.h"
 #include "MuT/NodeObjectPrivate.h"
+#include "MuT/StreamsPrivate.h"
+#include "MuR/Operations.h"
 
 
 namespace mu
@@ -21,6 +21,7 @@ namespace mu
 
         //! Detailed optimization options
         FModelOptimizationOptions OptimisationOptions;
+		FProxyFileContext DiskCacheContext;
 
 		uint64 EmbeddedDataBytesLimit = 1024;
 		uint64 PackagedDataBytesLimit = 1024*1024*64;
diff --git a/Engine/Plugins/Experimental/Mutable/Source/MutableTools/Private/MuT/DataPacker.cpp b/Engine/Plugins/Experimental/Mutable/Source/MutableTools/Private/MuT/DataPacker.cpp
--- a/Engine/Plugins/Experimental/Mutable/Source/MutableTools/Private/MuT/DataPacker.cpp
+++ b/Engine/Plugins/Experimental/Mutable/Source/MutableTools/Private/MuT/DataPacker.cpp
@@ -761,7 +761,7 @@ namespace mu
 					int32 newSize = pNew->GetDataSize();
                     if (float(oldSize) > float(newSize) * OptimizeOptions.MinRLECompressionGain)
                     {
-                        typed->SetValue(pNew, OptimizeOptions.bUseDiskCache);
+                        typed->SetValue(pNew, OptimizeOptions.DiskCacheContext);
                     }
                 }
                 else if ( accFormat.m_supportedFormats[typed][(int32)EImageFormat::IF_L_UBYTE_RLE] )
@@ -773,7 +773,7 @@ namespace mu
 					int32 newSize = pNew->GetDataSize();
                     if (float(oldSize) > float(newSize) * OptimizeOptions.MinRLECompressionGain)
                     {
-                        typed->SetValue(pNew, OptimizeOptions.bUseDiskCache);
+                        typed->SetValue(pNew, OptimizeOptions.DiskCacheContext);
                     }
                 }
             }
@@ -792,7 +792,7 @@ namespace mu
 				ASTOpConstantResource* typed = static_cast<ASTOpConstantResource*>(n.get());
                 Ptr<Mesh> pMesh = static_cast<const Mesh*>(typed->GetValue().get())->Clone();
                 MeshRemoveUnusedBufferSemantics( pMesh.get(), meshSemanticsVisitor.m_requiredSemantics[typed]);
-                typed->SetValue(pMesh, OptimizeOptions.bUseDiskCache);
+                typed->SetValue(pMesh, OptimizeOptions.DiskCacheContext);
             }
         });
 
diff --git a/Engine/Plugins/Experimental/Mutable/Source/MutableTools/Private/MuT/Streams.cpp b/Engine/Plugins/Experimental/Mutable/Source/MutableTools/Private/MuT/Streams.cpp
--- a/Engine/Plugins/Experimental/Mutable/Source/MutableTools/Private/MuT/Streams.cpp
+++ b/Engine/Plugins/Experimental/Mutable/Source/MutableTools/Private/MuT/Streams.cpp
@@ -4,9 +4,11 @@
 #include "MuT/Streams.h"
 
 #include "GenericPlatform/GenericPlatformFile.h"
+#include "GenericPlatform/GenericPlatformTime.h"
 #include "HAL/PlatformFileManager.h"
 #include "HAL/PlatformMath.h"
 #include "HAL/UnrealMemory.h"
+#include "Math/RandomStream.h"
 #include "Misc/AssertionMacros.h"
 #include "MuR/Image.h"
 #include "MuR/Model.h"
@@ -19,6 +21,14 @@
 
 namespace mu
 {
+	FProxyFileContext::FProxyFileContext()
+	{
+		uint32 Seed = FPlatformTime::Cycles();
+		FRandomStream RandomStream = FRandomStream((int32)Seed);
+		CurrentFileIndex = RandomStream.GetUnsignedInt();
+	}
+
+
 	InputFileStream::InputFileStream( const FString& File )
 	{
 		m_pD = new Private();
@@ -256,16 +266,6 @@ namespace mu
     }
 
 
-    //---------------------------------------------------------------------------------------------
-    //---------------------------------------------------------------------------------------------
-    //---------------------------------------------------------------------------------------------
-    Ptr<ResourceProxy<Image>> ProxyFactoryFiles::NewImageProxy(InputArchive& arch)
-    {        
-        Ptr<Image> t = Image::StaticUnserialise( arch );
-        return new ResourceProxyTempFile<Image>(t.get());
-    }
-
-
     //---------------------------------------------------------------------------------------------
     //---------------------------------------------------------------------------------------------
     //---------------------------------------------------------------------------------------------
diff --git a/Engine/Plugins/Experimental/Mutable/Source/MutableTools/Private/MuT/StreamsPrivate.h b/Engine/Plugins/Experimental/Mutable/Source/MutableTools/Private/MuT/StreamsPrivate.h
--- a/Engine/Plugins/Experimental/Mutable/Source/MutableTools/Private/MuT/StreamsPrivate.h
+++ b/Engine/Plugins/Experimental/Mutable/Source/MutableTools/Private/MuT/StreamsPrivate.h
@@ -11,6 +11,10 @@
 #include "MuT/Streams.h"
 #include "MuR/Platform.h"
 #include "MuR/Operations.h"
+#include "MuR/MutableRuntimeModule.h"
+
+#include <atomic>
+#include <inttypes.h>
 
 namespace mu
 {
@@ -50,6 +54,29 @@ namespace mu
         friend class OutputFileStream;
     };
 
+
+	/** Statistics about the proxy file usage. */
+	struct FProxyFileContext
+	{
+		FProxyFileContext();
+
+		/** Options */
+
+		/** Minimum data size in bytes to dumpt it to the disk. */
+		uint64 MinProxyFileSize = 1024 * 1024;
+
+		/** When creating temporary files, number of retries in case the OS-level call fails. */
+		uint64 MaxFileCreateAttempts = 256;
+
+		/** Statistics */
+		std::atomic<uint64> FilesWritten = 0;
+		std::atomic<uint64> FilesRead = 0;
+		std::atomic<uint64> BytesWritten = 0;
+		std::atomic<uint64> BytesRead = 0;
+
+		/** Internal data. */
+		std::atomic<uint64> CurrentFileIndex = 0;
+	};
     
     //---------------------------------------------------------------------------------------------
     //---------------------------------------------------------------------------------------------
@@ -58,13 +85,16 @@ namespace mu
     class MUTABLETOOLS_API ResourceProxyTempFile : public ResourceProxy<R>
     {
     private:
-        Ptr<const R> m_resource;
-        FString m_fileName;
-		uint64 m_fileSize = 0;
+        Ptr<const R> Resource;
+        FString FileName;
+		uint64 FileSize = 0;
 		FCriticalSection Mutex;
 
+		FProxyFileContext& Options;
+		
     public:
-        ResourceProxyTempFile( const R* resource )
+        ResourceProxyTempFile( const R* resource, FProxyFileContext& InOptions )
+			: Options(InOptions)
         {
 			IPlatformFile& PlatformFile = FPlatformFileManager::Get().GetPlatformFile();
 
@@ -79,48 +109,56 @@ namespace mu
             OutputArchive arch(&stream);
             R::Serialise( resource, arch );
 
-            if (stream.GetBufferSize()<=256*1024)
+            if (stream.GetBufferSize()<=Options.MinProxyFileSize)
             {
-                m_resource = resource;
+                Resource = resource;
             }
             else
             {
-                m_fileSize = stream.GetBufferSize();
+                FileSize = stream.GetBufferSize();
 
-				FString TempPath = FPlatformProcess::UserTempDir();
-				TempPath += TEXT("mut.temp");
+				FString Prefix = FPlatformProcess::UserTempDir();
 
-				FRandomStream RandomStream;
-				RandomStream.GenerateNewSeed();
-				
-				auto prefix = TempPath;
-				
+				uint32 PID = FPlatformProcess::GetCurrentProcessId();
+				Prefix += FString::Printf(TEXT("mut.temp.%u"), PID);
+								
 				FString FinalTempPath;
-				bool found = true;
-                while(found)
+				IFileHandle* ResourceFile = nullptr;
+				uint64 AttemptCount = 0;
+                while(!ResourceFile && AttemptCount< Options.MaxFileCreateAttempts)
                 {
-					FinalTempPath = TempPath+FString::Printf(TEXT(".%.8x.%.8x.%.8x.%.8x"), RandomStream.GetUnsignedInt(), RandomStream.GetUnsignedInt(), RandomStream.GetUnsignedInt(), RandomStream.GetUnsignedInt());
-					found = PlatformFile.FileExists(*FinalTempPath);
+					uint64 ThisThreadFileIndex = Options.CurrentFileIndex.load();
+					while (!Options.CurrentFileIndex.compare_exchange_strong(ThisThreadFileIndex, ThisThreadFileIndex + 1));
+
+					FinalTempPath = Prefix +FString::Printf(TEXT(".%.16" PRIx64), ThisThreadFileIndex);
+					ResourceFile = PlatformFile.OpenWrite(*FinalTempPath);
+					++AttemptCount;
                 }
-				IFileHandle* resourceFile = PlatformFile.OpenWrite(*FinalTempPath);
 
-                check(resourceFile);
-				resourceFile->Write( (const uint8*) stream.GetBuffer(), stream.GetBufferSize() );
-				delete resourceFile;
+				if (!ResourceFile)
+				{
+					UE_LOG(LogMutableCore, Error, TEXT("Failed to create temporary file. Disk full?"));
+					check(false);
+				}
 
-                m_fileName = FinalTempPath;
-            }
+				ResourceFile->Write( (const uint8*) stream.GetBuffer(), stream.GetBufferSize() );
+				delete ResourceFile;
+
+                FileName = FinalTempPath;
+				Options.FilesWritten++;
+				Options.BytesWritten += stream.GetBufferSize();
+			}
         }
 
         ~ResourceProxyTempFile()
         {
 			FScopeLock Lock(&Mutex);
 
-            if (!m_fileName.IsEmpty())
+            if (!FileName.IsEmpty())
             {
                 // Delete temp file
-				FPlatformFileManager::Get().GetPlatformFile().DeleteFile(*m_fileName);
-                m_fileName.Empty();
+				FPlatformFileManager::Get().GetPlatformFile().DeleteFile(*FileName);
+				FileName.Empty();
             }
         }
 
@@ -129,24 +167,27 @@ namespace mu
 			FScopeLock Lock(&Mutex);
 
             Ptr<const R> r;
-            if (m_resource)
+            if (Resource)
             {
                 // cached
-                r = m_resource;
+                r = Resource;
             }
-            else if (!m_fileName.IsEmpty())
+            else if (!FileName.IsEmpty())
             {
 				TArray<char> buf;
-				buf.SetNumUninitialized(m_fileSize);
-                auto resourceFile = FPlatformFileManager::Get().GetPlatformFile().OpenRead( *m_fileName );
+				buf.SetNumUninitialized(FileSize);
+				IFileHandle* resourceFile = FPlatformFileManager::Get().GetPlatformFile().OpenRead( *FileName);
                 check(resourceFile);
-                resourceFile->Read((uint8*)buf.GetData(), m_fileSize);
+                resourceFile->Read((uint8*)buf.GetData(), FileSize);
                 delete resourceFile;
 
-                InputMemoryStream stream( buf.GetData(), m_fileSize );
+                InputMemoryStream stream( buf.GetData(), FileSize );
                 InputArchive arch(&stream);
                 r =  R::StaticUnserialise( arch );
-            }
+			
+				Options.FilesRead++;
+				Options.BytesRead += FileSize;
+			}
             return r;
         }
 
diff --git a/Engine/Plugins/Experimental/Mutable/Source/MutableTools/Public/MuT/Compiler.h b/Engine/Plugins/Experimental/Mutable/Source/MutableTools/Public/MuT/Compiler.h
--- a/Engine/Plugins/Experimental/Mutable/Source/MutableTools/Public/MuT/Compiler.h
+++ b/Engine/Plugins/Experimental/Mutable/Source/MutableTools/Public/MuT/Compiler.h
@@ -67,7 +67,7 @@ namespace mu
 
 		//! Set the quality for the image compression algorithms. The level value is used internally
 		//! with System::SetImagecompressionQuality
-		void SetImageCompressionQuality(int32 quality);
+		void SetImageCompressionQuality(int32 Quality);
 
 		/** Set the image tiling strategy :
 		 * If 0 (default) there is no tiling. Otherwise, images will be generated in tiles of the given size or less, and assembled afterwards as a final step.
@@ -102,6 +102,9 @@ namespace mu
         //! 
         static const char* GetTextureLayoutStrategyName( TextureLayoutStrategy s );
 
+		/** Output some stats about the complete compilation to the log. */
+		void LogStats() const;
+
         //-----------------------------------------------------------------------------------------
         // Interface pattern
         //-----------------------------------------------------------------------------------------
diff --git a/Engine/Plugins/Experimental/Mutable/Source/MutableTools/Public/MuT/Streams.h b/Engine/Plugins/Experimental/Mutable/Source/MutableTools/Public/MuT/Streams.h
--- a/Engine/Plugins/Experimental/Mutable/Source/MutableTools/Public/MuT/Streams.h
+++ b/Engine/Plugins/Experimental/Mutable/Source/MutableTools/Public/MuT/Streams.h
@@ -156,15 +156,6 @@ namespace mu
 	};
 
 
-    //! ProxyFactory that provides proxies for data stored in temporary files.
-    //! \ingroup tools
-    class MUTABLETOOLS_API ProxyFactoryFiles : public InputArchiveWithProxies::ProxyFactory
-    {
-    public:
-        Ptr<ResourceProxy<mu::Image>> NewImageProxy(InputArchive& arch) override;
-    };
-
-
     //! ProxyFactory that provides proxies for data stored in a mutable_source file.
     //! \ingroup tools
     class MUTABLETOOLS_API ProxyFactoryMutableSourceFile : public InputArchiveWithProxies::ProxyFactory
