[mutable] Updated MTU so it can handle the addition of the new static benchmark field "ID_Descriptor"
- Added handling for that new field
- Fixed bug that prevented the correct separation of fields by changing the separator from ',' to ';' for the CSV and LOG generated lines.
- UX : Added _id being used by the data uploaded to MongoDB so it can later be easily found and read by humans.

#rnx

diff --git a/Engine/Plugins/Experimental/Mutable/Source/CustomizableObject/Private/MuCO/LogBenchmarkUtil.cpp b/Engine/Plugins/Experimental/Mutable/Source/CustomizableObject/Private/MuCO/LogBenchmarkUtil.cpp
--- a/Engine/Plugins/Experimental/Mutable/Source/CustomizableObject/Private/MuCO/LogBenchmarkUtil.cpp
+++ b/Engine/Plugins/Experimental/Mutable/Source/CustomizableObject/Private/MuCO/LogBenchmarkUtil.cpp
@@ -56,9 +56,9 @@ TSharedPtr<FArchive> CreateFile()
 	check(Archive);
 
 	const FString HeaderRow = TEXT(
-		"ID_CO,ID_COI,ID_UpdateType,ID_Descriptor,ID_UpdateResult,Time_Queue,Time_Update,Time_TaskGetMesh,Time_TaskLockCache,"
-		"Time_TaskGetImages,Time_TaskConvertResources,Time_TaskCallbacks,Memory_Update,Memory_Update_Real,"
-		"Time_TaskUpdateImage,Memory_TaskUpdateImage,Memory_TaskUpdateImage_Real");
+		"ID_CO;ID_COI;ID_UpdateType;ID_Descriptor;ID_UpdateResult;Time_Queue;Time_Update;Time_TaskGetMesh;Time_TaskLockCache;"
+		"Time_TaskGetImages;Time_TaskConvertResources;Time_TaskCallbacks;Memory_Update;Memory_Update_Real;"
+		"Time_TaskUpdateImage;Memory_TaskUpdateImage;Memory_TaskUpdateImage_Real");
 
 	LogBenchmarkUtil::Write(*Archive, HeaderRow);
 
@@ -326,7 +326,7 @@ void FLogBenchmarkUtil::FinishUpdateMesh(const TSharedRef<FUpdateContextPrivate>
 	const double Memory_UpdateEndPeakMB = (Context->UpdateEndPeakBytes / 1024.0) / 1024.0;
 	const double Memory_UpdateEndRealPeakMB = (Context->UpdateEndRealPeakBytes / 1024.0) / 1024.0;
 	
-	const FString UpdateString = FString::Printf(TEXT("%s,%s,%s,%s,%s,%f,%f,%f,%f,%f,%f,%f,%f,%f"), *ID_CO, *ID_COI, *ID_UpdateType, *ID_Descriptor, *ID_UpdateResult, Time_Queue, Time_Update, Time_TaskGetMesh, Time_TaskLockCache, Time_TaskGetImages, Time_TaskConvertResources, Time_TaskCallbacks, Memory_UpdateEndPeakMB, Memory_UpdateEndRealPeakMB);
+	const FString UpdateString = FString::Printf(TEXT("%s;%s;%s;%s;%s;%f;%f;%f;%f;%f;%f;%f;%f;%f"), *ID_CO, *ID_COI, *ID_UpdateType, *ID_Descriptor, *ID_UpdateResult, Time_Queue, Time_Update, Time_TaskGetMesh, Time_TaskLockCache, Time_TaskGetImages, Time_TaskConvertResources, Time_TaskCallbacks, Memory_UpdateEndPeakMB, Memory_UpdateEndRealPeakMB);
 	LogBenchmarkUtil::Write(*Archive, UpdateString);
 	Archive->Flush();
 }
@@ -349,7 +349,7 @@ void FLogBenchmarkUtil::FinishUpdateImage(const FString& CustomizableObjectPathN
 	const double Memory_TaskUpdateImagePeakMB = (TaskUpdateImageMemoryPeak / 1024.0) / 1024.0;
 	const double Memory_TaskUpdateImageRealPeakMB = (TaskUpdateImageRealMemoryPeak / 1024.0) / 1024.0;
 
-	const FString UpdateString = FString::Printf(TEXT("%s,%s,%s,,,,,,,,,,,,%f,%f,%f"), *ID_CO, *ID_COI, *ID_UpdateType, Time_TaskUpdateImage, Memory_TaskUpdateImagePeakMB,Memory_TaskUpdateImageRealPeakMB);
+	const FString UpdateString = FString::Printf(TEXT("%s;%s;%s;;;;;;;;;;;;%f;%f;%f"), *ID_CO, *ID_COI, *ID_UpdateType, Time_TaskUpdateImage, Memory_TaskUpdateImagePeakMB,Memory_TaskUpdateImageRealPeakMB);
 	LogBenchmarkUtil::Write(*Archive, UpdateString);
 	Archive->Flush();	
 }
