[Mutable] Add current compilation as pending work for the StreamingManager.

#rnx

diff --git a/Engine/Plugins/Experimental/Mutable/Source/CustomizableObject/Private/MuCO/CustomizableObjectSystem.cpp b/Engine/Plugins/Experimental/Mutable/Source/CustomizableObject/Private/MuCO/CustomizableObjectSystem.cpp
--- a/Engine/Plugins/Experimental/Mutable/Source/CustomizableObject/Private/MuCO/CustomizableObjectSystem.cpp
+++ b/Engine/Plugins/Experimental/Mutable/Source/CustomizableObject/Private/MuCO/CustomizableObjectSystem.cpp
@@ -2504,6 +2504,8 @@ bool UCustomizableObjectSystem::Tick(float DeltaTime)
 int32 UCustomizableObjectSystem::TickInternal()
 {
 	MUTABLE_CPUPROFILER_SCOPE(UCustomizableObjectSystem::TickInternal)
+
+	check(IsInGameThread());
 	
 	// Building instances is not enabled in servers. If at some point relevant collision or animation data is necessary for server logic this will need to be changed.
 #if UE_SERVER
@@ -2741,10 +2743,16 @@ int32 UCustomizableObjectSystem::TickInternal()
 		}
 	}
 
-	return Private->CurrentMutableOperation.IsValid() + 
+	int32 RemainingWork = Private->CurrentMutableOperation.IsValid() + 
 		Private->MutablePendingInstanceWork.Num() +
-		!!LODUpdateCandidateFound + // Still a pending LOD update. We can not use the size of RequestedLODUpdates since not all requests valid in future ticks.
+		static_cast<bool>(LODUpdateCandidateFound) + // Still a pending LOD update. We can not use the size of RequestedLODUpdates since not all requests valid in future ticks.
 		RemainingTasks;
+
+#if WITH_EDITOR
+	RemainingWork += static_cast<bool>(GetPrivate()->RecompileCustomizableObjectsCompiler); // Compiler only is valid if we are compiling a CO.
+#endif
+	
+	return RemainingWork;
 }
 
 
diff --git a/Engine/Plugins/Experimental/Mutable/Source/CustomizableObjectEditor/Private/MuCOE/CustomizableObjectCompileRunnable.cpp b/Engine/Plugins/Experimental/Mutable/Source/CustomizableObjectEditor/Private/MuCOE/CustomizableObjectCompileRunnable.cpp
--- a/Engine/Plugins/Experimental/Mutable/Source/CustomizableObjectEditor/Private/MuCOE/CustomizableObjectCompileRunnable.cpp
+++ b/Engine/Plugins/Experimental/Mutable/Source/CustomizableObjectEditor/Private/MuCOE/CustomizableObjectCompileRunnable.cpp
@@ -148,27 +148,7 @@ uint32 FCustomizableObjectCompileRunnable::Run()
 
 	auto ProviderTick = [this](float)
 		{
-			check(IsInGameThread());
-
-			constexpr double MaxSecondsPerFrame = 0.4;
-
-			double MaxTime = FPlatformTime::Seconds() + MaxSecondsPerFrame;
-
-			FReferenceResourceRequest Request;
-			while (PendingResourceReferenceRequests.Dequeue(Request))
-			{
-				*Request.ResolvedImage = LoadResourceReferenced(Request.ID);
-				Request.CompletionEvent->Trigger();
-
-				// Simple time limit enforcement to avoid blocking the game thread if there are many requests.
-				double CurrentTime = FPlatformTime::Seconds();
-				if (CurrentTime >= MaxTime)
-				{
-					break;
-				}
-			}
-
-			return true;
+			Tick();
 		};
 
 	CompilerOptions->SetReferencedResourceCallback([this, &ProviderTick](int32 ID, TSharedPtr<mu::Ptr<mu::Image>> ResolvedImage, bool bRunImmediatlyIfPossible)
@@ -193,9 +173,6 @@ uint32 FCustomizableObjectCompileRunnable::Run()
 		ProviderTick
 	);
 
-	// Register the tick function that will process the game-thread image reference resolve requests.
-	FTSTicker::FDelegateHandle ResolveReferenceResourcesTickerHandle = FTSTicker::GetCoreTicker().AddTicker(FTickerDelegate::CreateLambda(ProviderTick));
-
 	const int32 MinResidentMips = UTexture::GetStaticMinTextureResidentMipCount();
 	CompilerOptions->SetDataPackingStrategy( MinResidentMips, Options.EmbeddedDataBytesLimit, Options.PackagedDataBytesLimit );
 
@@ -248,8 +225,6 @@ uint32 FCustomizableObjectCompileRunnable::Run()
 
 	bThreadCompleted = true;
 
-	FTSTicker::GetCoreTicker().RemoveTicker( ResolveReferenceResourcesTickerHandle );
-
 	UE_LOG(LogMutable, Verbose, TEXT("PROFILE: [ %16.8f ] FCustomizableObjectCompileRunnable::Run end."), FPlatformTime::Seconds());
 
 	CompilerOptions->LogStats();
@@ -272,6 +247,30 @@ const TArray<FCustomizableObjectCompileRunnable::FError>& FCustomizableObjectCom
 }
 
 
+void FCustomizableObjectCompileRunnable::Tick()
+{
+	check(IsInGameThread());
+
+	constexpr double MaxSecondsPerFrame = 0.4;
+
+	double MaxTime = FPlatformTime::Seconds() + MaxSecondsPerFrame;
+
+	FReferenceResourceRequest Request;
+	while (PendingResourceReferenceRequests.Dequeue(Request))
+	{
+		*Request.ResolvedImage = LoadResourceReferenced(Request.ID);
+		Request.CompletionEvent->Trigger();
+
+		// Simple time limit enforcement to avoid blocking the game thread if there are many requests.
+		double CurrentTime = FPlatformTime::Seconds();
+		if (CurrentTime >= MaxTime)
+		{
+			break;
+		}
+	}
+}
+
+
 FCustomizableObjectSaveDDRunnable::FCustomizableObjectSaveDDRunnable(UCustomizableObject* CustomizableObject, const FCompilationOptions& InOptions)
 {
 	MUTABLE_CPUPROFILER_SCOPE(FCustomizableObjectSaveDDRunnable::FCustomizableObjectSaveDDRunnable)
diff --git a/Engine/Plugins/Experimental/Mutable/Source/CustomizableObjectEditor/Private/MuCOE/CustomizableObjectCompileRunnable.h b/Engine/Plugins/Experimental/Mutable/Source/CustomizableObjectEditor/Private/MuCOE/CustomizableObjectCompileRunnable.h
--- a/Engine/Plugins/Experimental/Mutable/Source/CustomizableObjectEditor/Private/MuCOE/CustomizableObjectCompileRunnable.h
+++ b/Engine/Plugins/Experimental/Mutable/Source/CustomizableObjectEditor/Private/MuCOE/CustomizableObjectCompileRunnable.h
@@ -68,7 +68,7 @@ public:
 	//
 	const TArray<FError>& GetArrayErrors() const;
 
-public:
+	void Tick();
 
 	TSharedPtr<mu::Model, ESPMode::ThreadSafe> Model;
 
diff --git a/Engine/Plugins/Experimental/Mutable/Source/CustomizableObjectEditor/Private/MuCOE/CustomizableObjectCompiler.cpp b/Engine/Plugins/Experimental/Mutable/Source/CustomizableObjectEditor/Private/MuCOE/CustomizableObjectCompiler.cpp
--- a/Engine/Plugins/Experimental/Mutable/Source/CustomizableObjectEditor/Private/MuCOE/CustomizableObjectCompiler.cpp
+++ b/Engine/Plugins/Experimental/Mutable/Source/CustomizableObjectEditor/Private/MuCOE/CustomizableObjectCompiler.cpp
@@ -55,6 +55,11 @@ FCustomizableObjectCompiler::FCustomizableObjectCompiler() : FCustomizableObject
 
 bool FCustomizableObjectCompiler::Tick()
 {
+	if (CompileTask)
+	{
+		CompileTask->Tick();
+	}
+	
 	bool bUpdated = false;
 
 	if (CompileTask.IsValid() && CompileTask->IsCompleted())
