[mutable] Improved layout block editor toolbar.
+ Moved some parameter widgets to the node details.

#rnx

diff --git a/Engine/Plugins/Experimental/Mutable/Source/CustomizableObjectEditor/Private/MuCOE/Nodes/CustomizableObjectNodeLayoutBlocksDetails.cpp b/Engine/Plugins/Experimental/Mutable/Source/CustomizableObjectEditor/Private/MuCOE/Nodes/CustomizableObjectNodeLayoutBlocksDetails.cpp
--- a/Engine/Plugins/Experimental/Mutable/Source/CustomizableObjectEditor/Private/MuCOE/Nodes/CustomizableObjectNodeLayoutBlocksDetails.cpp
+++ b/Engine/Plugins/Experimental/Mutable/Source/CustomizableObjectEditor/Private/MuCOE/Nodes/CustomizableObjectNodeLayoutBlocksDetails.cpp
@@ -5,6 +5,8 @@
 #include "DetailLayoutBuilder.h"
 #include "DetailWidgetRow.h"
 #include "IDetailsView.h"
+#include "IDetailGroup.h"
+#include "Layout/Visibility.h"
 #include "MuCOE/CustomizableObjectLayout.h"
 #include "MuCOE/Nodes/CustomizableObjectNodeLayoutBlocks.h"
 #include "MuCOE/SCustomizableObjectNodeLayoutBlocksEditor.h"
@@ -15,6 +17,7 @@
 #include "Widgets/Layout/SBox.h"
 #include "Widgets/SWidget.h"
 #include "Widgets/Text/STextBlock.h"
+#include "Widgets/Input/STextComboBox.h"
 
 class FString;
 
@@ -38,87 +41,141 @@ void FCustomizableObjectNodeLayoutBlocksDetails::CustomizeDetails( IDetailLayout
 	}
 
 	IDetailCategoryBuilder& CustomizableObjectCategory = DetailBuilder.EditCategory( "LayoutOptions" );
-	IDetailCategoryBuilder& BlocksCategory = DetailBuilder.EditCategory( "LayoutEditor" );
+	IDetailCategoryBuilder& BlocksCategory = DetailBuilder.EditCategory( "LayoutBlocksEditor" );
 
-	if (Node)
+	if (Node.IsValid())
 	{
 		LayoutBlocksEditor = SNew(SCustomizableObjectNodeLayoutBlocksEditor);
 
-		CustomizableObjectCategory.AddCustomRow(LOCTEXT("BlocksDetails_LayoutOptions", "LayoutOptions"))
+		TSharedPtr<FString> CurrentSize, CurrentStrategy, CurrentMaxSize, CurrentReductionMethod;
+		FillComboBoxOptionsArrays(CurrentSize, CurrentStrategy, CurrentMaxSize, CurrentReductionMethod);
+
+		// Layout size selector widget
+		CustomizableObjectCategory.AddCustomRow(LOCTEXT("BlocksDetails_SizeSelector", "SizeSelector"))
+		.NameContent()
 		[
-			SNew(SVerticalBox)
-			+SVerticalBox::Slot()
-			.Padding(0.0f,5.0f,0.0f,2.5f)
-			.AutoHeight()
-			[
-				SNew(SHorizontalBox)
-				+SHorizontalBox::Slot()
-				.HAlign(HAlign_Left)
-				.VAlign(VAlign_Center)
-				.FillWidth(0.75f)
-				[
-					SNew(STextBlock)
-					.Text( LOCTEXT("BlocksDetails_LayoutOptions_IgnoreLodsCheckBox", "Ignore Unassigned Vertices Warning:"))
-					.ToolTipText(LOCTEXT("BlocksDetails_LayoutOptions_IgnoreLodsCheckBox_Tooltip",
-						"If true, warning message \"Source mesh has vertices not assigned to any layout block\" will be ignored."
-						"\n Note:"
-						"\n This warning can appear when a CO has more than one LOD using the same Layout Block node and these LODs have been generated using the automatic LOD generation."
-						"\n (At high LODs, some vertices may have been displaced from their original position which means they could have been displaced outside their layout blocks.)"
-						"\n Ignoring these warnings can cause some visual artifacts that may or may not be visually important at higher LODs."))
-					.Font(IDetailLayoutBuilder::GetDetailFont())
-				]
-				
-				+SHorizontalBox::Slot()
-				.HAlign(HAlign_Left)
-				.FillWidth(0.25f)
-				.Padding(5.0f, 0.0f, 0.0f, 0.0f)
-				[
-					SNew(SCheckBox)
-					.IsChecked(Node->Layout ? (Node->Layout->GetIgnoreVertexLayoutWarnings() ? ECheckBoxState::Checked : ECheckBoxState::Unchecked) : ECheckBoxState::Unchecked)
-					.OnCheckStateChanged(this, &FCustomizableObjectNodeLayoutBlocksDetails::OnIgnoreErrorsCheckStateChanged)
-				]
-			]
-			
-			+ SVerticalBox::Slot()
-			.Padding(0.0f, 2.5f, 0.0f, 5.0f)
-			.AutoHeight()
-			[
-				SNew(SHorizontalBox)
-				+SHorizontalBox::Slot()
-				.HAlign(HAlign_Left)
-				.VAlign(VAlign_Center)
-				.FillWidth(0.75f)
-				[
-					SAssignNew(LODSelectorTextWidget, STextBlock)
-					.Text( LOCTEXT("BlocksDetails_LayoutOptions_IgnoreLod", "First LOD to ignore:"))
-					.ToolTipText(LOCTEXT("BlocksDetails_LayoutOptions_IgnoreLod_Tooltip", "LOD from which vertex warning messages will be ignored."))
-					.Font(IDetailLayoutBuilder::GetDetailFont())
-					.IsEnabled(Node->Layout ? (Node->Layout->GetIgnoreVertexLayoutWarnings() ? true : false) : false)
-				]
-				
-				+ SHorizontalBox::Slot()
-				.HAlign(HAlign_Left)
-				.FillWidth(0.25f)
-				.Padding(5.0f,0.0f,0.0f,0.0f)
-				[
-					SAssignNew(LODSelectorWidget, SSpinBox<int32>)
-					.Value_Lambda([this]()->int32
-					{
-						if (Node->Layout)
-						{
-							return Node->Layout->GetFirstLODToIgnoreWarnings();
-						}
-
-						return 0;
-					})
-					.IsEnabled(Node->Layout ? (Node->Layout->GetIgnoreVertexLayoutWarnings() ? true : false ) : false)
-					.OnValueChanged(this, &FCustomizableObjectNodeLayoutBlocksDetails::OnLODBoxValueChanged)
-					.MinValue(0)
-					.Delta(1)
-					.AlwaysUsesDeltaSnap(true)
-					.MinDesiredWidth(40.0f)
-				]
-			]
+			SNew(STextBlock)
+			.Text(LOCTEXT("LayoutGridSizeText", "Grid Size"))
+			.Font(DetailBuilder.GetDetailFont())
+		]
+		.ValueContent()
+		[
+			SNew(STextComboBox)
+			.InitiallySelectedItem(CurrentSize)
+			.OptionsSource(&LayoutGridSizes)
+			.OnSelectionChanged(this, &FCustomizableObjectNodeLayoutBlocksDetails::OnGridSizeChanged)
+			.Font(DetailBuilder.GetDetailFont())
+		];
+
+		// Layout strategy selector group widget
+		IDetailGroup* LayoutStrategyOptionsGroup = &CustomizableObjectCategory.AddGroup(TEXT("LayoutStrategyOptionsGroup"), LOCTEXT("LayoutStrategyGroup", "Layout Strategy Group"), false, true);
+		LayoutStrategyOptionsGroup->HeaderRow()
+		.NameContent()
+		[
+			SNew(STextBlock)
+			.Text(LOCTEXT("LayoutStrategy_Text", "Layout Strategy:"))
+			.ToolTipText(LOCTEXT("LayoutStrategyTooltup", "Selects the packing strategy: Resizable Layout or Fixed Layout"))
+			.Font(DetailBuilder.GetDetailFont())
+		]
+		.ValueContent()
+		[
+			SNew(STextComboBox)
+			.InitiallySelectedItem(CurrentStrategy)
+			.OptionsSource(&LayoutPackingStrategies)
+			.OnSelectionChanged(this, &FCustomizableObjectNodeLayoutBlocksDetails::OnLayoutPackingStrategyChanged)
+			.Font(DetailBuilder.GetDetailFont())
+		];
+
+		// Max layout size selector widget
+		LayoutStrategyOptionsGroup->AddWidgetRow()
+		.Visibility(TAttribute<EVisibility>(this, &FCustomizableObjectNodeLayoutBlocksDetails::FixedStrategyOptionsVisibility))
+		.NameContent()
+		[
+			SNew(STextBlock)
+			.Text(LOCTEXT("MaxLayoutSize_Text", "Max Layout Size:"))
+			.Font(DetailBuilder.GetDetailFont())
+		]
+		.ValueContent()
+		[
+			SNew(STextComboBox)
+			.InitiallySelectedItem(CurrentMaxSize)
+			.OptionsSource(&LayoutGridSizes)
+			.OnSelectionChanged(this, &FCustomizableObjectNodeLayoutBlocksDetails::OnMaxGridSizeChanged)
+			.Font(DetailBuilder.GetDetailFont())
+		];
+
+		// Reduction method selector widget
+		LayoutStrategyOptionsGroup->AddWidgetRow()
+		.Visibility(TAttribute<EVisibility>(this, &FCustomizableObjectNodeLayoutBlocksDetails::FixedStrategyOptionsVisibility))
+		.NameContent()
+		[
+			SNew(STextBlock)
+			.Text(LOCTEXT("ReductionMethod_Text", "Reduction Method:"))
+			.ToolTipText(LOCTEXT("Reduction_Method_Tooltip", "Select how blocks will be reduced in case that they do not fit in the layout:"
+				"\n Halve: blocks will be reduced by half each time."
+				"\n Unit: blocks will be reduced by one unit each time."))
+			.Font(DetailBuilder.GetDetailFont())
+		]
+		.ValueContent()
+		[
+			SNew(STextComboBox)
+			.InitiallySelectedItem(CurrentReductionMethod)
+			.OptionsSource(&BlockReductionMethods)
+			.OnSelectionChanged(this, &FCustomizableObjectNodeLayoutBlocksDetails::OnReductionMethodChanged)
+			.Font(DetailBuilder.GetDetailFont())
+		];
+
+		// Warning selector group widget
+		IDetailGroup* IgnoreWarningsGroup = &CustomizableObjectCategory.AddGroup(TEXT("IgnoreWarningsOptionsGroup"), LOCTEXT("IgnoreWarningsOptions", "Ignore Unassigned Vertives Warning group"), false, true);
+		IgnoreWarningsGroup->HeaderRow()
+		.NameContent()
+		[
+			SNew(STextBlock)
+			.Text(LOCTEXT("LayoutOptions_IgnoreLodsCheckBox_Text", "Ignore Unassigned Vertices Warning:"))
+			.ToolTipText(LOCTEXT("LayoutOptions_IgnoreLodsCheckBox_Tooltip",
+				"If true, warning message \"Source mesh has vertices not assigned to any layout block\" will be ignored."
+				"\n Note:"
+				"\n This warning can appear when a CO has more than one LOD using the same Layout Block node and these LODs have been generated using the automatic LOD generation."
+				"\n (At high LODs, some vertices may have been displaced from their original position which means they could have been displaced outside their layout blocks.)"
+				"\n Ignoring these warnings can cause some visual artifacts that may or may not be visually important at higher LODs."))
+			.Font(IDetailLayoutBuilder::GetDetailFont())
+		]
+		.ValueContent()
+		[
+			SNew(SCheckBox)
+			.IsChecked((Node->Layout && Node->Layout->GetIgnoreVertexLayoutWarnings()) ? ECheckBoxState::Checked : ECheckBoxState::Unchecked)
+			.OnCheckStateChanged(this, &FCustomizableObjectNodeLayoutBlocksDetails::OnIgnoreErrorsCheckStateChanged)
+		];
+		
+		// LOD selector widget
+		IgnoreWarningsGroup->AddWidgetRow()
+		.NameContent()
+		[
+			SAssignNew(LODSelectorTextWidget, STextBlock)
+			.Text(LOCTEXT("LayoutOptions_IgnoreLod_Text", "First LOD to ignore:"))
+			.ToolTipText(LOCTEXT("LayoutOptions_IgnoreLod_Tooltip", "LOD from which vertex warning messages will be ignored."))
+			.Font(IDetailLayoutBuilder::GetDetailFont())
+			.IsEnabled(Node->Layout ? Node->Layout->GetIgnoreVertexLayoutWarnings() : false)
+		]
+		.ValueContent()
+		[
+			SAssignNew(LODSelectorWidget, SSpinBox<int32>)
+			.Value_Lambda([this]()->int32
+			{
+				if (Node->Layout)
+				{
+					return Node->Layout->GetFirstLODToIgnoreWarnings();
+				}
+
+				return 0;
+			})
+			.IsEnabled(Node->Layout ? Node->Layout->GetIgnoreVertexLayoutWarnings() : false)
+			.OnValueChanged(this, &FCustomizableObjectNodeLayoutBlocksDetails::OnLODBoxValueChanged)
+			.MinValue(0)
+			.Delta(1)
+			.AlwaysUsesDeltaSnap(true)
+			.MinDesiredWidth(40.0f)
+			.Font(IDetailLayoutBuilder::GetDetailFont())
 		];
 
 		BlocksCategory.AddCustomRow( LOCTEXT("BlocksDetails_BlockInstructions", "BlockInstructions") )
@@ -141,87 +198,140 @@ void FCustomizableObjectNodeLayoutBlocksDetails::CustomizeDetails( IDetailLayout
 			.Text( LOCTEXT( "Node not found", "Node not found" ) )
 		];
 	}
-
 }
 
 
-FIntPoint FCustomizableObjectNodeLayoutBlocksDetails::GetGridSize() const
+void FCustomizableObjectNodeLayoutBlocksDetails::OnIgnoreErrorsCheckStateChanged(ECheckBoxState State)
 {
 	if (Node->Layout)
 	{
-		return Node->Layout->GetGridSize();
-	}
+		bool bStateBool = State == ECheckBoxState::Checked;
+		Node->Layout->SetIgnoreVertexLayoutWarnings(bStateBool);
 
-	return FIntPoint(0);
+		LODSelectorWidget->SetEnabled(bStateBool);
+		LODSelectorTextWidget->SetEnabled(bStateBool);
+	}
 }
 
 
-void FCustomizableObjectNodeLayoutBlocksDetails::OnBlockChanged( int BlockIndex, FIntRect Block )
+void FCustomizableObjectNodeLayoutBlocksDetails::OnLODBoxValueChanged(int32 Value)
 {
 	if (Node->Layout)
 	{
-		Node->Layout->Blocks[BlockIndex].Min = Block.Min;
-		Node->Layout->Blocks[BlockIndex].Max = Block.Max;
-		Node->PostEditChange();
+		Node->Layout->SetIgnoreWarningsLOD(Value);
 	}
 }
 
 
-TArray<FIntRect> FCustomizableObjectNodeLayoutBlocksDetails::GetBlocks() const
+void FCustomizableObjectNodeLayoutBlocksDetails::FillComboBoxOptionsArrays(TSharedPtr<FString>& CurrGridSize, TSharedPtr<FString>& CurrStrategy, TSharedPtr<FString>& CurrMaxSize, TSharedPtr<FString>& CurrRedMethod)
 {
-	TArray<FIntRect> Blocks;
-
 	if (Node->Layout)
 	{
-		Blocks.SetNum(Node->Layout->Blocks.Num());
+		// Static const variable?
+		int32 MaxGridSize = 32;
+		LayoutGridSizes.Empty();
 
-		for (int BlockIndex = 0; BlockIndex < Node->Layout->Blocks.Num(); ++BlockIndex)
+		for (int32 Size = 1; Size <= MaxGridSize; Size *= 2)
 		{
-			Blocks[BlockIndex] = FIntRect(Node->Layout->Blocks[BlockIndex].Min, Node->Layout->Blocks[BlockIndex].Max);
+			LayoutGridSizes.Add(MakeShareable(new FString(FString::Printf(TEXT("%d x %d"), Size, Size))));
+
+			if (Node->Layout->GetGridSize() == FIntPoint(Size))
+			{
+				CurrGridSize = LayoutGridSizes.Last();
+			}
+
+			if (Node->Layout->GetMaxGridSize() == FIntPoint(Size))
+			{
+				CurrMaxSize = LayoutGridSizes.Last();
+			}
 		}
+
+		LayoutPackingStrategies.Empty();
+		LayoutPackingStrategies.Add(MakeShareable(new FString("Resizable")));
+		LayoutPackingStrategies.Add(MakeShareable(new FString("Fixed")));
+		LayoutPackingStrategies.Add(MakeShareable(new FString("Overlay")));
+		CurrStrategy = LayoutPackingStrategies[(uint32)Node->Layout->GetPackingStrategy()];
+
+		BlockReductionMethods.Empty();
+		BlockReductionMethods.Add(MakeShareable(new FString("Halve")));
+		BlockReductionMethods.Add(MakeShareable(new FString("Unitary")));
+		CurrRedMethod = BlockReductionMethods[(uint32)Node->Layout->GetBlockReductionMethod()];
 	}
+}
 
-	return Blocks;
+EVisibility FCustomizableObjectNodeLayoutBlocksDetails::FixedStrategyOptionsVisibility() const
+{
+	return Node->Layout->GetPackingStrategy() == ECustomizableObjectTextureLayoutPackingStrategy::Fixed ? EVisibility::Visible : EVisibility::Collapsed;
 }
 
 
-void FCustomizableObjectNodeLayoutBlocksDetails::OnGridComboBoxSelectionChanged(TSharedPtr<FString> Selection, ESelectInfo::Type SelectInfo)
+void FCustomizableObjectNodeLayoutBlocksDetails::OnGridSizeChanged(TSharedPtr<FString> NewSelection, ESelectInfo::Type SelectInfo)
 {
-	/*
-	if (CustomInstance)
+	if (Node->Layout)
 	{
-		mu::Model* Model = CustomInstance->CustomizableObject->GetModel();
-
-		CustomInstance->PreEditChange(NULL);
-		CustomInstance->State = Model->FindState( *Selection );
-		CustomInstance->PostEditChange();
+		int Size = 1 << LayoutGridSizes.Find(NewSelection);
 
-		ResetParamBox();
+		if (Node->Layout->GetGridSize().X != Size || Node->Layout->GetGridSize().Y != Size)
+		{
+			Node->Layout->SetGridSize(FIntPoint(Size));
+
+			// Adjust all the blocks sizes
+			for (int b = 0; b < Node->Layout->Blocks.Num(); ++b)
+			{
+				Node->Layout->Blocks[b].Min.X = FMath::Min(Node->Layout->Blocks[b].Min.X, Size - 1);
+				Node->Layout->Blocks[b].Min.Y = FMath::Min(Node->Layout->Blocks[b].Min.Y, Size - 1);
+				Node->Layout->Blocks[b].Max.X = FMath::Min(Node->Layout->Blocks[b].Max.X, Size);
+				Node->Layout->Blocks[b].Max.Y = FMath::Min(Node->Layout->Blocks[b].Max.Y, Size);
+			}
+
+			Node->MarkPackageDirty();
+		}
 	}
-	*/
 }
 
 
-void FCustomizableObjectNodeLayoutBlocksDetails::OnIgnoreErrorsCheckStateChanged(ECheckBoxState State)
+void FCustomizableObjectNodeLayoutBlocksDetails::OnLayoutPackingStrategyChanged(TSharedPtr<FString> NewSelection, ESelectInfo::Type SelectInfo)
 {
 	if (Node->Layout)
 	{
-		bool bStateBool = State == ECheckBoxState::Checked;
-		Node->Layout->SetIgnoreVertexLayoutWarnings(bStateBool);
+		uint32 selection = LayoutPackingStrategies.IndexOfByKey(NewSelection);
 
-		LODSelectorWidget->SetEnabled(bStateBool);
-		LODSelectorTextWidget->SetEnabled(bStateBool);
+		if (Node->Layout->GetPackingStrategy() != (ECustomizableObjectTextureLayoutPackingStrategy)selection)
+		{
+			Node->Layout->SetPackingStrategy((ECustomizableObjectTextureLayoutPackingStrategy)selection);
+			Node->MarkPackageDirty();
+		}
 	}
 }
 
 
-void FCustomizableObjectNodeLayoutBlocksDetails::OnLODBoxValueChanged(int32 Value)
+void FCustomizableObjectNodeLayoutBlocksDetails::OnMaxGridSizeChanged(TSharedPtr<FString> NewSelection, ESelectInfo::Type SelectInfo)
 {
 	if (Node->Layout)
 	{
-		Node->Layout->SetIgnoreWarningsLOD(Value);
+		int Size = 1 << LayoutGridSizes.Find(NewSelection);
+
+		if (Node->Layout->GetMaxGridSize().X != Size || Node->Layout->GetMaxGridSize().Y != Size)
+		{
+			Node->Layout->SetMaxGridSize(FIntPoint(Size));
+			Node->Layout->MarkPackageDirty();
+		}
 	}
 }
 
 
+void FCustomizableObjectNodeLayoutBlocksDetails::OnReductionMethodChanged(TSharedPtr<FString> NewSelection, ESelectInfo::Type SelectInfo)
+{
+	if (Node->Layout)
+	{
+		uint32 selection = BlockReductionMethods.IndexOfByKey(NewSelection);
+
+		if (Node->Layout->GetBlockReductionMethod() != (ECustomizableObjectLayoutBlockReductionMethod)selection)
+		{
+			Node->Layout->SetBlockReductionMethod((ECustomizableObjectLayoutBlockReductionMethod)selection);
+			Node->MarkPackageDirty();
+		}
+	}
+}
+
 #undef LOCTEXT_NAMESPACE
diff --git a/Engine/Plugins/Experimental/Mutable/Source/CustomizableObjectEditor/Private/MuCOE/Nodes/CustomizableObjectNodeLayoutBlocksDetails.h b/Engine/Plugins/Experimental/Mutable/Source/CustomizableObjectEditor/Private/MuCOE/Nodes/CustomizableObjectNodeLayoutBlocksDetails.h
--- a/Engine/Plugins/Experimental/Mutable/Source/CustomizableObjectEditor/Private/MuCOE/Nodes/CustomizableObjectNodeLayoutBlocksDetails.h
+++ b/Engine/Plugins/Experimental/Mutable/Source/CustomizableObjectEditor/Private/MuCOE/Nodes/CustomizableObjectNodeLayoutBlocksDetails.h
@@ -4,6 +4,7 @@
 
 
 #include "IDetailCustomization.h"
+#include "UObject/WeakObjectPtr.h"
 
 enum class ECheckBoxState : uint8;
 namespace ESelectInfo { enum Type : int; }
@@ -12,6 +13,8 @@ class FString;
 class IDetailLayoutBuilder;
 class STextBlock;
 class SWidget;
+class UCustomizableObjectNodeLayoutBlocks;
+struct EVisibility;
 
 class FCustomizableObjectNodeLayoutBlocksDetails : public IDetailCustomization
 {
@@ -25,7 +28,24 @@ public:
 
 private:
 
-	class UCustomizableObjectNodeLayoutBlocks* Node;
+	/** Returns the visibility of the Fixed layout widgets */
+	EVisibility FixedStrategyOptionsVisibility() const;
+
+	/** Fills the combo box arrays sources */
+	void FillComboBoxOptionsArrays(TSharedPtr<FString>& CurrGridSize, TSharedPtr<FString>& CurrStrategy, TSharedPtr<FString>& CurrMaxSize, TSharedPtr<FString>& CurrRedMethod);
+
+	/** Layout Options Callbacks */
+	void OnGridSizeChanged(TSharedPtr<FString> NewSelection, ESelectInfo::Type SelectInfo);
+	void OnMaxGridSizeChanged(TSharedPtr<FString> NewSelection, ESelectInfo::Type SelectInfo);
+	void OnReductionMethodChanged(TSharedPtr<FString> NewSelection, ESelectInfo::Type SelectInfo);
+	void OnLayoutPackingStrategyChanged(TSharedPtr<FString> NewSelection, ESelectInfo::Type SelectInfo);
+	void OnIgnoreErrorsCheckStateChanged(ECheckBoxState State);
+	void OnLODBoxValueChanged(int32 Value);
+
+private:
+
+	/** Weak pointer to the node */
+	TWeakObjectPtr<UCustomizableObjectNodeLayoutBlocks> Node;
 
 	// Layout block editor widget
 	TSharedPtr<class SCustomizableObjectNodeLayoutBlocksEditor> LayoutBlocksEditor;
@@ -36,12 +56,13 @@ private:
 
 	TArray< TSharedPtr<FString> > GridComboOptions;
 
-	void OnGridComboBoxSelectionChanged(TSharedPtr<FString> Selection, ESelectInfo::Type SelectInfo);
-	void OnBlockChanged( int BlockIndex, FIntRect Block );
-	void OnIgnoreErrorsCheckStateChanged(ECheckBoxState State);
-	void OnLODBoxValueChanged(int32 Value);
+	/** List of available layout grid sizes. */
+	TArray< TSharedPtr< FString > > LayoutGridSizes;
+
+	/** List of available layout packing strategies. */
+	TArray< TSharedPtr< FString > > LayoutPackingStrategies;
 
-	FIntPoint GetGridSize() const;
-	TArray<FIntRect> GetBlocks() const;
+	/** List of available block reduction methods. */
+	TArray< TSharedPtr< FString > > BlockReductionMethods;
 
 };
diff --git a/Engine/Plugins/Experimental/Mutable/Source/CustomizableObjectEditor/Private/MuCOE/SCustomizableObjectNodeLayoutBlocksEditor.cpp b/Engine/Plugins/Experimental/Mutable/Source/CustomizableObjectEditor/Private/MuCOE/SCustomizableObjectNodeLayoutBlocksEditor.cpp
--- a/Engine/Plugins/Experimental/Mutable/Source/CustomizableObjectEditor/Private/MuCOE/SCustomizableObjectNodeLayoutBlocksEditor.cpp
+++ b/Engine/Plugins/Experimental/Mutable/Source/CustomizableObjectEditor/Private/MuCOE/SCustomizableObjectNodeLayoutBlocksEditor.cpp
@@ -96,16 +96,6 @@ void SCustomizableObjectNodeLayoutBlocksEditor::SetCurrentLayout(UCustomizableOb
 			BuildLayoutToolBar()
 		]
 
-		+ SVerticalBox::Slot()
-		.HAlign(EHorizontalAlignment::HAlign_Left)
-		.AutoHeight()
-		.Padding(3.0f,5.0f,0.0f,5.0f)
-		[
-			SNew(SImage)
-			.Image(UE_MUTABLE_GET_BRUSH(TEXT("Icons.Info")))
-			.ToolTip(GenerateInfoToolTip())
-		]
-
 		+ SVerticalBox::Slot()
 		.FillHeight(1)
 		[
@@ -140,57 +130,12 @@ void SCustomizableObjectNodeLayoutBlocksEditor::AddReferencedObjects( FReference
 
 TSharedRef<SWidget> SCustomizableObjectNodeLayoutBlocksEditor::BuildLayoutToolBar()
 {
-	FToolBarBuilder LayoutToolbarBuilder(UICommandList, FMultiBoxCustomization::None);
+	FSlimHorizontalToolBarBuilder LayoutToolbarBuilder(UICommandList, FMultiBoxCustomization::None, nullptr, true);
+	LayoutToolbarBuilder.SetLabelVisibility(EVisibility::Visible);
 
 	//Getting toolbar style
 	const ISlateStyle* const StyleSet = &FCoreStyle::Get();
 	const FName& StyleName = "ToolBar";
-	
-	// Build toolbar widgets
-	LayoutGridSizes.Empty();
-	int32 MaxGridSize = 32;
-	for(int32 Size = 1; Size <= MaxGridSize; Size*=2 )
-	{
-		LayoutGridSizes.Add( MakeShareable( new FString( FString::Printf(TEXT("%d x %d"), Size, Size ) ) ) );
-	}
-	
-	LayoutGridSizeWidget = SNew(SVerticalBox);
-
-	LayoutGridSizeWidget->AddSlot()
-	.AutoHeight()
-	.Padding(0.0f,0.0f,0.0f,10.0f)
-	[
-		SAssignNew(LayoutGridSizeCombo, STextComboBox)
-		.OptionsSource(&LayoutGridSizes)
-		.OnSelectionChanged(this, &SCustomizableObjectNodeLayoutBlocksEditor::OnGridSizeChanged)
-		.IsEnabled(FSlateApplication::Get().GetNormalExecutionAttribute())
-	];
-
-	LayoutGridSizeWidget->AddSlot()
-	.AutoHeight()
-	.HAlign(EHorizontalAlignment::HAlign_Center)
-	.VAlign(EVerticalAlignment::VAlign_Bottom)
-	[
-		SNew(STextBlock).Text(LOCTEXT("LayoutGridSizeText", "Grid Size"))
-		.ShadowOffset(FVector2D::UnitVector)
-		.ColorAndOpacity(FLinearColor::Gray)
-	];
-
-	if (CurrentLayout)
-	{
-		int grid = CurrentLayout->GetGridSize().X;
-		int option = -1;
-		while ( grid )
-		{
-			option++;
-			grid >>= 1;
-		}
-
-		if ( LayoutGridSizes.IsValidIndex(option) )
-		{
-			LayoutGridSizeCombo->SetSelectedItem( LayoutGridSizes[option] );
-		}
-	}
 
 	LayoutToolbarBuilder.BeginSection("Blocks");
 	{
@@ -200,212 +145,35 @@ TSharedRef<SWidget> SCustomizableObjectNodeLayoutBlocksEditor::BuildLayoutToolBa
 	}
 	LayoutToolbarBuilder.EndSection();
 
-	LayoutToolbarBuilder.BeginSection("Grid");
-	{		
-		//LayoutToolbarBuilder.AddComboButton(FUIAction(),FOnGetContent::CreateSP(
-		//		this,&SCustomizableObjectNodeLayoutBlocksEditor::GenerateLayoutGridOptionsMenuContent),
-		//	TAttribute<FText>(),TAttribute<FText>(),FSlateIcon());
-
-		LayoutToolbarBuilder.AddWidget(LayoutGridSizeWidget.ToSharedRef());
-	}
-	LayoutToolbarBuilder.EndSection();
-
-	
-	StrategyWidget = BuildLayoutStrategyWidgets(StyleSet, StyleName);
-	
-	LayoutToolbarBuilder.BeginSection("Packing Strategy");
+	LayoutToolbarBuilder.BeginSection("Info");
 	{
-		LayoutToolbarBuilder.AddWidget(StrategyWidget.ToSharedRef());
-	}
-	LayoutToolbarBuilder.EndSection();
-
-	return
-		SNew(SHorizontalBox)
-
-		+SHorizontalBox::Slot()
-		.Padding(4,0)
-		[
-			SNew(SBorder)
-			.Padding(2.0f)
-			.BorderImage(UE_MUTABLE_GET_BRUSH("NoBorder"))
-			.IsEnabled( FSlateApplication::Get().GetNormalExecutionAttribute() )
+		LayoutToolbarBuilder.AddWidget
+		(
+			SNew(SBox)
+			.VAlign(EVerticalAlignment::VAlign_Center)
+			.HAlign(EHorizontalAlignment::HAlign_Left)
 			[
-				LayoutToolbarBuilder.MakeWidget()
+				SNew(SImage)
+				.Image(UE_MUTABLE_GET_BRUSH(TEXT("Icons.Info")))
+				.ToolTip(GenerateInfoToolTip())
 			]
-		]
-		;
-}
-
-
-TSharedRef<SWidget> SCustomizableObjectNodeLayoutBlocksEditor::BuildLayoutStrategyWidgets(const ISlateStyle* Style, const FName& StyleName)
-{
-	MaxLayoutGridSizes.Empty();
-	int32 MaxGridSize = 32;
-	for (int32 Size = 1; Size <= MaxGridSize; Size *= 2)
-	{
-		MaxLayoutGridSizes.Add(MakeShareable(new FString(FString::Printf(TEXT("%d x %d"), Size, Size))));
+		);
 	}
-
-	LayoutPackingStrategies.Empty();
-	LayoutPackingStrategies.Add(MakeShareable(new FString("Resizable")));
-	LayoutPackingStrategies.Add(MakeShareable(new FString("Fixed")));
-	LayoutPackingStrategies.Add(MakeShareable(new FString("Overlay")));
-
-	BlockReductionMethods.Empty();
-	BlockReductionMethods.Add(MakeShareable(new FString("Halve")));
-	BlockReductionMethods.Add(MakeShareable(new FString("Unitary")));
+	LayoutToolbarBuilder.EndSection();
 	
-	MaxLayoutGridSizeCombo = SNew(STextComboBox)
-		.OptionsSource(&MaxLayoutGridSizes)
-		.OnSelectionChanged(this, &SCustomizableObjectNodeLayoutBlocksEditor::OnMaxGridSizeChanged)
-		.IsEnabled(FSlateApplication::Get().GetNormalExecutionAttribute());
-
-	LayoutPackingStrategyCombo = SNew(STextComboBox)
-		.OptionsSource(&LayoutPackingStrategies)
-		.OnSelectionChanged(this, &SCustomizableObjectNodeLayoutBlocksEditor::OnLayoutPackingStrategyChanged)
-		.IsEnabled(FSlateApplication::Get().GetNormalExecutionAttribute());
-
-	BlockReductionMethodsCombo = SNew(STextComboBox)
-		.OptionsSource(&BlockReductionMethods)
-		.OnSelectionChanged(this, &SCustomizableObjectNodeLayoutBlocksEditor::OnReductionMethodChanged)
-		.IsEnabled(FSlateApplication::Get().GetNormalExecutionAttribute());
-
-
-	LayoutStrategyWidget = SNew(SHorizontalBox)
-		+ SHorizontalBox::Slot()
-		.AutoWidth()
-		.Padding(0.0f,2.0f,2.0f,0.0f)
-		[
-			SNew(STextBlock)
-			.Text(LOCTEXT("LayoutStrategy_Text", "Layout Strategy:"))
-			.ToolTipText(LOCTEXT("LayoutStrategyTooltup","Selects the packing strategy: Resizable Layout or Fixed Layout"))
-			.ShadowOffset(FVector2D::UnitVector)
-			.ColorAndOpacity(FLinearColor::Gray)
-		]
-
-		+ SHorizontalBox::Slot()
-		.AutoWidth()
-		.HAlign(EHorizontalAlignment::HAlign_Right)
-		[
-			SNew(SVerticalBox)
-			+ SVerticalBox::Slot()
-			.AutoHeight()
-			[
-				LayoutPackingStrategyCombo.ToSharedRef()
-			]
-		];
-
-	FixedLayoutWidget = SNew(SVerticalBox)
-		+SVerticalBox::Slot()
-		.AutoHeight()
-		[
-			SNew(SHorizontalBox)
-			+SHorizontalBox::Slot()
-			.AutoWidth()
-			.Padding(0.0f, 2.0f, 2.0f, 0.0f)
-			[
-				SNew(STextBlock)
-				.Text(LOCTEXT("MaxLayoutSize_Text", "Max Layout Size:"))
-				.ShadowOffset(FVector2D::UnitVector)
-				.ColorAndOpacity(FLinearColor::Gray)
-			]
-			+ SHorizontalBox::Slot()
-			.AutoWidth()
-			.HAlign(EHorizontalAlignment::HAlign_Right)
-			[
-				MaxLayoutGridSizeCombo.ToSharedRef()
-			]
-		]
-		+ SVerticalBox::Slot()
-		.AutoHeight()
-		[
-			SNew(SHorizontalBox)
-			+ SHorizontalBox::Slot()
-			.AutoWidth()
-			.Padding(0.0f, 2.0f, 2.0f, 0.0f)
-			[
-				SNew(STextBlock)
-				.Text(LOCTEXT("ReductionMethod_Text", "Reduction Method:"))
-				.ToolTipText(LOCTEXT("Reduction_Method_Tooltip", "Select how blocks will be reduced in case that they do not fit in the layout:"
-				"\n Halve: blocks will be reduced by half each time."
-				"\n Unit: blocks will be reduced by one unit each time."))
-				.ShadowOffset(FVector2D::UnitVector)
-				.ColorAndOpacity(FLinearColor::Gray)
-			]
-			
-			+ SHorizontalBox::Slot()
-			.AutoWidth()
-			.Padding(0.0f, 0.0f, 0.0f, 8.0f)
-			.HAlign(EHorizontalAlignment::HAlign_Right)
-			[
-				BlockReductionMethodsCombo.ToSharedRef()
-			]
-		];
-
-
-	//Modify after the creation of all the widgets
-	if (CurrentLayout)
-	{
-		int grid = CurrentLayout->GetMaxGridSize().X;
-		int option = -1;
-		while (grid)
-		{
-			option++;
-			grid >>= 1;
-		}
-
-		if (MaxLayoutGridSizes.IsValidIndex(option))
-		{
-			MaxLayoutGridSizeCombo->SetSelectedItem(MaxLayoutGridSizes[option]);
-		}
-
-		if (LayoutPackingStrategies.IsValidIndex((uint32)CurrentLayout->GetPackingStrategy()))
-		{
-			LayoutPackingStrategyCombo->SetSelectedItem(LayoutPackingStrategies[(uint32)CurrentLayout->GetPackingStrategy()]);
-		}
-
-		if (BlockReductionMethods.IsValidIndex((uint32)CurrentLayout->GetBlockReductionMethod()))
-		{
-			BlockReductionMethodsCombo->SetSelectedItem(BlockReductionMethods[(uint32)CurrentLayout->GetBlockReductionMethod()]);
-		}
-	}
-
-	return SNew(SVerticalBox)
-		//.IsEnabled(CurrentLayout ? !CurrentLayout->GetNode()->GetGraphEditor()->GetCustomizableObject()->bDisableTextureLayoutManagement : false)
-		+ SVerticalBox::Slot().VAlign(EVerticalAlignment::VAlign_Center)
-		[
-			LayoutStrategyWidget.ToSharedRef()
-		]
-		+SVerticalBox::Slot()
+	return
+	SNew(SHorizontalBox)
+	+SHorizontalBox::Slot()
+	.Padding(4,0)
+	[
+		SNew(SBorder)
+		.Padding(2.0f)
+		.BorderImage(UE_MUTABLE_GET_BRUSH("NoBorder"))
+		.IsEnabled( FSlateApplication::Get().GetNormalExecutionAttribute() )
 		[
-			FixedLayoutWidget.ToSharedRef()
+			LayoutToolbarBuilder.MakeWidget()
 		]
-	;
-}
-
-
-void SCustomizableObjectNodeLayoutBlocksEditor::OnGridSizeChanged( TSharedPtr<FString> NewSelection, ESelectInfo::Type SelectInfo )
-{
-	if (CurrentLayout)
-	{
-		int Size = 1 << LayoutGridSizes.Find(NewSelection);
-
-		if (CurrentLayout->GetGridSize().X!=Size || CurrentLayout->GetGridSize().Y!=Size )
-		{
-			CurrentLayout->SetGridSize(FIntPoint(Size));
-
-			// Adjust all the blocks sizes
-			for ( int b=0; b< CurrentLayout->Blocks.Num(); ++b )
-			{
-				CurrentLayout->Blocks[b].Min.X = FMath::Min( CurrentLayout->Blocks[b].Min.X, Size-1 );
-				CurrentLayout->Blocks[b].Min.Y = FMath::Min( CurrentLayout->Blocks[b].Min.Y, Size-1 );
-				CurrentLayout->Blocks[b].Max.X = FMath::Min( CurrentLayout->Blocks[b].Max.X, Size );
-				CurrentLayout->Blocks[b].Max.Y = FMath::Min( CurrentLayout->Blocks[b].Max.Y, Size );
-			}
-
-			CurrentLayout->MarkPackageDirty();
-		}
-	}
+	];
 }
 
 
@@ -567,60 +335,6 @@ void SCustomizableObjectNodeLayoutBlocksEditor::OnSetBlockReductionByTwo(bool bI
 }
 
 
-void SCustomizableObjectNodeLayoutBlocksEditor::OnLayoutPackingStrategyChanged(TSharedPtr<FString> NewSelection, ESelectInfo::Type SelectInfo)
-{
-	if (CurrentLayout && FixedLayoutWidget)
-	{
-		uint32 selection = LayoutPackingStrategies.IndexOfByKey(NewSelection);
-
-		if (!selection)
-		{
-			FixedLayoutWidget->SetVisibility(EVisibility::Collapsed);
-		}
-		else
-		{
-			FixedLayoutWidget->SetVisibility(EVisibility::Visible);
-		}
-
-		if (CurrentLayout->GetPackingStrategy() != (ECustomizableObjectTextureLayoutPackingStrategy)selection)
-		{
-			CurrentLayout->SetPackingStrategy((ECustomizableObjectTextureLayoutPackingStrategy)selection);
-			CurrentLayout->MarkPackageDirty();
-		}
-	}
-}
-
-void SCustomizableObjectNodeLayoutBlocksEditor::OnMaxGridSizeChanged(TSharedPtr<FString> NewSelection, ESelectInfo::Type SelectInfo)
-{
-	if (CurrentLayout)
-	{
-		int Size = 1 << MaxLayoutGridSizes.Find(NewSelection);
-
-		if (CurrentLayout->GetMaxGridSize().X != Size || CurrentLayout->GetMaxGridSize().Y != Size)
-		{
-			CurrentLayout->SetMaxGridSize(FIntPoint(Size));
-			
-			CurrentLayout->MarkPackageDirty();
-		}
-	}
-}
-
-
-void SCustomizableObjectNodeLayoutBlocksEditor::OnReductionMethodChanged(TSharedPtr<FString> NewSelection, ESelectInfo::Type SelectInfo)
-{
-	if (CurrentLayout)
-	{
-		uint32 selection = BlockReductionMethods.IndexOfByKey(NewSelection);
-
-		if (CurrentLayout->GetBlockReductionMethod() != (ECustomizableObjectLayoutBlockReductionMethod)selection)
-		{
-			CurrentLayout->SetBlockReductionMethod((ECustomizableObjectLayoutBlockReductionMethod)selection);
-			CurrentLayout->MarkPackageDirty();
-		}
-	}
-}
-
-
 TArray<FCustomizableObjectLayoutBlock> SCustomizableObjectNodeLayoutBlocksEditor::GetBlocks() const
 {
 	TArray<FCustomizableObjectLayoutBlock> Blocks;
@@ -702,4 +416,3 @@ TSharedPtr<IToolTip> SCustomizableObjectNodeLayoutBlocksEditor::GenerateInfoTool
 
 
 #undef LOCTEXT_NAMESPACE
-
diff --git a/Engine/Plugins/Experimental/Mutable/Source/CustomizableObjectEditor/Private/MuCOE/SCustomizableObjectNodeLayoutBlocksEditor.h b/Engine/Plugins/Experimental/Mutable/Source/CustomizableObjectEditor/Private/MuCOE/SCustomizableObjectNodeLayoutBlocksEditor.h
--- a/Engine/Plugins/Experimental/Mutable/Source/CustomizableObjectEditor/Private/MuCOE/SCustomizableObjectNodeLayoutBlocksEditor.h
+++ b/Engine/Plugins/Experimental/Mutable/Source/CustomizableObjectEditor/Private/MuCOE/SCustomizableObjectNodeLayoutBlocksEditor.h
@@ -58,38 +58,6 @@ private:
 	/** */
 	TSharedPtr<SCustomizableObjectLayoutGrid> LayoutGridWidget;
 
-	TSharedPtr< SVerticalBox > LayoutGridSizeWidget;
-
-	/** Widget for displaying the available layout block grid sizes. */
-	TSharedPtr< STextComboBox > LayoutGridSizeCombo;
-	TSharedPtr< STextComboBox > MaxLayoutGridSizeCombo;
-
-	/** List of available layout grid sizes. */
-	TArray< TSharedPtr< FString > > LayoutGridSizes;
-
-	/** List of available layout max grid sizes. */
-	TArray< TSharedPtr< FString > > MaxLayoutGridSizes;
-
-	/** List of available layout packing strategies. */
-	TArray< TSharedPtr< FString > > LayoutPackingStrategies;
-
-	/** Widget to select the available layout packing strategies. */
-	TSharedPtr< STextComboBox > LayoutPackingStrategyCombo;
-
-	/** List of available block reduction methods. */
-	TArray< TSharedPtr< FString > > BlockReductionMethods;
-
-	/** Widget to select the available layout block reduction methods. */
-	TSharedPtr< STextComboBox > BlockReductionMethodsCombo;
-
-	/** Widget to select the layout packing strategy */
-	TSharedPtr< SHorizontalBox> LayoutStrategyWidget;
-
-	/** Widget to select the fixed layout properties */
-	TSharedPtr< SVerticalBox> FixedLayoutWidget;
-
-	TSharedPtr<SWidget> StrategyWidget;
-
 	/** The list of UI Commands executable */
 	TSharedRef<FUICommandList> UICommandList;
 
@@ -102,7 +70,6 @@ private:
 
 	/** Callbacks from the layout block editor. */
 	TSharedRef<SWidget> BuildLayoutToolBar();
-	TSharedRef<SWidget> BuildLayoutStrategyWidgets(const ISlateStyle* Style, const FName& StyleName);
 
 	void OnAddBlock();
 	void OnAddBlockAt(const FIntPoint Min, const FIntPoint Max);
@@ -110,11 +77,6 @@ private:
 
 	/** Generate layout blocks using UVs*/
 	void OnGenerateBlocks();
-		
-	/** . */
-	void OnGridSizeChanged( TSharedPtr<FString> NewSelection, ESelectInfo::Type SelectInfo );
-	void OnMaxGridSizeChanged(TSharedPtr<FString> NewSelection, ESelectInfo::Type SelectInfo);
-	void OnReductionMethodChanged(TSharedPtr<FString> NewSelection, ESelectInfo::Type SelectInfo);
 
 	/** Sets the block priority from the input text. */
 	void OnSetBlockPriority(int32 InValue);
@@ -125,9 +87,6 @@ private:
 	/** Sets the block reduction ReduceByTwo option. */
 	void OnSetBlockReductionByTwo(bool bInValue);
 
-	/** Called when the packing strategy has changed. */
-	void OnLayoutPackingStrategyChanged(TSharedPtr<FString> NewSelection, ESelectInfo::Type SelectInfo);
-
 	TSharedPtr<IToolTip> GenerateInfoToolTip() const;
 
 };
diff --git a/Engine/Plugins/Experimental/Mutable/Source/CustomizableObjectEditor/Private/MuCOE/SCustomizableObjectNodeLayoutBlocksSelector.cpp b/Engine/Plugins/Experimental/Mutable/Source/CustomizableObjectEditor/Private/MuCOE/SCustomizableObjectNodeLayoutBlocksSelector.cpp
--- a/Engine/Plugins/Experimental/Mutable/Source/CustomizableObjectEditor/Private/MuCOE/SCustomizableObjectNodeLayoutBlocksSelector.cpp
+++ b/Engine/Plugins/Experimental/Mutable/Source/CustomizableObjectEditor/Private/MuCOE/SCustomizableObjectNodeLayoutBlocksSelector.cpp
@@ -149,7 +149,8 @@ void SCustomizableObjectNodeLayoutBlocksSelector::AddReferencedObjects( FReferen
 
 TSharedRef<SWidget> SCustomizableObjectNodeLayoutBlocksSelector::BuildLayoutToolBar()
 {
-	FToolBarBuilder LayoutToolbarBuilder(UICommandList, FMultiBoxCustomization::None);
+	FSlimHorizontalToolBarBuilder LayoutToolbarBuilder(UICommandList, FMultiBoxCustomization::None, nullptr, true);
+	LayoutToolbarBuilder.SetLabelVisibility(EVisibility::Visible);
 
 	// Build toolbar widgets
 
