[mutable] Initial support for project-specific versioning of child Customizable Objects and table rows without Mutable having to know anything about how that project implements versioning. Child objects can already be versioned. Table rows are pending.
#rnx

diff --git a/Engine/Plugins/Experimental/Mutable/Source/CustomizableObject/Classes/MuCO/CustomizableObject.h b/Engine/Plugins/Experimental/Mutable/Source/CustomizableObject/Classes/MuCO/CustomizableObject.h
--- a/Engine/Plugins/Experimental/Mutable/Source/CustomizableObject/Classes/MuCO/CustomizableObject.h
+++ b/Engine/Plugins/Experimental/Mutable/Source/CustomizableObject/Classes/MuCO/CustomizableObject.h
@@ -487,6 +487,22 @@ public:
 	* "SetReplaceDiscardedWithReferenceMeshEnabled" (See the c++ section). */
 	UPROPERTY(EditAnywhere, Category = CustomizableObject)
 	TArray<TObjectPtr<USkeletalMesh>> ReferenceSkeletalMeshes;
+
+	/**
+	  * The optional VersionBridge asset, which must implement the ICustomizableObjectVersionBridgeInterface, will be used to
+	  * decide which Mutable child CustomizableObjects and table rows must be included in a compilation/cook depending on its 
+	  * version struct/column by comparing it to the game-specific version system.
+	*/
+	UPROPERTY(EditAnywhere, Category = Versioning)
+	TObjectPtr<UObject> VersionBridge;
+
+	/**
+	  * This optional struct is used to define which version this child CustomizableObject belongs to. It will be used during
+	  * cook/compilation to decide whether this CO should be included or not in the final compiled CO. To be used, the root
+	  * CO must have defined the VersionBridge property, which must implement the ICustomizableObjectVersionBridgeInterface
+	*/
+	UPROPERTY(EditAnywhere, Category = Versioning)
+	FInstancedStruct VersionStruct;
 #endif
 	
 	UPROPERTY(EditAnywhere, Category = CustomizableObject, meta = (DisplayName = "LOD Settings"))
diff --git a/Engine/Plugins/Experimental/Mutable/Source/CustomizableObjectEditor/Private/MuCOE/CustomizableObjectCompiler.cpp b/Engine/Plugins/Experimental/Mutable/Source/CustomizableObjectEditor/Private/MuCOE/CustomizableObjectCompiler.cpp
--- a/Engine/Plugins/Experimental/Mutable/Source/CustomizableObjectEditor/Private/MuCOE/CustomizableObjectCompiler.cpp
+++ b/Engine/Plugins/Experimental/Mutable/Source/CustomizableObjectEditor/Private/MuCOE/CustomizableObjectCompiler.cpp
@@ -28,6 +28,7 @@
 #include "MuCO/CustomizableObjectSystemPrivate.h"
 #include "MuR/Model.h"
 #include "MuR/ModelPrivate.h"
+#include "MuCOE/CustomizableObjectVersionBridge.h"
 
 class UTexture2D;
 
@@ -172,6 +173,16 @@ void FCustomizableObjectCompiler::Compile(UCustomizableObject& Object, const FCo
 	{
 		UE_LOG(LogMutable, Warning, TEXT("Failed to compile Customizable Object [%s]. Mutable is disabled. To enable it set the CVar Mutable.Enabled to true."), *Object.GetName());
 		SetCompilationState(ECustomizableObjectCompilationState::Failed);
+
+		return;
+	}
+
+	if (Object.VersionBridge && !Object.VersionBridge->GetClass()->ImplementsInterface(UCustomizableObjectVersionBridgeInterface::StaticClass()))
+	{
+		UE_LOG(LogMutable, Warning, TEXT("In Customizable Object [%s], the VersionBridge asset [%s] does not implement the required UCustomizableObjectVersionBridgeInterface."), 
+			*Object.GetName(), *Object.VersionBridge.GetName());
+		SetCompilationState(ECustomizableObjectCompilationState::Failed);
+
 		return;
 	}
 
@@ -334,6 +345,31 @@ void FCustomizableObjectCompiler::ProcessChildObjectsRecursively(UCustomizableOb
 			continue;
 		}
 
+		if (ChildObject->VersionStruct.IsValid())
+		{
+			if (!GenerationContext.Object->VersionBridge)
+			{
+				UE_LOG(LogMutable, Warning, TEXT("The child Customizable Object [%s] defines its VersionStruct Property but its root CustomizableObject doesn't define the VersionBridge property. There's no way to verify the VersionStruct has to be included in this compilation, so the child CustomizableObject will be omitted."), 
+					*ChildObject->GetName());
+				continue;
+			}
+
+			ICustomizableObjectVersionBridgeInterface* CustomizableObjectVersionBridgeInterface = Cast<ICustomizableObjectVersionBridgeInterface>(GenerationContext.Object->VersionBridge);
+
+			if (CustomizableObjectVersionBridgeInterface)
+			{
+				if (!CustomizableObjectVersionBridgeInterface->IsVersionStructIncludedInCurrentRelease(ChildObject->VersionStruct))
+				{
+					continue;
+				}
+			}
+			else
+			{
+				// This should never happen as the ICustomizableObjectVersionBridgeInterface was already checked at the start of the compilation
+				ensure(false);
+			}
+		}
+
 		ArrayAlreadyProcessedChild.Add(ReferenceName);
 
 		if (!bMultipleBaseObjectsFound)
diff --git a/Engine/Plugins/Experimental/Mutable/Source/CustomizableObjectEditor/Private/MuCOE/CustomizableObjectDetails.cpp b/Engine/Plugins/Experimental/Mutable/Source/CustomizableObjectEditor/Private/MuCOE/CustomizableObjectDetails.cpp
--- a/Engine/Plugins/Experimental/Mutable/Source/CustomizableObjectEditor/Private/MuCOE/CustomizableObjectDetails.cpp
+++ b/Engine/Plugins/Experimental/Mutable/Source/CustomizableObjectEditor/Private/MuCOE/CustomizableObjectDetails.cpp
@@ -145,6 +145,34 @@ void FCustomizableObjectDetails::CustomizeDetails( IDetailLayoutBuilder& DetailB
 			Property->MarkResetToDefaultCustomized();
 		}
 	}
+
+	TSharedRef<IPropertyHandle> VersionBridgeProperty = DetailBuilder.GetProperty("VersionBridge");
+
+	if (VersionBridgeProperty->IsValidHandle() && CustomizableObject)
+	{
+		if (CustomizableObject->IsChildObject())
+		{
+			VersionBridgeProperty->MarkHiddenByCustomization();
+		}
+		else
+		{
+			VersionBridgeProperty->MarkResetToDefaultCustomized();
+		}
+	}
+
+	TSharedRef<IPropertyHandle> VersionStructProperty = DetailBuilder.GetProperty("VersionStruct");
+
+	if (VersionStructProperty->IsValidHandle() && CustomizableObject)
+	{
+		if (CustomizableObject->IsChildObject())
+		{
+			VersionStructProperty->MarkResetToDefaultCustomized();
+		}
+		else
+		{
+			VersionStructProperty->MarkHiddenByCustomization();
+		}
+	}
 }
 
 
diff --git a/Engine/Plugins/Experimental/Mutable/Source/CustomizableObjectEditor/Public/MuCOE/CustomizableObjectVersionBridge.h b/Engine/Plugins/Experimental/Mutable/Source/CustomizableObjectEditor/Public/MuCOE/CustomizableObjectVersionBridge.h
new file mode 100644
--- /dev/null
+++ b/Engine/Plugins/Experimental/Mutable/Source/CustomizableObjectEditor/Public/MuCOE/CustomizableObjectVersionBridge.h
@@ -0,0 +1,44 @@
+// Copyright Epic Games, Inc. All Rights Reserved.
+
+#pragma once
+
+#include "UObject/Interface.h"
+#include "InstancedStruct.h"
+#include "CustomizableObjectVersionBridge.generated.h"
+
+UINTERFACE(MinimalAPI, Blueprintable)
+class UCustomizableObjectVersionBridgeInterface : public UInterface
+{
+	GENERATED_BODY()
+};
+
+class ICustomizableObjectVersionBridgeInterface
+{
+	GENERATED_BODY()
+
+public:
+	/** Interface function declarations */
+
+	/** Checks whether the version encoded in VersionStruct is older or equal to the CurrentRelease, meaning that all
+	  *	Mutable child Customizable Objects marked with this version that return true should be compiled and packaged.
+	  * The VersionStruct will encode the game-specific version struct, and should be accessed from the game-specific
+	  * implementation of this interface in the the VersionBridge property in the CustomizableObject with the
+	  * VersionStruct.GetPtr<GameSpecificVersionType>() method. The VersionStruct in child CustomizableObjects
+	  * must be of the GameSpecificVersionType type.
+	  * @param VersionStruct Struct encoding the version of the Mutable table row
+	  * @return True if the table row is included in the current release
+	*/
+	virtual bool IsVersionStructIncludedInCurrentRelease(const FInstancedStruct& VersionStruct) const { return false; };
+
+	/** Checks whether the version encoded in VersionProperty/CellData is older or equal to the CurrentRelease, meaning that all
+	  *	Mutable table rows marked with this version that return true should be compiled and packaged.
+	  * The VersionProperty/CellData will encode the game-specific version struct, and should be accessed from the game-specific
+	  * implementation of this interface in the the VersionBridge property in the CustomizableObject by casting CellData
+	  * to the GameSpecificVersionType. The Version column in Mutable data table must be of the
+	  * GameSpecificVersionType type.
+	  * @param VersionProperty Property of the column encoding the version of the Mutable table row. Useful for type checking.
+	  * @param CellData pointer to the GameSpecificVersionType struct encoding the version of the Mutable table row
+	  * @return True if the table row is included in the current release
+	*/
+	virtual bool IsVersionPropertyIncludedInCurrentRelease(const FProperty& VersionProperty, const uint8* CellData) const { return false; };
+};
