[mutable] Added ticking engine wait before we compile the CO to ensure that no engine tick related actions will run alongside the instance updates that get run later.
- Added Wait method that while waitting do tick the engine for it to complete it's tick related work.

#rnx

diff --git a/Engine/Plugins/Experimental/Mutable/Source/MutableValidation/Private/MuV/COIBakingTestCommandlet.cpp b/Engine/Plugins/Experimental/Mutable/Source/MutableValidation/Private/MuV/COIBakingTestCommandlet.cpp
--- a/Engine/Plugins/Experimental/Mutable/Source/MutableValidation/Private/MuV/COIBakingTestCommandlet.cpp
+++ b/Engine/Plugins/Experimental/Mutable/Source/MutableValidation/Private/MuV/COIBakingTestCommandlet.cpp
@@ -69,6 +69,9 @@ int32 UCOIBakingTestCommandlet::Main(const FString& Params)
 
 	// Perform a blocking search to ensure all assets used by mutable are reachable using the AssetRegistry
 	PrepareAssetRegistry();
+
+	// Make sure there is nothing else that the engine needs to do before starting our test
+	Wait(60);
 	
 	// Compile it's CO (using current config)
 	UCustomizableObject* InstanceCustomizableObject = TargetInstance->GetCustomizableObject();
@@ -77,7 +80,6 @@ int32 UCOIBakingTestCommandlet::Main(const FString& Params)
 		UE_LOG(LogMutable,Error,TEXT("The instance %s does not have a CO to compile : Exitting commandlet."), *TargetInstance->GetName());
 		return 1;
 	}
-
 	
 	// Set the target platform to be using for the compilation. Must not be a nullptr
 	FCompilationOptions CompilationOptions = InstanceCustomizableObject->CompileOptions;
diff --git a/Engine/Plugins/Experimental/Mutable/Source/MutableValidation/Private/MuV/COIBulkUpdateTestCommandlet.cpp b/Engine/Plugins/Experimental/Mutable/Source/MutableValidation/Private/MuV/COIBulkUpdateTestCommandlet.cpp
--- a/Engine/Plugins/Experimental/Mutable/Source/MutableValidation/Private/MuV/COIBulkUpdateTestCommandlet.cpp
+++ b/Engine/Plugins/Experimental/Mutable/Source/MutableValidation/Private/MuV/COIBulkUpdateTestCommandlet.cpp
@@ -115,6 +115,9 @@ int32 UCOIBulkUpdateTestCommandlet::Main(const FString& Params)
 	
 	// ------ Execution of the actual mutable operations ------ 
 	
+	// Make sure there is nothing else that the engine needs to do before starting our test
+	Wait(60);
+	
 	// Create the updater object so we can later call for the update a target instance
 	InstanceUpdater = NewObject<UCOIUpdater>();
 	
diff --git a/Engine/Plugins/Experimental/Mutable/Source/MutableValidation/Private/MuV/CustomizableObjectValidationCommandlet.cpp b/Engine/Plugins/Experimental/Mutable/Source/MutableValidation/Private/MuV/CustomizableObjectValidationCommandlet.cpp
--- a/Engine/Plugins/Experimental/Mutable/Source/MutableValidation/Private/MuV/CustomizableObjectValidationCommandlet.cpp
+++ b/Engine/Plugins/Experimental/Mutable/Source/MutableValidation/Private/MuV/CustomizableObjectValidationCommandlet.cpp
@@ -63,6 +63,9 @@ int32 UCustomizableObjectValidationCommandlet::Main(const FString& Params)
 	
 	// Perform a blocking search to ensure all assets used by mutable are reachable using the AssetRegistry
 	PrepareAssetRegistry();
+
+	// Make sure there is nothing else that the engine needs to do before starting our test
+	Wait(60);
 	
 	// Compile the Customizable Object ------------------------------------------------------------------------------ //
 	bool bWasCoCompilationSuccessful = false;
diff --git a/Engine/Plugins/Experimental/Mutable/Source/MutableValidation/Private/MuV/ValidationUtils.cpp b/Engine/Plugins/Experimental/Mutable/Source/MutableValidation/Private/MuV/ValidationUtils.cpp
--- a/Engine/Plugins/Experimental/Mutable/Source/MutableValidation/Private/MuV/ValidationUtils.cpp
+++ b/Engine/Plugins/Experimental/Mutable/Source/MutableValidation/Private/MuV/ValidationUtils.cpp
@@ -109,11 +109,31 @@ bool CompileCustomizableObject(UCustomizableObject* InCustomizableObject, const
 		// TODO: Add logs for the other relevant configs of the model being compiled
 	}
 
-
 	return bWasCoCompilationSuccessful;
 }
 
 
+void Wait(const double ToWaitSeconds)
+{
+	check (ToWaitSeconds > 0);
+	
+	const double EndSeconds = FPlatformTime::Seconds() + ToWaitSeconds;
+	UE_LOG(LogMutable,Display,TEXT("Holding test execution for %f seconds."),ToWaitSeconds);
+	while (FPlatformTime::Seconds() < EndSeconds)
+	{
+		// Tick the engine
+		CommandletHelpers::TickEngine();
+
+		// Stop if exit was requested
+		if (IsEngineExitRequested())
+		{
+			break;
+		}
+	}
+
+	UE_LOG(LogMutable,Display,TEXT("Resuming test execution."));
+}
+
 
 bool UCOIUpdater::UpdateInstance(UCustomizableObjectInstance* InInstance)
 {
diff --git a/Engine/Plugins/Experimental/Mutable/Source/MutableValidation/Private/MuV/ValidationUtils.h b/Engine/Plugins/Experimental/Mutable/Source/MutableValidation/Private/MuV/ValidationUtils.h
--- a/Engine/Plugins/Experimental/Mutable/Source/MutableValidation/Private/MuV/ValidationUtils.h
+++ b/Engine/Plugins/Experimental/Mutable/Source/MutableValidation/Private/MuV/ValidationUtils.h
@@ -19,6 +19,13 @@ struct FCompilationOptions;
 void PrepareAssetRegistry();
 
 
+/**
+ * Hold the thread for the time specified while ticking the engine.
+ * @param ToWaitSeconds The time in seconds we want to hold the execution of the thread
+ */
+void Wait(const double ToWaitSeconds);
+
+
 /**
  * Compiles a CO synchronously
  * @param InCustomizableObject The Customizable Object we want to synchronously compile
