[Mutable] Fix "Use of arithmetic operator on Boolean type(s)" static analyzer warning.

#rnx

diff --git a/Engine/Plugins/Experimental/Mutable/Source/CustomizableObject/Private/MuCO/CustomizableObjectSystem.cpp b/Engine/Plugins/Experimental/Mutable/Source/CustomizableObject/Private/MuCO/CustomizableObjectSystem.cpp
--- a/Engine/Plugins/Experimental/Mutable/Source/CustomizableObject/Private/MuCO/CustomizableObjectSystem.cpp
+++ b/Engine/Plugins/Experimental/Mutable/Source/CustomizableObject/Private/MuCO/CustomizableObjectSystem.cpp
@@ -2745,11 +2745,11 @@ int32 UCustomizableObjectSystem::TickInternal()
 
 	int32 RemainingWork = Private->CurrentMutableOperation.IsValid() + 
 		Private->MutablePendingInstanceWork.Num() +
-		static_cast<bool>(LODUpdateCandidateFound) + // Still a pending LOD update. We can not use the size of RequestedLODUpdates since not all requests valid in future ticks.
+		static_cast<int32>(LODUpdateCandidateFound != nullptr) + // Still a pending LOD update. We can not use the size of RequestedLODUpdates since not all requests valid in future ticks.
 		RemainingTasks;
 
 #if WITH_EDITOR
-	RemainingWork += static_cast<bool>(GetPrivate()->RecompileCustomizableObjectsCompiler); // Compiler only is valid if we are compiling a CO.
+	RemainingWork += static_cast<int32>(GetPrivate()->RecompileCustomizableObjectsCompiler != nullptr); // Compiler only is valid if we are compiling a CO.
 #endif
 	
 	return RemainingWork;
