[Mutable] Don't zero global memory counter if not benchmarking.

#rnx

diff --git a/Engine/Plugins/Experimental/Mutable/Source/CustomizableObject/Private/MuCO/CustomizableObjectMipDataProvider.cpp b/Engine/Plugins/Experimental/Mutable/Source/CustomizableObject/Private/MuCO/CustomizableObjectMipDataProvider.cpp
--- a/Engine/Plugins/Experimental/Mutable/Source/CustomizableObject/Private/MuCO/CustomizableObjectMipDataProvider.cpp
+++ b/Engine/Plugins/Experimental/Mutable/Source/CustomizableObject/Private/MuCO/CustomizableObjectMipDataProvider.cpp
@@ -131,9 +131,12 @@ namespace impl
 		MUTABLE_CPUPROFILER_SCOPE(Task_Mutable_UpdateImage);
 		const double StartTime = FPlatformTime::Seconds();
 		
-		// Cache memory used when starting the update of the image
-		OperationData->ImageUpdateStartBytes = mu::FGlobalMemoryCounter::GetCounter();
-		mu::FGlobalMemoryCounter::Zero();
+		if (CVarEnableBenchmark.GetValueOnAnyThread())
+		{
+			// Cache memory used when starting the update of the image
+			OperationData->ImageUpdateStartBytes = mu::FGlobalMemoryCounter::GetCounter();
+			mu::FGlobalMemoryCounter::Zero();
+		}
 		
 		// Any external texture that may be needed for this update will be requested from Mutable Core's GetImage
 		// which will safely access the GlobalExternalImages map, and then just get the cached image or issue a disk read
diff --git a/Engine/Plugins/Experimental/Mutable/Source/CustomizableObject/Private/MuCO/CustomizableObjectSystem.cpp b/Engine/Plugins/Experimental/Mutable/Source/CustomizableObject/Private/MuCO/CustomizableObjectSystem.cpp
--- a/Engine/Plugins/Experimental/Mutable/Source/CustomizableObject/Private/MuCO/CustomizableObjectSystem.cpp
+++ b/Engine/Plugins/Experimental/Mutable/Source/CustomizableObject/Private/MuCO/CustomizableObjectSystem.cpp
@@ -1582,8 +1582,11 @@ namespace impl
 		UCustomizableObjectSystem* System = UCustomizableObjectSystem::GetInstanceChecked(); // Save since UCustomizableObjectSystem::BeginDestroy always waits for all tasks to finish
 		const UCustomizableObjectSystemPrivate* SystemPrivate = System->GetPrivate();
 		
-		Operation->UpdateStartBytes = mu::FGlobalMemoryCounter::GetCounter();
-		mu::FGlobalMemoryCounter::Zero();
+		if (CVarEnableBenchmark.GetValueOnAnyThread())
+		{
+			Operation->UpdateStartBytes = mu::FGlobalMemoryCounter::GetCounter();
+			mu::FGlobalMemoryCounter::Zero();
+		}
 
 		// Prepare streaming for the current customizable object
 		check(SystemPrivate->Streamer != nullptr);
