[mutable] Re-enable concurrency in compilation when cooking.

#rnx

diff --git a/Engine/Plugins/Experimental/Mutable/Source/CustomizableObjectEditor/Private/MuCOE/CustomizableObjectCompileRunnable.cpp b/Engine/Plugins/Experimental/Mutable/Source/CustomizableObjectEditor/Private/MuCOE/CustomizableObjectCompileRunnable.cpp
--- a/Engine/Plugins/Experimental/Mutable/Source/CustomizableObjectEditor/Private/MuCOE/CustomizableObjectCompileRunnable.cpp
+++ b/Engine/Plugins/Experimental/Mutable/Source/CustomizableObjectEditor/Private/MuCOE/CustomizableObjectCompileRunnable.cpp
@@ -24,7 +24,7 @@ class ITargetPlatform;
 
 TAutoConsoleVariable<bool> CVarMutableCompilerConcurrency(
 	TEXT("mutable.ForceCompilerConcurrency"),
-	false,
+	true,
 	TEXT("Force the use of multithreading when compiling CustomizableObjects both in editor and cook commandlets."),
 	ECVF_Default);
 
diff --git a/Engine/Plugins/Experimental/Mutable/Source/MutableTools/Private/MuT/CodeOptimiser.cpp b/Engine/Plugins/Experimental/Mutable/Source/MutableTools/Private/MuT/CodeOptimiser.cpp
--- a/Engine/Plugins/Experimental/Mutable/Source/MutableTools/Private/MuT/CodeOptimiser.cpp
+++ b/Engine/Plugins/Experimental/Mutable/Source/MutableTools/Private/MuT/CodeOptimiser.cpp
@@ -782,7 +782,7 @@ namespace mu
 			/** Protect access to the original AST being optimized. */
 			FCriticalSection ASTAccessLock;
 
-			// Launch the tasks. Do it from a task in the ASTPipe to make sure to avoid race conditions in the main AST manipulation.
+			// Launch the tasks.
 			UE::Tasks::FTask LaunchTask = UE::Tasks::Launch(TEXT("ConstantGeneratorLaunchTasks"), 
 				[&ConstantSubgraphs, &GetRequisites, &ASTAccessLock, Pass, InOptions]()
 				{
