[Mutable] Fix assert when inspecting an image operation in the debugger.
- Runtime loading data was used to determine if a rom is loaded or not, but this data may not be up to date in different contexts.

#rnx

diff --git a/Engine/Plugins/Experimental/Mutable/Source/MutableRuntime/Internal/MuR/ModelPrivate.h b/Engine/Plugins/Experimental/Mutable/Source/MutableRuntime/Internal/MuR/ModelPrivate.h
--- a/Engine/Plugins/Experimental/Mutable/Source/MutableRuntime/Internal/MuR/ModelPrivate.h
+++ b/Engine/Plugins/Experimental/Mutable/Source/MutableRuntime/Internal/MuR/ModelPrivate.h
@@ -337,18 +337,18 @@ namespace mu
 		//! Return true if the given ROM is loaded.
 		FORCEINLINE bool IsRomLoaded(int32 RomIndex) const
 		{
-			const bool bIsRomLoaded = LoadedMemTrackedRoms.IsValidIndex(RomIndex);
-			check(bIsRomLoaded == Invoke([&]()
+			switch (m_roms[RomIndex].ResourceType)
 			{
-				switch (m_roms[RomIndex].ResourceType)
-				{
-				case DT_IMAGE: return ConstantImageLODs[m_roms[RomIndex].ResourceIndex].Value.get() != nullptr;
-				case DT_MESH: return ConstantMeshes[m_roms[RomIndex].ResourceIndex].Value.get() != nullptr;
-				default: return false;
-				}
-			}));
+				case DT_IMAGE: 
+					return ConstantImageLODs[m_roms[RomIndex].ResourceIndex].Value.get() != nullptr;
+				case DT_MESH: 
+					return ConstantMeshes[m_roms[RomIndex].ResourceIndex].Value.get() != nullptr;
+				default:
+					check(false);
+					break;
+			}
 
-			return bIsRomLoaded;
+			return false;
 		}
 
 		/** Unload a rom resource. Return the size of the unloaded rom.*/
