[Mutable] TextureProject nodes now return the reference texture first in FindReferenceImage.

#rnx

diff --git a/Engine/Plugins/Experimental/Mutable/Source/CustomizableObjectEditor/Private/MuCOE/GenerateMutableSource/GenerateMutableSource.cpp b/Engine/Plugins/Experimental/Mutable/Source/CustomizableObjectEditor/Private/MuCOE/GenerateMutableSource/GenerateMutableSource.cpp
--- a/Engine/Plugins/Experimental/Mutable/Source/CustomizableObjectEditor/Private/MuCOE/GenerateMutableSource/GenerateMutableSource.cpp
+++ b/Engine/Plugins/Experimental/Mutable/Source/CustomizableObjectEditor/Private/MuCOE/GenerateMutableSource/GenerateMutableSource.cpp
@@ -611,20 +611,27 @@ UTexture2D* FindReferenceImage(const UEdGraphPin* Pin, FMutableGraphGenerationCo
 
 	else if (const UCustomizableObjectNodeTextureProject* TypedNodeProj = Cast<UCustomizableObjectNodeTextureProject>(Node))
 	{
-		int TexIndex = -1;// TypedNodeProj->OutputPins.Find((UEdGraphPin*)Pin);
-		for (int32 i = 0; i < TypedNodeProj->GetNumOutputs(); ++i)
+		if (TypedNodeProj->ReferenceTexture)
 		{
-			if (TypedNodeProj->OutputPins(i) == Pin)
+			Result = TypedNodeProj->ReferenceTexture;
+		}
+		else
+		{
+			int TexIndex = -1;// TypedNodeProj->OutputPins.Find((UEdGraphPin*)Pin);
+			for (int32 i = 0; i < TypedNodeProj->GetNumOutputs(); ++i)
 			{
-				TexIndex = i;
+				if (TypedNodeProj->OutputPins(i) == Pin)
+				{
+					TexIndex = i;
+				}
 			}
-		}
 
-		check(TexIndex >= 0 && TexIndex < TypedNodeProj->GetNumTextures());
+			check(TexIndex >= 0 && TexIndex < TypedNodeProj->GetNumTextures());
 
-		if (const UEdGraphPin* ConnectedPin = FollowInputPin(*TypedNodeProj->TexturePins(TexIndex)))
-		{
-			Result = FindReferenceImage(ConnectedPin, GenerationContext);
+			if (const UEdGraphPin* ConnectedPin = FollowInputPin(*TypedNodeProj->TexturePins(TexIndex)))
+			{
+				Result = FindReferenceImage(ConnectedPin, GenerationContext);
+			}
 		}
 	}
 
