[mutable] Added fixed value for working memory and clear working memory after each instance update
- Enabled automatic working memory clearing after each instance update
- Set the working mememory to be used during all tests (defaulted to 51200 KB)
- Set by default all tests running without using the disk as memory.

diff --git a/Engine/Plugins/Experimental/Mutable/Source/MutableValidation/Private/MuV/COIBakingTestCommandlet.cpp b/Engine/Plugins/Experimental/Mutable/Source/MutableValidation/Private/MuV/COIBakingTestCommandlet.cpp
--- a/Engine/Plugins/Experimental/Mutable/Source/MutableValidation/Private/MuV/COIBakingTestCommandlet.cpp
+++ b/Engine/Plugins/Experimental/Mutable/Source/MutableValidation/Private/MuV/COIBakingTestCommandlet.cpp
@@ -72,6 +72,8 @@ int32 UCOIBakingTestCommandlet::Main(const FString& Params)
 
 	// Make sure there is nothing else that the engine needs to do before starting our test
 	Wait(60);
+
+	LogMutableSettings();
 	
 	// Compile it's CO (using current config)
 	UCustomizableObject* InstanceCustomizableObject = TargetInstance->GetCustomizableObject();
@@ -85,6 +87,7 @@ int32 UCOIBakingTestCommandlet::Main(const FString& Params)
 	FCompilationOptions CompilationOptions = InstanceCustomizableObject->CompileOptions;
 	ITargetPlatformManagerModule& TPM = GetTargetPlatformManagerRef();
 	CompilationOptions.TargetPlatform = TPM.GetRunningTargetPlatform();
+	CompilationOptions.bUseDiskCompilation = false;
 	if (!CompileCustomizableObject(InstanceCustomizableObject,true, &CompilationOptions))
 	{
 		UE_LOG(LogMutable,Error,TEXT("Failed to compile the target CO. Exitting commandlet."));
diff --git a/Engine/Plugins/Experimental/Mutable/Source/MutableValidation/Private/MuV/COIBulkUpdateTestCommandlet.cpp b/Engine/Plugins/Experimental/Mutable/Source/MutableValidation/Private/MuV/COIBulkUpdateTestCommandlet.cpp
--- a/Engine/Plugins/Experimental/Mutable/Source/MutableValidation/Private/MuV/COIBulkUpdateTestCommandlet.cpp
+++ b/Engine/Plugins/Experimental/Mutable/Source/MutableValidation/Private/MuV/COIBulkUpdateTestCommandlet.cpp
@@ -35,6 +35,8 @@ int32 UCOIBulkUpdateTestCommandlet::Main(const FString& Params)
 	
 	// Load the asset registry system so we can proceed without issues
 	PrepareAssetRegistry();
+
+	LogMutableSettings();
 	
 	// Cache all UAssets (find a way to not have them in memory yet, not until we need them)
 	TArray<FAssetData> FoundAssetData;
@@ -135,6 +137,8 @@ int32 UCOIBulkUpdateTestCommandlet::Main(const FString& Params)
 		// Set the compilation platform based on what the system is currently running on
 		FCompilationOptions CompilationOptions = CustomizableObjectToCompile->CompileOptions;
 		CompilationOptions.TargetPlatform = TargetCompilationPlatform;
+
+		CompilationOptions.bUseDiskCompilation = false;
 		
 		// Compile the current CO object
 		const FString CustomizableObjectName = CustomizableObjectToCompile->GetName();
diff --git a/Engine/Plugins/Experimental/Mutable/Source/MutableValidation/Private/MuV/CustomizableObjectValidationCommandlet.cpp b/Engine/Plugins/Experimental/Mutable/Source/MutableValidation/Private/MuV/CustomizableObjectValidationCommandlet.cpp
--- a/Engine/Plugins/Experimental/Mutable/Source/MutableValidation/Private/MuV/CustomizableObjectValidationCommandlet.cpp
+++ b/Engine/Plugins/Experimental/Mutable/Source/MutableValidation/Private/MuV/CustomizableObjectValidationCommandlet.cpp
@@ -66,6 +66,8 @@ int32 UCustomizableObjectValidationCommandlet::Main(const FString& Params)
 
 	// Make sure there is nothing else that the engine needs to do before starting our test
 	Wait(60);
+
+	LogMutableSettings();
 	
 	// Compile the Customizable Object ------------------------------------------------------------------------------ //
 	bool bWasCoCompilationSuccessful = false;
@@ -80,12 +82,8 @@ int32 UCustomizableObjectValidationCommandlet::Main(const FString& Params)
 		CompilationOptions.TargetPlatform = TargetCompilationPlatform;
 
 		// Disk cache usage for compilation operation : Override if the user provided an argument with a different value than the default one of the CO
-		bool bUseDiskCacheForCompilation = false;
-		if (FParse::Bool(*Params,TEXT("UseDiskCompilation="),bUseDiskCacheForCompilation))
-		{
-			// Set the disk cache configuration
-			CompilationOptions.bUseDiskCompilation = bUseDiskCacheForCompilation;
-		}
+		CompilationOptions.bUseDiskCompilation = false;
+		FParse::Bool(*Params,TEXT("UseDiskCompilation="),CompilationOptions.bUseDiskCompilation);
 
 		bWasCoCompilationSuccessful = CompileCustomizableObject(ToTestCustomizableObject, true, &CompilationOptions);
 	}
@@ -229,7 +227,7 @@ int32 UCustomizableObjectValidationCommandlet::Main(const FString& Params)
     }
     else
     {
-        UE_LOG(LogMutable,Error,TEXT("The generation of Customizable object instances  was not succesfull."));
+        UE_LOG(LogMutable,Error,TEXT("The generation of Customizable object instances was not succesfull."));
     }
 	
 	
diff --git a/Engine/Plugins/Experimental/Mutable/Source/MutableValidation/Private/MuV/ValidationUtils.cpp b/Engine/Plugins/Experimental/Mutable/Source/MutableValidation/Private/MuV/ValidationUtils.cpp
--- a/Engine/Plugins/Experimental/Mutable/Source/MutableValidation/Private/MuV/ValidationUtils.cpp
+++ b/Engine/Plugins/Experimental/Mutable/Source/MutableValidation/Private/MuV/ValidationUtils.cpp
@@ -112,6 +112,14 @@ bool CompileCustomizableObject(UCustomizableObject* InCustomizableObject, const
 	return bWasCoCompilationSuccessful;
 }
 
+void LogMutableSettings()
+{
+	const int32 WorkingMemory = UCustomizableObjectSystem::GetInstanceChecked()->GetWorkingMemory() ;
+	UE_LOG(LogMutable,Log, TEXT("(int) working_memory_bytes : %d"), WorkingMemory * 1024)
+	UE_LOG(LogMutable, Display, TEXT("The mutable updates will use as working memory the value of %d KB"), WorkingMemory)
+	
+	// Expand this when adding new controls from the .xml file
+}
 
 void Wait(const double ToWaitSeconds)
 {
diff --git a/Engine/Plugins/Experimental/Mutable/Source/MutableValidation/Private/MuV/ValidationUtils.h b/Engine/Plugins/Experimental/Mutable/Source/MutableValidation/Private/MuV/ValidationUtils.h
--- a/Engine/Plugins/Experimental/Mutable/Source/MutableValidation/Private/MuV/ValidationUtils.h
+++ b/Engine/Plugins/Experimental/Mutable/Source/MutableValidation/Private/MuV/ValidationUtils.h
@@ -36,6 +36,14 @@ void Wait(const double ToWaitSeconds);
 bool CompileCustomizableObject(UCustomizableObject* InCustomizableObject, const bool bLogMutableLogs = true, const FCompilationOptions* InCompilationOptionsOverride = nullptr);
 
 
+/**
+ * Logs some configuration data related to how mutable will compile and then generate instances. We do this so we can later
+ * Isolate tests using different configurations.
+ * @note Add new logs each time you add a way to change the configuration of the test from the .xml testing file
+ */
+void LogMutableSettings();
+
+
 /**
  * Helping class that handles the async update of the provided instance. It will also wait for the mips of it so they get streamed.
  */
