[mutable] Added layout options to table node.

#rnx

diff --git a/Engine/Plugins/Experimental/Mutable/Source/CustomizableObjectEditor/Private/MuCOE/Nodes/CustomizableObjectNodeLayoutBlocksDetails.h b/Engine/Plugins/Experimental/Mutable/Source/CustomizableObjectEditor/Private/MuCOE/Nodes/CustomizableObjectNodeLayoutBlocksDetails.h
--- a/Engine/Plugins/Experimental/Mutable/Source/CustomizableObjectEditor/Private/MuCOE/Nodes/CustomizableObjectNodeLayoutBlocksDetails.h
+++ b/Engine/Plugins/Experimental/Mutable/Source/CustomizableObjectEditor/Private/MuCOE/Nodes/CustomizableObjectNodeLayoutBlocksDetails.h
@@ -54,8 +54,6 @@ private:
 	TSharedPtr<SWidget> LODSelectorWidget;
 	TSharedPtr<STextBlock> LODSelectorTextWidget;
 
-	TArray< TSharedPtr<FString> > GridComboOptions;
-
 	/** List of available layout grid sizes. */
 	TArray< TSharedPtr< FString > > LayoutGridSizes;
 
diff --git a/Engine/Plugins/Experimental/Mutable/Source/CustomizableObjectEditor/Private/MuCOE/Nodes/CustomizableObjectNodeTableDetails.cpp b/Engine/Plugins/Experimental/Mutable/Source/CustomizableObjectEditor/Private/MuCOE/Nodes/CustomizableObjectNodeTableDetails.cpp
--- a/Engine/Plugins/Experimental/Mutable/Source/CustomizableObjectEditor/Private/MuCOE/Nodes/CustomizableObjectNodeTableDetails.cpp
+++ b/Engine/Plugins/Experimental/Mutable/Source/CustomizableObjectEditor/Private/MuCOE/Nodes/CustomizableObjectNodeTableDetails.cpp
@@ -8,14 +8,21 @@
 #include "Engine/SkeletalMesh.h"
 #include "Engine/StaticMesh.h"
 #include "GameplayTagContainer.h"
+#include "IDetailGroup.h"
 #include "IDetailsView.h"
+#include "Layout/Visibility.h"
 #include "MuCOE/CustomizableObjectLayout.h"
 #include "MuCOE/Nodes/CustomizableObjectNodeTable.h"
 #include "MuCOE/SCustomizableObjectNodeLayoutBlocksEditor.h"
 #include "MuCOE/UnrealEditorPortabilityHelpers.h"
 #include "Styling/SlateColor.h"
-#include "Widgets/Input/STextComboBox.h"
 #include "Widgets/Input/SButton.h"
+#include "Widgets/Input/SCheckBox.h"
+#include "Widgets/Input/SSpinBox.h"
+#include "Widgets/Layout/SBox.h"
+#include "Widgets/SWidget.h"
+#include "Widgets/Text/STextBlock.h"
+#include "Widgets/Input/STextComboBox.h"
 
 #define LOCTEXT_NAMESPACE "CustomizableObjectDetails"
 
@@ -41,6 +48,7 @@ void FCustomizableObjectNodeTableDetails::CustomizeDetails(const TSharedPtr<IDet
 	if (Node.IsValid())
 	{
 		IDetailCategoryBuilder& CustomizableObjectCategory = DetailBuilder->EditCategory("TableProperties");
+		//DetailBuilder->HideProperty("VersionColumn");
 		IDetailCategoryBuilder& UICategory = DetailBuilder->EditCategory("UI");
 		DetailBuilder->HideProperty("ParamUIMetadataColumn");
 		IDetailCategoryBuilder& AnimationCategory = DetailBuilder->EditCategory("AnimationProperties");
@@ -256,9 +264,10 @@ void FCustomizableObjectNodeTableDetails::CustomizeDetails(const TSharedPtr<IDet
 			]
 		];
 
+		SelectedLayout = nullptr;
 		LayoutBlocksEditor = SNew(SCustomizableObjectNodeLayoutBlocksEditor);
 
-		LayoutCategory.AddCustomRow(LOCTEXT("LayoutEditor", "Layout Editor"))
+		LayoutCategory.AddCustomRow(LOCTEXT("TableLayoutEditor_MeshSelector", "Mesh Selector"))
 		[
 			SNew(SVerticalBox)
 			
@@ -275,6 +284,7 @@ void FCustomizableObjectNodeTableDetails::CustomizeDetails(const TSharedPtr<IDet
 					SNew(STextBlock)
 					.Text(LOCTEXT("LayoutMeshColumnText", "Mesh Column: "))
 					.ToolTipText(LOCTEXT("LayoutMeshColumnTooltip", "Select a mesh from the Data Table to edit its layout blocks."))
+					.Font(IDetailLayoutBuilder::GetDetailFont())
 				]
 		
 				+ SHorizontalBox::Slot()
@@ -287,19 +297,93 @@ void FCustomizableObjectNodeTableDetails::CustomizeDetails(const TSharedPtr<IDet
 						.OptionsSource(&LayoutMeshColumnOptionNames)
 						.InitiallySelectedItem(nullptr)
 						.OnSelectionChanged(this, &FCustomizableObjectNodeTableDetails::OnLayoutMeshColumnComboBoxSelectionChanged)
+						.Font(IDetailLayoutBuilder::GetDetailFont())
 					]
 				]
 			]
+		];
 
-			+SVerticalBox::Slot()
-			.Padding(0.0f,20.0f,0.0f,0.0f)
+		// Layout size selector widget
+		LayoutCategory.AddCustomRow(LOCTEXT("TableBlocksDetails_SizeSelector", "SizeSelector"))
+		.Visibility(TAttribute<EVisibility>(this, &FCustomizableObjectNodeTableDetails::LayoutOptionsVisibility))
+		.NameContent()
+		[
+			SNew(STextBlock)
+			.Text(LOCTEXT("TableLayoutGridSizeText", "Grid Size"))
+			.Font(IDetailLayoutBuilder::GetDetailFont())
+		]
+		.ValueContent()
+		[
+			SAssignNew(GridSizeComboBox, STextComboBox)
+			.OptionsSource(&LayoutGridSizes)
+			.OnSelectionChanged(this, &FCustomizableObjectNodeTableDetails::OnGridSizeChanged)
+			.Font(IDetailLayoutBuilder::GetDetailFont())
+		];
+
+		// Layout strategy selector group widget
+		IDetailGroup* LayoutStrategyOptionsGroup = &LayoutCategory.AddGroup(TEXT("TableLayoutStrategyOptionsGroup"), LOCTEXT("TableLayoutStrategyGroup", "Table Layout Strategy Group"), false, true);
+		LayoutStrategyOptionsGroup->HeaderRow()
+		.Visibility(TAttribute<EVisibility>(this, &FCustomizableObjectNodeTableDetails::LayoutOptionsVisibility))
+		.NameContent()
+		[
+			SNew(STextBlock)
+			.Text(LOCTEXT("TableLayoutStrategy_Text", "Layout Strategy:"))
+			.ToolTipText(LOCTEXT("TableLayoutStrategyTooltip", "Selects the packing strategy: Resizable Layout or Fixed Layout"))
+			.Font(IDetailLayoutBuilder::GetDetailFont())
+		]
+		.ValueContent()
+		[
+			SAssignNew(StrategyComboBox, STextComboBox)
+			.OptionsSource(&LayoutPackingStrategies)
+			.OnSelectionChanged(this, &FCustomizableObjectNodeTableDetails::OnLayoutPackingStrategyChanged)
+			.Font(IDetailLayoutBuilder::GetDetailFont())
+		];
+
+		// Max layout size selector widget
+		LayoutStrategyOptionsGroup->AddWidgetRow()
+		.Visibility(TAttribute<EVisibility>(this, &FCustomizableObjectNodeTableDetails::FixedStrategyOptionsVisibility))
+		.NameContent()
+		[
+			SNew(STextBlock)
+			.Text(LOCTEXT("TableMaxLayoutSize_Text", "Max Layout Size:"))
+			.Font(IDetailLayoutBuilder::GetDetailFont())
+		]
+		.ValueContent()
+		[
+			SAssignNew(MaxGridSizeComboBox, STextComboBox)
+			.OptionsSource(&LayoutGridSizes)
+			.OnSelectionChanged(this, &FCustomizableObjectNodeTableDetails::OnMaxGridSizeChanged)
+			.Font(IDetailLayoutBuilder::GetDetailFont())
+		];
+
+		// Reduction method selector widget
+		LayoutStrategyOptionsGroup->AddWidgetRow()
+		.Visibility(TAttribute<EVisibility>(this, &FCustomizableObjectNodeTableDetails::FixedStrategyOptionsVisibility))
+		.NameContent()
+		[
+			SNew(STextBlock)
+			.Text(LOCTEXT("TableReductionMethod_Text", "Reduction Method:"))
+			.ToolTipText(LOCTEXT("TableReduction_Method_Tooltip", "Select how blocks will be reduced in case that they do not fit in the layout:"
+				"\n Halve: blocks will be reduced by half each time."
+				"\n Unit: blocks will be reduced by one unit each time."))
+			.Font(IDetailLayoutBuilder::GetDetailFont())
+		]
+		.ValueContent()
+		[
+			SAssignNew(ReductionMethodComboBox, STextComboBox)
+			.OptionsSource(&BlockReductionMethods)
+			.OnSelectionChanged(this, &FCustomizableObjectNodeTableDetails::OnReductionMethodChanged)
+			.Font(IDetailLayoutBuilder::GetDetailFont())
+		];
+
+		// Block editor Widget
+		LayoutCategory.AddCustomRow(LOCTEXT("TableLayoutEditor", "Layout Editor"))
+		[
+			SNew(SBox)
+			.HeightOverride(700.0f)
+			.WidthOverride(700.0f)
 			[
-				SNew(SBox)
-				.HeightOverride(700.0f)
-				.WidthOverride(700.0f)
-				[
-					LayoutBlocksEditor.ToSharedRef()
-				]
+				LayoutBlocksEditor.ToSharedRef()
 			]
 		];
 
@@ -481,6 +565,9 @@ void FCustomizableObjectNodeTableDetails::OnLayoutMeshColumnComboBoxSelectionCha
 					if (PinData->Layouts[LayoutIndex]->GetLayoutName() == *Selection)
 					{
 						LayoutBlocksEditor->SetCurrentLayout(PinData->Layouts[LayoutIndex]);
+						SelectedLayout = PinData->Layouts[LayoutIndex];
+
+						FillLayoutComboBoxOptions();
 					}
 				}
 			}
@@ -706,8 +793,126 @@ void FCustomizableObjectNodeTableDetails::OnMutableMetaDataColumnComboBoxSelecti
 		MutableMetaDataComboBox->ClearSelection();
 		MutableMetaDataComboBox->RefreshOptions();
 	}
-	
 }
 
 
+EVisibility FCustomizableObjectNodeTableDetails::LayoutOptionsVisibility() const
+{
+	return SelectedLayout.IsValid() ? EVisibility::Visible : EVisibility::Collapsed;
+}
+
+
+EVisibility FCustomizableObjectNodeTableDetails::FixedStrategyOptionsVisibility() const
+{
+	return (SelectedLayout.IsValid() && SelectedLayout->GetPackingStrategy() == ECustomizableObjectTextureLayoutPackingStrategy::Fixed) ? EVisibility::Visible : EVisibility::Collapsed;
+}
+
+
+void FCustomizableObjectNodeTableDetails::FillLayoutComboBoxOptions()
+{
+	if (SelectedLayout.IsValid() && GridSizeComboBox.IsValid() && StrategyComboBox.IsValid()
+		&& MaxGridSizeComboBox.IsValid() && ReductionMethodComboBox.IsValid())
+	{
+		// Static const variable?
+		int32 MaxGridSize = 32;
+		LayoutGridSizes.Empty();
+
+		for (int32 Size = 1; Size <= MaxGridSize; Size *= 2)
+		{
+			LayoutGridSizes.Add(MakeShareable(new FString(FString::Printf(TEXT("%d x %d"), Size, Size))));
+
+			if (SelectedLayout->GetGridSize() == FIntPoint(Size))
+			{
+				GridSizeComboBox->SetSelectedItem(LayoutGridSizes.Last());
+			}
+
+			if (SelectedLayout->GetMaxGridSize() == FIntPoint(Size))
+			{
+				MaxGridSizeComboBox->SetSelectedItem(LayoutGridSizes.Last());
+			}
+		}
+
+		LayoutPackingStrategies.Empty();
+		LayoutPackingStrategies.Add(MakeShareable(new FString("Resizable")));
+		LayoutPackingStrategies.Add(MakeShareable(new FString("Fixed")));
+		LayoutPackingStrategies.Add(MakeShareable(new FString("Overlay")));
+		StrategyComboBox->SetSelectedItem(LayoutPackingStrategies[(uint32)SelectedLayout->GetPackingStrategy()]);
+
+		BlockReductionMethods.Empty();
+		BlockReductionMethods.Add(MakeShareable(new FString("Halve")));
+		BlockReductionMethods.Add(MakeShareable(new FString("Unitary")));
+		ReductionMethodComboBox->SetSelectedItem(BlockReductionMethods[(uint32)SelectedLayout->GetBlockReductionMethod()]);
+	}
+}
+
+
+void FCustomizableObjectNodeTableDetails::OnGridSizeChanged(TSharedPtr<FString> NewSelection, ESelectInfo::Type SelectInfo)
+{
+	if (SelectedLayout.IsValid())
+	{
+		int Size = 1 << LayoutGridSizes.Find(NewSelection);
+
+		if (SelectedLayout->GetGridSize().X != Size || SelectedLayout->GetGridSize().Y != Size)
+		{
+			SelectedLayout->SetGridSize(FIntPoint(Size));
+
+			// Adjust all the blocks sizes
+			for (int b = 0; b < SelectedLayout->Blocks.Num(); ++b)
+			{
+				SelectedLayout->Blocks[b].Min.X = FMath::Min(SelectedLayout->Blocks[b].Min.X, Size - 1);
+				SelectedLayout->Blocks[b].Min.Y = FMath::Min(SelectedLayout->Blocks[b].Min.Y, Size - 1);
+				SelectedLayout->Blocks[b].Max.X = FMath::Min(SelectedLayout->Blocks[b].Max.X, Size);
+				SelectedLayout->Blocks[b].Max.Y = FMath::Min(SelectedLayout->Blocks[b].Max.Y, Size);
+			}
+
+			Node->MarkPackageDirty();
+		}
+	}
+}
+
+
+void FCustomizableObjectNodeTableDetails::OnLayoutPackingStrategyChanged(TSharedPtr<FString> NewSelection, ESelectInfo::Type SelectInfo)
+{
+	if (SelectedLayout.IsValid())
+	{
+		uint32 selection = LayoutPackingStrategies.IndexOfByKey(NewSelection);
+
+		if (SelectedLayout->GetPackingStrategy() != (ECustomizableObjectTextureLayoutPackingStrategy)selection)
+		{
+			SelectedLayout->SetPackingStrategy((ECustomizableObjectTextureLayoutPackingStrategy)selection);
+			Node->MarkPackageDirty();
+		}
+	}
+}
+
+
+void FCustomizableObjectNodeTableDetails::OnMaxGridSizeChanged(TSharedPtr<FString> NewSelection, ESelectInfo::Type SelectInfo)
+{
+	if (SelectedLayout.IsValid())
+	{
+		int Size = 1 << LayoutGridSizes.Find(NewSelection);
+
+		if (SelectedLayout->GetMaxGridSize().X != Size || SelectedLayout->GetMaxGridSize().Y != Size)
+		{
+			SelectedLayout->SetMaxGridSize(FIntPoint(Size));
+			SelectedLayout->MarkPackageDirty();
+		}
+	}
+}
+
+
+void FCustomizableObjectNodeTableDetails::OnReductionMethodChanged(TSharedPtr<FString> NewSelection, ESelectInfo::Type SelectInfo)
+{
+	if (SelectedLayout.IsValid())
+	{
+		uint32 selection = BlockReductionMethods.IndexOfByKey(NewSelection);
+
+		if (SelectedLayout->GetBlockReductionMethod() != (ECustomizableObjectLayoutBlockReductionMethod)selection)
+		{
+			SelectedLayout->SetBlockReductionMethod((ECustomizableObjectLayoutBlockReductionMethod)selection);
+			Node->MarkPackageDirty();
+		}
+	}
+}
+
 #undef LOCTEXT_NAMESPACE
diff --git a/Engine/Plugins/Experimental/Mutable/Source/CustomizableObjectEditor/Private/MuCOE/Nodes/CustomizableObjectNodeTableDetails.h b/Engine/Plugins/Experimental/Mutable/Source/CustomizableObjectEditor/Private/MuCOE/Nodes/CustomizableObjectNodeTableDetails.h
--- a/Engine/Plugins/Experimental/Mutable/Source/CustomizableObjectEditor/Private/MuCOE/Nodes/CustomizableObjectNodeTableDetails.h
+++ b/Engine/Plugins/Experimental/Mutable/Source/CustomizableObjectEditor/Private/MuCOE/Nodes/CustomizableObjectNodeTableDetails.h
@@ -7,14 +7,16 @@
 #include "UObject/WeakObjectPtr.h"
 
 namespace ESelectInfo { enum Type : int; }
-
 class FString;
 class IDetailLayoutBuilder;
-class STextComboBox;
 class FReply;
 class SButton;
 class SCustomizableObjectNodeLayoutBlocksEditor;
+class STextBlock;
+class STextComboBox;
+class UCustomizableObjectLayout;
 class UCustomizableObjectNodeTable;
+struct EVisibility;
 struct FSlateColor;
 
 /** Copy Material node details panel. Hides all properties from the inheret Material node. */
@@ -85,6 +87,23 @@ private:
 	// OnComboBoxSelectionChanged Callback for Layout ComboBox
 	void OnMutableMetaDataColumnComboBoxSelectionReset();
 
+	// Layout options visibility
+	EVisibility LayoutOptionsVisibility() const;
+
+	/** Returns the visibility of the Fixed layout widgets */
+	EVisibility FixedStrategyOptionsVisibility() const;
+
+	/** Fills the combo box arrays sources */
+	void FillLayoutComboBoxOptions();
+
+	/** Layout Options Callbacks */
+	void OnGridSizeChanged(TSharedPtr<FString> NewSelection, ESelectInfo::Type SelectInfo);
+	void OnMaxGridSizeChanged(TSharedPtr<FString> NewSelection, ESelectInfo::Type SelectInfo);
+	void OnReductionMethodChanged(TSharedPtr<FString> NewSelection, ESelectInfo::Type SelectInfo);
+	void OnLayoutPackingStrategyChanged(TSharedPtr<FString> NewSelection, ESelectInfo::Type SelectInfo);
+
+private:
+
 	// Pointer to the node represented in this details
 	TWeakObjectPtr<UCustomizableObjectNodeTable> Node;
 
@@ -130,6 +149,26 @@ private:
 	// Layout block editor widget
 	TSharedPtr<SCustomizableObjectNodeLayoutBlocksEditor> LayoutBlocksEditor;
 
+	TWeakObjectPtr<UCustomizableObjectLayout> SelectedLayout;
+
 	// Pointer to the Detail Builder to force the refresh on recontruct the node
 	TWeakPtr<IDetailLayoutBuilder> DetailBuilderPtr = nullptr;
+
+	/** List of available layout grid sizes. */
+	TArray< TSharedPtr< FString > > LayoutGridSizes;
+
+	/** List of available layout packing strategies. */
+	TArray< TSharedPtr< FString > > LayoutPackingStrategies;
+
+	/** List of available block reduction methods. */
+	TArray< TSharedPtr< FString > > BlockReductionMethods;
+
+	// ComboBox widget to select a Grid Size from the Selected Layout
+	TSharedPtr<STextComboBox> GridSizeComboBox;
+	// ComboBox widget to select a Strategy from the Selected Layout
+	TSharedPtr<STextComboBox> StrategyComboBox;
+	// ComboBox widget to select a Max Grid Size from the Selected Layout
+	TSharedPtr<STextComboBox> MaxGridSizeComboBox;
+	// ComboBox widget to select a Reduction Method from the Selected Layout
+	TSharedPtr<STextComboBox> ReductionMethodComboBox;
 };
