[Mutable] Remove unnecessary program copy on GetVisibleChildren CodeViewer function.

#rnx

diff --git a/Engine/Plugins/Experimental/Mutable/Source/CustomizableObjectEditor/Private/MuCOE/SMutableCodeViewer.cpp b/Engine/Plugins/Experimental/Mutable/Source/CustomizableObjectEditor/Private/MuCOE/SMutableCodeViewer.cpp
--- a/Engine/Plugins/Experimental/Mutable/Source/CustomizableObjectEditor/Private/MuCOE/SMutableCodeViewer.cpp
+++ b/Engine/Plugins/Experimental/Mutable/Source/CustomizableObjectEditor/Private/MuCOE/SMutableCodeViewer.cpp
@@ -1726,7 +1726,7 @@ void SMutableCodeViewer::OnExpansionChanged(TSharedPtr<FMutableCodeTreeElement>
 void SMutableCodeViewer::GetVisibleChildren(TSharedPtr<FMutableCodeTreeElement> InInfo, TSet<TSharedPtr<FMutableCodeTreeElement>>& OutChildren)
 {
 	check(MutableModel);
-	const mu::FProgram MutableProgram = MutableModel->GetPrivate()->m_program;
+	const mu::FProgram& MutableProgram = MutableModel->GetPrivate()->m_program;
 	
 	TArray<TSharedPtr<FMutableCodeTreeElement>> ToSearchForChildren;
 	ToSearchForChildren.Add(InInfo);
