[Mutable] Newly added child COs and composite Tables now are detected by the Participating Objects system.

Added a new way of accessing COE module functions from CO module:
- Renamed CustomizableObjectEditorModule.h to ICustomizableObjectEditorModule.h.
- Moved ICustomizableObjectEditorModule from COE to CO.
- Created a new CustomizableObjectEditorModule.h file.
- Moved FCustomizableObjectEditorModule declaration to CustomizableObjectEditorModule.h file.
- Renamed ICustomizableObjectEditorModule Get() to GetChecked().
- Created a new ICustomizableObjectEditorModule Get() that can return a nullptr (useful to get the COE module from CO module).
- Removed unused ICustomizableObjectEditorModule::GetAssetCategory

Changes:
- IsCompilationOutOfDate now also checks for newly added child CO and  composite Tables.
- Merged ShowOnScreenCompileWarnings and HideOnScreenCompileWarnings.
- ShowOnScreenCompileWarnings moved to a delayed ticker (1.0s).
- ShowOnScreenCompileWarnings moved from CO to COE module.

#rnx

diff --git a/Engine/Plugins/Experimental/Mutable/Source/CustomizableObject/Internal/MuCO/CustomizableObjectPrivate.h b/Engine/Plugins/Experimental/Mutable/Source/CustomizableObject/Internal/MuCO/CustomizableObjectPrivate.h
--- a/Engine/Plugins/Experimental/Mutable/Source/CustomizableObject/Internal/MuCO/CustomizableObjectPrivate.h
+++ b/Engine/Plugins/Experimental/Mutable/Source/CustomizableObject/Internal/MuCO/CustomizableObjectPrivate.h
@@ -46,7 +46,7 @@ public:
 	virtual void ForceFinishCompilation() {};
 
 	// Return true if this object doesn't reference a parent object.
-	virtual bool IsRootObject(const class UCustomizableObject* Object) const { return true; }
+	virtual bool IsRootObject(const UCustomizableObject* Object) const { return true; }
 
 	/** Provides the caller with the warning and error messages produced during compilation */
 	virtual void GetCompilationMessages(TArray<FText>& OutWarningMessages, TArray<FText>& OutErrorMessages) const = 0;
@@ -622,6 +622,8 @@ class CUSTOMIZABLEOBJECT_API UCustomizableObjectPrivate : public UObject
 #endif
 
 public:
+	UCustomizableObjectPrivate();
+	
 	void SetModel(const TSharedPtr<mu::Model, ESPMode::ThreadSafe>& Model, const FGuid Identifier);
 	const TSharedPtr<mu::Model, ESPMode::ThreadSafe>& GetModel();
 	TSharedPtr<const mu::Model, ESPMode::ThreadSafe> GetModel() const;
@@ -719,9 +721,11 @@ public:
 	/** Return the MinLOD index to generate based on the active LODSettings (PerPlatformMinLOD or PerQualityLevelMinLOD) */
 	int32 GetMinLODIndex() const;
 	
-#if WITH_EDITORONLY_DATA
-	/** Return true if the CO is not compiled or the ParticipatingObjects system has detected a change (participating objects dirty or re-saved since last compilation). */
-	bool IsCompilationOutOfDate() const;
+#if WITH_EDITOR
+	/** See ICustomizableObjectEditorModule::IsCompilationOutOfDate. */
+	bool IsCompilationOutOfDate(TArray<FName>* OutOfDatePackages = nullptr) const;
+
+	void OnParticipatingObjectDirty(UPackage* Package, bool);
 #endif
 
 	TMap<uint64, FMutableStreamableBlock>& GetHashToStreamableBlock();
@@ -767,12 +771,15 @@ public:
 	 * Updated each time the CO is compiled and saved in the Derived Data. */
 	TMap<FName, FGuid> ParticipatingObjects;
 
+	/** List of Participating Objects (packages) has been marked as dirty since the last compilation. */
+	TArray<FName> DirtyParticipatingObjects;
+	
 	/** If the object is compiled, this flag is false unless it was compiled with maximum optimizations. If the object is not compiled, its value is meaningless. */
 	bool bIsCompiledWithoutOptimization = true;
 
 	/** This is a non-user-controlled flag to disable streaming (set at object compilation time, depending on optimization). */
 	bool bDisableTextureStreaming = false;
-
+	
 	ECustomizableObjectCompilationState CompilationState = ECustomizableObjectCompilationState::None;
 	
 #if WITH_EDITOR
@@ -804,6 +811,6 @@ public:
 	// This is a manual version number for the binary blobs in this asset.
 	// Increasing it invalidates all the previously compiled models.
 	// Warning: If while merging code both versions have changed, take the highest+1.
-	static constexpr int32 CurrentSupportedVersion = 431;
+	static constexpr int32 CurrentSupportedVersion = 432;
 };
 
diff --git a/Engine/Plugins/Experimental/Mutable/Source/CustomizableObject/Internal/MuCO/CustomizableObjectSystemPrivate.h b/Engine/Plugins/Experimental/Mutable/Source/CustomizableObject/Internal/MuCO/CustomizableObjectSystemPrivate.h
--- a/Engine/Plugins/Experimental/Mutable/Source/CustomizableObject/Internal/MuCO/CustomizableObjectSystemPrivate.h
+++ b/Engine/Plugins/Experimental/Mutable/Source/CustomizableObject/Internal/MuCO/CustomizableObjectSystemPrivate.h
@@ -722,13 +722,6 @@ public:
 
 	void UnCacheTextureParameters(const TArray<FCustomizableObjectTextureParameterValue>& TextureParameters) const;
 
-#if WITH_EDITORONLY_DATA
-	/**	PIE CO on-screen warnings. */
-	static void ShowOnScreenCompileWarnings();
-
-	static void HideOnScreenCompileWarnings(const UCustomizableObjectPrivate& ObjectPrivate);
-#endif
-
 	/** Tick, in the Game Thread, anything that could potentially block the Mutable Thread. */
 	void TickMutableThreadDependencies();
 
diff --git a/Engine/Plugins/Experimental/Mutable/Source/CustomizableObject/Internal/MuCO/ICustomizableObjectEditorModule.h b/Engine/Plugins/Experimental/Mutable/Source/CustomizableObject/Internal/MuCO/ICustomizableObjectEditorModule.h
new file mode 100644
--- /dev/null
+++ b/Engine/Plugins/Experimental/Mutable/Source/CustomizableObject/Internal/MuCO/ICustomizableObjectEditorModule.h
@@ -0,0 +1,53 @@
+// Copyright Epic Games, Inc. All Rights Reserved.
+
+#pragma once
+
+#include "Modules/ModuleManager.h"
+#include "Templates/SharedPointer.h"
+
+class FCustomizableObjectEditorLogger;
+class ICustomizableObjectDebugger;
+class ICustomizableObjectEditor;
+class ICustomizableObjectInstanceEditor;
+class IToolkitHost;
+class UCustomizableObject;
+class UCustomizableObjectPrivate;
+class UCustomizableObjectInstance;
+class FExtensibilityManager;
+
+extern const FName CustomizableObjectEditorAppIdentifier;
+extern const FName CustomizableObjectInstanceEditorAppIdentifier;
+extern const FName CustomizableObjectPopulationEditorAppIdentifier;
+extern const FName CustomizableObjectPopulationClassEditorAppIdentifier;
+extern const FName CustomizableObjectDebuggerAppIdentifier;
+
+
+/**
+ * Customizable object editor module interface
+ */
+class ICustomizableObjectEditorModule : public IModuleInterface
+{
+public:
+	static ICustomizableObjectEditorModule* Get()
+	{
+		return FModuleManager::LoadModulePtr<ICustomizableObjectEditorModule>("CustomizableObjectEditor");
+	}
+	
+	static ICustomizableObjectEditorModule& GetChecked()
+	{
+		return FModuleManager::LoadModuleChecked<ICustomizableObjectEditorModule>("CustomizableObjectEditor");
+	}
+
+	virtual TSharedPtr<FExtensibilityManager> GetCustomizableObjectEditorToolBarExtensibilityManager() { return nullptr; }
+	virtual TSharedPtr<FExtensibilityManager> GetCustomizableObjectEditorMenuExtensibilityManager() { return nullptr; }
+	
+	/** Returns the module logger. */
+	virtual FCustomizableObjectEditorLogger& GetLogger() = 0;
+	
+	/** Return if the CO is not compiled or the ParticipatingObjects system has detected a change (participating objects dirty or re-saved since last compilation).
+	  * @param Object object to check.
+	  * @param OutOfDatePackages list of out of date packages.
+   	  * @return true if the compilation is out of date. */
+	virtual bool IsCompilationOutOfDate(const UCustomizableObject& Object, TArray<FName>* OutOfDatePackages = nullptr) const = 0;
+};
+ 
diff --git a/Engine/Plugins/Experimental/Mutable/Source/CustomizableObject/Private/MuCO/CustomizableObject.cpp b/Engine/Plugins/Experimental/Mutable/Source/CustomizableObject/Private/MuCO/CustomizableObject.cpp
--- a/Engine/Plugins/Experimental/Mutable/Source/CustomizableObject/Private/MuCO/CustomizableObject.cpp
+++ b/Engine/Plugins/Experimental/Mutable/Source/CustomizableObject/Private/MuCO/CustomizableObject.cpp
@@ -35,6 +35,8 @@
 #include "Serialization/MemoryWriter.h"
 #include "UObject/AssetRegistryTagsContext.h"
 #include "UObject/ObjectSaveContext.h"
+#include "MuCO/ICustomizableObjectEditorModule.h"
+#include "MuCO/CustomizableObjectSystemPrivate.h"
 
 #if WITH_EDITOR
 #include "Editor.h"
@@ -712,34 +714,15 @@ void UCustomizableObjectPrivate::LoadCompiledData(FArchive& MemoryReader, const
 			if (!bIsCooking)
 			{
 				MemoryReader << LocalModelResource.EditorOnlyMorphTargetReconstructionData;
-			}
+				
+				DirtyParticipatingObjects.Empty();
 
-			if (!bIsCooking)
-			{
-				const FAssetRegistryModule& AssetRegistryModule = FModuleManager::LoadModuleChecked<FAssetRegistryModule>("AssetRegistry");
+				TArray<FName> OutOfDatePackages;
+				bInvalidateModel = IsCompilationOutOfDate(&OutOfDatePackages);
 
-				for (TTuple<FName, FGuid>& ParticipatingObject : ParticipatingObjects)
+				if (bInvalidateModel)
 				{
-					FAssetPackageData AssetPackageData;
-					const UE::AssetRegistry::EExists Result = AssetRegistryModule.Get().TryGetAssetPackageData(ParticipatingObject.Key, AssetPackageData);
-					if (Result == UE::AssetRegistry::EExists::Exists)
-					{
-						PRAGMA_DISABLE_DEPRECATION_WARNINGS
-							const FGuid PackageGuid = AssetPackageData.PackageGuid;
-						PRAGMA_ENABLE_DEPRECATION_WARNINGS
-
-							bInvalidateModel = PackageGuid != ParticipatingObject.Value;
-					}
-					else
-					{
-						bInvalidateModel = true;
-					}
-
-					if (bInvalidateModel)
-					{
-						UE_LOG(LogMutable, Display, TEXT("Invalidating compiled data due to changes in %s."), *ParticipatingObject.Key.ToString());
-						break;
-					}
+					UE_LOG(LogMutable, Display, TEXT("Invalidating compiled data due to changes in %s."), *OutOfDatePackages[0].ToString());
 				}
 			}
 		}
@@ -1210,8 +1193,6 @@ FString UCustomizableObject::GetStateParameterName(int32 StateIndex, int32 Param
 void UCustomizableObjectPrivate::PostCompile()
 {
 	PostCompileDelegate.Broadcast();
-
-	UCustomizableObjectSystemPrivate::HideOnScreenCompileWarnings(*this);	
 }
 #endif
 
@@ -1889,6 +1870,14 @@ void FSkeletonCache::Add(const TArray<uint16>& Key, USkeleton* Value)
 }
 
 
+UCustomizableObjectPrivate::UCustomizableObjectPrivate()
+{
+#if WITH_EDITOR
+	UPackage::PackageMarkedDirtyEvent.AddUObject(this, &UCustomizableObjectPrivate::OnParticipatingObjectDirty);
+#endif
+}
+
+
 void UCustomizableObjectPrivate::SetModel(const TSharedPtr<mu::Model, ESPMode::ThreadSafe>& Model, const FGuid Id)
 {
 	if (MutableModel == Model
@@ -1946,57 +1935,24 @@ FModelResources& UCustomizableObjectPrivate::GetModelResources(bool bIsCooking)
 #endif
 
 
-#if WITH_EDITORONLY_DATA
-bool UCustomizableObjectPrivate::IsCompilationOutOfDate() const
+#if WITH_EDITOR
+bool UCustomizableObjectPrivate::IsCompilationOutOfDate(TArray<FName>* OutOfDatePackages) const
 {
-	if (Status.Get() != FCustomizableObjectStatus::EState::ModelLoaded)
-	{
-		return false;
-	}
-	
-	if (GetPackage()->IsDirty())
+	if (const ICustomizableObjectEditorModule* Module = ICustomizableObjectEditorModule::Get())
 	{
-		return true;
+		return Module->IsCompilationOutOfDate(*GetPublic(), OutOfDatePackages);
 	}
-	
-	const FAssetRegistryModule& AssetRegistryModule = FModuleManager::LoadModuleChecked<FAssetRegistryModule>("AssetRegistry");
 
-	for (const TTuple<FName, FGuid>& ParticipatingObject : ParticipatingObjects)
-	{
-		TSoftObjectPtr Object(ParticipatingObject.Key.ToString());
-		if (Object) // If loaded
-		{
-			PRAGMA_DISABLE_DEPRECATION_WARNINGS
-			const FGuid PackageGuid = Object->GetPackage()->GetGuid();
-			PRAGMA_ENABLE_DEPRECATION_WARNINGS
-			
-			if (Object->GetPackage()->IsDirty() || PackageGuid != ParticipatingObject.Value)
-			{
-				return true;
-			}
-		}
-		else // Not loaded
-		{
-			FAssetPackageData AssetPackageData;
-			const UE::AssetRegistry::EExists Result = AssetRegistryModule.Get().TryGetAssetPackageData(ParticipatingObject.Key, AssetPackageData);
-				
-			if (Result != UE::AssetRegistry::EExists::Exists)
-			{
-				return true;
-			}
+	return false;		
+}
 
-			PRAGMA_DISABLE_DEPRECATION_WARNINGS
-			const FGuid PackageGuid = AssetPackageData.PackageGuid;
-			PRAGMA_ENABLE_DEPRECATION_WARNINGS
-			
-			if (PackageGuid != ParticipatingObject.Value)
-			{
-				return true;
-			}
-		}
-	}
 
-	return false;
+void UCustomizableObjectPrivate::OnParticipatingObjectDirty(UPackage* Package, bool)
+{
+	if (ParticipatingObjects.Contains(Package->GetFName()))
+	{
+		DirtyParticipatingObjects.AddUnique(Package->GetFName());		
+	}
 }
 #endif
 
diff --git a/Engine/Plugins/Experimental/Mutable/Source/CustomizableObject/Private/MuCO/CustomizableObjectSystem.cpp b/Engine/Plugins/Experimental/Mutable/Source/CustomizableObject/Private/MuCO/CustomizableObjectSystem.cpp
--- a/Engine/Plugins/Experimental/Mutable/Source/CustomizableObject/Private/MuCO/CustomizableObjectSystem.cpp
+++ b/Engine/Plugins/Experimental/Mutable/Source/CustomizableObject/Private/MuCO/CustomizableObjectSystem.cpp
@@ -2723,10 +2723,6 @@ int32 UCustomizableObjectSystem::TickInternal()
 	TickRecompileCustomizableObjects();
 	GetPrivate()->TickMutableThreadDependencies();
 #endif
-
-#if WITH_EDITORONLY_DATA
-	UCustomizableObjectSystemPrivate::ShowOnScreenCompileWarnings();
-#endif
 	
 	const int32 RemainingTasks = Private->MutableTaskGraph.Tick();
 
@@ -3231,87 +3227,6 @@ int32 UCustomizableObjectSystem::GetWorkingMemory() const
 }
 
 
-#if WITH_EDITORONLY_DATA
-void UCustomizableObjectSystemPrivate::ShowOnScreenCompileWarnings()
-{
-	TSet<const UCustomizableObject*> Objects;
-	
-	for (TObjectIterator<UCustomizableObjectInstanceUsage> CustomizableObjectInstanceUsage; CustomizableObjectInstanceUsage; ++CustomizableObjectInstanceUsage)
-	{
-		if (!IsValid(*CustomizableObjectInstanceUsage) || CustomizableObjectInstanceUsage->IsTemplate())
-		{
-			continue;
-		}
-		
-		const UCustomizableObjectInstance* Instance = CustomizableObjectInstanceUsage->GetCustomizableObjectInstance();
-		if (!Instance)
-		{
-			continue;
-		}
-
-		const UCustomizableObject* Object = Cast<UCustomizableObject>(Instance->GetCustomizableObject());
-		if (!Object)
-		{
-			continue;
-		}
-		
-		const USkeletalMeshComponent* Parent = Cast<USkeletalMeshComponent>(CustomizableObjectInstanceUsage->GetAttachParent());
-		if (!Parent)
-		{
-			continue;
-		}
-
-		const UWorld* World = Parent->GetWorld();
-		if (!World)
-		{
-			continue;
-		}
-
-		if (World->WorldType != EWorldType::PIE)
-		{
-			continue;
-		}
-
-		Objects.Add(Object);
-	}
-
-	for (const UCustomizableObject* Object : Objects)
-	{
-		// Show a warning if the compilation was not done with optimizations.
-		const uint64 KeyCompiledWithOptimization = reinterpret_cast<uint64>(Object);
-		if (Object->GetPrivate()->bIsCompiledWithoutOptimization && // Quicker to check
-			!GEngine->OnScreenDebugMessageExists(KeyCompiledWithOptimization))
-		{
-			FString Msg = FString::Printf(TEXT("Warning: Customizable Object [%s] was compiled without optimization."), *Object->GetName());
-			GEngine->AddOnScreenDebugMessage(KeyCompiledWithOptimization, 10.0f, FColor::Red, Msg);
-		}
-
-		const uint64 KeyCompiledOutOfData = reinterpret_cast<uint64>(Object) + KEY_OFFSET_COMPILATION_OUT_OF_DATE; // Offset added to avoid collision with bIsCompiledWithoutOptimization warning
-		if (!GEngine->OnScreenDebugMessageExists(KeyCompiledWithOptimization) && // Quicker to check
-			Object->GetPrivate()->IsCompilationOutOfDate())
-		{
-			FString Msg = FString::Printf(TEXT("Warning: Customizable Object [%s] compilation out of date. Save all referenced assets and recompile."), *Object->GetName());
-			GEngine->AddOnScreenDebugMessage(KeyCompiledOutOfData, 10.0f, FColor::Red, Msg);
-		}
-	}
-}
-
-
-void UCustomizableObjectSystemPrivate::HideOnScreenCompileWarnings(const UCustomizableObjectPrivate& ObjectPrivate)
-{
-	if (ObjectPrivate.bIsCompiledWithoutOptimization)
-	{
-		GEngine->RemoveOnScreenDebugMessage(reinterpret_cast<uint64>(&ObjectPrivate));
-	}
-
-	if (!ObjectPrivate.IsCompilationOutOfDate())
-	{
-		GEngine->RemoveOnScreenDebugMessage(reinterpret_cast<uint64>(&ObjectPrivate) + KEY_OFFSET_COMPILATION_OUT_OF_DATE);
-	}
-}
-#endif
-
-
 void UCustomizableObjectSystemPrivate::TickMutableThreadDependencies()
 {
 	check(IsInGameThread());
diff --git a/Engine/Plugins/Experimental/Mutable/Source/CustomizableObjectEditor/Classes/MuCOE/Nodes/CustomizableObjectNodeTable.h b/Engine/Plugins/Experimental/Mutable/Source/CustomizableObjectEditor/Classes/MuCOE/Nodes/CustomizableObjectNodeTable.h
--- a/Engine/Plugins/Experimental/Mutable/Source/CustomizableObjectEditor/Classes/MuCOE/Nodes/CustomizableObjectNodeTable.h
+++ b/Engine/Plugins/Experimental/Mutable/Source/CustomizableObjectEditor/Classes/MuCOE/Nodes/CustomizableObjectNodeTable.h
@@ -395,7 +395,9 @@ public:
 	/** Returns the column id */
 	FGuid GetColumnIdByName(const FName& ColumnName) const;
 
-
+	/** Return the list of UDataTable that will be used to compose the final UDataTable. */
+	TArray<FAssetData> GetParentTables() const;
+	
 private:
 
 	/** Number of properties to know when the node needs an update */
diff --git a/Engine/Plugins/Experimental/Mutable/Source/CustomizableObjectEditor/Private/MuCOE/AssetDefinition_CustomizableObject.cpp b/Engine/Plugins/Experimental/Mutable/Source/CustomizableObjectEditor/Private/MuCOE/AssetDefinition_CustomizableObject.cpp
--- a/Engine/Plugins/Experimental/Mutable/Source/CustomizableObjectEditor/Private/MuCOE/AssetDefinition_CustomizableObject.cpp
+++ b/Engine/Plugins/Experimental/Mutable/Source/CustomizableObjectEditor/Private/MuCOE/AssetDefinition_CustomizableObject.cpp
@@ -14,7 +14,7 @@
 #include "MuCO/CustomizableObject.h"
 #include "MuCO/CustomizableObjectInstance.h"
 #include "MuCO/CustomizableObjectSystem.h"
-#include "MuCOE/CustomizableObjectEditorModule.h"
+#include "MuCO/ICustomizableObjectEditorModule.h"
 #include "Subsystems/AssetEditorSubsystem.h"
 #include "ToolMenus.h"
 
@@ -27,8 +27,8 @@ EAssetCommandResult UAssetDefinition_CustomizableObject::OpenAssets(const FAsset
 	
 	for (UCustomizableObject* Object : OpenArgs.LoadObjects<UCustomizableObject>())
 	{
-		ICustomizableObjectEditorModule* CustomizableObjectEditorModule = &FModuleManager::LoadModuleChecked<ICustomizableObjectEditorModule>( "CustomizableObjectEditor" );
-		CustomizableObjectEditorModule->CreateCustomizableObjectEditor(Mode, OpenArgs.ToolkitHost, Object);
+		const TSharedPtr<FCustomizableObjectEditor> Editor = MakeShared<FCustomizableObjectEditor>(*Object);
+		Editor->InitCustomizableObjectEditor(Mode, OpenArgs.ToolkitHost);
 	}
 
 	return EAssetCommandResult::Handled;
diff --git a/Engine/Plugins/Experimental/Mutable/Source/CustomizableObjectEditor/Private/MuCOE/AssetDefinition_CustomizableObjectInstance.cpp b/Engine/Plugins/Experimental/Mutable/Source/CustomizableObjectEditor/Private/MuCOE/AssetDefinition_CustomizableObjectInstance.cpp
--- a/Engine/Plugins/Experimental/Mutable/Source/CustomizableObjectEditor/Private/MuCOE/AssetDefinition_CustomizableObjectInstance.cpp
+++ b/Engine/Plugins/Experimental/Mutable/Source/CustomizableObjectEditor/Private/MuCOE/AssetDefinition_CustomizableObjectInstance.cpp
@@ -4,9 +4,10 @@
 
 #include "ContentBrowserMenuContexts.h"
 #include "ContentBrowserModule.h"
+#include "CustomizableObjectInstanceEditor.h"
 #include "Editor.h"
 #include "ToolMenus.h"
-#include "MuCOE/CustomizableObjectEditorModule.h"
+#include "MuCO/ICustomizableObjectEditorModule.h"
 #include "Subsystems/AssetEditorSubsystem.h"
 
 #define LOCTEXT_NAMESPACE "CustomizableObjectInstance"
@@ -17,8 +18,8 @@ EAssetCommandResult UAssetDefinition_CustomizableObjectInstance::OpenAssets(cons
 
 	for (UCustomizableObjectInstance* Object : OpenArgs.LoadObjects<UCustomizableObjectInstance>())
 	{
-		ICustomizableObjectEditorModule* CustomizableObjectEditorModule = &FModuleManager::LoadModuleChecked<ICustomizableObjectEditorModule>( "CustomizableObjectEditor" );
-		CustomizableObjectEditorModule->CreateCustomizableObjectInstanceEditor(Mode, OpenArgs.ToolkitHost, Object);
+		const TSharedPtr<FCustomizableObjectInstanceEditor> Editor = MakeShared<FCustomizableObjectInstanceEditor>();
+		Editor->InitCustomizableObjectInstanceEditor(Mode, OpenArgs.ToolkitHost, Object);
 	}
 
 	return EAssetCommandResult::Handled;
diff --git a/Engine/Plugins/Experimental/Mutable/Source/CustomizableObjectEditor/Private/MuCOE/CustomizableObjectCompiler.cpp b/Engine/Plugins/Experimental/Mutable/Source/CustomizableObjectEditor/Private/MuCOE/CustomizableObjectCompiler.cpp
--- a/Engine/Plugins/Experimental/Mutable/Source/CustomizableObjectEditor/Private/MuCOE/CustomizableObjectCompiler.cpp
+++ b/Engine/Plugins/Experimental/Mutable/Source/CustomizableObjectEditor/Private/MuCOE/CustomizableObjectCompiler.cpp
@@ -19,6 +19,8 @@
 #include "MuCOE/GraphTraversal.h"
 #include "MuCOE/ICustomizableObjectPopulationModule.h"
 #include "MuCOE/Nodes/CustomizableObjectNodeObjectGroup.h"
+#include "MuCOE/Nodes/CustomizableObjectNodeTable.h"
+#include "MuCOE/CustomizableObjectEditorModule.h"
 #include "UObject/ICookInfo.h"
 #include "UObject/UObjectIterator.h"
 #include "Widgets/Notifications/SNotificationList.h"
@@ -1118,10 +1120,27 @@ void FCustomizableObjectCompiler::CompileInternal(UCustomizableObject* Object, c
 
 		if (!ParamNamesToSelectedOptions.Num())
 		{
-			const FName PackageName = Object->GetPackage()->GetFName();
-			GenerationContext.ParticipatingObjects.Remove(PackageName); // Remove self CO reference.
+			// Get possible objects used in the compilation that are not directly referenced.
+			// Due to this check being done also in PIE (to detect out of date compilations), it has to be performant. Therefore we are gathering a relaxed set.
+			// For example, a referencing Customizable Object may not be used if it is not assigned in any Group Node. In the relaxed set we include those regardless.
+			// Notice that, to avoid automatic compilations/warnings, the set of referencing objects set found here must coincide with the set found when loading the
+			// model (discard previous compilations) or when showing PIE warnings.
+			TArray<FName> ReferencingObjectNames;
+			GetReferencingPackages(*Object, ReferencingObjectNames);
+
+			for (const FName& ReferencingObjectName : ReferencingObjectNames)
+			{
+				const TSoftObjectPtr SoftObjectPtr(ReferencingObjectName.ToString());
+
+				if (const UObject* ReferencingObject = SoftObjectPtr.LoadSynchronous())
+				{
+					GenerationContext.AddParticipatingObject(*ReferencingObject);					
+				}
+			}
 			
+			// Copy final array of participating objects
 			Object->GetPrivate()->ParticipatingObjects = MoveTemp(GenerationContext.ParticipatingObjects);
+			Object->GetPrivate()->DirtyParticipatingObjects.Empty();
 		}
 
 		if (CompileTask.IsValid()) // Don't start compilation if there's a compilation running
diff --git a/Engine/Plugins/Experimental/Mutable/Source/CustomizableObjectEditor/Private/MuCOE/CustomizableObjectEditor.cpp b/Engine/Plugins/Experimental/Mutable/Source/CustomizableObjectEditor/Private/MuCOE/CustomizableObjectEditor.cpp
--- a/Engine/Plugins/Experimental/Mutable/Source/CustomizableObjectEditor/Private/MuCOE/CustomizableObjectEditor.cpp
+++ b/Engine/Plugins/Experimental/Mutable/Source/CustomizableObjectEditor/Private/MuCOE/CustomizableObjectEditor.cpp
@@ -23,7 +23,7 @@
 #include "MuCOE/CustomizableObjectCustomSettings.h"
 #include "MuCOE/CustomizableObjectEditorActions.h"
 #include "MuCOE/CustomizableObjectEditorLogger.h"
-#include "MuCOE/CustomizableObjectEditorModule.h"
+#include "MuCO/ICustomizableObjectEditorModule.h"
 #include "MuCOE/CustomizableObjectEditorNodeContextCommands.h"
 #include "MuCOE/CustomizableObjectEditorStyle.h"
 #include "MuCOE/CustomizableObjectEditorViewportClient.h"
@@ -149,14 +149,6 @@ void FCustomizableObjectEditor::UnregisterTabSpawners(const TSharedRef<class FTa
 }
 
 
-TSharedRef<FCustomizableObjectEditor> FCustomizableObjectEditor::Create(const EToolkitMode::Type Mode, const TSharedPtr<IToolkitHost>& InitToolkitHost, UCustomizableObject* ObjectToEdit)
-{
-	TSharedRef<FCustomizableObjectEditor> Editor = MakeShareable(new FCustomizableObjectEditor(*ObjectToEdit));
-	Editor->InitCustomizableObjectEditor(Mode, InitToolkitHost);
-	return Editor;
-}
-
-
 FCustomizableObjectEditor::~FCustomizableObjectEditor()
 {
 	if (PreviewInstance && HelperCallback)
diff --git a/Engine/Plugins/Experimental/Mutable/Source/CustomizableObjectEditor/Private/MuCOE/CustomizableObjectEditor.h b/Engine/Plugins/Experimental/Mutable/Source/CustomizableObjectEditor/Private/MuCOE/CustomizableObjectEditor.h
--- a/Engine/Plugins/Experimental/Mutable/Source/CustomizableObjectEditor/Private/MuCOE/CustomizableObjectEditor.h
+++ b/Engine/Plugins/Experimental/Mutable/Source/CustomizableObjectEditor/Private/MuCOE/CustomizableObjectEditor.h
@@ -160,14 +160,7 @@ class FCustomizableObjectEditor :
 	public FEditorUndoClient
 {
 public:
-	/**
-	 * Create a new Customizable Object editor. Called immediately after construction.
-	 *
-	 * @param	Mode					Asset editing mode for this editor (standalone or world-centric)
-	 * @param	InitToolkitHost			When Mode is WorldCentric, this is the level editor instance to spawn this editor within
-	 * @param	ObjectToEdit			The object to edit
-	 */
-	static TSharedRef<FCustomizableObjectEditor> Create(const EToolkitMode::Type Mode, const TSharedPtr<IToolkitHost>& InitToolkitHost, UCustomizableObject* ObjectToEdit);
+	explicit FCustomizableObjectEditor(UCustomizableObject& ObjectToEdit);
 
 	virtual ~FCustomizableObjectEditor() override;
 
@@ -280,8 +273,6 @@ public:
 	void DebugObject() const;
 	
 private:
-	explicit FCustomizableObjectEditor(UCustomizableObject& ObjectToEdit);
-
 	TSharedRef<SDockTab> SpawnTab_Viewport(const FSpawnTabArgs& Args);
 	TSharedRef<SDockTab> SpawnTab_ObjectProperties(const FSpawnTabArgs& Args);
 	TSharedRef<SDockTab> SpawnTab_InstanceProperties(const FSpawnTabArgs& Args);
diff --git a/Engine/Plugins/Experimental/Mutable/Source/CustomizableObjectEditor/Private/MuCOE/CustomizableObjectEditorLogger.cpp b/Engine/Plugins/Experimental/Mutable/Source/CustomizableObjectEditor/Private/MuCOE/CustomizableObjectEditorLogger.cpp
--- a/Engine/Plugins/Experimental/Mutable/Source/CustomizableObjectEditor/Private/MuCOE/CustomizableObjectEditorLogger.cpp
+++ b/Engine/Plugins/Experimental/Mutable/Source/CustomizableObjectEditor/Private/MuCOE/CustomizableObjectEditorLogger.cpp
@@ -6,7 +6,7 @@
 #include "Logging/MessageLog.h"
 #include "MessageLogModule.h"
 #include "Misc/UObjectToken.h"
-#include "MuCOE/CustomizableObjectEditorModule.h"
+#include "MuCO/ICustomizableObjectEditorModule.h"
 #include "MuCOE/CustomizableObjectGraph.h"
 #include "MuCOE/ICustomizableObjectEditor.h"
 #include "MuCOE/Nodes/CustomizableObjectNode.h"
@@ -218,13 +218,13 @@ void FLogParameters::Log()
 
 FLogParameters FCustomizableObjectEditorLogger::CreateLog(const FText& Text)
 {
-	FCustomizableObjectEditorLogger& Logger = ICustomizableObjectEditorModule::Get().GetLogger();
+	FCustomizableObjectEditorLogger& Logger = ICustomizableObjectEditorModule::GetChecked().GetLogger();
 	return FLogParameters(Logger, Text);
 }
 
 FLogParameters FCustomizableObjectEditorLogger::CreateLog(FText&& Text)
 {
-	FCustomizableObjectEditorLogger& Logger = ICustomizableObjectEditorModule::Get().GetLogger();
+	FCustomizableObjectEditorLogger& Logger = ICustomizableObjectEditorModule::GetChecked().GetLogger();
 	return FLogParameters(Logger, MoveTemp(Text));
 }
 
@@ -311,7 +311,7 @@ void FCustomizableObjectEditorLogger::Log(FLogParameters& LogParameters)
 
 void FCustomizableObjectEditorLogger::DismissNotification(ELoggerCategory Category)
 {
-	FCustomizableObjectEditorLogger& Logger = ICustomizableObjectEditorModule::Get().GetLogger();
+	FCustomizableObjectEditorLogger& Logger = ICustomizableObjectEditorModule::GetChecked().GetLogger();
 	if (const FCategoryData* Result = Logger.CategoriesData.Find(Category))
 	{
 		if (const TSharedPtr<SNotificationItem> LastNotification = Result->Notification.Pin())
diff --git a/Engine/Plugins/Experimental/Mutable/Source/CustomizableObjectEditor/Private/MuCOE/CustomizableObjectEditorModule.cpp b/Engine/Plugins/Experimental/Mutable/Source/CustomizableObjectEditor/Private/MuCOE/CustomizableObjectEditorModule.cpp
--- a/Engine/Plugins/Experimental/Mutable/Source/CustomizableObjectEditor/Private/MuCOE/CustomizableObjectEditorModule.cpp
+++ b/Engine/Plugins/Experimental/Mutable/Source/CustomizableObjectEditor/Private/MuCOE/CustomizableObjectEditorModule.cpp
@@ -57,11 +57,12 @@
 #include "MuCOE/Nodes/CustomizableObjectNodeTableDetails.h"
 #include "MuCOE/Widgets/CustomizableObjectLODReductionSettings.h"
 #include "PropertyEditorModule.h"
+#include "AssetRegistry/AssetRegistryModule.h"
 #include "MuCO/CustomizableObjectInstance.h"
 #include "MuCOE/Nodes/CustomizableObjectNodeGroupProjectorParameter.h"
 #include "UObject/UObjectIterator.h"
 #include "Subsystems/PlacementSubsystem.h"
-
+#include "Components/SkeletalMeshComponent.h"
 
 class AActor;
 class FString;
@@ -78,51 +79,95 @@ const FName CustomizableObjectDebuggerAppIdentifier = FName(TEXT("CustomizableOb
 
 #define LOCTEXT_NAMESPACE "MutableSettings"
 
-/**
- * StaticMesh editor module
- */
-class FCustomizableObjectEditorModule : public ICustomizableObjectEditorModule
-{
-public:
-	// IModuleInterface interface
-	void StartupModule() override;
-	void ShutdownModule() override;
-
-	// ICustomizableObjectEditorModule interface
-	TSharedRef<ICustomizableObjectEditor> CreateCustomizableObjectEditor( const EToolkitMode::Type Mode, const TSharedPtr< IToolkitHost >& InitToolkitHost, UCustomizableObject* CustomizableObject ) override;
-	TSharedRef<ICustomizableObjectInstanceEditor> CreateCustomizableObjectInstanceEditor( const EToolkitMode::Type Mode, const TSharedPtr< IToolkitHost >& InitToolkitHost, UCustomizableObjectInstance* CustomizableObjectInstance ) override;
-	virtual EAssetTypeCategories::Type GetAssetCategory() const override;
-	virtual FCustomizableObjectEditorLogger& GetLogger() override;
 
-	virtual TSharedPtr<FExtensibilityManager> GetCustomizableObjectEditorToolBarExtensibilityManager() override { return CustomizableObjectEditor_ToolBarExtensibilityManager; }
-	virtual TSharedPtr<FExtensibilityManager> GetCustomizableObjectEditorMenuExtensibilityManager() override { return CustomizableObjectEditor_MenuExtensibilityManager; };
+constexpr float ShowOnScreenCompileWarningsTickerTime = 1.0f;
 
-public:
-	bool HandleSettingsSaved();
-	void RegisterSettings();
 
-private:	
-	TSharedPtr<FExtensibilityManager> CustomizableObjectEditor_ToolBarExtensibilityManager;
-	TSharedPtr<FExtensibilityManager> CustomizableObjectEditor_MenuExtensibilityManager;
+void ShowOnScreenCompileWarnings()
+{
+	TSet<const UCustomizableObject*> Objects;
+	
+	for (TObjectIterator<UCustomizableObjectInstanceUsage> CustomizableObjectInstanceUsage; CustomizableObjectInstanceUsage; ++CustomizableObjectInstanceUsage)
+	{
+		if (!IsValid(*CustomizableObjectInstanceUsage) || CustomizableObjectInstanceUsage->IsTemplate())
+		{
+			continue;
+		}
+		
+		const UCustomizableObjectInstance* Instance = CustomizableObjectInstanceUsage->GetCustomizableObjectInstance();
+		if (!Instance)
+		{
+			continue;
+		}
 
-	/** List of registered custom details to remove later. */
-	TArray<FName> RegisteredCustomDetails;
+		const UCustomizableObject* Object = Cast<UCustomizableObject>(Instance->GetCustomizableObject());
+		if (!Object)
+		{
+			continue;
+		}
+		
+		const USkeletalMeshComponent* Parent = Cast<USkeletalMeshComponent>(CustomizableObjectInstanceUsage->GetAttachParent());
+		if (!Parent)
+		{
+			continue;
+		}
 
-	/** Custom asset category. */
-	EAssetTypeCategories::Type CustomizableObjectAssetCategory;
+		const UWorld* World = Parent->GetWorld();
+		if (!World)
+		{
+			continue;
+		}
 
-	/** Register Custom details. Also adds them to RegisteredCustomDetails list. */
-	void RegisterCustomDetails(FPropertyEditorModule& PropertyModule, const UClass* Class, FOnGetDetailCustomizationInstance DetailLayoutDelegate);
+		if (World->WorldType != EWorldType::PIE)
+		{
+			continue;
+		}
 
-	FCustomizableObjectEditorLogger Logger;
+		Objects.Add(Object);
+	}
 
-	// Command to look for Customizable Object Instance in the player pawn of the current world and open its Customizable Object Instance Editor
-	IConsoleCommand* LaunchCOIECommand;
-	static void OpenCOIE(const TArray<FString>& Arguments);
+	for (const UCustomizableObject* Object : Objects)
+	{
+		if (Object->GetPrivate()->Status.Get() != FCustomizableObjectStatus::EState::ModelLoaded)
+		{
+			continue;
+		}
+	
+		// Show a warning if the compilation was not done with optimizations.
+		const uint64 KeyCompiledWithOptimization = reinterpret_cast<uint64>(Object);
+		if (Object->GetPrivate()->bIsCompiledWithoutOptimization)
+		{
+			FString Msg = FString::Printf(TEXT("Customizable Object [%s] was compiled without optimization."), *Object->GetName());
+			GEngine->AddOnScreenDebugMessage(KeyCompiledWithOptimization, ShowOnScreenCompileWarningsTickerTime * 2.0f, FColor::Yellow, Msg);
+		}
+		else
+		{
+			GEngine->RemoveOnScreenDebugMessage(KeyCompiledWithOptimization);
+		}
+		
+		const uint64 KeyCompiledOutOfDate = reinterpret_cast<uint64>(Object) + KEY_OFFSET_COMPILATION_OUT_OF_DATE; // Offset added to avoid collision with bIsCompiledWithOptimization warning
+		TArray<FName> OutOfDatePackages;
+		if (Object->GetPrivate()->IsCompilationOutOfDate(&OutOfDatePackages))
+		{
+			FString Msg = FString::Printf(TEXT("Customizable Object [%s] compilation out of date. See the Output Log for more information."), *Object->GetName());
+			GEngine->AddOnScreenDebugMessage(KeyCompiledOutOfDate, ShowOnScreenCompileWarningsTickerTime * 2.0f, FColor::Yellow, Msg);
+			
+			if (!GEngine->OnScreenDebugMessageExists(KeyCompiledOutOfDate))
+			{
+				UE_LOG(LogMutable, Verbose, TEXT("Warning: Customizable Object [%s] compilation out of date. Modified packages since last compilation:"), *Object->GetName());
+				for (const FName& OutOfDatePackage : OutOfDatePackages)
+				{
+					UE_LOG(LogMutable, Verbose, TEXT("%s"), *OutOfDatePackage.ToString());
+				}
+			}
+		}
+		else
+		{
+			GEngine->RemoveOnScreenDebugMessage(KeyCompiledOutOfDate);
+		}
+	}
+}
 
-	/** Register the COI factory */
-	void RegisterFactory();
-};
 
 IMPLEMENT_MODULE( FCustomizableObjectEditorModule, CustomizableObjectEditor );
 
@@ -205,6 +250,11 @@ void FCustomizableObjectEditorModule::StartupModule()
 		TEXT("Looks for a Customizable Object Instance within the player pawn and opens its Customizable Object Instance Editor. Specify slot ID to control which component is edited."),
 		FConsoleCommandWithArgsDelegate::CreateStatic(&FCustomizableObjectEditorModule::OpenCOIE));
 
+	WarningsTickerHandle = FTSTicker::GetCoreTicker().AddTicker(TEXT("ShowOnScreenCompileWarnings"), ShowOnScreenCompileWarningsTickerTime, [](float)
+	{
+		ShowOnScreenCompileWarnings();
+		return true;
+	});
 }
 
 
@@ -227,25 +277,11 @@ void FCustomizableObjectEditorModule::ShutdownModule()
 	FCoreDelegates::OnPostEngineInit.RemoveAll(this);
 
 	FCustomizableObjectEditorStyle::Shutdown();
-}
 
-
-TSharedRef<ICustomizableObjectEditor> FCustomizableObjectEditorModule::CreateCustomizableObjectEditor( const EToolkitMode::Type Mode, const TSharedPtr< IToolkitHost >& InitToolkitHost, UCustomizableObject* CustomizableObject )
-{
-	return FCustomizableObjectEditor::Create(Mode, InitToolkitHost,CustomizableObject);
+	FTSTicker::GetCoreTicker().RemoveTicker(WarningsTickerHandle);
 }
 
 
-TSharedRef<ICustomizableObjectInstanceEditor> FCustomizableObjectEditorModule::CreateCustomizableObjectInstanceEditor( const EToolkitMode::Type Mode, const TSharedPtr< IToolkitHost >& InitToolkitHost, UCustomizableObjectInstance* CustomizableObjectInstance )
-{
-	TSharedRef<FCustomizableObjectInstanceEditor> NewCustomizableObjectInstanceEditor(new FCustomizableObjectInstanceEditor());
-	NewCustomizableObjectInstanceEditor->InitCustomizableObjectInstanceEditor(Mode, InitToolkitHost,CustomizableObjectInstance);
-	return NewCustomizableObjectInstanceEditor;
-}
-
-
-
-
 FCustomizableObjectEditorLogger& FCustomizableObjectEditorModule::GetLogger()
 {
 	return Logger;
@@ -377,9 +413,176 @@ void FCustomizableObjectEditorModule::RegisterFactory()
 }
 
 
-EAssetTypeCategories::Type FCustomizableObjectEditorModule::GetAssetCategory() const
+/** Recursively get all Customizable Objects that reference the given Customizable Object. */
+void GetReferencingCustomizableObjects(FName CustomizableObjectName, TArray<FName>& VisitedObjectNames, TArray<FName>& ObjectNames)
 {
-	return CustomizableObjectAssetCategory;
+	if (VisitedObjectNames.Contains(CustomizableObjectName))
+	{
+		return;
+	}
+
+	VisitedObjectNames.Add(CustomizableObjectName);
+
+	TArray<FName> ReferencedObjectNames;
+	
+	const FAssetRegistryModule& AssetRegistryModule = FModuleManager::LoadModuleChecked<FAssetRegistryModule>("AssetRegistry");
+	AssetRegistryModule.Get().GetReferencers(CustomizableObjectName, ReferencedObjectNames, UE::AssetRegistry::EDependencyCategory::Package, UE::AssetRegistry::EDependencyQuery::Hard);
+
+	TArray<FAssetData> AssetDataArray;
+
+	FARFilter Filter;
+	Filter.PackageNames = MoveTemp(ReferencedObjectNames);
+	
+	AssetRegistryModule.Get().GetAssets(Filter, AssetDataArray);
+
+	for (FAssetData AssetData : AssetDataArray)
+	{
+		if (AssetData.GetClass() == UCustomizableObject::StaticClass())
+		{
+			FName ReferencedObjectName = AssetData.GetPackage()->GetFName();
+	
+			ObjectNames.Add(ReferencedObjectName);
+
+			GetReferencingCustomizableObjects(ReferencedObjectName, VisitedObjectNames, ObjectNames);
+		}			
+	}
+}
+
+
+void GetReferencingPackages(const UCustomizableObject& Object, TArray<FName>& ObjectNames)
+{
+	// Gather all child CustomizableObjects
+	TArray<FName> VisitedObjectNames;
+	GetReferencingCustomizableObjects(Object.GetPackage()->GetFName(), VisitedObjectNames, ObjectNames);
+
+	// Gather all tables which will composite the final tables
+	TArray<FName> CustomizableObjectNames = ObjectNames;
+	for (const FName& CustomizableObjectName : CustomizableObjectNames)
+	{
+		const TSoftObjectPtr SoftObjectPtr(CustomizableObjectName.ToString());
+
+		const UCustomizableObject* ChildCustomizableObject = Cast<UCustomizableObject>(SoftObjectPtr.LoadSynchronous());
+		if (!ChildCustomizableObject)
+		{
+			continue;
+		}
+
+		TArray<UCustomizableObjectNodeTable*> TableNodes;
+		ChildCustomizableObject->Source->GetNodesOfClass(TableNodes);
+
+		FARFilter Filter;
+		Filter.ClassPaths.Add(FTopLevelAssetPath(UDataTable::StaticClass()));
+
+		for (const UCustomizableObjectNodeTable* TableNode : TableNodes)
+		{
+			TArray<FAssetData> DataTableAssets = TableNode->GetParentTables();
+
+			for (const FAssetData& DataTableAsset : DataTableAssets)
+			{
+				if (DataTableAsset.IsValid())
+				{
+					ObjectNames.AddUnique(DataTableAsset.PackageName);
+				}
+			}
+		}		
+	}
+}
+
+
+bool FCustomizableObjectEditorModule::IsCompilationOutOfDate(const UCustomizableObject& Object, TArray<FName>* OutOfDatePackages) const
+{
+	if (!Object.GetPrivate()->DirtyParticipatingObjects.IsEmpty())
+	{
+		if (OutOfDatePackages)
+		{
+			OutOfDatePackages->Append(Object.GetPrivate()->DirtyParticipatingObjects);
+		}
+		else
+		{
+			return false;
+		}
+	}
+	
+	const FAssetRegistryModule& AssetRegistryModule = FModuleManager::LoadModuleChecked<FAssetRegistryModule>("AssetRegistry");
+
+	for (const TTuple<FName, FGuid>& ParticipatingObject : Object.GetPrivate()->ParticipatingObjects)
+	{
+		TSoftObjectPtr SoftObjectPtr(ParticipatingObject.Key.ToString());
+		if (SoftObjectPtr) // If loaded
+		{
+			PRAGMA_DISABLE_DEPRECATION_WARNINGS
+			const FGuid PackageGuid = SoftObjectPtr->GetPackage()->GetGuid();
+			PRAGMA_ENABLE_DEPRECATION_WARNINGS
+			
+			if (PackageGuid != ParticipatingObject.Value)
+			{
+				if (OutOfDatePackages)
+				{
+					OutOfDatePackages->AddUnique(ParticipatingObject.Key);
+				}
+				else
+				{
+					return false;
+				}
+			}
+		}
+		else // Not loaded
+		{
+			FAssetPackageData AssetPackageData;
+			const UE::AssetRegistry::EExists Result = AssetRegistryModule.Get().TryGetAssetPackageData(ParticipatingObject.Key, AssetPackageData);
+				
+			if (Result != UE::AssetRegistry::EExists::Exists)
+			{
+				if (OutOfDatePackages)
+				{
+					OutOfDatePackages->AddUnique(ParticipatingObject.Key);
+				}
+				else
+				{
+					return false;
+				}
+			}
+
+			PRAGMA_DISABLE_DEPRECATION_WARNINGS
+			const FGuid PackageGuid = AssetPackageData.PackageGuid;
+			PRAGMA_ENABLE_DEPRECATION_WARNINGS
+			
+			if (PackageGuid != ParticipatingObject.Value)
+			{
+				if (OutOfDatePackages)
+				{
+					OutOfDatePackages->AddUnique(ParticipatingObject.Key);
+				}
+				else
+				{
+					return false;
+				}
+			}
+		}
+	}
+
+	TArray<FName> ReferencingObjectNames;
+	GetReferencingPackages(Object, ReferencingObjectNames);
+	
+	for (const FName& ObjectName : ReferencingObjectNames)
+	{
+		TSoftObjectPtr ReferencingObject(ObjectName.ToString()); 
+
+		if ((ReferencingObject && ReferencingObject->GetPackage()->IsDirty()) ||
+			!Object.GetPrivate()->ParticipatingObjects.Contains(ObjectName)) // Must be in the participating objects, if not it means it did not exist when compiling the object.
+		{
+			if (OutOfDatePackages)
+			{
+				OutOfDatePackages->AddUnique(ObjectName);
+			}
+			else
+			{
+				return false;
+			}
+		}
+	}
+
+	return !OutOfDatePackages->IsEmpty();	
 }
 
 
diff --git a/Engine/Plugins/Experimental/Mutable/Source/CustomizableObjectEditor/Private/MuCOE/CustomizableObjectEditorModule.h b/Engine/Plugins/Experimental/Mutable/Source/CustomizableObjectEditor/Private/MuCOE/CustomizableObjectEditorModule.h
new file mode 100644
--- /dev/null
+++ b/Engine/Plugins/Experimental/Mutable/Source/CustomizableObjectEditor/Private/MuCOE/CustomizableObjectEditorModule.h
@@ -0,0 +1,66 @@
+// Copyright Epic Games, Inc. All Rights Reserved.
+
+#pragma once
+
+#include "MuCO/ICustomizableObjectEditorModule.h"
+
+#include "MuCOE/CustomizableObjectEditorLogger.h"
+#include "Toolkits/AssetEditorToolkit.h"
+
+class USkeletalMeshComponent;
+class FPropertyEditorModule;
+class FExtensibilityManager;
+class ICustomizableObjectEditor;
+class ICustomizableObjectInstanceEditor;
+class ICustomizableObjectDebugger;
+
+
+/** Get a list of packages that are used by the compilation but are not directly referenced.
+  * List includes:
+  * - Child UCustomizableObjects: Have inverted references.
+  * - UDataTable: Data Tables used by Composite Data Tables are indirectly referenced by the UStruct and filtered by path. */
+void GetReferencingPackages(const UCustomizableObject& Object, TArray<FName>& ObjectNames);
+
+
+/**
+ * StaticMesh editor module
+ */
+class FCustomizableObjectEditorModule : public ICustomizableObjectEditorModule
+{
+public:
+	// IModuleInterface interface
+	virtual void StartupModule() override;
+	virtual void ShutdownModule() override;
+
+	// ICustomizableObjectEditorModule interface
+	virtual FCustomizableObjectEditorLogger& GetLogger() override;
+	virtual bool IsCompilationOutOfDate(const UCustomizableObject& Object, TArray<FName>* OutOfDatePackages) const override;
+
+	virtual TSharedPtr<FExtensibilityManager> GetCustomizableObjectEditorToolBarExtensibilityManager() override { return CustomizableObjectEditor_ToolBarExtensibilityManager; }
+	virtual TSharedPtr<FExtensibilityManager> GetCustomizableObjectEditorMenuExtensibilityManager() override { return CustomizableObjectEditor_MenuExtensibilityManager; }
+
+	bool HandleSettingsSaved();
+	void RegisterSettings();
+
+private:	
+	TSharedPtr<FExtensibilityManager> CustomizableObjectEditor_ToolBarExtensibilityManager;
+	TSharedPtr<FExtensibilityManager> CustomizableObjectEditor_MenuExtensibilityManager;
+
+	/** List of registered custom details to remove later. */
+	TArray<FName> RegisteredCustomDetails;
+
+	/** Register Custom details. Also adds them to RegisteredCustomDetails list. */
+	void RegisterCustomDetails(FPropertyEditorModule& PropertyModule, const UClass* Class, FOnGetDetailCustomizationInstance DetailLayoutDelegate);
+
+	FCustomizableObjectEditorLogger Logger;
+
+	// Command to look for Customizable Object Instance in the player pawn of the current world and open its Customizable Object Instance Editor
+	IConsoleCommand* LaunchCOIECommand;
+
+	FTSTicker::FDelegateHandle WarningsTickerHandle;
+	
+	static void OpenCOIE(const TArray<FString>& Arguments);
+
+	/** Register the COI factory */
+	void RegisterFactory();
+};
\ No newline at end of file
diff --git a/Engine/Plugins/Experimental/Mutable/Source/CustomizableObjectEditor/Private/MuCOE/CustomizableObjectInstanceEditor.cpp b/Engine/Plugins/Experimental/Mutable/Source/CustomizableObjectEditor/Private/MuCOE/CustomizableObjectInstanceEditor.cpp
--- a/Engine/Plugins/Experimental/Mutable/Source/CustomizableObjectEditor/Private/MuCOE/CustomizableObjectInstanceEditor.cpp
+++ b/Engine/Plugins/Experimental/Mutable/Source/CustomizableObjectEditor/Private/MuCOE/CustomizableObjectInstanceEditor.cpp
@@ -21,7 +21,7 @@
 #include "MuCO/CustomizableSkeletalComponent.h"
 #include "MuCOE/CustomizableObjectCompiler.h"
 #include "MuCOE/CustomizableObjectEditorActions.h"
-#include "MuCOE/CustomizableObjectEditorModule.h"
+#include "MuCO/ICustomizableObjectEditorModule.h"
 #include "MuCOE/CustomizableObjectEditorViewportClient.h"
 #include "MuCOE/CustomizableObjectInstanceEditorActions.h"
 #include "MuCOE/CustomizableObjectPreviewScene.h"
diff --git a/Engine/Plugins/Experimental/Mutable/Source/CustomizableObjectEditor/Private/MuCOE/GenerateMutableSource/GenerateMutableSourceTable.cpp b/Engine/Plugins/Experimental/Mutable/Source/CustomizableObjectEditor/Private/MuCOE/GenerateMutableSource/GenerateMutableSourceTable.cpp
--- a/Engine/Plugins/Experimental/Mutable/Source/CustomizableObjectEditor/Private/MuCOE/GenerateMutableSource/GenerateMutableSourceTable.cpp
+++ b/Engine/Plugins/Experimental/Mutable/Source/CustomizableObjectEditor/Private/MuCOE/GenerateMutableSource/GenerateMutableSourceTable.cpp
@@ -844,6 +844,20 @@ mu::TablePtr GenerateMutableSourceTable(const UDataTable* DataTable, const UCust
 }
 
 
+void AddCompositeTablesToParticipatingObjetcts(const UDataTable* Table, FMutableGraphGenerationContext& GenerationContext)
+{
+	if (const UCompositeDataTable* CompositeTable = Cast<UCompositeDataTable>(Table))
+	{
+		GenerationContext.AddParticipatingObject(*CompositeTable);
+
+		/*for (const TArray<TObjectPtr<UDataTable>>& ParentTable : CompositeTable.ParentTables) // TODO 
+		{
+			AddCompositeTablesToParticipatingObjetcts(ParentTable, GenerationContext);
+		}*/
+	}
+}
+
+
 UDataTable* GetDataTable(const UCustomizableObjectNodeTable* TableNode, FMutableGraphGenerationContext& GenerationContext)
 {
 	UDataTable* OutDataTable = nullptr;
@@ -857,6 +871,8 @@ UDataTable* GetDataTable(const UCustomizableObjectNodeTable* TableNode, FMutable
 		OutDataTable = TableNode->Table;
 	}
 
+	AddCompositeTablesToParticipatingObjetcts(OutDataTable, GenerationContext);
+
 	return OutDataTable;
 }
 
@@ -893,26 +909,7 @@ UDataTable* GenerateDataTableFromStruct(const UCustomizableObjectNodeTable* Tabl
 		return nullptr;
 	}
 
-	TArray<FName> ReferencedTables;
-	AssetRegistry.Get()->GetReferencers(TableNode->Structure.GetPackage().GetFName(), ReferencedTables, UE::AssetRegistry::EDependencyCategory::Package, UE::AssetRegistry::EDependencyQuery::NoRequirements);
-	
-	FARFilter Filter;
-	Filter.ClassPaths.Add(FTopLevelAssetPath(UDataTable::StaticClass()));
-
-	for (const FName& ReferencedTable : ReferencedTables)
-	{
-		Filter.PackageNames.Add(ReferencedTable);
-	}
-
-	for (int32 PathIndex = 0; PathIndex < TableNode->FilterPaths.Num(); ++PathIndex)
-	{
-		Filter.PackagePaths.Add(TableNode->FilterPaths[PathIndex]);
-	}
-
-	Filter.bRecursivePaths = true;
-
-	TArray<FAssetData> DataTableAssets;
-	AssetRegistry.Get()->GetAssets(Filter, DataTableAssets);
+	TArray<FAssetData> DataTableAssets = TableNode->GetParentTables();
 
 	UCompositeDataTable* CompositeDataTable = NewObject<UCompositeDataTable>();
 	CompositeDataTable->RowStruct = TableNode->Structure;
diff --git a/Engine/Plugins/Experimental/Mutable/Source/CustomizableObjectEditor/Private/MuCOE/Nodes/CustomizableObjectNodeTable.cpp b/Engine/Plugins/Experimental/Mutable/Source/CustomizableObjectEditor/Private/MuCOE/Nodes/CustomizableObjectNodeTable.cpp
--- a/Engine/Plugins/Experimental/Mutable/Source/CustomizableObjectEditor/Private/MuCOE/Nodes/CustomizableObjectNodeTable.cpp
+++ b/Engine/Plugins/Experimental/Mutable/Source/CustomizableObjectEditor/Private/MuCOE/Nodes/CustomizableObjectNodeTable.cpp
@@ -21,6 +21,8 @@
 #include "Rendering/SkeletalMeshLODModel.h"
 #include "Rendering/SkeletalMeshModel.h"
 #include "Animation/AnimInstance.h"
+#include "AssetRegistry/AssetRegistryModule.h"
+#include "AssetRegistry/IAssetRegistry.h"
 
 class ICustomizableObjectEditor;
 class UCustomizableObjectNodeRemapPins;
@@ -1422,4 +1424,34 @@ FGuid UCustomizableObjectNodeTable::GetColumnIdByName(const FName& ColumnName) c
 }
 
 
+TArray<FAssetData> UCustomizableObjectNodeTable::GetParentTables() const
+{
+	FAssetRegistryModule& AssetRegistryModule = FModuleManager::LoadModuleChecked<FAssetRegistryModule>("AssetRegistry");
+	IAssetRegistry& AssetRegistry = AssetRegistryModule.GetRegistry();
+	
+	TArray<FName> ReferencedTables;
+	AssetRegistry.Get()->GetReferencers(Structure.GetPackage().GetFName(), ReferencedTables, UE::AssetRegistry::EDependencyCategory::Package, UE::AssetRegistry::EDependencyQuery::NoRequirements);
+	
+	FARFilter Filter;
+	Filter.ClassPaths.Add(FTopLevelAssetPath(UDataTable::StaticClass()));
+
+	for (const FName& ReferencedTable : ReferencedTables)
+	{
+		Filter.PackageNames.Add(ReferencedTable);
+	}
+
+	for (int32 PathIndex = 0; PathIndex < FilterPaths.Num(); ++PathIndex)
+	{
+		Filter.PackagePaths.Add(FilterPaths[PathIndex]);
+	}
+
+	Filter.bRecursivePaths = true;
+
+	TArray<FAssetData> DataTableAssets;
+	AssetRegistry.Get()->GetAssets(Filter, DataTableAssets);
+
+	return DataTableAssets;
+}
+
+
 #undef LOCTEXT_NAMESPACE
diff --git a/Engine/Plugins/Experimental/Mutable/Source/CustomizableObjectEditor/Public/MuCOE/CustomizableObjectEditorModule.h b/Engine/Plugins/Experimental/Mutable/Source/CustomizableObjectEditor/Public/MuCOE/CustomizableObjectEditorModule.h
deleted file mode 100644
diff --git a/Engine/Plugins/Experimental/Mutable/Source/CustomizableObjectPopulationEditor/Private/MuCOPE/AssetDefinition_CustomizableObjectPopulation.cpp b/Engine/Plugins/Experimental/Mutable/Source/CustomizableObjectPopulationEditor/Private/MuCOPE/AssetDefinition_CustomizableObjectPopulation.cpp
--- a/Engine/Plugins/Experimental/Mutable/Source/CustomizableObjectPopulationEditor/Private/MuCOPE/AssetDefinition_CustomizableObjectPopulation.cpp
+++ b/Engine/Plugins/Experimental/Mutable/Source/CustomizableObjectPopulationEditor/Private/MuCOPE/AssetDefinition_CustomizableObjectPopulation.cpp
@@ -7,7 +7,7 @@
 #include "IContentBrowserSingleton.h"
 #include "ToolMenus.h"
 #include "AssetRegistry/AssetRegistryModule.h"
-#include "MuCOE/CustomizableObjectEditorModule.h"
+#include "MuCO/ICustomizableObjectEditorModule.h"
 #include "MuCOP/CustomizableObjectPopulation.h"
 #include "MuCOPE/CustomizableObjectPopulationEditorModule.h"
 
diff --git a/Engine/Plugins/Experimental/Mutable/Source/CustomizableObjectPopulationEditor/Private/MuCOPE/AssetDefinition_CustomizableObjectPopulationClass.cpp b/Engine/Plugins/Experimental/Mutable/Source/CustomizableObjectPopulationEditor/Private/MuCOPE/AssetDefinition_CustomizableObjectPopulationClass.cpp
--- a/Engine/Plugins/Experimental/Mutable/Source/CustomizableObjectPopulationEditor/Private/MuCOPE/AssetDefinition_CustomizableObjectPopulationClass.cpp
+++ b/Engine/Plugins/Experimental/Mutable/Source/CustomizableObjectPopulationEditor/Private/MuCOPE/AssetDefinition_CustomizableObjectPopulationClass.cpp
@@ -3,7 +3,7 @@
 
 #include "MuCOPE/AssetDefinition_CustomizableObjectPopulationClass.h"
 
-#include "MuCOE/CustomizableObjectEditorModule.h"
+#include "MuCO/ICustomizableObjectEditorModule.h"
 #include "MuCOP/CustomizableObjectPopulationClass.h"
 #include "MuCOPE/CustomizableObjectPopulationEditorModule.h"
 
