[mutable] Made the Overlay layout strategy not clamp UVs to the 0-1 range, and not warn about it.
#rnx

diff --git a/Engine/Plugins/Experimental/Mutable/Source/MutableTools/Private/MuT/CodeGenerator_Mesh.cpp b/Engine/Plugins/Experimental/Mutable/Source/MutableTools/Private/MuT/CodeGenerator_Mesh.cpp
--- a/Engine/Plugins/Experimental/Mutable/Source/MutableTools/Private/MuT/CodeGenerator_Mesh.cpp
+++ b/Engine/Plugins/Experimental/Mutable/Source/MutableTools/Private/MuT/CodeGenerator_Mesh.cpp
@@ -412,6 +412,7 @@ namespace mu
 		// Get a copy of the UVs as vec2<float> to work with them. 
 		{
 			bool bNonNormalizedUVs = false;
+			const bool bIsOverlayLayout = GeneratedLayout->GetLayoutPackingStrategy() == mu::EPackStrategy::OVERLAY_LAYOUT;
 
 			const uint8* pVertices = pData;
 			for (int32 VertexIndex = 0; VertexIndex < NumVertices; ++VertexIndex)
@@ -428,7 +429,8 @@ namespace mu
 				}
 
 				// Check that UVs are normalized. If not, clamp the values and throw a warning.
-				if (MeshOptions.bNormalizeUVs && (UV[0] < 0.f || UV[0] > 1.f || UV[1] < 0.f || UV[1] > 1.f))
+				if (MeshOptions.bNormalizeUVs && !bIsOverlayLayout
+					&& (UV[0] < 0.f || UV[0] > 1.f || UV[1] < 0.f || UV[1] > 1.f))
 				{
 					UV[0] = FMath::Clamp(UV[0], 0.f, 1.f);
 					UV[1] = FMath::Clamp(UV[1], 0.f, 1.f);
@@ -439,7 +441,7 @@ namespace mu
 			}
 
 			// Mutable does not support non-normalized UVs
-			if (bNonNormalizedUVs)
+			if (bNonNormalizedUVs && !bIsOverlayLayout)
 			{
 				FString Msg = FString::Printf(TEXT("Source mesh has non-normalized UVs in LOD %d"), m_currentParents.Last().m_lod );
 				m_pErrorLog->GetPrivate()->Add(Msg, ELMT_WARNING, errorContext);
