[mutable] Fixed Crash when compiling CustomizableObjects with null Texture Parameters.
#rnx
#tests Editor

diff --git a/Engine/Plugins/Experimental/Mutable/Source/CustomizableObjectEditor/Private/MuCOE/GenerateMutableSource/GenerateMutableSource.cpp b/Engine/Plugins/Experimental/Mutable/Source/CustomizableObjectEditor/Private/MuCOE/GenerateMutableSource/GenerateMutableSource.cpp
--- a/Engine/Plugins/Experimental/Mutable/Source/CustomizableObjectEditor/Private/MuCOE/GenerateMutableSource/GenerateMutableSource.cpp
+++ b/Engine/Plugins/Experimental/Mutable/Source/CustomizableObjectEditor/Private/MuCOE/GenerateMutableSource/GenerateMutableSource.cpp
@@ -1670,6 +1670,8 @@ int32 GetPlatformLODBias(int32 TextureSize, int32 NumMips, int32 MaxPlatformSize
 
 uint32 ComputeLODBiasForTexture(const FMutableGraphGenerationContext& GenerationContext, const UTexture2D* Texture, const UTexture2D* ReferenceTexture, int32 BaseTextureSize)
 {
+	ensure(Texture);
+
 	constexpr int32 MaxAllowedLODBias = 6;
 
 	// Force a large LODBias for debug
diff --git a/Engine/Plugins/Experimental/Mutable/Source/CustomizableObjectEditor/Private/MuCOE/GenerateMutableSource/GenerateMutableSourceImage.cpp b/Engine/Plugins/Experimental/Mutable/Source/CustomizableObjectEditor/Private/MuCOE/GenerateMutableSource/GenerateMutableSourceImage.cpp
--- a/Engine/Plugins/Experimental/Mutable/Source/CustomizableObjectEditor/Private/MuCOE/GenerateMutableSource/GenerateMutableSourceImage.cpp
+++ b/Engine/Plugins/Experimental/Mutable/Source/CustomizableObjectEditor/Private/MuCOE/GenerateMutableSource/GenerateMutableSourceImage.cpp
@@ -255,10 +255,18 @@ mu::NodeImagePtr GenerateMutableSourceImage(const UEdGraphPin* Pin, FMutableGrap
 		Result = ImageNode;
 
 		UTexture2D* Texture = TypedNodeMesh->FindTextureForPin(Pin);
-		ImageNode->SetValue( GenerateImageConstant(Texture, GenerationContext, false).get() );
 
-		const uint32 MipsToSkip = ComputeLODBiasForTexture(GenerationContext, Texture, nullptr, ReferenceTextureSize);
-		Result = ResizeTextureByNumMips(ImageNode, MipsToSkip);
+		if (Texture)
+		{
+			ImageNode->SetValue(GenerateImageConstant(Texture, GenerationContext, false).get());
+
+			const uint32 MipsToSkip = ComputeLODBiasForTexture(GenerationContext, Texture, nullptr, ReferenceTextureSize);
+			Result = ResizeTextureByNumMips(ImageNode, MipsToSkip);
+		}
+		else
+		{
+			Result = mu::NodeImagePtr();
+		}
 	}
 
 	else if (const UCustomizableObjectNodeTextureInterpolate* TypedNodeInterp = Cast<UCustomizableObjectNodeTextureInterpolate>(Node))
diff --git a/Engine/Plugins/Experimental/Mutable/Source/CustomizableObjectEditor/Private/MuCOE/GenerateMutableSource/GenerateMutableSourceSurface.cpp b/Engine/Plugins/Experimental/Mutable/Source/CustomizableObjectEditor/Private/MuCOE/GenerateMutableSource/GenerateMutableSourceSurface.cpp
--- a/Engine/Plugins/Experimental/Mutable/Source/CustomizableObjectEditor/Private/MuCOE/GenerateMutableSource/GenerateMutableSourceSurface.cpp
+++ b/Engine/Plugins/Experimental/Mutable/Source/CustomizableObjectEditor/Private/MuCOE/GenerateMutableSource/GenerateMutableSourceSurface.cpp
@@ -607,12 +607,19 @@ mu::NodeSurfacePtr GenerateMutableSourceSurface(const UEdGraphPin * Pin, FMutabl
 								{
 									UTexture2D* Texture2D = TypedNodeMat->GetImageValue(ImageIndex);
 
-									const mu::NodeImageConstantPtr ConstImageNode = new mu::NodeImageConstant();
-									mu::Ptr<mu::Image> ImageConstant = GenerateImageConstant(Texture2D, GenerationContext, false);
-									ConstImageNode->SetValue(ImageConstant.get());
+									if (Texture2D)
+									{
+										const mu::NodeImageConstantPtr ConstImageNode = new mu::NodeImageConstant();
+										mu::Ptr<mu::Image> ImageConstant = GenerateImageConstant(Texture2D, GenerationContext, false);
+										ConstImageNode->SetValue(ImageConstant.get());
 
-									const uint32 MipsToSkip = ComputeLODBiasForTexture(GenerationContext, Texture2D, nullptr, Props.TextureSize);
-									return ResizeTextureByNumMips(ConstImageNode, MipsToSkip);
+										const uint32 MipsToSkip = ComputeLODBiasForTexture(GenerationContext, Texture2D, nullptr, Props.TextureSize);
+										return ResizeTextureByNumMips(ConstImageNode, MipsToSkip);
+									}
+									else
+									{
+										return mu::NodeImagePtr();
+									}
 								}
 							}
 							else
