[mutable] Reduce memory peak during the link stage of the CustomizableObject compilation.
Replace concurrent compilation locks with access-specific locks. It may produce better concurrency.
Try to avoid some unnecessary compiled data copies during cook.

#rnx

diff --git a/Engine/Plugins/Experimental/Mutable/Source/CustomizableObject/Internal/MuCO/CustomizableObjectPrivate.h b/Engine/Plugins/Experimental/Mutable/Source/CustomizableObject/Internal/MuCO/CustomizableObjectPrivate.h
--- a/Engine/Plugins/Experimental/Mutable/Source/CustomizableObject/Internal/MuCO/CustomizableObjectPrivate.h
+++ b/Engine/Plugins/Experimental/Mutable/Source/CustomizableObject/Internal/MuCO/CustomizableObjectPrivate.h
@@ -699,7 +699,7 @@ public:
 	void LoadCompiledDataFromDisk();
 
 	/** Cache platform data for cook */
-	void CachePlatformData(const ITargetPlatform* InTargetPlatform, const TArray64<uint8>& InObjectBytes, const TArray64<uint8>& InBulkBytes, const TArray64<uint8>& InMorphBytes);
+	void CachePlatformData(const ITargetPlatform* InTargetPlatform, TArray64<uint8>& InObjectBytes, TArray64<uint8>& InBulkBytes, TArray64<uint8>& InMorphBytes);
 	
 	/** Loads data previously compiled in BeginCacheForCookedPlatformData onto the UProperties in *this,
 	  * in preparation for saving the cooked package for *this or for a CustomizableObjectInstance using *this.
diff --git a/Engine/Plugins/Experimental/Mutable/Source/CustomizableObject/Private/MuCO/CustomizableObject.cpp b/Engine/Plugins/Experimental/Mutable/Source/CustomizableObject/Private/MuCO/CustomizableObject.cpp
--- a/Engine/Plugins/Experimental/Mutable/Source/CustomizableObject/Private/MuCO/CustomizableObject.cpp
+++ b/Engine/Plugins/Experimental/Mutable/Source/CustomizableObject/Private/MuCO/CustomizableObject.cpp
@@ -812,7 +812,7 @@ void UCustomizableObjectPrivate::LoadCompiledDataFromDisk()
 }
 
 
-void UCustomizableObjectPrivate::CachePlatformData(const ITargetPlatform* InTargetPlatform, const TArray64<uint8>& InModelBytes, const TArray64<uint8>& InBulkBytes, const TArray64<uint8>& InMorphBytes)
+void UCustomizableObjectPrivate::CachePlatformData(const ITargetPlatform* InTargetPlatform, TArray64<uint8>& InModelBytes, TArray64<uint8>& InBulkBytes, TArray64<uint8>& InMorphBytes)
 {
 	MUTABLE_CPUPROFILER_SCOPE(CachePlatformData)
 
@@ -828,11 +828,8 @@ void UCustomizableObjectPrivate::CachePlatformData(const ITargetPlatform* InTarg
 	Data.ModelData.Append(InModelBytes);
 
 	// Cache streamable bulk data
-	Data.StreamableData.SetNumUninitialized(InBulkBytes.Num(), EAllowShrinking::No);
-	FMemory::Memcpy(Data.StreamableData.GetData(), InBulkBytes.GetData(), InBulkBytes.Num());
-
-	Data.MorphData.SetNumUninitialized(InMorphBytes.Num(), EAllowShrinking::No);
-	FMemory::Memcpy(Data.MorphData.GetData(), InMorphBytes.GetData(), Data.MorphData.Num());
+	Data.StreamableData = MoveTemp(InBulkBytes);
+	Data.MorphData = MoveTemp(InMorphBytes);
 }
 
 
diff --git a/Engine/Plugins/Experimental/Mutable/Source/CustomizableObjectEditor/Private/MuCOE/CustomizableObjectCompileRunnable.cpp b/Engine/Plugins/Experimental/Mutable/Source/CustomizableObjectEditor/Private/MuCOE/CustomizableObjectCompileRunnable.cpp
--- a/Engine/Plugins/Experimental/Mutable/Source/CustomizableObjectEditor/Private/MuCOE/CustomizableObjectCompileRunnable.cpp
+++ b/Engine/Plugins/Experimental/Mutable/Source/CustomizableObjectEditor/Private/MuCOE/CustomizableObjectCompileRunnable.cpp
@@ -398,22 +398,6 @@ uint32 FCustomizableObjectSaveDDRunnable::Run()
 }
 
 
-TArray64<uint8>& FCustomizableObjectSaveDDRunnable::GetModelBytes()
-{
-	return Bytes;
-}
-
-
-TArray64<uint8>& FCustomizableObjectSaveDDRunnable::GetBulkBytes()
-{
-	return BulkDataBytes;
-}
-
-TArray64<uint8>& FCustomizableObjectSaveDDRunnable::GetMorphBytes()
-{
-	return MorphDataBytes;
-}
-
 bool FCustomizableObjectSaveDDRunnable::IsCompleted() const
 {
 	return bThreadCompleted;
diff --git a/Engine/Plugins/Experimental/Mutable/Source/CustomizableObjectEditor/Private/MuCOE/CustomizableObjectCompileRunnable.h b/Engine/Plugins/Experimental/Mutable/Source/CustomizableObjectEditor/Private/MuCOE/CustomizableObjectCompileRunnable.h
--- a/Engine/Plugins/Experimental/Mutable/Source/CustomizableObjectEditor/Private/MuCOE/CustomizableObjectCompileRunnable.h
+++ b/Engine/Plugins/Experimental/Mutable/Source/CustomizableObjectEditor/Private/MuCOE/CustomizableObjectCompileRunnable.h
@@ -92,15 +92,6 @@ public:
 	// FRunnable interface
 	uint32 Run() override;
 
-	// Bytes where the model is stored
-	TArray64<uint8>& GetModelBytes();
-
-	// Bytes where the streamed data files are stored
-	TArray64<uint8>& GetBulkBytes();
-
-	// Bytes where the streamed morph data files are stored
-	TArray64<uint8>& GetMorphBytes();
-
 	//
 	bool IsCompleted() const;
 
@@ -121,6 +112,13 @@ private:
 
 	TSharedPtr<mu::Model, ESPMode::ThreadSafe> Model;
 
+	bool bIsCooking = false;
+
+	// Whether the thread has finished running
+	std::atomic<bool> bThreadCompleted = false;
+
+public:
+
 	// Bytes where the model is stored
 	TArray64<uint8> Bytes;
 
@@ -130,8 +128,4 @@ private:
 	// Bytes store streameable files coming form the CO itself.
 	TArray64<uint8> MorphDataBytes;
 
-	bool bIsCooking = false;
-
-	// Whether the thread has finished running
-	std::atomic<bool> bThreadCompleted = false;
 };
diff --git a/Engine/Plugins/Experimental/Mutable/Source/CustomizableObjectEditor/Private/MuCOE/CustomizableObjectCompiler.cpp b/Engine/Plugins/Experimental/Mutable/Source/CustomizableObjectEditor/Private/MuCOE/CustomizableObjectCompiler.cpp
--- a/Engine/Plugins/Experimental/Mutable/Source/CustomizableObjectEditor/Private/MuCOE/CustomizableObjectCompiler.cpp
+++ b/Engine/Plugins/Experimental/Mutable/Source/CustomizableObjectEditor/Private/MuCOE/CustomizableObjectCompiler.cpp
@@ -1172,6 +1172,7 @@ void FCustomizableObjectCompiler::CompileInternal(UCustomizableObject* Object, c
 			{
 				SaveDDTask->Init();
 				SaveDDTask->Run();
+				CurrentObject->GetPrivate()->GetModel()->GetPrivate()->UnloadRoms();
 				FinishSavingDerivedData();
 			}
 
@@ -1348,9 +1349,9 @@ void FCustomizableObjectCompiler::FinishSavingDerivedData()
 	{
 		CurrentObject->GetPrivate()->CachePlatformData(
 				SaveDDTask->GetTargetPlatform(), 
-				SaveDDTask->GetModelBytes(), 
-				SaveDDTask->GetBulkBytes(),
-				SaveDDTask->GetMorphBytes());
+				SaveDDTask->Bytes, 
+				SaveDDTask->BulkDataBytes,
+				SaveDDTask->MorphDataBytes);
 	}
 
 	// Order matters
diff --git a/Engine/Plugins/Experimental/Mutable/Source/MutableTools/Internal/MuT/AST.h b/Engine/Plugins/Experimental/Mutable/Source/MutableTools/Internal/MuT/AST.h
--- a/Engine/Plugins/Experimental/Mutable/Source/MutableTools/Internal/MuT/AST.h
+++ b/Engine/Plugins/Experimental/Mutable/Source/MutableTools/Internal/MuT/AST.h
@@ -554,16 +554,18 @@ namespace mu
         // Code linking
         //---------------------------------------------------------------------------------------------
 
-        //!
-        static void FullLink( Ptr<ASTOp>& root, FProgram&, FLinkerOptions*);
-
-        //!
-        static void ClearLinkData( Ptr<ASTOp>& root );
+		/** Convert the operation graph at Root into code in the given program.
+		* Potentially destroys the data in this operation, so it shouldn't be used after calling Link.
+		*/
+		static OP::ADDRESS FullLink( Ptr<ASTOp>& Root, FProgram&, FLinkerOptions* );
 
     private:
 
-        //!
-        virtual void Link( FProgram& program, FLinkerOptions* Options ) = 0;
+        /** Convert this operation into code in the given program. 
+		* It assumes children have been linked already
+		* Potentially destroys the data in this operation, so it shouldn't be used after calling Link. 
+		*/
+        virtual void Link( FProgram&, FLinkerOptions* ) = 0;
 
     protected:
 
diff --git a/Engine/Plugins/Experimental/Mutable/Source/MutableTools/Private/MuT/AST.cpp b/Engine/Plugins/Experimental/Mutable/Source/MutableTools/Private/MuT/AST.cpp
--- a/Engine/Plugins/Experimental/Mutable/Source/MutableTools/Private/MuT/AST.cpp
+++ b/Engine/Plugins/Experimental/Mutable/Source/MutableTools/Private/MuT/AST.cpp
@@ -246,24 +246,20 @@ mu::Ptr<ASTOp> ASTOp::DeepClone( const Ptr<ASTOp>& root )
 
 
 //-------------------------------------------------------------------------------------------------
-void ASTOp::FullLink( Ptr<ASTOp>& root, FProgram& program, FLinkerOptions* Options )
+OP::ADDRESS ASTOp::FullLink( Ptr<ASTOp>& Root, FProgram& Program, FLinkerOptions* Options )
 {
     MUTABLE_CPUPROFILER_SCOPE(AST_FullLink);
 
-    Traverse_BottomUp_Unique( root,
-                              [&](Ptr<ASTOp> n){ n->Link(program, Options); },
+    Traverse_BottomUp_Unique( Root,
+                              [&](Ptr<ASTOp> n){ n->Link(Program, Options); },
                               [&](Ptr<const ASTOp> n){ return n->linkedAddress==0; });
-}
 
+	OP::ADDRESS Result = Root->linkedAddress;
 
-//-------------------------------------------------------------------------------------------------
-void ASTOp::ClearLinkData( Ptr<ASTOp>& root )
-{
-    MUTABLE_CPUPROFILER_SCOPE(AST_ClearLinkData);
-    ASTOpList roots;
-    roots.Add(root);
-    Traverse_TopDown_Unique_Imprecise( roots,
-                             [&](const Ptr<ASTOp>& n){ n->linkedAddress = 0; return true; });
+	// This signals the caller that the Root pointer shouldn't be used anymore.
+	Root = nullptr;
+
+	return Result;
 }
 
 
diff --git a/Engine/Plugins/Experimental/Mutable/Source/MutableTools/Private/MuT/ASTOpConstantResource.cpp b/Engine/Plugins/Experimental/Mutable/Source/MutableTools/Private/MuT/ASTOpConstantResource.cpp
--- a/Engine/Plugins/Experimental/Mutable/Source/MutableTools/Private/MuT/ASTOpConstantResource.cpp
+++ b/Engine/Plugins/Experimental/Mutable/Source/MutableTools/Private/MuT/ASTOpConstantResource.cpp
@@ -29,13 +29,13 @@ namespace mu
 
 
 	//-------------------------------------------------------------------------------------------------
-	bool ASTOpConstantResource::IsEqual(const ASTOp& otherUntyped) const
+	bool ASTOpConstantResource::IsEqual(const ASTOp& OtherUntyped) const
 	{
-		if (otherUntyped.GetOpType()==GetOpType())
+		if (OtherUntyped.GetOpType()==GetOpType())
 		{
-			const ASTOpConstantResource* other = static_cast<const ASTOpConstantResource*>(&otherUntyped);
-			return type == other->type && hash == other->hash &&
-				loadedValue == other->loadedValue && proxy == other->proxy;
+			const ASTOpConstantResource* Other = static_cast<const ASTOpConstantResource*>(&OtherUntyped);
+			return Type == Other->Type && ValueHash == Other->ValueHash &&
+				LoadedValue == Other->LoadedValue && Proxy == Other->Proxy;
 		}
 		return false;
 	}
@@ -45,10 +45,10 @@ namespace mu
 	mu::Ptr<ASTOp> ASTOpConstantResource::Clone(MapChildFuncRef) const
 	{
 		Ptr<ASTOpConstantResource> n = new ASTOpConstantResource();
-		n->type = type;
-		n->proxy = proxy;
-		n->loadedValue = loadedValue;
-		n->hash = hash;
+		n->Type = Type;
+		n->Proxy = Proxy;
+		n->LoadedValue = LoadedValue;
+		n->ValueHash = ValueHash;
 		return n;
 	}
 
@@ -56,8 +56,8 @@ namespace mu
 	//-------------------------------------------------------------------------------------------------
 	uint64 ASTOpConstantResource::Hash() const
 	{
-		uint64 res = std::hash<uint64>()(uint64(type));
-		hash_combine(res, hash);
+		uint64 res = std::hash<uint64>()(uint64(Type));
+		hash_combine(res, ValueHash);
 		return res;
 	}
 
@@ -181,7 +181,7 @@ namespace mu
 
 		if (!linkedAddress && !bLinkedAndNull)
 		{
-			if (type == OP_TYPE::ME_CONSTANT)
+			if (Type == OP_TYPE::ME_CONSTANT)
 			{
 				OP::MeshConstantArgs args;
 				FMemory::Memset(&args, 0, sizeof(args));
@@ -218,7 +218,7 @@ namespace mu
 
 				linkedAddress = (OP::ADDRESS)program.m_opAddress.Num();
 				program.m_opAddress.Add((uint32_t)program.m_byteCode.Num());
-				AppendCode(program.m_byteCode, type);
+				AppendCode(program.m_byteCode, Type);
 				AppendCode(program.m_byteCode, args);
 			}
 			else
@@ -228,7 +228,7 @@ namespace mu
 
 				bool bValidData = true;
 
-				switch (type)
+				switch (Type)
 				{
 				case OP_TYPE::IM_CONSTANT:
 				{
@@ -244,6 +244,7 @@ namespace mu
 					{
 						args.value = AddConstantImage( program, pTyped, *Options);
 					}
+
 					break;
 				}
 				case OP_TYPE::LA_CONSTANT:
@@ -260,8 +261,8 @@ namespace mu
 				if (bValidData)
 				{
 					linkedAddress = (OP::ADDRESS)program.m_opAddress.Num();
-					program.m_opAddress.Add((uint32_t)program.m_byteCode.Num());
-					AppendCode(program.m_byteCode, type);
+					program.m_opAddress.Add((uint32)program.m_byteCode.Num());
+					AppendCode(program.m_byteCode, Type);
 					AppendCode(program.m_byteCode, args);
 				}
 				else
@@ -271,6 +272,10 @@ namespace mu
 					bLinkedAndNull = true;
 				}
 			}
+
+			// Clear stored value to reduce memory usage.
+			LoadedValue = nullptr;
+			Proxy = nullptr;
 		}
 	}
 
@@ -280,7 +285,7 @@ namespace mu
 	{
 		FImageDesc res;
 
-		if (type == OP_TYPE::IM_CONSTANT)
+		if (Type == OP_TYPE::IM_CONSTANT)
 		{
 			// TODO: cache to avoid disk loading
 			Ptr<const Image> pConst = static_cast<const Image*>(GetValue().get());
@@ -300,7 +305,7 @@ namespace mu
 	//-------------------------------------------------------------------------------------------------
 	void ASTOpConstantResource::GetBlockLayoutSize(int blockIndex, int* pBlockX, int* pBlockY, FBlockLayoutSizeCache*)
 	{
-		switch (type)
+		switch (Type)
 		{
 		case OP_TYPE::LA_CONSTANT:
 		{
@@ -333,7 +338,7 @@ namespace mu
 	//-------------------------------------------------------------------------------------------------
 	void ASTOpConstantResource::GetLayoutBlockSize(int* pBlockX, int* pBlockY)
 	{
-		switch (type)
+		switch (Type)
 		{
 
 		case OP_TYPE::IM_CONSTANT:
@@ -353,7 +358,7 @@ namespace mu
 	//-------------------------------------------------------------------------------------------------
 	bool ASTOpConstantResource::GetNonBlackRect(FImageRect& maskUsage) const
 	{
-		if (type == OP_TYPE::IM_CONSTANT)
+		if (Type == OP_TYPE::IM_CONSTANT)
 		{
 			// TODO: cache
 			Ptr<const Image> pMask = static_cast<const Image*>(GetValue().get());
@@ -369,7 +374,7 @@ namespace mu
 	bool ASTOpConstantResource::IsImagePlainConstant(FVector4f& colour) const
 	{
 		bool res = false;
-		switch (type)
+		switch (Type)
 		{
 
 		case OP_TYPE::IM_CONSTANT:
@@ -422,25 +427,25 @@ namespace mu
 	//-------------------------------------------------------------------------------------------------
 	uint64 ASTOpConstantResource::GetValueHash() const
 	{
-		return hash;
+		return ValueHash;
 	}
 
 
 	//-------------------------------------------------------------------------------------------------
 	mu::Ptr<const RefCounted> ASTOpConstantResource::GetValue() const
 	{
-		if (loadedValue)
+		if (LoadedValue)
 		{
-			return loadedValue;
+			return LoadedValue;
 		}
 		else
 		{
-			switch (type)
+			switch (Type)
 			{
 
 			case OP_TYPE::IM_CONSTANT:
 			{
-				Ptr<ResourceProxy<Image>> typedProxy = static_cast<ResourceProxy<Image>*>(proxy.get());
+				Ptr<ResourceProxy<Image>> typedProxy = static_cast<ResourceProxy<Image>*>(Proxy.get());
 				Ptr<const Image> r = typedProxy->Get();
 				return r;
 			}
@@ -458,7 +463,7 @@ namespace mu
 	//-------------------------------------------------------------------------------------------------
 	void ASTOpConstantResource::SetValue(const Ptr<const RefCounted>& v, FProxyFileContext* DiskCacheContext)
 	{
-		switch (type)
+		switch (Type)
 		{
 		case OP_TYPE::IM_CONSTANT:
 		{
@@ -468,15 +473,15 @@ namespace mu
 			OutputArchive arch(&stream);
 			Image::Serialise(r.get(), arch);
 
-			hash = CityHash64(static_cast<const char*>(stream.GetBuffer()), stream.GetBufferSize());
+			ValueHash = CityHash64(static_cast<const char*>(stream.GetBuffer()), stream.GetBufferSize());
 
 			if (DiskCacheContext)
 			{
-				proxy = new ResourceProxyTempFile<Image>(r.get(), *DiskCacheContext);
+				Proxy = new ResourceProxyTempFile<Image>(r.get(), *DiskCacheContext);
 			}
 			else
 			{
-				loadedValue = r;
+				LoadedValue = r;
 			}
 			break;
 		}
@@ -489,9 +494,9 @@ namespace mu
 			Ptr<const Mesh> r = static_cast<const Mesh*>(v.get());
 			Mesh::Serialise(r.get(), arch);
 
-			hash = CityHash64(static_cast<const char*>(stream.GetBuffer()), stream.GetBufferSize());
+			ValueHash = CityHash64(static_cast<const char*>(stream.GetBuffer()), stream.GetBufferSize());
 
-			loadedValue = v;
+			LoadedValue = v;
 			break;
 		}
 
@@ -503,14 +508,14 @@ namespace mu
 			Ptr<const Layout> r = static_cast<const Layout*>(v.get());
 			Layout::Serialise(r.get(), arch);
 
-			hash = CityHash64(static_cast<const char*>(stream.GetBuffer()), stream.GetBufferSize());
+			ValueHash = CityHash64(static_cast<const char*>(stream.GetBuffer()), stream.GetBufferSize());
 
-			loadedValue = v;
+			LoadedValue = v;
 			break;
 		}
 
 		default:
-			loadedValue = v;
+			LoadedValue = v;
 			break;
 		}
 	}
@@ -519,7 +524,7 @@ namespace mu
 	//-------------------------------------------------------------------------------------------------
 	mu::Ptr<ImageSizeExpression> ASTOpConstantResource::GetImageSizeExpression() const
 	{
-		if (type==OP_TYPE::IM_CONSTANT)
+		if (Type==OP_TYPE::IM_CONSTANT)
 		{
 			Ptr<ImageSizeExpression> pRes = new ImageSizeExpression;
 			pRes->type = ImageSizeExpression::ISET_CONSTANT;
diff --git a/Engine/Plugins/Experimental/Mutable/Source/MutableTools/Private/MuT/ASTOpConstantResource.h b/Engine/Plugins/Experimental/Mutable/Source/MutableTools/Private/MuT/ASTOpConstantResource.h
--- a/Engine/Plugins/Experimental/Mutable/Source/MutableTools/Private/MuT/ASTOpConstantResource.h
+++ b/Engine/Plugins/Experimental/Mutable/Source/MutableTools/Private/MuT/ASTOpConstantResource.h
@@ -21,11 +21,11 @@ namespace mu
 	private:
 
 		//!
-		Ptr<const RefCounted> loadedValue;
-		Ptr<RefCounted> proxy;
+		Ptr<const RefCounted> LoadedValue;
+		Ptr<RefCounted> Proxy;
 
 		//! Value hash
-		uint64 hash;
+		uint64 ValueHash;
 
 		//! We tried to link already but the result is a null op.
 		bool bLinkedAndNull = false;
@@ -33,7 +33,7 @@ namespace mu
 	public:
 
 		//! Type of switch
-		OP_TYPE type;
+		OP_TYPE Type;
 
 	public:
 
@@ -54,7 +54,7 @@ namespace mu
 
 
 		// ASTOp interface
-		OP_TYPE GetOpType() const override { return type; }
+		OP_TYPE GetOpType() const override { return Type; }
 		void ForEachChild(const TFunctionRef<void(ASTChild&)>) override;
 		bool IsEqual(const ASTOp& otherUntyped) const override;
 		Ptr<ASTOp> Clone(MapChildFuncRef mapChild) const override;
diff --git a/Engine/Plugins/Experimental/Mutable/Source/MutableTools/Private/MuT/ASTOpMeshFormat.cpp b/Engine/Plugins/Experimental/Mutable/Source/MutableTools/Private/MuT/ASTOpMeshFormat.cpp
--- a/Engine/Plugins/Experimental/Mutable/Source/MutableTools/Private/MuT/ASTOpMeshFormat.cpp
+++ b/Engine/Plugins/Experimental/Mutable/Source/MutableTools/Private/MuT/ASTOpMeshFormat.cpp
@@ -281,7 +281,7 @@ mu::Ptr<ASTOp> Sink_MeshFormatAST::Visit(const mu::Ptr<ASTOp>& at, const ASTOpMe
 		MeshPtrConst pTargetMorphFormat = MakeMorphTargetFormat(pTargetFormat);
 
 		mu::Ptr<ASTOpConstantResource> motaop = new ASTOpConstantResource();
-		motaop->type = OP_TYPE::ME_CONSTANT;
+		motaop->Type = OP_TYPE::ME_CONSTANT;
 		motaop->SetValue(pTargetMorphFormat, nullptr);
 		auto targetMorphFormatAt = motaop;
 
@@ -321,7 +321,7 @@ mu::Ptr<ASTOp> Sink_MeshFormatAST::Visit(const mu::Ptr<ASTOp>& at, const ASTOpMe
 		MeshPtrConst pTargetMorphFormat = MakeMorphTargetFormat(pTargetFormat);
 
 		mu::Ptr<ASTOpConstantResource> motaop = new ASTOpConstantResource();
-		motaop->type = OP_TYPE::ME_CONSTANT;
+		motaop->Type = OP_TYPE::ME_CONSTANT;
 		motaop->SetValue(pTargetMorphFormat, nullptr );
 		auto targetMorphFormatAt = motaop;
 
diff --git a/Engine/Plugins/Experimental/Mutable/Source/MutableTools/Private/MuT/CodeGenerator.cpp b/Engine/Plugins/Experimental/Mutable/Source/MutableTools/Private/MuT/CodeGenerator.cpp
--- a/Engine/Plugins/Experimental/Mutable/Source/MutableTools/Private/MuT/CodeGenerator.cpp
+++ b/Engine/Plugins/Experimental/Mutable/Source/MutableTools/Private/MuT/CodeGenerator.cpp
@@ -979,7 +979,7 @@ namespace mu
 						// Constant layout to start with
 						{
 							Ptr<ASTOpConstantResource> op = new ASTOpConstantResource();
-							op->type = OP_TYPE::LA_CONSTANT;
+							op->Type = OP_TYPE::LA_CONSTANT;
 
 							op->SetValue(pLayout, m_compilerOptions->OptimisationOptions.DiskCacheContext);
 							layoutOp = op;
@@ -1005,7 +1005,7 @@ namespace mu
 								Ptr<ASTOp> layoutFragmentAd;
 								{
 									Ptr<ASTOpConstantResource> op = new ASTOpConstantResource();
-									op->type = OP_TYPE::LA_CONSTANT;
+									op->Type = OP_TYPE::LA_CONSTANT;
 
 									Ptr<const Layout> pCloned = data.GeneratedLayouts[LayoutIndex];
 									op->SetValue(pCloned, m_compilerOptions->OptimisationOptions.DiskCacheContext);
diff --git a/Engine/Plugins/Experimental/Mutable/Source/MutableTools/Private/MuT/CodeGenerator_Image.cpp b/Engine/Plugins/Experimental/Mutable/Source/MutableTools/Private/MuT/CodeGenerator_Image.cpp
--- a/Engine/Plugins/Experimental/Mutable/Source/MutableTools/Private/MuT/CodeGenerator_Image.cpp
+++ b/Engine/Plugins/Experimental/Mutable/Source/MutableTools/Private/MuT/CodeGenerator_Image.cpp
@@ -210,7 +210,7 @@ namespace mu
 		else
 		{
 			Ptr<ASTOpConstantResource> op = new ASTOpConstantResource();
-			op->type = OP_TYPE::IM_CONSTANT;
+			op->Type = OP_TYPE::IM_CONSTANT;
 			op->SetValue(pImage, m_compilerOptions->OptimisationOptions.DiskCacheContext);
 			Result.op = op;
 		}
@@ -1437,7 +1437,7 @@ namespace mu
                 // For wrapping projector we need the entire mesh. The actual project operation
                 // will remove the faces that are not in the layout block we are generating.
                 Ptr<ASTOpConstantResource> cop = new ASTOpConstantResource();
-                cop->type = OP_TYPE::ME_CONSTANT;
+                cop->Type = OP_TYPE::ME_CONSTANT;
 				Ptr<Mesh> FormatMeshResult = new Mesh();
 				CreateMeshOptimisedForWrappingProjection(FormatMeshResult.get(), node.m_layout);
 
@@ -1468,7 +1468,7 @@ namespace mu
 
                 // Reformat the mesh to a more efficient format for this operation
                 Ptr<ASTOpConstantResource> cop = new ASTOpConstantResource();
-                cop->type = OP_TYPE::ME_CONSTANT;
+                cop->Type = OP_TYPE::ME_CONSTANT;
 
 				Ptr<Mesh> FormatMeshResult = new Mesh();
                 CreateMeshOptimisedForProjection(FormatMeshResult.get(), node.m_layout);
@@ -1493,7 +1493,7 @@ namespace mu
             // This argument is required
             Ptr<const Mesh> TempMesh = new Mesh();
             Ptr<ASTOpConstantResource> cop = new ASTOpConstantResource();
-            cop->type = OP_TYPE::ME_CONSTANT;
+            cop->Type = OP_TYPE::ME_CONSTANT;
             cop->SetValue(TempMesh, m_compilerOptions->OptimisationOptions.DiskCacheContext);
 			ProjectOp->SetChild(ProjectOp->op.args.MeshProject.mesh, cop );
             m_pErrorLog->GetPrivate()->Add( "Projector mesh not set.", ELMT_ERROR, node.m_errorContext );
diff --git a/Engine/Plugins/Experimental/Mutable/Source/MutableTools/Private/MuT/CodeGenerator_Mesh.cpp b/Engine/Plugins/Experimental/Mutable/Source/MutableTools/Private/MuT/CodeGenerator_Mesh.cpp
--- a/Engine/Plugins/Experimental/Mutable/Source/MutableTools/Private/MuT/CodeGenerator_Mesh.cpp
+++ b/Engine/Plugins/Experimental/Mutable/Source/MutableTools/Private/MuT/CodeGenerator_Mesh.cpp
@@ -1255,7 +1255,7 @@ namespace mu
         NodeMeshConstant::Private& node = *constant->GetPrivate();
 
         Ptr<ASTOpConstantResource> op = new ASTOpConstantResource();
-        op->type = OP_TYPE::ME_CONSTANT;
+        op->Type = OP_TYPE::ME_CONSTANT;
 		OutResult.baseMeshOp = op;
 		OutResult.meshOp = op;
 		OutResult.GeneratedLayouts.Empty();
@@ -1506,7 +1506,7 @@ namespace mu
             }
 
             Ptr<ASTOpConstantResource> cop = new ASTOpConstantResource();
-            cop->type = OP_TYPE::ME_CONSTANT;
+            cop->Type = OP_TYPE::ME_CONSTANT;
             cop->SetValue( pFormatMesh, m_compilerOptions->OptimisationOptions.DiskCacheContext );
             op->Format = cop;
 
diff --git a/Engine/Plugins/Experimental/Mutable/Source/MutableTools/Private/MuT/CodeOptimiser.cpp b/Engine/Plugins/Experimental/Mutable/Source/MutableTools/Private/MuT/CodeOptimiser.cpp
--- a/Engine/Plugins/Experimental/Mutable/Source/MutableTools/Private/MuT/CodeOptimiser.cpp
+++ b/Engine/Plugins/Experimental/Mutable/Source/MutableTools/Private/MuT/CodeOptimiser.cpp
@@ -24,6 +24,7 @@
 #include "MuR/MutableRuntimeModule.h"
 
 #include "Tasks/Task.h"
+#include "Misc/ScopeRWLock.h"
 
 #include <unordered_set>
 
@@ -329,17 +330,6 @@ namespace mu
 		{
 			auto& container = visited[(int)n->GetOpType()];
 
-			// debug
-	//        for( const auto& o : container )
-	//        {
-	//            if (*n==*o)
-	//            {
-	//                auto thisHash = n->Hash();
-	//                auto otherHash = o->Hash();
-	//                check(thisHash==otherHash);
-	//            }
-	//        }
-
 			// Insert will tell us if it was already there
 			auto it = container.insert(n);
 			if( !it.second )
@@ -498,8 +488,7 @@ namespace mu
 			LinkerOptions.bSeparateImageMips = false;
 
 			TSharedPtr<const Model> model = MakeShared<Model>();
-			ASTOp::FullLink(SourceCloned, model->GetPrivate()->m_program, &LinkerOptions);
-			OP::ADDRESS at = SourceCloned->linkedAddress;
+			OP::ADDRESS at = ASTOp::FullLink(SourceCloned, model->GetPrivate()->m_program, &LinkerOptions);
 
 			FProgram::FState state;
 			state.m_root = at;
@@ -520,7 +509,7 @@ namespace mu
 				if (pMesh)
 				{
 					mu::Ptr<ASTOpConstantResource> constantOp = new ASTOpConstantResource();
-					constantOp->type = OP_TYPE::ME_CONSTANT;
+					constantOp->Type = OP_TYPE::ME_CONSTANT;
 					constantOp->SetValue( pMesh, DiskCacheContext );
 					Result = constantOp;
 				  }
@@ -536,7 +525,7 @@ namespace mu
 				if (pImage)
 				{
 					mu::Ptr<ASTOpConstantResource> constantOp = new ASTOpConstantResource();
-					constantOp->type = OP_TYPE::IM_CONSTANT;
+					constantOp->Type = OP_TYPE::IM_CONSTANT;
 					constantOp->SetValue( pImage, DiskCacheContext );
 					Result = constantOp;
 				}
@@ -552,7 +541,7 @@ namespace mu
 				if (pLayout)
 				{
 					mu::Ptr<ASTOpConstantResource> constantOp = new ASTOpConstantResource();
-					constantOp->type = OP_TYPE::LA_CONSTANT;
+					constantOp->Type = OP_TYPE::LA_CONSTANT;
 					constantOp->SetValue( pLayout, DiskCacheContext);
 					Result = constantOp;
 				}
@@ -604,8 +593,6 @@ namespace mu
 			}
 
 			pSystem->GetPrivate()->EndBuild();
-
-			SourceCloned = nullptr;
 		}
 
 	};
@@ -780,7 +767,7 @@ namespace mu
 		if (bUseConcurrency)
 		{
 			/** Protect access to the original AST being optimized. */
-			FCriticalSection ASTAccessLock;
+			FRWLock ASTAccessLock;
 
 			// Launch the tasks.
 			UE::Tasks::FTask LaunchTask = UE::Tasks::Launch(TEXT("ConstantGeneratorLaunchTasks"), 
@@ -788,7 +775,7 @@ namespace mu
 				{
 					MUTABLE_CPUPROFILER_SCOPE(ConstantGenerator_LaunchTasks);
 
-					FScopeLock Lock(&ASTAccessLock);
+					FReadScopeLock Lock(ASTAccessLock);
 
 					FImageOperator ImOp = FImageOperator::GetDefault(InOptions->ImageFormatFunc);
 
@@ -819,10 +806,10 @@ namespace mu
 							UE::Tasks::FTask CompleteTask = UE::Tasks::Launch(TEXT("MutableResolveComplete"),
 								[SubgraphRoot, InOptions, ResolveImage, &ASTAccessLock]()
 								{
-									FScopeLock Lock(&ASTAccessLock);
+									FWriteScopeLock Lock(ASTAccessLock);
 
 									Ptr<ASTOpConstantResource> ConstantOp = new ASTOpConstantResource;
-									ConstantOp->type = OP_TYPE::IM_CONSTANT;
+									ConstantOp->Type = OP_TYPE::IM_CONSTANT;
 									ConstantOp->SetValue(ResolveImage->get(), InOptions->OptimisationOptions.DiskCacheContext);
 									ASTOp::Replace(SubgraphRoot, ConstantOp);
 								},
@@ -844,7 +831,7 @@ namespace mu
 							// Launch the preparation on the AST-modification pipe
 							UE::Tasks::FTask PrepareTask = UE::Tasks::Launch(TEXT("MutableConstantPrepare"), [TaskPtr, &ASTAccessLock]()
 								{
-									FScopeLock Lock(&ASTAccessLock);
+									FReadScopeLock Lock(ASTAccessLock);
 
 									// We need the clone because linking modifies ASTOp state and also to be safe for concurrency.
 									TaskPtr->SourceCloned = ASTOp::DeepClone(TaskPtr->Source);
@@ -863,7 +850,7 @@ namespace mu
 							// Launch the completion on the AST-modification pipe
 							UE::Tasks::FTask CompleteTask = UE::Tasks::Launch(TEXT("MutableConstantComplete"), [TaskPtr = MoveTemp(Task), &ASTAccessLock]()
 								{
-									FScopeLock Lock(&ASTAccessLock);
+									FWriteScopeLock Lock(ASTAccessLock);
 
 									ASTOp::Replace(TaskPtr->Source, TaskPtr->Result);
 									TaskPtr->Source = nullptr;
@@ -937,7 +924,7 @@ namespace mu
 								UE::Tasks::FTaskEvent ReferenceCompletionEvent = InOptions->OptimisationOptions.ReferencedResourceProvider(ImageID, ResolveImage, bRunImmediatlyIfPossible);
 								ReferenceCompletionEvent.Wait();
 								Ptr<ASTOpConstantResource> ConstantOp = new ASTOpConstantResource;
-								ConstantOp->type = OP_TYPE::IM_CONSTANT;
+								ConstantOp->Type = OP_TYPE::IM_CONSTANT;
 								ConstantOp->SetValue(ResolveImage->get(), InOptions->OptimisationOptions.DiskCacheContext);
 								ASTOp::Replace(SubgraphRoot, ConstantOp);
 							};
@@ -1159,7 +1146,7 @@ namespace mu
 		}
 
 		// List of meshes that require a skeleton
-		vector<mu::Ptr<ASTOpConstantResource>> m_meshesRequiringSkeleton;
+		TArray<mu::Ptr<ASTOpConstantResource>> MeshesRequiringSkeleton;
 
 	private:
 
@@ -1179,12 +1166,7 @@ namespace mu
 
 				if (currentProtected)
 				{
-					if ( std::find(m_meshesRequiringSkeleton.begin(), m_meshesRequiringSkeleton.end(), typedOp)
-						 ==
-						 m_meshesRequiringSkeleton.end() )
-					{
-						m_meshesRequiringSkeleton.push_back(typedOp);
-					}
+					MeshesRequiringSkeleton.AddUnique(typedOp);
 				}
 
 				return false;
@@ -1263,7 +1245,7 @@ namespace mu
 						  const mu::Ptr<Skeleton>& pFinalSkeleton )
 		{
 			m_pAddMeshOp = pAddMeshOp;
-			m_contributingMeshes = std::move(contributingMeshes);
+			m_contributingMeshes = MoveTemp(contributingMeshes);
 			m_pFinalSkeleton = pFinalSkeleton;
 		}
 	};
@@ -1352,7 +1334,7 @@ namespace mu
 						{
 							NewMesh->CheckIntegrity();
 							mu::Ptr<ASTOpConstantResource> newOp = new ASTOpConstantResource();
-							newOp->type = OP_TYPE::ME_CONSTANT;
+							newOp->Type = OP_TYPE::ME_CONSTANT;
 							newOp->SetValue(NewMesh, options.DiskCacheContext);
 
 							ASTOp::Replace(at, newOp);
diff --git a/Engine/Plugins/Experimental/Mutable/Source/MutableTools/Private/MuT/Compiler.cpp b/Engine/Plugins/Experimental/Mutable/Source/MutableTools/Private/MuT/Compiler.cpp
--- a/Engine/Plugins/Experimental/Mutable/Source/MutableTools/Private/MuT/Compiler.cpp
+++ b/Engine/Plugins/Experimental/Mutable/Source/MutableTools/Private/MuT/Compiler.cpp
@@ -340,10 +340,9 @@ namespace mu
         {
 			LinkerOptions.MinTextureResidentMipCount = m_pD->m_options->GetPrivate()->MinTextureResidentMipCount;
 
-            ASTOp::FullLink(s.root,program, &LinkerOptions);
             if (s.root)
             {
-                s.state.m_root = s.root->linkedAddress;
+				s.state.m_root = ASTOp::FullLink(s.root, program, &LinkerOptions);
             }
             else
             {
