[mutable] Added custom version support to data table nodes.

#rnx

diff --git a/Engine/Plugins/Experimental/Mutable/Source/CustomizableObjectEditor/Classes/MuCOE/Nodes/CustomizableObjectNodeTable.h b/Engine/Plugins/Experimental/Mutable/Source/CustomizableObjectEditor/Classes/MuCOE/Nodes/CustomizableObjectNodeTable.h
--- a/Engine/Plugins/Experimental/Mutable/Source/CustomizableObjectEditor/Classes/MuCOE/Nodes/CustomizableObjectNodeTable.h
+++ b/Engine/Plugins/Experimental/Mutable/Source/CustomizableObjectEditor/Classes/MuCOE/Nodes/CustomizableObjectNodeTable.h
@@ -216,6 +216,10 @@ public:
 	UPROPERTY(EditAnywhere, Category = TableProperties)
 	ETableTextureType DefaultImageMode = ETableTextureType::MUTABLE_TEXTURE;
 
+	/** Name of the column that contains the Version options. */
+	UPROPERTY(EditAnywhere, Category = TableProperties)
+	FName VersionColumn;
+
 	UPROPERTY(EditAnywhere, Category = UI, meta = (DisplayName = "Parameter UI Metadata"))
 	FMutableParamUIMetadata ParamUIMetadata;
 
@@ -357,11 +361,6 @@ public:
 	USkeletalMesh* GetSkeletalMeshAt(const UEdGraphPin* Pin, const UDataTable* DataTable, const FName& RowName) const;
 	TSoftClassPtr<UAnimInstance> GetAnimInstanceAt(const UEdGraphPin* Pin, const UDataTable* DataTable, const FName& RowName) const;
 
-
-	// Return the name of the enabled rows in the data table.
-	// Returns the name if the row has a bool column set as false (true == disabled)
-	TArray<FName> GetRowNames(const UDataTable* DataTable) const;
-
 	// Changes the image mode of a pin
 	// bSetDefault param: if true sets the pin to be equal to the default mode (same as node)
 	void ChangeImagePinMode(UEdGraphPin* Pin, bool bSetDefault = false);
diff --git a/Engine/Plugins/Experimental/Mutable/Source/CustomizableObjectEditor/Private/MuCOE/GenerateMutableSource/GenerateMutableSourceMesh.cpp b/Engine/Plugins/Experimental/Mutable/Source/CustomizableObjectEditor/Private/MuCOE/GenerateMutableSource/GenerateMutableSourceMesh.cpp
--- a/Engine/Plugins/Experimental/Mutable/Source/CustomizableObjectEditor/Private/MuCOE/GenerateMutableSource/GenerateMutableSourceMesh.cpp
+++ b/Engine/Plugins/Experimental/Mutable/Source/CustomizableObjectEditor/Private/MuCOE/GenerateMutableSource/GenerateMutableSourceMesh.cpp
@@ -2239,7 +2239,7 @@ TArray<TTuple<USkeletalMesh*, TSoftClassPtr<UAnimInstance>>> GetSkeletalMeshesIn
 
 		if (DataTable)
 		{
-			for (const FName& RowName : TableNode->GetRowNames(DataTable))
+			for (const FName& RowName : GetRowsToCompile(*DataTable, *TableNode,GenerationContext))
 			{
 				USkeletalMesh* SkeletalMesh = TableNode->GetSkeletalMeshAt(SourceMeshPin, DataTable, RowName);
 				TSoftClassPtr<UAnimInstance> MeshAnimInstance = TableNode->GetAnimInstanceAt(SourceMeshPin, DataTable, RowName);
@@ -2702,7 +2702,7 @@ mu::NodeMeshPtr GenerateMorphMesh(const UEdGraphPin* Pin,
 		UDataTable* DataTable = GetDataTable(TypedNodeTable, GenerationContext);
 
 		// Generate a new Column for each morph
-		const TArray<FName>& RowNames = TypedNodeTable->GetRowNames(DataTable);
+		const TArray<FName>& RowNames = GetRowsToCompile(*DataTable, *TypedNodeTable, GenerationContext);
 		int32 NumRows = RowNames.Num();
 
 		// Should exist
diff --git a/Engine/Plugins/Experimental/Mutable/Source/CustomizableObjectEditor/Private/MuCOE/GenerateMutableSource/GenerateMutableSourceTable.cpp b/Engine/Plugins/Experimental/Mutable/Source/CustomizableObjectEditor/Private/MuCOE/GenerateMutableSource/GenerateMutableSourceTable.cpp
--- a/Engine/Plugins/Experimental/Mutable/Source/CustomizableObjectEditor/Private/MuCOE/GenerateMutableSource/GenerateMutableSourceTable.cpp
+++ b/Engine/Plugins/Experimental/Mutable/Source/CustomizableObjectEditor/Private/MuCOE/GenerateMutableSource/GenerateMutableSourceTable.cpp
@@ -22,6 +22,7 @@
 #include "PhysicsEngine/PhysicsAsset.h"
 #include "Rendering/SkeletalMeshLODModel.h"
 #include "Rendering/SkeletalMeshModel.h"
+#include "MuCOE/CustomizableObjectVersionBridge.h"
 
 #define LOCTEXT_NAMESPACE "CustomizableObjectEditor"
 
@@ -637,6 +638,20 @@ bool FillTableColumn(const UCustomizableObjectNodeTable* TableNode,	mu::TablePtr
 }
 
 
+uint8* GetCellData(const FName& RowName, const UDataTable& DataTable, const FProperty& ColumnProperty)
+{
+	// Get Row Data
+	uint8* RowData = DataTable.FindRowUnchecked(RowName);
+
+	if (RowData)
+	{
+		return ColumnProperty.ContainerPtrToValuePtr<uint8>(RowData, 0);
+	}
+
+	return nullptr;
+}
+
+
 FName GetAnotherOption(FName SelectedOptionName, const TArray<FName>& RowNames)
 {
 	for (const FName& CandidateOption : RowNames)
@@ -651,6 +666,51 @@ FName GetAnotherOption(FName SelectedOptionName, const TArray<FName>& RowNames)
 }
 
 
+TArray<FName> GetEnabledRows(const UDataTable& DataTable, const UCustomizableObjectNodeTable& TableNode)
+{
+	TArray<FName> RowNames;
+	const UScriptStruct* TableStruct = DataTable.GetRowStruct();
+
+	if (!TableStruct)
+	{
+		return RowNames;
+	}
+
+	TArray<FName> TableRowNames = DataTable.GetRowNames();
+	FBoolProperty* BoolProperty = nullptr;
+
+	for (TFieldIterator<FProperty> PropertyIt(TableStruct); PropertyIt && TableNode.bDisableCheckedRows; ++PropertyIt)
+	{
+		BoolProperty = CastField<FBoolProperty>(*PropertyIt);
+
+		if (BoolProperty)
+		{
+			for (const FName& RowName : TableRowNames)
+			{
+				if (uint8* CellData = GetCellData(RowName, DataTable, *BoolProperty))
+				{
+					if (!BoolProperty->GetPropertyValue(CellData))
+					{
+						RowNames.Add(RowName);
+					}
+				}
+			}
+
+			// There should be only one Bool column
+			break;
+		}
+	}
+
+	// There is no Bool column or we don't want to disable rows
+	if (!BoolProperty)
+	{
+		return TableRowNames;
+	}
+
+	return RowNames;
+}
+
+
 void RestrictRowNamesToSelectedOption(TArray<FName>& InOutRowNames, const UCustomizableObjectNodeTable& TableNode, FMutableGraphGenerationContext& GenerationContext)
 {
 	// If the param is in the map restrict to only the selected option
@@ -678,6 +738,57 @@ void RestrictRowNamesToSelectedOption(TArray<FName>& InOutRowNames, const UCusto
 }
 
 
+void RestrictRowContentByVersion( TArray<FName>& InOutRowNames, const UDataTable& DataTable, const UCustomizableObjectNodeTable& TableNode, FMutableGraphGenerationContext& GenerationContext)
+{
+	FProperty* ColumnProperty = DataTable.FindTableProperty(TableNode.VersionColumn);
+
+	if (!ColumnProperty)
+	{
+		return;
+	}
+
+	ICustomizableObjectVersionBridgeInterface* CustomizableObjectVersionBridgeInterface = Cast<ICustomizableObjectVersionBridgeInterface>(GenerationContext.Object->VersionBridge);
+	if (!CustomizableObjectVersionBridgeInterface)
+	{
+		const FString Message = "Found a data table with at least a row with a Custom Version asset but the Root Object does not have a Version Bridge asset assigned.";
+		GenerationContext.Compiler->CompilerLog(FText::FromString(Message), &TableNode, EMessageSeverity::Error);
+		return;
+	}
+
+	TArray<FName> OutRowNames;
+	OutRowNames.Reserve(InOutRowNames.Num());
+
+	for (int32 RowIndex = 0; RowIndex < InOutRowNames.Num(); ++RowIndex)
+	{
+		if (uint8* CellData = GetCellData(InOutRowNames[RowIndex], DataTable, *ColumnProperty))
+		{
+			if (!CustomizableObjectVersionBridgeInterface->IsVersionPropertyIncludedInCurrentRelease(*ColumnProperty, CellData))
+			{
+				continue;
+			}
+
+			OutRowNames.Add(InOutRowNames[RowIndex]);
+		}
+	}
+
+	InOutRowNames = OutRowNames;
+}
+
+
+TArray<FName> GetRowsToCompile(const UDataTable& DataTable, const UCustomizableObjectNodeTable& TableNode, FMutableGraphGenerationContext& GenerationContext)
+{
+	TArray<FName> RowNames = GetEnabledRows(DataTable, TableNode);
+
+	if (!RowNames.IsEmpty())
+	{
+		RestrictRowNamesToSelectedOption(RowNames, TableNode, GenerationContext);
+		RestrictRowContentByVersion(RowNames, DataTable, TableNode, GenerationContext);
+	}
+
+	return RowNames;
+}
+
+
 bool GenerateTableColumn(const UCustomizableObjectNodeTable* TableNode, const UEdGraphPin* Pin, mu::TablePtr MutableTable, const FString& DataTableColumnName, const FProperty* ColumnProperty,
 	const int32 LODIndexConnected, const int32 SectionIndexConnected, const int32 LODIndex, const int32 SectionIndex, const bool bOnlyConnectedLOD, FMutableGraphGenerationContext& GenerationContext)
 {
@@ -698,28 +809,18 @@ bool GenerateTableColumn(const UCustomizableObjectNodeTable* TableNode, const UE
 	GenerationContext.AddParticipatingObject(*DataTable);
 
 	// Getting names of the rows to access the information
-	TArray<FName> RowNames = TableNode->GetRowNames(DataTable);
-	RestrictRowNamesToSelectedOption(RowNames, *TableNode, GenerationContext);
+	TArray<FName> RowNames = GetRowsToCompile(*DataTable, *TableNode, GenerationContext);
 
 	for (int32 RowIndex = 0; RowIndex < RowNames.Num(); ++RowIndex)
 	{
-		// Getting Row Data
-		uint8* RowData = DataTable->FindRowUnchecked(RowNames[RowIndex]);
-
-		if (RowData)
+		if (uint8* CellData = GetCellData(RowNames[RowIndex], *DataTable, *ColumnProperty))
 		{
-			// Getting Cell Data
-			uint8* CellData = ColumnProperty->ContainerPtrToValuePtr<uint8>(RowData, 0);
+			bool bCellGenerated = FillTableColumn(TableNode, MutableTable, DataTableColumnName, RowNames[RowIndex].ToString(), RowIndex, CellData, ColumnProperty,
+				LODIndexConnected, SectionIndexConnected, LODIndex, SectionIndex, bOnlyConnectedLOD, GenerationContext);
 
-			if (CellData)
+			if (!bCellGenerated)
 			{
-				bool bCellGenerated = FillTableColumn(TableNode, MutableTable, DataTableColumnName, RowNames[RowIndex].ToString(), RowIndex, CellData, ColumnProperty,
-					LODIndexConnected, SectionIndexConnected, LODIndex, SectionIndex, bOnlyConnectedLOD, GenerationContext);
-
-				if (!bCellGenerated)
-				{
-					return false;
-				}
+				return false;
 			}
 		}
 	}
@@ -737,8 +838,7 @@ void GenerateTableParameterUIData(const UDataTable* DataTable, const UCustomizab
 	if (!GenerationContext.ParameterUIDataMap.Contains(TableNode->ParameterName))
 	{
 		// Getting Table and row names to access the information
-		TArray<FName> RowNames = TableNode->GetRowNames(DataTable);
-		RestrictRowNamesToSelectedOption(RowNames, *TableNode, GenerationContext);
+		TArray<FName> RowNames = GetRowsToCompile(*DataTable, *TableNode, GenerationContext);
 
 		FParameterUIData ParameterUIData(TableNode->ParameterName, TableNode->ParamUIMetadata, EMutableParameterType::Int);
 		ParameterUIData.IntegerParameterGroupType = TableNode->bAddNoneOption ? ECustomizableObjectGroupType::COGT_ONE_OR_NONE : ECustomizableObjectGroupType::COGT_ONE;
@@ -772,23 +872,16 @@ void GenerateTableParameterUIData(const UDataTable* DataTable, const UCustomizab
 			
 			for (int32 NameIndex = 0; NameIndex < RowNames.Num(); ++NameIndex)
 			{
-				// Getting Row Data
-				if (const uint8* RowData = DataTable->FindRowUnchecked(RowNames[NameIndex]))
+				if (uint8* CellData = GetCellData(RowNames[NameIndex], *DataTable, *ColumnProperty))
 				{
-					// Getting Cell Data
-					if (const uint8* CellData = ColumnProperty->ContainerPtrToValuePtr<uint8>(RowData, 0))
-					{
-						// Getting cell value
-						FMutableParamUIMetadata Value = *(FMutableParamUIMetadata*)CellData;
-						ParameterUIDataRef.ArrayIntegerParameterOption.Add(FIntegerParameterUIData(RowNames[NameIndex].ToString(), Value));
-					}
+					FMutableParamUIMetadata Value = *(FMutableParamUIMetadata*)CellData;
+					ParameterUIDataRef.ArrayIntegerParameterOption.Add(FIntegerParameterUIData(RowNames[NameIndex].ToString(), Value));
 				}
 			}
 		}
 		else
 		{
 			GenerationContext.Compiler->CompilerLog(FText::FromString(WrongTypeMessage), TableNode);
-
 			return;
 		}
 	}
@@ -815,8 +908,7 @@ mu::TablePtr GenerateMutableSourceTable(const UDataTable* DataTable, const UCust
 	if (const UScriptStruct* TableStruct = DataTable->GetRowStruct())
 	{
 		// Getting Table and row names to access the information
-		TArray<FName> RowNames = TableNode->GetRowNames(DataTable);
-		RestrictRowNamesToSelectedOption(RowNames, *TableNode, GenerationContext);
+		TArray<FName> RowNames = GetRowsToCompile(*DataTable, *TableNode, GenerationContext);
 
 		// Adding and filling Name Column
 		MutableTable->AddColumn("Name", mu::ETableColumnType::String);
@@ -995,7 +1087,6 @@ UDataTable* GenerateDataTableFromStruct(const UCustomizableObjectNodeTable* Tabl
 }
 
 
-
 void LogRowGenerationMessage(const UCustomizableObjectNodeTable* TableNode, const UDataTable* DataTable, FMutableGraphGenerationContext& GenerationContext, const FString& Message, const FString& RowName)
 {
 	FString FinalMessage = Message;
diff --git a/Engine/Plugins/Experimental/Mutable/Source/CustomizableObjectEditor/Private/MuCOE/GenerateMutableSource/GenerateMutableSourceTable.h b/Engine/Plugins/Experimental/Mutable/Source/CustomizableObjectEditor/Private/MuCOE/GenerateMutableSource/GenerateMutableSourceTable.h
--- a/Engine/Plugins/Experimental/Mutable/Source/CustomizableObjectEditor/Private/MuCOE/GenerateMutableSource/GenerateMutableSourceTable.h
+++ b/Engine/Plugins/Experimental/Mutable/Source/CustomizableObjectEditor/Private/MuCOE/GenerateMutableSource/GenerateMutableSourceTable.h
@@ -65,6 +65,10 @@ void GenerateTableParameterUIData(const UDataTable* DataTable, const UCustomizab
 /** Gets the data table needed during the compilation process */
 UDataTable* GetDataTable(const UCustomizableObjectNodeTable* TableNode, FMutableGraphGenerationContext& GenerationContext);
 
+/** Gets all the rows of a data table that are going to be compiled. Some rows can be disabled with a bool column or an asset version system 
+	@return Array with all the names of the rows that are going to be compiled */
+TArray<FName> GetRowsToCompile(const UDataTable& DataTable, const UCustomizableObjectNodeTable& TableNode, FMutableGraphGenerationContext& GenerationContext);
+
 /** Generates a Data Table from the Data Tables referenced in a Script Struct */
 UDataTable* GenerateDataTableFromStruct(const UCustomizableObjectNodeTable* TableNode, FMutableGraphGenerationContext& GenerationContext);
 
diff --git a/Engine/Plugins/Experimental/Mutable/Source/CustomizableObjectEditor/Private/MuCOE/Nodes/CustomizableObjectNodeTable.cpp b/Engine/Plugins/Experimental/Mutable/Source/CustomizableObjectEditor/Private/MuCOE/Nodes/CustomizableObjectNodeTable.cpp
--- a/Engine/Plugins/Experimental/Mutable/Source/CustomizableObjectEditor/Private/MuCOE/Nodes/CustomizableObjectNodeTable.cpp
+++ b/Engine/Plugins/Experimental/Mutable/Source/CustomizableObjectEditor/Private/MuCOE/Nodes/CustomizableObjectNodeTable.cpp
@@ -287,6 +287,7 @@ void UCustomizableObjectNodeTable::PostEditChangeProperty(FPropertyChangedEvent&
 		if (PropertyThatChanged->GetName() == TEXT("Table") || PropertyThatChanged->GetName() == TEXT("Structure"))
 		{
 			ParamUIMetadataColumn = NAME_None;
+			VersionColumn = NAME_None;
 			ReconstructNode();
 		}
 		else if (PropertyThatChanged->GetName() == TEXT("DefaultImageMode"))
@@ -1140,58 +1141,6 @@ TSoftClassPtr<UAnimInstance> UCustomizableObjectNodeTable::GetAnimInstanceAt(con
 }
 
 
-TArray<FName> UCustomizableObjectNodeTable::GetRowNames(const UDataTable* DataTable) const
-{
-	TArray<FName> RowNames;
-
-	if (DataTable)
-	{
-		const UScriptStruct* TableStruct = DataTable->GetRowStruct();
-
-		if (!TableStruct)
-		{
-			return RowNames;
-		}
-
-		TArray<FName> TableRowNames = DataTable->GetRowNames();
-		FBoolProperty* BoolProperty = nullptr;
-
-		for (TFieldIterator<FProperty> PropertyIt(TableStruct); PropertyIt && bDisableCheckedRows; ++PropertyIt)
-		{
-			BoolProperty = CastField<FBoolProperty>(*PropertyIt);
-
-			if (BoolProperty)
-			{
-				for (const FName& RowName : TableRowNames)
-				{
-					if (uint8* RowData = DataTable->FindRowUnchecked(RowName))
-					{
-						if (uint8* CellData = BoolProperty->ContainerPtrToValuePtr<uint8>(RowData, 0))
-						{
-							if (!BoolProperty->GetPropertyValue(CellData))
-							{
-								RowNames.Add(RowName);
-							}
-						}
-					}
-				}
-
-				// There should be only one Bool column
-				break;
-			}
-		}
-
-		// There is no Bool column or we don't want to disable rows
-		if (!BoolProperty)
-		{
-			return TableRowNames;
-		}
-	}
-
-	return RowNames;
-}
-
-
 void UCustomizableObjectNodeTableImagePinData::ConvertArrayTextureToAnyTextureFixup()
 {
 	if (bIsArrayTexture_DEPRECATED)
diff --git a/Engine/Plugins/Experimental/Mutable/Source/CustomizableObjectEditor/Private/MuCOE/Nodes/CustomizableObjectNodeTableDetails.cpp b/Engine/Plugins/Experimental/Mutable/Source/CustomizableObjectEditor/Private/MuCOE/Nodes/CustomizableObjectNodeTableDetails.cpp
--- a/Engine/Plugins/Experimental/Mutable/Source/CustomizableObjectEditor/Private/MuCOE/Nodes/CustomizableObjectNodeTableDetails.cpp
+++ b/Engine/Plugins/Experimental/Mutable/Source/CustomizableObjectEditor/Private/MuCOE/Nodes/CustomizableObjectNodeTableDetails.cpp
@@ -59,6 +59,28 @@ void FCustomizableObjectNodeTableDetails::CustomizeDetails(const TSharedPtr<IDet
 
 		GenerateMeshColumnComboBoxOptions();
 		TSharedPtr<FString> CurrentMutableMetadataColumn = GenerateMutableMetaDataColumnComboBoxOptions();
+		TSharedPtr<FString> CurrentVersionColumn = GenerateVersionColumnComboBoxOptions();
+
+		CustomizableObjectCategory.AddProperty("ParameterName");
+		CustomizableObjectCategory.AddCustomRow(LOCTEXT("VersionColumn_Selector","VersionColumn"))
+		.NameContent()
+		[
+			SNew(STextBlock)
+			.Text(LOCTEXT("VersionColumn_SelectorText","Version Column"))
+			.ToolTipText(LOCTEXT("VersionColumn_SelectorTooltip","Select the column that contains the version of each row."))
+			.Font(DetailBuilder->GetDetailFont())
+		]
+		.ValueContent()
+		[
+			SAssignNew(VersionColumnsComboBox,STextComboBox)
+			.InitiallySelectedItem(CurrentVersionColumn)
+			.OptionsSource(&VersionColumnsOptionNames)
+			.OnComboBoxOpening(this, &FCustomizableObjectNodeTableDetails::OnOpenVersionColumnComboBox)
+			.OnSelectionChanged(this, &FCustomizableObjectNodeTableDetails::OnVersionColumnComboBoxSelectionChanged)
+			.Font(DetailBuilder->GetDetailFont())
+			.ColorAndOpacity(this, &FCustomizableObjectNodeTableDetails::GetVersionColumnComboBoxTextColor, &VersionColumnsOptionNames)
+		]
+		.OverrideResetToDefault(FResetToDefaultOverride::Create(FSimpleDelegate::CreateSP(this, &FCustomizableObjectNodeTableDetails::OnVersionColumnComboBoxSelectionReset)));
 
 		UICategory.AddCustomRow(LOCTEXT("MutableUIMetadataColumn_Selector","MutableUIMetadataColumn"))
 		.NameContent()
@@ -796,6 +818,95 @@ void FCustomizableObjectNodeTableDetails::OnMutableMetaDataColumnComboBoxSelecti
 }
 
 
+TSharedPtr<FString> FCustomizableObjectNodeTableDetails::GenerateVersionColumnComboBoxOptions()
+{
+	const UScriptStruct* TableStruct = Node->GetTableNodeStruct();
+	TSharedPtr<FString> CurrentSelection;
+	VersionColumnsOptionNames.Reset();
+
+	if (!TableStruct)
+	{
+		return CurrentSelection;
+	}
+
+	// Iterating struct Options
+	for (TFieldIterator<FProperty> It(TableStruct); It; ++It)
+	{
+		FProperty* ColumnProperty = *It;
+
+		if (!ColumnProperty)
+		{
+			continue;
+		}
+
+		TSharedPtr<FString> Option = MakeShareable(new FString(DataTableUtils::GetPropertyExportName(ColumnProperty)));
+		VersionColumnsOptionNames.Add(Option);
+
+		if (*Option == Node->VersionColumn)
+		{
+			CurrentSelection = VersionColumnsOptionNames.Last();
+		}
+	}
+
+	if (!Node->VersionColumn.IsNone() && !CurrentSelection)
+	{
+		VersionColumnsOptionNames.Add(MakeShareable(new FString(Node->VersionColumn.ToString())));
+		CurrentSelection = VersionColumnsOptionNames.Last();
+	}
+
+	return CurrentSelection;
+}
+
+
+void FCustomizableObjectNodeTableDetails::OnOpenVersionColumnComboBox()
+{
+	TSharedPtr<FString> CurrentSelection = GenerateVersionColumnComboBoxOptions();
+
+	if (VersionColumnsComboBox.IsValid())
+	{
+		VersionColumnsComboBox->ClearSelection();
+		VersionColumnsComboBox->RefreshOptions();
+		VersionColumnsComboBox->SetSelectedItem(CurrentSelection);
+	}
+}
+
+
+void FCustomizableObjectNodeTableDetails::OnVersionColumnComboBoxSelectionChanged(TSharedPtr<FString> Selection, ESelectInfo::Type SelectInfo)
+{
+	if (Selection && Node->VersionColumn != FName(*Selection)
+		&& (SelectInfo == ESelectInfo::OnKeyPress || SelectInfo == ESelectInfo::OnMouseClick))
+	{
+		Node->VersionColumn = FName(*Selection);
+		Node->MarkPackageDirty();
+	}
+}
+
+
+FSlateColor FCustomizableObjectNodeTableDetails::GetVersionColumnComboBoxTextColor(TArray<TSharedPtr<FString>>* CurrentOptions) const
+{
+	if (Node->FindTableProperty(Node->GetTableNodeStruct(), Node->VersionColumn) || Node->VersionColumn.IsNone())
+	{
+		return FSlateColor::UseForeground();
+	}
+
+	// Table Struct null or does not contain the selected property anymore
+	return FSlateColor(FLinearColor(0.9f, 0.05f, 0.05f, 1.0f));
+}
+
+
+void FCustomizableObjectNodeTableDetails::OnVersionColumnComboBoxSelectionReset()
+{
+	Node->VersionColumn = NAME_None;
+
+	if (VersionColumnsComboBox.IsValid())
+	{
+		GenerateVersionColumnComboBoxOptions();
+		VersionColumnsComboBox->ClearSelection();
+		VersionColumnsComboBox->RefreshOptions();
+	}
+}
+
+
 EVisibility FCustomizableObjectNodeTableDetails::LayoutOptionsVisibility() const
 {
 	return SelectedLayout.IsValid() ? EVisibility::Visible : EVisibility::Collapsed;
diff --git a/Engine/Plugins/Experimental/Mutable/Source/CustomizableObjectEditor/Private/MuCOE/Nodes/CustomizableObjectNodeTableDetails.h b/Engine/Plugins/Experimental/Mutable/Source/CustomizableObjectEditor/Private/MuCOE/Nodes/CustomizableObjectNodeTableDetails.h
--- a/Engine/Plugins/Experimental/Mutable/Source/CustomizableObjectEditor/Private/MuCOE/Nodes/CustomizableObjectNodeTableDetails.h
+++ b/Engine/Plugins/Experimental/Mutable/Source/CustomizableObjectEditor/Private/MuCOE/Nodes/CustomizableObjectNodeTableDetails.h
@@ -7,6 +7,7 @@
 #include "UObject/WeakObjectPtr.h"
 
 namespace ESelectInfo { enum Type : int; }
+
 class FString;
 class IDetailLayoutBuilder;
 class FReply;
@@ -75,6 +76,10 @@ private:
 	// Returns the current selected option or a null pointer
 	TSharedPtr<FString> GenerateMutableMetaDataColumnComboBoxOptions();
 
+	// Generates MutableMetadata columns combobox options
+	// Returns the current selected option or a null pointer
+	TSharedPtr<FString> GenerateVersionColumnComboBoxOptions();
+
 	// Callback to regenerate the combobox options
 	void OnOpenMutableMetadataComboBox();
 
@@ -87,6 +92,18 @@ private:
 	// OnComboBoxSelectionChanged Callback for Layout ComboBox
 	void OnMutableMetaDataColumnComboBoxSelectionReset();
 
+	// Callback to regenerate the combobox options
+	void OnOpenVersionColumnComboBox();
+	
+	// OnComboBoxSelectionChanged Callback for Layout ComboBox
+	void OnVersionColumnComboBoxSelectionChanged(TSharedPtr<FString> Selection, ESelectInfo::Type SelectInfo);
+	
+	// Sets the combo box selection color
+	FSlateColor GetVersionColumnComboBoxTextColor(TArray<TSharedPtr<FString>>* CurrentOptions) const;
+	
+	// OnComboBoxSelectionChanged Callback for Layout ComboBox
+	void OnVersionColumnComboBoxSelectionReset();
+
 	// Layout options visibility
 	EVisibility LayoutOptionsVisibility() const;
 
@@ -142,6 +159,12 @@ private:
 
 	// ComboBox widget to select a MutableMetaDatacolumn from the NodeTable
 	TSharedPtr<STextComboBox> MutableMetaDataComboBox;
+
+	// Array with the name of the Version columns
+	TArray<TSharedPtr<FString>> VersionColumnsOptionNames;
+	
+	// ComboBox widget to select a VersionColumn from the NodeTable
+	TSharedPtr<STextComboBox> VersionColumnsComboBox;
 	
 	// Button to clear selections of the animation comboboxes
 	TSharedPtr<SButton> ClearButton;
