[mutable] Fixed static analysis Issue 440620 in file GenerateMutableSource.cpp.
#rnx
#tests Editor

diff --git a/Engine/Plugins/Experimental/Mutable/Source/CustomizableObjectEditor/Private/MuCOE/GenerateMutableSource/GenerateMutableSource.cpp b/Engine/Plugins/Experimental/Mutable/Source/CustomizableObjectEditor/Private/MuCOE/GenerateMutableSource/GenerateMutableSource.cpp
--- a/Engine/Plugins/Experimental/Mutable/Source/CustomizableObjectEditor/Private/MuCOE/GenerateMutableSource/GenerateMutableSource.cpp
+++ b/Engine/Plugins/Experimental/Mutable/Source/CustomizableObjectEditor/Private/MuCOE/GenerateMutableSource/GenerateMutableSource.cpp
@@ -1668,10 +1668,8 @@ int32 GetPlatformLODBias(int32 TextureSize, int32 NumMips, int32 MaxPlatformSize
 }
 
 
-uint32 ComputeLODBiasForTexture(const FMutableGraphGenerationContext& GenerationContext, const UTexture2D* Texture, const UTexture2D* ReferenceTexture, int32 BaseTextureSize)
+uint32 ComputeLODBiasForTexture(const FMutableGraphGenerationContext& GenerationContext, const UTexture2D& Texture, const UTexture2D* ReferenceTexture, int32 BaseTextureSize)
 {
-	ensure(Texture);
-
 	constexpr int32 MaxAllowedLODBias = 6;
 
 	// Force a large LODBias for debug
@@ -1681,7 +1679,7 @@ uint32 ComputeLODBiasForTexture(const FMutableGraphGenerationContext& Generation
 	}
 
 	// Max size and number of mips from Texture. 
-	const int32 SourceSize = (int32)FMath::Max3(Texture->Source.GetSizeX(),Texture->Source.GetSizeY(),(int64)1);
+	const int32 SourceSize = (int32)FMath::Max3(Texture.Source.GetSizeX(),Texture.Source.GetSizeY(),(int64)1);
 	const int32 NumMipsSource = FMath::CeilLogTwo(SourceSize) + 1;
 
 	// When the BaseTextureSize is known, skip mips until the texture is equal or smaller.
@@ -1699,7 +1697,7 @@ uint32 ComputeLODBiasForTexture(const FMutableGraphGenerationContext& Generation
 	const UTextureLODSettings& LODSettings = GenerationContext.Options.TargetPlatform->GetTextureLODSettings();
 
 	// Get the MaxTextureSize for the TargetPlatform.
-	const int32 MaxTextureSize = GetMaxTextureSize(ReferenceTexture ? *ReferenceTexture : *Texture, LODSettings);
+	const int32 MaxTextureSize = GetMaxTextureSize(ReferenceTexture ? *ReferenceTexture : Texture, LODSettings);
 
 	if (ReferenceTexture)
 	{
@@ -1727,7 +1725,7 @@ uint32 ComputeLODBiasForTexture(const FMutableGraphGenerationContext& Generation
 	const int64 TextureSize = SourceSize >> PlatformLODBias;
 
 	// Additional LODBias of the Texture
-	const int32 TextureLODBias = LODSettings.CalculateLODBias(TextureSize, TextureSize, 0, Texture->LODGroup, Texture->LODBias, 0, Texture->MipGenSettings, Texture->IsCurrentlyVirtualTextured());
+	const int32 TextureLODBias = LODSettings.CalculateLODBias(TextureSize, TextureSize, 0, Texture.LODGroup, Texture.LODBias, 0, Texture.MipGenSettings, Texture.IsCurrentlyVirtualTextured());
 
 	return FMath::Max(PlatformLODBias + TextureLODBias, 0);
 }
diff --git a/Engine/Plugins/Experimental/Mutable/Source/CustomizableObjectEditor/Private/MuCOE/GenerateMutableSource/GenerateMutableSource.h b/Engine/Plugins/Experimental/Mutable/Source/CustomizableObjectEditor/Private/MuCOE/GenerateMutableSource/GenerateMutableSource.h
--- a/Engine/Plugins/Experimental/Mutable/Source/CustomizableObjectEditor/Private/MuCOE/GenerateMutableSource/GenerateMutableSource.h
+++ b/Engine/Plugins/Experimental/Mutable/Source/CustomizableObjectEditor/Private/MuCOE/GenerateMutableSource/GenerateMutableSource.h
@@ -864,7 +864,7 @@ uint32 GetBaseTextureSize(const FMutableGraphGenerationContext& GenerationContex
 
 // Computes the LOD bias for a texture given the current mesh LOD and automatic LOD settings, the reference texture settings
 // and whether it's being built for a server or not
-uint32 ComputeLODBiasForTexture(const FMutableGraphGenerationContext& GenerationContext, const UTexture2D* Texture,const UTexture2D* ReferenceTexture = nullptr, int32 MaxTextureSizeInGame = 0);
+uint32 ComputeLODBiasForTexture(const FMutableGraphGenerationContext& GenerationContext, const UTexture2D& Texture, const UTexture2D* ReferenceTexture = nullptr, int32 MaxTextureSizeInGame = 0);
 
 // Max texture size to set on the ImageProperties
 int32 GetMaxTextureSize(const UTexture2D& ReferenceTexture, const UTextureLODSettings& LODSettings);
diff --git a/Engine/Plugins/Experimental/Mutable/Source/CustomizableObjectEditor/Private/MuCOE/GenerateMutableSource/GenerateMutableSourceGroupProjector.cpp b/Engine/Plugins/Experimental/Mutable/Source/CustomizableObjectEditor/Private/MuCOE/GenerateMutableSource/GenerateMutableSourceGroupProjector.cpp
--- a/Engine/Plugins/Experimental/Mutable/Source/CustomizableObjectEditor/Private/MuCOE/GenerateMutableSource/GenerateMutableSourceGroupProjector.cpp
+++ b/Engine/Plugins/Experimental/Mutable/Source/CustomizableObjectEditor/Private/MuCOE/GenerateMutableSource/GenerateMutableSourceGroupProjector.cpp
@@ -480,7 +480,7 @@ bool GenerateMutableSourceGroupProjector(const UEdGraphPin* Pin, FMutableGraphGe
 					mu::NodeImageConstantPtr ImageNode = new mu::NodeImageConstant();
 					ImageNode->SetValue(ImageConstant.get());
 
-					const uint32 MipsToSkip = ComputeLODBiasForTexture(GenerationContext, Texture, ProjParamNode->ReferenceTexture) + AdditionalLODBias;
+					const uint32 MipsToSkip = ComputeLODBiasForTexture(GenerationContext, *Texture, ProjParamNode->ReferenceTexture) + AdditionalLODBias;
 					SwitchNode->SetOption(SelectorIndex, ResizeTextureByNumMips(ImageNode, MipsToSkip));
 				}
 				else
diff --git a/Engine/Plugins/Experimental/Mutable/Source/CustomizableObjectEditor/Private/MuCOE/GenerateMutableSource/GenerateMutableSourceImage.cpp b/Engine/Plugins/Experimental/Mutable/Source/CustomizableObjectEditor/Private/MuCOE/GenerateMutableSource/GenerateMutableSourceImage.cpp
--- a/Engine/Plugins/Experimental/Mutable/Source/CustomizableObjectEditor/Private/MuCOE/GenerateMutableSource/GenerateMutableSourceImage.cpp
+++ b/Engine/Plugins/Experimental/Mutable/Source/CustomizableObjectEditor/Private/MuCOE/GenerateMutableSource/GenerateMutableSourceImage.cpp
@@ -171,7 +171,7 @@ mu::NodeImagePtr GenerateMutableSourceImage(const UEdGraphPin* Pin, FMutableGrap
 				GenerationContext.GeneratedImages.Add(ImageKey, ImageNode);
 			}
 
-			const uint32 MipsToSkip = ComputeLODBiasForTexture(GenerationContext, BaseTexture, nullptr, ReferenceTextureSize);
+			const uint32 MipsToSkip = ComputeLODBiasForTexture(GenerationContext, *BaseTexture, nullptr, ReferenceTextureSize);
 			Result = ResizeTextureByNumMips(ImageNode, MipsToSkip);
 		}
 		else
@@ -224,7 +224,7 @@ mu::NodeImagePtr GenerateMutableSourceImage(const UEdGraphPin* Pin, FMutableGrap
 		{
 			GenerationContext.AddParticipatingObject(*ReferenceTexture);
 
-			const uint32 LODBias = ComputeLODBiasForTexture(GenerationContext, TypedNodeParam->ReferenceValue, ReferenceTexture, ReferenceTextureSize);
+			const uint32 LODBias = ComputeLODBiasForTexture(GenerationContext, *TypedNodeParam->ReferenceValue, ReferenceTexture, ReferenceTextureSize);
 			TextureSize.X = FMath::Max(ReferenceTexture->Source.GetSizeX() >> LODBias, 1);
 			TextureSize.Y = FMath::Max(ReferenceTexture->Source.GetSizeY() >> LODBias, 1);
 		}
@@ -260,7 +260,7 @@ mu::NodeImagePtr GenerateMutableSourceImage(const UEdGraphPin* Pin, FMutableGrap
 		{
 			ImageNode->SetValue(GenerateImageConstant(Texture, GenerationContext, false).get());
 
-			const uint32 MipsToSkip = ComputeLODBiasForTexture(GenerationContext, Texture, nullptr, ReferenceTextureSize);
+			const uint32 MipsToSkip = ComputeLODBiasForTexture(GenerationContext, *Texture, nullptr, ReferenceTextureSize);
 			Result = ResizeTextureByNumMips(ImageNode, MipsToSkip);
 		}
 		else
@@ -1096,7 +1096,7 @@ mu::NodeImagePtr GenerateMutableSourceImage(const UEdGraphPin* Pin, FMutableGrap
 							{
 								mu::FImageDesc ImageDesc = GenerateImageDescriptor(DefaultTexture2D);
 
-								uint32 LODBias = ReferenceTextureSize > 0 ? ComputeLODBiasForTexture(GenerationContext, DefaultTexture2D, nullptr, ReferenceTextureSize) : 0;
+								uint32 LODBias = ReferenceTextureSize > 0 ? ComputeLODBiasForTexture(GenerationContext, *DefaultTexture2D, nullptr, ReferenceTextureSize) : 0;
 								ImageDesc.m_size[0] = ImageDesc.m_size[0] >> LODBias;
 								ImageDesc.m_size[1] = ImageDesc.m_size[1] >> LODBias;
 								
diff --git a/Engine/Plugins/Experimental/Mutable/Source/CustomizableObjectEditor/Private/MuCOE/GenerateMutableSource/GenerateMutableSourceSurface.cpp b/Engine/Plugins/Experimental/Mutable/Source/CustomizableObjectEditor/Private/MuCOE/GenerateMutableSource/GenerateMutableSourceSurface.cpp
--- a/Engine/Plugins/Experimental/Mutable/Source/CustomizableObjectEditor/Private/MuCOE/GenerateMutableSource/GenerateMutableSourceSurface.cpp
+++ b/Engine/Plugins/Experimental/Mutable/Source/CustomizableObjectEditor/Private/MuCOE/GenerateMutableSource/GenerateMutableSourceSurface.cpp
@@ -613,7 +613,7 @@ mu::NodeSurfacePtr GenerateMutableSourceSurface(const UEdGraphPin * Pin, FMutabl
 										mu::Ptr<mu::Image> ImageConstant = GenerateImageConstant(Texture2D, GenerationContext, false);
 										ConstImageNode->SetValue(ImageConstant.get());
 
-										const uint32 MipsToSkip = ComputeLODBiasForTexture(GenerationContext, Texture2D, nullptr, Props.TextureSize);
+										const uint32 MipsToSkip = ComputeLODBiasForTexture(GenerationContext, *Texture2D, nullptr, Props.TextureSize);
 										return ResizeTextureByNumMips(ConstImageNode, MipsToSkip);
 									}
 									else
@@ -637,7 +637,7 @@ mu::NodeSurfacePtr GenerateMutableSourceSurface(const UEdGraphPin * Pin, FMutabl
 						{
 							// Apply base LODBias. It will be propagated to most images.
 							const uint32 FirstLODAvailable = GenerationContext.Options.bUseLODAsBias ? GenerationContext.FirstLODAvailable : 0;
-							const uint32 BaseLODBias = ComputeLODBiasForTexture(GenerationContext, ReferenceTexture) + FirstLODAvailable;
+							const uint32 BaseLODBias = ComputeLODBiasForTexture(GenerationContext, *ReferenceTexture) + FirstLODAvailable;
 							mu::NodeImagePtr LastImage = ResizeTextureByNumMips(ImageNode, BaseLODBias);
 
 							mu::NodeImageMipmapPtr MipmapImage = new mu::NodeImageMipmap();
@@ -681,7 +681,7 @@ mu::NodeSurfacePtr GenerateMutableSourceSurface(const UEdGraphPin * Pin, FMutabl
 									CompositeNormalImage->SetValue(ImageConstant.get());
 
 									mu::NodeImageMipmapPtr NormalCompositeMipmapImage = new mu::NodeImageMipmap();
-									const uint32 MipsToSkip = ComputeLODBiasForTexture(GenerationContext, ReferenceCompositeNormalTexture, ReferenceTexture);
+									const uint32 MipsToSkip = ComputeLODBiasForTexture(GenerationContext, *ReferenceCompositeNormalTexture, ReferenceTexture);
 									NormalCompositeMipmapImage->SetSource(ResizeTextureByNumMips(CompositeNormalImage, MipsToSkip));
 									NormalCompositeMipmapImage->SetMipmapGenerationSettings(mu::EMipmapFilterType::MFT_SimpleAverage, mu::EAddressMode::None, 1.0f, true);
 
