[mutable] Fix image upscaling filtering.
#rnx

diff --git a/Engine/Plugins/Experimental/Mutable/Source/MutableRuntime/Private/MuR/OpImageResize.cpp b/Engine/Plugins/Experimental/Mutable/Source/MutableRuntime/Private/MuR/OpImageResize.cpp
--- a/Engine/Plugins/Experimental/Mutable/Source/MutableRuntime/Private/MuR/OpImageResize.cpp
+++ b/Engine/Plugins/Experimental/Mutable/Source/MutableRuntime/Private/MuR/OpImageResize.cpp
@@ -46,31 +46,21 @@ namespace mu
 			const uint8* BaseBuf = Base->GetData() + y * BaseSizeX * NC;
 			uint8* DestBuf = Dest->GetData() + y * DestSizeX * NC;
 
+			uint32 LastPX = BaseSizeX - 1;
 			for (int32 x = 0; x < DestSizeX; ++x)
 			{
 				uint32 px = px_16 >> 16;
 				uint32 epx_16 = px_16 + dx_16;
 
-				if ((px_16 & 0xffff0000) == ((epx_16 - 1) & 0xffff0000))
-				{
-					// One fraction
-					for (int32 c = 0; c < NC; ++c)
-					{
-						DestBuf[c] = BaseBuf[px * NC + c];
-					}
-				}
-				else
-				{
-					// Two fractions
-					uint32 frac1 = (px_16 & 0xffff);
-					uint32 frac0 = 0x10000 - frac1;
-
-					for (int32 c = 0; c < NC; ++c)
-					{
-						DestBuf[c] = uint8((BaseBuf[px * NC + c] * frac0 + BaseBuf[(px + 1) * NC + c] * frac1) >> 16);
-					}
+				uint32 NextPX = FMath::Min(px + 1, LastPX);
+				
+				// Two fractions
+				uint32 frac1 = (px_16 & 0xffff);
+				uint32 frac0 = 0x10000 - frac1;
 
-					++px;
+				for (int32 c = 0; c < NC; ++c)
+				{
+					DestBuf[c] = uint8((BaseBuf[px * NC + c] * frac0 + BaseBuf[NextPX * NC + c] * frac1) >> 16);
 				}
 
 				px_16 = epx_16;
@@ -436,11 +426,13 @@ namespace mu
 		{
 			uint32 dy_16 = (uint32(BaseSizeY) << 16) / DestSizeY;
 
+			uint32 LastPY = BaseSizeY - 1;
+
 			// Linear filtering
 			// \todo: optimise: swap loops, etc.
 			//for ( int32 x=0; x<SizeX; ++x )
 			const auto& ProcessColumn = [
-				Dest, Base, SizeX, DestSizeY, dy_16
+				Dest, Base, SizeX, DestSizeY, LastPY, dy_16
 			] (uint32 x)
 			{
 				uint32 py_16 = 0;
@@ -452,28 +444,17 @@ namespace mu
 					uint32 py = py_16 >> 16;
 					uint32 epy_16 = py_16 + dy_16;
 
-					if ((py_16 & 0xffff0000) == ((epy_16 - 1) & 0xffff0000))
-					{
-						// One fraction
-						for (int32 c = 0; c < NC; ++c)
-						{
-							DestBuf[c] = BaseBuf[(py * SizeX + x) * NC + c];
-						}
-					}
-					else
-					{
-						// Two fractions
-						uint32 frac1 = (py_16 & 0xffff);
-						uint32 frac0 = 0x10000 - frac1;
+					uint32 NextPY = FMath::Min(py + 1, LastPY);
 
-						for (int32 c = 0; c < NC; ++c)
-						{
-							DestBuf[c] = (uint8)((BaseBuf[(py * SizeX + x) * NC + c] * frac0 +
-								BaseBuf[((py + 1) * SizeX + x) * NC + c] * frac1
-								) >> 16);
-						}
+					// Two fractions
+					uint32 frac1 = (py_16 & 0xffff);
+					uint32 frac0 = 0x10000 - frac1;
 
-						++py;
+					for (int32 c = 0; c < NC; ++c)
+					{
+						DestBuf[c] = (uint8)((BaseBuf[(py * SizeX + x) * NC + c] * frac0 +
+							BaseBuf[(NextPY * SizeX + x) * NC + c] * frac1
+							) >> 16);
 					}
 
 					py_16 = epy_16;
@@ -481,7 +462,7 @@ namespace mu
 				}
 			};
 
-				ParallelFor(SizeX, ProcessColumn);
+			ParallelFor(SizeX, ProcessColumn);
 		}
 	}
 
