[mutable] Swizzle of an alpha channel from a texture that doesn't have one will return plain white instead of black.
#rnx

diff --git a/Engine/Plugins/Experimental/Mutable/Source/MutableRuntime/Private/MuR/OpImageSwizzle.cpp b/Engine/Plugins/Experimental/Mutable/Source/MutableRuntime/Private/MuR/OpImageSwizzle.cpp
--- a/Engine/Plugins/Experimental/Mutable/Source/MutableRuntime/Private/MuR/OpImageSwizzle.cpp
+++ b/Engine/Plugins/Experimental/Mutable/Source/MutableRuntime/Private/MuR/OpImageSwizzle.cpp
@@ -73,7 +73,7 @@ namespace mu
 				}
 			}
 
-			bool filled = false;
+			bool bFilled = false;
 
 			constexpr int32 NumBatchElems = 4096*2;
 			const int32 NumBatches = FMath::DivideAndRoundUp(PixelCount, NumBatchElems);
@@ -107,7 +107,7 @@ namespace mu
 							ParallelFor(NumBatches, ProcessBatch);
 						}
 
-						filled = true;
+						bFilled = true;
 					}
 					break;
 
@@ -134,7 +134,7 @@ namespace mu
 							ParallelFor(NumBatches, ProcessBatch);
 						}
 
-						filled = true;
+						bFilled = true;
 					}
 					break;
 
@@ -161,7 +161,7 @@ namespace mu
 							ParallelFor(NumBatches, ProcessBatch);
 						}
 
-						filled = true;
+						bFilled = true;
 					}
 					break;
 
@@ -195,7 +195,7 @@ namespace mu
 							ParallelFor(NumBatches, ProcessBatch);
 						}
 
-						filled = true;
+						bFilled = true;
 					}
 					break;
 
@@ -204,16 +204,18 @@ namespace mu
 				}
 			}
 
-			if (!filled)
+			if (!bFilled)
 			{
-				// Source not set. Clear to 0
-				auto ProcessBatch = [pDestBuf, NumDestChannels, PixelCount, NumBatchElems](int32 BatchId)
+				// Source not set. Clear to 0 if it is a color channel, or to 255 if it is alpha
+				uint8 Value = (Channels[i] == 3) ? 255 : 0;
+
+				auto ProcessBatch = [pDestBuf, NumDestChannels, PixelCount, NumBatchElems, Value](int32 BatchId)
 				{
 					const int32 BatchBegin = BatchId * NumBatchElems;
 					const int32 BatchEnd = FMath::Min(BatchBegin + NumBatchElems, PixelCount);
 					for (int32 p = BatchBegin; p < BatchEnd; ++p)
 					{
-						pDestBuf[p * NumDestChannels] = 0;
+						pDestBuf[p * NumDestChannels] = Value;
 					}
 				};
 
