[Mutable] Store Image LODs in different buffers.

#rnx

diff --git a/Engine/Plugins/Experimental/Mutable/Source/CustomizableObject/Internal/MuCO/CustomizableObjectPrivate.h b/Engine/Plugins/Experimental/Mutable/Source/CustomizableObject/Internal/MuCO/CustomizableObjectPrivate.h
--- a/Engine/Plugins/Experimental/Mutable/Source/CustomizableObject/Internal/MuCO/CustomizableObjectPrivate.h
+++ b/Engine/Plugins/Experimental/Mutable/Source/CustomizableObject/Internal/MuCO/CustomizableObjectPrivate.h
@@ -811,6 +811,6 @@ public:
 	// This is a manual version number for the binary blobs in this asset.
 	// Increasing it invalidates all the previously compiled models.
 	// Warning: If while merging code both versions have changed, take the highest+1.
-	static constexpr int32 CurrentSupportedVersion = 432;
+	static constexpr int32 CurrentSupportedVersion = 435;
 };
 
diff --git a/Engine/Plugins/Experimental/Mutable/Source/CustomizableObject/Private/MuCO/BusyWaits_Deprecated.cpp b/Engine/Plugins/Experimental/Mutable/Source/CustomizableObject/Private/MuCO/BusyWaits_Deprecated.cpp
--- a/Engine/Plugins/Experimental/Mutable/Source/CustomizableObject/Private/MuCO/BusyWaits_Deprecated.cpp
+++ b/Engine/Plugins/Experimental/Mutable/Source/CustomizableObject/Private/MuCO/BusyWaits_Deprecated.cpp
@@ -110,11 +110,19 @@ namespace impl_deprecated
 
 					// Force the right number of mips. The missing data will be black.
 					const mu::Ptr<mu::Image> NewImage = new mu::Image(Image.Image->GetSizeX(), Image.Image->GetSizeY(), FullMipCount, Image.Image->GetFormat(), mu::EInitializationType::Black);
-					check(NewImage);
-					if (NewImage->GetDataSize() >= Image.Image->GetDataSize())
+					check(NewImage);	
+					// Formats with BytesPerBlock == 0 will not allocate memory. This type of images are not expected here.
+					check(!NewImage->DataStorage.IsEmpty());
+
+					for (int32 L = 0; L < RealMipCount; ++L)
 					{
-						FMemory::Memcpy(NewImage->GetData(), Image.Image->GetData(), Image.Image->GetDataSize());
+						TArrayView<uint8> DestView = NewImage->DataStorage.GetLOD(L);
+						TArrayView<const uint8> SrcView = Image.Image->DataStorage.GetLOD(L);
+
+						check(DestView.Num() == SrcView.Num());
+						FMemory::Memcpy(DestView.GetData(), SrcView.GetData(), DestView.Num());
 					}
+
 					Image.Image = NewImage;
 				}
 
diff --git a/Engine/Plugins/Experimental/Mutable/Source/CustomizableObject/Private/MuCO/CustomizableObjectInstance.cpp b/Engine/Plugins/Experimental/Mutable/Source/CustomizableObject/Private/MuCO/CustomizableObjectInstance.cpp
--- a/Engine/Plugins/Experimental/Mutable/Source/CustomizableObject/Private/MuCO/CustomizableObjectInstance.cpp
+++ b/Engine/Plugins/Experimental/Mutable/Source/CustomizableObject/Private/MuCO/CustomizableObjectInstance.cpp
@@ -2957,7 +2957,6 @@ FTexturePlatformData* MutableCreateImagePlatformData(mu::Ptr<const mu::Image> Mu
 	PlatformData->PixelFormat = PlatformFormat;
 
 	// Allocate mipmaps.
-	const uint8_t* MutableData = MutableImage->GetMipData( FirstLOD );
 
 	if (!FMath::IsPowerOfTwo(SizeX) || !FMath::IsPowerOfTwo(SizeY))
 	{
@@ -2991,10 +2990,12 @@ FTexturePlatformData* MutableCreateImagePlatformData(mu::Ptr<const mu::Image> Mu
 			Mip->BulkData.Lock(LOCK_READ_WRITE);
 			Mip->BulkData.ClearBulkDataFlags(BULKDATA_SingleUse);
 
-			uint32 SourceDataSize = MutableImage->GetLODDataSize(MipLevelMutable);
+			const uint8* MutableData = MutableImage->GetLODData(MipLevelMutable);
+			const uint32 SourceDataSize = MutableImage->GetLODDataSize(MipLevelMutable);
+
 			uint32 DestDataSize = (MutableFormat == mu::EImageFormat::IF_RGB_UBYTE)
-				? (SourceDataSize/3*4)
-				: SourceDataSize;
+					? (SourceDataSize/3) * 4
+					: SourceDataSize;
 			void* pData = Mip->BulkData.Realloc(DestDataSize);
 
 			// Special inefficient cases
@@ -3035,8 +3036,6 @@ FTexturePlatformData* MutableCreateImagePlatformData(mu::Ptr<const mu::Image> Mu
 				FMemory::Memcpy(pData, MutableData, SourceDataSize);
 			}
 
-			MutableData += SourceDataSize;
-
 			Mip->BulkData.Unlock();
 		}
 		else
@@ -3113,12 +3112,12 @@ void ConvertImage(UTexture2D* Texture, mu::Ptr<const mu::Image> MutableImage, co
 		// Expand the image.
 		mu::ImagePtr Converted = new mu::Image(MutableImage->GetSizeX(), MutableImage->GetSizeY(), MutableImage->GetLODCount(), mu::EImageFormat::IF_RGBA_UBYTE, mu::EInitializationType::NotInitialized);
 
-		for (int LODIndex = 0; LODIndex < Converted->GetLODCount(); ++LODIndex)
+		for (int32 LODIndex = 0; LODIndex < Converted->GetLODCount(); ++LODIndex)
 		{
-			int32 pixelCount = MutableImage->GetLODDataSize(LODIndex)/3;
-			const uint8_t* pSource = MutableImage->GetMipData(LODIndex);
-			uint8_t* pTarget = Converted->GetMipData(LODIndex);
-			for (int32 p = 0; p < pixelCount; ++p)
+			int32 PixelCount = MutableImage->GetLODDataSize(LODIndex)/3;
+			const uint8* pSource = MutableImage->GetMipData(LODIndex);
+			uint8* pTarget = Converted->GetMipData(LODIndex);
+			for (int32 p = 0; p < PixelCount; ++p)
 			{
 				pTarget[4 * p + 0] = pSource[3 * p + 0];
 				pTarget[4 * p + 1] = pSource[3 * p + 1];
@@ -3138,11 +3137,11 @@ void ConvertImage(UTexture2D* Texture, mu::Ptr<const mu::Image> MutableImage, co
 		// Swizzle the image.
 		// \TODO: Raise a warning?
 		mu::ImagePtr Converted = new mu::Image(MutableImage->GetSizeX(), MutableImage->GetSizeY(), 1, mu::EImageFormat::IF_RGBA_UBYTE, mu::EInitializationType::NotInitialized);
-		int32 pixelCount = MutableImage->GetSizeX() * MutableImage->GetSizeY();
+		int32 PixelCount = MutableImage->GetSizeX() * MutableImage->GetSizeY();
 
-		const uint8_t* pSource = MutableImage->GetData();
-		uint8_t* pTarget = Converted->GetData();
-		for (int32 p = 0; p<pixelCount; ++p)
+		const uint8* pSource = MutableImage->GetLODData(0);
+		uint8* pTarget = Converted->GetLODData(0);
+		for (int32 p = 0; p < PixelCount; ++p)
 		{
 			pTarget[4 * p + 0] = pSource[4 * p + 2];
 			pTarget[4 * p + 1] = pSource[4 * p + 1];
diff --git a/Engine/Plugins/Experimental/Mutable/Source/CustomizableObject/Private/MuCO/CustomizableObjectMipDataProvider.cpp b/Engine/Plugins/Experimental/Mutable/Source/CustomizableObject/Private/MuCO/CustomizableObjectMipDataProvider.cpp
--- a/Engine/Plugins/Experimental/Mutable/Source/CustomizableObject/Private/MuCO/CustomizableObjectMipDataProvider.cpp
+++ b/Engine/Plugins/Experimental/Mutable/Source/CustomizableObject/Private/MuCO/CustomizableObjectMipDataProvider.cpp
@@ -212,11 +212,19 @@ namespace impl
 
 						// Force the right number of mips. The missing data will be black.
 						mu::Ptr<mu::Image> NewImage = new mu::Image(Image->GetSizeX(), Image->GetSizeY(), FullMipCount, Image->GetFormat(), mu::EInitializationType::Black);
-						check(NewImage);
-						if (NewImage->GetDataSize() >= Image->GetDataSize())
+
+						// Formats with BytesPerBlock == 0 will not allocate memory. This type of images are not expected here.
+						check(!NewImage->DataStorage.IsEmpty());
+
+						for (int32 L = 0; L < RealMipCount; ++L)
 						{
-							FMemory::Memcpy(NewImage->GetData(), Image->GetData(), Image->GetDataSize());
+							TArrayView<uint8> DestView = NewImage->DataStorage.GetLOD(L);
+							TArrayView<const uint8> SrcView = Image->DataStorage.GetLOD(L);
+
+							check(DestView.Num() == SrcView.Num());
+							FMemory::Memcpy(DestView.GetData(), SrcView.GetData(), DestView.Num());
 						}
+
 						Image = NewImage;
 					}
 
diff --git a/Engine/Plugins/Experimental/Mutable/Source/CustomizableObject/Private/MuCO/UnrealMutableImageProvider.cpp b/Engine/Plugins/Experimental/Mutable/Source/CustomizableObject/Private/MuCO/UnrealMutableImageProvider.cpp
--- a/Engine/Plugins/Experimental/Mutable/Source/CustomizableObject/Private/MuCO/UnrealMutableImageProvider.cpp
+++ b/Engine/Plugins/Experimental/Mutable/Source/CustomizableObject/Private/MuCO/UnrealMutableImageProvider.cpp
@@ -8,6 +8,7 @@
 #include "MuCO/UnrealToMutableTextureConversionUtils.h"
 #include "TextureResource.h"
 #include "MuR/Parameters.h"
+#include "MuR/ImageTypes.h"
 
 
 //-------------------------------------------------------------------------------------------------
@@ -49,7 +50,7 @@ namespace
 		if (pSource)
 		{
 			OutResult->Init(SizeX, SizeY, LODs, MutableFormat, mu::EInitializationType::NotInitialized);
-			FMemory::Memcpy(OutResult->GetData(), pSource, OutResult->GetDataSize());
+			FMemory::Memcpy(OutResult->GetLODData(0), pSource, OutResult->GetLODDataSize(0));
 			Texture->GetPlatformData()->Mips[MipmapsToSkip].BulkData.Unlock();
 		}
 		else
@@ -241,11 +242,12 @@ TTuple<UE::Tasks::FTask, TFunction<void()>> FUnrealMutableImageProvider::GetImag
 				return Invoke(TrivialReturn);
 			}
 
-			int SizeX = TextureToLoad->GetSizeX() >> MipIndex;
-			int SizeY = TextureToLoad->GetSizeY() >> MipIndex;
+			int32 SizeX = TextureToLoad->GetSizeX() >> MipIndex;
+			int32 SizeY = TextureToLoad->GetSizeY() >> MipIndex;
 
+			check(LODs == 1);
 			mu::Ptr<mu::Image> Image = new mu::Image(SizeX, SizeY, LODs, MutImageFormat, mu::EInitializationType::NotInitialized);
-			MutImageDataSize = Image->GetDataSize();
+			TArrayView<uint8> MutImageDataView = Image->DataStorage.GetLOD(0);
 
 			// In a packaged game the bulk data has to be loaded
 			// Get the actual file to read the mip 0 data, do not keep any reference to TextureToLoad because once outside of the lock
@@ -254,7 +256,7 @@ TTuple<UE::Tasks::FTask, TFunction<void()>> FUnrealMutableImageProvider::GetImag
 			BulkDataSize = BulkData.GetBulkDataSize();
 			check(BulkDataSize > 0);
 
-			if (BulkDataSize != MutImageDataSize)
+			if (BulkDataSize != MutImageDataView.Num())
 			{
 				UE_LOG(LogMutable, Warning, TEXT("Failed to get external image [%s]. Bulk data size is different than the expected size. BulkData size [%d]. Mutable image data size [%d]."),
 					*Id.ToString(),	BulkDataSize, MutImageDataSize);
@@ -270,7 +272,7 @@ TTuple<UE::Tasks::FTask, TFunction<void()>> FUnrealMutableImageProvider::GetImag
 
 				TFunction<void(bool, IBulkDataIORequest*)> IOCallback =
 					[
-						MutImageDataSize,
+						MutImageDataView,
 						MutImageFormat,
 						Format,
 						Image,
@@ -295,10 +297,10 @@ TTuple<UE::Tasks::FTask, TFunction<void()>> FUnrealMutableImageProvider::GetImag
 
 					uint8* Results = IORequest->GetReadResults(); // required?
 
-					if (Results && Image->GetDataSize() == (int32)IORequest->GetSize())
+					if (Results && MutImageDataView.Num() == (int32)IORequest->GetSize())
 					{
 						check(BulkDataSize == (int32)IORequest->GetSize());
-						check(Results == Image->GetData());
+						check(Results == MutImageDataView.GetData());
 
 						ResultCallback(Image);
 						return;
@@ -311,11 +313,11 @@ TTuple<UE::Tasks::FTask, TFunction<void()>> FUnrealMutableImageProvider::GetImag
 							GetPixelFormatString(Format),
 							(int32)MutImageFormat);
 					}
-					else if (MutImageDataSize != (int32)IORequest->GetSize())
+					else if (MutImageDataView.Num() != (int32)IORequest->GetSize())
 					{
 						UE_LOG(LogMutable, Warning, TEXT("Failed to get external image. Requested size is different than the expected size. RequestSize: [%lld]. ExpectedSize: [%d]. Format: [%s]. MutableFormat: [%d]."),
 							IORequest->GetSize(),
-							Image->GetDataSize(),
+							MutImageDataView.Num(),
 							GetPixelFormatString(Format),
 							(int32)MutImageFormat);
 					}
@@ -332,7 +334,7 @@ TTuple<UE::Tasks::FTask, TFunction<void()>> FUnrealMutableImageProvider::GetImag
 				// This can *not* be done in the IOCallback because it would cause a deadlock so it is deferred to the returned
 				// cleanup function. Another solution could be to spwan a new task that depends on the 
 				// IORequestComplitionEvent which deletes it.
-				IORequest = BulkData.CreateStreamingRequest(EAsyncIOPriorityAndFlags::AIOP_High, &IOCallback, Image->GetData());
+				IORequest = BulkData.CreateStreamingRequest(EAsyncIOPriorityAndFlags::AIOP_High, &IOCallback, MutImageDataView.GetData());
 
 				if (IORequest)
 				{
@@ -368,7 +370,7 @@ TTuple<UE::Tasks::FTask, TFunction<void()>> FUnrealMutableImageProvider::GetImag
 				
 				if (Data)
 				{
-					FMemory::Memcpy(Image->GetData(), Data, BulkDataSize);
+					FMemory::Memcpy(MutImageDataView.GetData(), Data, BulkDataSize);
 
 					BulkData.Unlock();
 					ResultCallback(Image);
@@ -559,7 +561,7 @@ void FUnrealMutableImageProvider::CacheImage(FName Id, bool bUser)
 					{
 						FIntVector desc = Provider->GetTextureParameterValueSize(Id);
 						pResult = new mu::Image(desc[0], desc[1], 1, mu::EImageFormat::IF_RGBA_UBYTE, mu::EInitializationType::Black);
-						Provider->GetTextureParameterValueData(Id, pResult->GetData());
+						Provider->GetTextureParameterValueData(Id, pResult->GetLODData(0));
 						break;
 					}
 
@@ -711,28 +713,30 @@ void FUnrealMutableImageProvider::UnCacheImages(const mu::Parameters& Parameters
 mu::ImagePtr FUnrealMutableImageProvider::CreateDummy()
 {
 	// Create a dummy image
-	const int size = DUMMY_IMAGE_DESC.m_size[0];
-	const int checkerSize = 4;
-	constexpr int checkerTileCount = 2;
+	const int32 Size = DUMMY_IMAGE_DESC.m_size[0];
+	const int32 CheckerSize = 4;
+	constexpr int32 CheckerTileCount = 2;
 	
 #if !UE_BUILD_SHIPPING
-	uint8_t colours[checkerTileCount][4] = { { 255, 255, 0, 255 },{ 0, 0, 255, 255 } };
+	uint8 Colors[CheckerTileCount][4] = {{255, 255, 0, 255}, {0, 0, 255, 255}};
 #else
-	uint8_t colours[checkerTileCount][4] = { { 255, 255, 0, 0 },  { 0, 0, 255, 0 } };
+	uint8 Colors[CheckerTileCount][4] = {{255, 255, 0, 0}, {0, 0, 255, 0}};
 #endif
 
-	mu::ImagePtr pResult = new mu::Image(size, size, DUMMY_IMAGE_DESC.m_lods, DUMMY_IMAGE_DESC.m_format, mu::EInitializationType::NotInitialized);
+	mu::ImagePtr pResult = new mu::Image(Size, Size, DUMMY_IMAGE_DESC.m_lods, DUMMY_IMAGE_DESC.m_format, mu::EInitializationType::NotInitialized);
 
-	uint8_t* pData = pResult->GetData();
-	for (int x = 0; x < size; ++x)
+	check(pResult->GetLODCount() == 1);
+	check(pResult->GetFormat() == mu::EImageFormat::IF_RGBA_UBYTE || pResult->GetFormat() == mu::EImageFormat::IF_BGRA_UBYTE);
+	uint8* pData = pResult->GetLODData(0);
+	for (int32 X = 0; X < Size; ++X)
 	{
-		for (int y = 0; y < size; ++y)
+		for (int32 Y = 0; Y < Size; ++Y)
 		{
-			int checkerIndex = ((x / checkerSize) + (y / checkerSize)) % checkerTileCount;
-			pData[0] = colours[checkerIndex][0];
-			pData[1] = colours[checkerIndex][1];
-			pData[2] = colours[checkerIndex][2];
-			pData[3] = colours[checkerIndex][3];
+			int32 CheckerIndex = ((X / CheckerSize) + (Y / CheckerSize)) % CheckerTileCount;
+			pData[0] = Colors[CheckerIndex][0];
+			pData[1] = Colors[CheckerIndex][1];
+			pData[2] = Colors[CheckerIndex][2];
+			pData[3] = Colors[CheckerIndex][3];
 			pData += 4;
 		}
 	}
diff --git a/Engine/Plugins/Experimental/Mutable/Source/CustomizableObject/Private/MuCO/UnrealToMutableTextureConversionUtils.cpp b/Engine/Plugins/Experimental/Mutable/Source/CustomizableObject/Private/MuCO/UnrealToMutableTextureConversionUtils.cpp
--- a/Engine/Plugins/Experimental/Mutable/Source/CustomizableObject/Private/MuCO/UnrealToMutableTextureConversionUtils.cpp
+++ b/Engine/Plugins/Experimental/Mutable/Source/CustomizableObject/Private/MuCO/UnrealToMutableTextureConversionUtils.cpp
@@ -254,8 +254,10 @@ EUnrealToMutableConversionError ConvertTextureUnrealSourceToMutable(mu::Image* O
 	case ERawImageFormat::G8:
 	{
 		MUTABLE_CPUPROFILER_SCOPE(NoConvert);
+        check(LODs == 1);
+
 		OutResult->Init(SizeX, SizeY, LODs, mu::EImageFormat::IF_L_UBYTE, mu::EInitializationType::NotInitialized);
-		OutResult->m_data = MoveTemp(TempImage.RawData);
+		OutResult->DataStorage.GetInternalArray(0) = MoveTemp(TempImage.RawData);
 		break;
 	}
 
@@ -275,7 +277,8 @@ EUnrealToMutableConversionError ConvertTextureUnrealSourceToMutable(mu::Image* O
 		{
 			MUTABLE_CPUPROFILER_SCOPE(ToRGBA);
 			OutResult->Init(SizeX, SizeY, LODs, mu::EImageFormat::IF_RGBA_UBYTE, mu::EInitializationType::NotInitialized);
-			uint8* DataDest = OutResult->GetData();
+            check(LODs == 1);
+            uint8* DataDest = OutResult->GetLODData(0);
 
 			// Convert to RGBA8 while copying
 			TArrayView64<FColor> ImageDataView = TempImage.AsBGRA8();
@@ -294,7 +297,8 @@ EUnrealToMutableConversionError ConvertTextureUnrealSourceToMutable(mu::Image* O
 
 			// TODO: add support for a mu::IF_RGBX_UBYTE?			
 			OutResult->Init(SizeX, SizeY, LODs, mu::EImageFormat::IF_RGB_UBYTE, mu::EInitializationType::NotInitialized);
-			uint8* DataDest = OutResult->GetData();
+			check(LODs == 1);
+            uint8* DataDest = OutResult->GetLODData(0);
 
 			// Convert to RGB8 while copying
 			TArrayView64<FColor> ImageDataView = TempImage.AsBGRA8();
diff --git a/Engine/Plugins/Experimental/Mutable/Source/CustomizableObjectEditor/Private/MuCOE/SMutableCodeViewer.cpp b/Engine/Plugins/Experimental/Mutable/Source/CustomizableObjectEditor/Private/MuCOE/SMutableCodeViewer.cpp
--- a/Engine/Plugins/Experimental/Mutable/Source/CustomizableObjectEditor/Private/MuCOE/SMutableCodeViewer.cpp
+++ b/Engine/Plugins/Experimental/Mutable/Source/CustomizableObjectEditor/Private/MuCOE/SMutableCodeViewer.cpp
@@ -2512,10 +2512,10 @@ namespace
 				mu::EInitializationType::NotInitialized);
 
 			// Generate an alpha-tested circle with an horizontal gradient color.
-			uint8* Data = Image->GetData();
+			uint8* Data = Image->GetLODData(0);
 			int32 CircleRadius = (Size * 2) / 5;
 			int32 CircleRadius2 = CircleRadius * CircleRadius;
-			int32 Color[3] = { 255,128,0 };
+			int32 Color[3] = {255, 128, 0};
 
 			int32 LogSize = FMath::CeilLogTwo(Size);
 
diff --git a/Engine/Plugins/Experimental/Mutable/Source/CustomizableObjectEditor/Private/MuCOE/UnrealBakeHelpers.cpp b/Engine/Plugins/Experimental/Mutable/Source/CustomizableObjectEditor/Private/MuCOE/UnrealBakeHelpers.cpp
--- a/Engine/Plugins/Experimental/Mutable/Source/CustomizableObjectEditor/Private/MuCOE/UnrealBakeHelpers.cpp
+++ b/Engine/Plugins/Experimental/Mutable/Source/CustomizableObjectEditor/Private/MuCOE/UnrealBakeHelpers.cpp
@@ -193,14 +193,14 @@ UTexture2D* FUnrealBakeHelpers::BakeHelper_CreateAssetTexture(UTexture2D* SrcTex
 
 	// Create a mutable image from the platform data.
 	mu::EImageFormat MutableFormat = UnrealToMutablePixelFormat(SrcTex->GetPlatformData()->PixelFormat, SrcTex->HasAlphaChannel());
-	mu::Ptr<mu::Image> PlatformImage = new mu::Image(sx,sy,1,MutableFormat, mu::EInitializationType::NotInitialized);
+	mu::Ptr<mu::Image> PlatformImage = new mu::Image(sx, sy, 1, MutableFormat, mu::EInitializationType::NotInitialized);
 
 	constexpr int32 MipIndex = 0;
-	const uint8_t* SourceData = reinterpret_cast<const uint8_t*>(SrcTex->GetPlatformData()->Mips[MipIndex].BulkData.LockReadOnly());
+	const uint8* SourceData = reinterpret_cast<const uint8*>(SrcTex->GetPlatformData()->Mips[MipIndex].BulkData.LockReadOnly());
 	check(SourceData); // A mutable-generated texture should always contain platform data
 	int32 PlatformDataSize = SrcTex->GetPlatformData()->Mips[MipIndex].BulkData.GetBulkDataSize();
 	check(PlatformImage->GetDataSize()== PlatformDataSize);
-	FMemory::Memcpy( PlatformImage->GetData(), SourceData, PlatformDataSize);
+	FMemory::Memcpy(PlatformImage->GetLODData(0), SourceData, PlatformDataSize);
 	SourceData = nullptr;
 	SrcTex->GetPlatformData()->Mips[MipIndex].BulkData.Unlock();
 
@@ -226,20 +226,22 @@ UTexture2D* FUnrealBakeHelpers::BakeHelper_CreateAssetTexture(UTexture2D* SrcTex
 	// Copy the decompressed data to the texture source data
 	int32 SourceDataSize = DupTex->Source.CalcMipSize(MipIndex);
 
+	TArrayView<uint8> UncompressedView = UncompressedImage->DataStorage.GetLOD(0);
+	
 	// If this doesn't match, more cases have to be added to the switch above.
-	check(UncompressedImage->GetDataSize() == SourceDataSize);
+	check(UncompressedView.Num() == SourceDataSize);
 
-	uint8_t* Dest = DupTex->Source.LockMip(MipIndex);
+	uint8* Dest = DupTex->Source.LockMip(MipIndex);
 	check(Dest);
-	FMemory::Memcpy(Dest, UncompressedImage->GetData(), SourceDataSize);
+	FMemory::Memcpy(Dest, UncompressedView.GetData(), SourceDataSize);
 
 	// Probably can be integrated in the pixel format
 	const bool bNeedsRBSwizzle = PixelFormat == TSF_BGRA8;
 	if (bNeedsRBSwizzle)
 	{
-		for (int x = 0; x < sx * sy; ++x)
+		for (int32 x = 0; x < sx * sy; ++x)
 		{
-			uint8_t temp = Dest[0];
+			uint8 temp = Dest[0];
 			Dest[0] = Dest[2];
 			Dest[2] = temp;
 			Dest += 4;
diff --git a/Engine/Plugins/Experimental/Mutable/Source/MutableRuntime/Internal/MuR/ImagePrivate.h b/Engine/Plugins/Experimental/Mutable/Source/MutableRuntime/Internal/MuR/ImagePrivate.h
--- a/Engine/Plugins/Experimental/Mutable/Source/MutableRuntime/Internal/MuR/ImagePrivate.h
+++ b/Engine/Plugins/Experimental/Mutable/Source/MutableRuntime/Internal/MuR/ImagePrivate.h
@@ -12,10 +12,6 @@
 
 namespace mu
 {
-	MUTABLE_DEFINE_ENUM_SERIALISABLE( EMipmapFilterType );
-	MUTABLE_DEFINE_ENUM_SERIALISABLE( EImageFormat );
-	MUTABLE_DEFINE_ENUM_SERIALISABLE( EAddressMode );
-
 
     //---------------------------------------------------------------------------------------------
     //!
diff --git a/Engine/Plugins/Experimental/Mutable/Source/MutableRuntime/Internal/MuR/ModelPrivate.h b/Engine/Plugins/Experimental/Mutable/Source/MutableRuntime/Internal/MuR/ModelPrivate.h
--- a/Engine/Plugins/Experimental/Mutable/Source/MutableRuntime/Internal/MuR/ModelPrivate.h
+++ b/Engine/Plugins/Experimental/Mutable/Source/MutableRuntime/Internal/MuR/ModelPrivate.h
@@ -569,7 +569,6 @@ namespace mu
             return index;
         }
 
-
         //! Get a constant image, assuming it is fully loaded. The image constant will be composed with lodaded mips if necessary.
 		template <typename CreateImageFunc>
         void GetConstant( int32 ConstantIndex, ImagePtrConst& res, int32 MipsToSkip, const CreateImageFunc& CreateImage) const
@@ -596,34 +595,31 @@ namespace mu
 			{
 				MUTABLE_CPUPROFILER_SCOPE(ComposeConstantImage);
 
-				Ptr<Image> Result = CreateImage( CurrentMip->GetSizeX(), CurrentMip->GetSizeY(), FinalLODs, CurrentMip->GetFormat(), EInitializationType::NotInitialized );
+				Ptr<Image> Result = CreateImage(CurrentMip->GetSizeX(), CurrentMip->GetSizeY(), FinalLODs, CurrentMip->GetFormat(), EInitializationType::NotInitialized);
 				Result->m_flags = CurrentMip->m_flags;
 
 				// Some non-block pixel formats require separate memory size calculation
-				if (!Result->GetDataSize())
+				if (Result->DataStorage.IsEmpty())
 				{
-					int32 TotalSize = 0;
 					for (int32 LOD = 0; LOD < FinalLODs; ++LOD)
 					{
 						int32 LODIndex = m_constantImageLODIndices[ResultLODIndexIndex + LOD];
-						int32 MipSizeBytes = ConstantImageLODs[LODIndex].Value->GetDataSize();
-						TotalSize += MipSizeBytes;
+						int32 MipSizeBytes = ConstantImageLODs[LODIndex].Value->GetLODDataSize(0);
+						Result->DataStorage.ResizeLOD(LOD, MipSizeBytes);
 					}
-					Result->m_data.SetNum(TotalSize);
 				}
 
-				int32 WrittenDataBytes = 0;
-				uint8* Data = Result->GetData();
-				for (int32 LOD=0; LOD<FinalLODs; ++LOD)
+				for (int32 LOD = 0; LOD < FinalLODs; ++LOD)
 				{
 					check(CurrentMip->GetLODCount() == 1);
 					check(CurrentMip->GetFormat() == Result->GetFormat());
-					int32 MipSizeBytes = CurrentMip->GetDataSize();
-					check( Result->GetDataSize()>=WrittenDataBytes+MipSizeBytes);
 
-					FMemory::Memcpy( Data, CurrentMip->GetData(), MipSizeBytes);
-					Data += MipSizeBytes;
-					WrittenDataBytes += MipSizeBytes;
+					TArrayView<uint8> ResultLODView = Result->DataStorage.GetLOD(LOD);
+					TArrayView<const uint8> CurrentMipView = CurrentMip->DataStorage.GetLOD(0);
+					
+					check(CurrentMipView.Num() == ResultLODView.Num());
+
+					FMemory::Memcpy(ResultLODView.GetData(), CurrentMipView.GetData(), ResultLODView.Num());
 
 					if (LOD + 1 < FinalLODs)
 					{
@@ -637,7 +633,6 @@ namespace mu
 			}
 		}
 
-
         void GetConstant(int32 ConstantIndex, MeshPtrConst& res) const
         {
 			res = ConstantMeshes[ConstantIndex].Value;
diff --git a/Engine/Plugins/Experimental/Mutable/Source/MutableRuntime/Internal/MuR/SerialisationPrivate.h b/Engine/Plugins/Experimental/Mutable/Source/MutableRuntime/Internal/MuR/SerialisationPrivate.h
--- a/Engine/Plugins/Experimental/Mutable/Source/MutableRuntime/Internal/MuR/SerialisationPrivate.h
+++ b/Engine/Plugins/Experimental/Mutable/Source/MutableRuntime/Internal/MuR/SerialisationPrivate.h
@@ -141,48 +141,43 @@ namespace mu
 	};
 
 
-	//---------------------------------------------------------------------------------------------
-#define MUTABLE_IMPLEMENT_POD_SERIALISABLE(T)							\
-		template<>														\
-		void DLLEXPORT operator<< <T>( OutputArchive& arch, const T& t )			\
-		{																\
-			arch.GetPrivate()->m_pStream->Write( &t, sizeof(T) );		\
-		}																\
-																		\
-		template<>														\
-		void DLLEXPORT operator>> <T>( InputArchive& arch, T& t )					\
-		{																\
-			arch.GetPrivate()->m_pStream->Read( &t, sizeof(T) );		\
-		}																\
+#define MUTABLE_IMPLEMENT_POD_SERIALISABLE(Type)				     \
+    void DLLEXPORT operator<<(OutputArchive& Arch, const Type& T)    \
+    {																 \
+        Arch.GetPrivate()->m_pStream->Write(&T, sizeof(Type));		 \
+    }																 \
+                                                                     \
+    void DLLEXPORT operator>>(InputArchive& Arch, Type& T)		     \
+    {																 \
+        Arch.GetPrivate()->m_pStream->Read(&T, sizeof(Type));		 \
+    }																 \
 		
 
-#define MUTABLE_IMPLEMENT_POD_VECTOR_SERIALISABLE(T)									     \
-	template<typename Alloc>                                                             \
-	void operator <<(OutputArchive& arch, const TArray<T, Alloc>& v)					 \
-	{                                                                                    \
-		uint32 Num = uint32(v.Num());													 \
-		arch << Num;																	 \
-		if (Num)																		 \
-		{                                                                                \
-			arch.GetPrivate()->m_pStream->Write(&v[0], Num * sizeof(T));				 \
-		}                                                                                \
-	}                                                                                    \
-	                                                                                     \
-	template<typename Alloc>                                                             \
-	void operator >>(InputArchive& arch, TArray<T, Alloc>& v)                            \
-	{                                                                                    \
-		uint32 Num;																		 \
-		arch >> Num;																	 \
-		v.SetNum(Num);																	 \
-		if (Num)																		 \
-		{                                                                                \
-			arch.GetPrivate()->m_pStream->Read(&v[0], Num * sizeof(T));					 \
-		}                                                                                \
-	}                                                                                    \
-
-	//---------------------------------------------------------------------------------------------
-	// TVariant custom serialize. Based on the default serialization.
-	//---------------------------------------------------------------------------------------------
+#define MUTABLE_IMPLEMENT_POD_VECTOR_SERIALISABLE(Type)                        \
+    template<class Alloc>                                                      \
+	void operator<<(OutputArchive& Arch, const TArray<Type, Alloc>& V)         \
+	{                                                                          \
+		uint32 Num = uint32(V.Num());                                          \
+		Arch << Num;                                                           \
+		if (Num)                                                               \
+		{                                                                      \
+			Arch.GetPrivate()->m_pStream->Write(&V[0], Num * sizeof(Type));    \
+		}                                                                      \
+	}                                                                          \
+                                                                               \
+    template<class Alloc>                                                      \
+	void operator>>(InputArchive& Arch, TArray<Type, Alloc>& V)                \
+	{                                                                          \
+		uint32 Num;                                                            \
+		Arch >> Num;                                                           \
+		V.SetNum(Num);                                                         \
+		if (Num)                                                               \
+		{                                                                      \
+			Arch.GetPrivate()->m_pStream->Read(&V[0], Num * sizeof(Type));     \
+		}                                                                      \
+	}                                                                          \
+
+	/** TVariant custom serialize. Based on the default serialization. */
 	template <typename... Ts>
 	void operator<<(OutputArchive& Ar, const TVariant<Ts...>& Variant)
 	{
@@ -234,54 +229,49 @@ namespace mu
 		TVariantLoadFromInputArchiveLookup<Ts...>::Load(static_cast<SIZE_T>(Index), Ar, Variant);
 	}
 
+#define MUTABLE_IMPLEMENT_ENUM_SERIALISABLE(Type)						\
+        void DLLEXPORT operator<<(OutputArchive& Arch, const Type& T)   \
+		{																\
+            uint32 V = (uint32)T;                                       \
+            Arch.GetPrivate()->m_pStream->Write(&V, sizeof(uint32));    \
+		}																\
+																		\
+        void DLLEXPORT operator>>(InputArchive& Arch, Type& T)   		\
+		{																\
+            uint32 V;													\
+            Arch.GetPrivate()->m_pStream->Read(&V, sizeof(uint32));     \
+			T = (Type)V;												\
+		}																\
 
-	//---------------------------------------------------------------------------------------------
-#define MUTABLE_IMPLEMENT_ENUM_SERIALISABLE(T)								\
-		template<>															\
-        void DLLEXPORT operator<< <T>( OutputArchive& arch, const T& t )	\
-		{																	\
-            uint32 v = (uint32)t;                                         	\
-            arch.GetPrivate()->m_pStream->Write( &v, sizeof(uint32) );  	\
-		}																	\
-																			\
-		template<>															\
-        void DLLEXPORT operator>> <T>( InputArchive& arch, T& t )   		\
-		{																	\
-            uint32 v;														\
-            arch.GetPrivate()->m_pStream->Read( &v, sizeof(uint32) );		\
-			t = (T)v;														\
-		}																	\
-
-    //---------------------------------------------------------------------------------------------
-    template< typename T0, typename T1 >
-    inline void operator<< ( OutputArchive& arch, const std::pair<T0,T1>& v )
+    template<typename T0, typename T1>
+    inline void operator<<(OutputArchive& Arch, const std::pair<T0, T1>& V)
     {
-        arch << v.first;
-        arch << v.second;
+        Arch << V.first;
+        Arch << V.second;
     }
 
-    template< typename T0, typename T1 >
-    inline void operator>> ( InputArchive& arch, std::pair<T0,T1>& v )
+    template<typename T0, typename T1>
+    inline void operator>>(InputArchive& Arch, std::pair<T0, T1>& V)
     {
-        arch >> v.first;
-        arch >> v.second;
+        Arch >> V.first;
+        Arch >> V.second;
     }
-
 	
-	//---------------------------------------------------------------------------------------------
-	template<typename T, uint32 Size> void operator<<(OutputArchive& arch, const TStaticArray<T, Size>& v)
+	template<typename T, uint32 Size, uint32 Align> 
+    void operator<<(OutputArchive& Arch, const TStaticArray<T, Size, Align>& V)
 	{
-		for (uint32 i = 0; i < Size; ++i)
+		for (int32 I = 0; I < Size; ++I)
 		{
-			arch << v[i];
+			Arch << V[I];
 		}
 	}
 
-	template<typename T, uint32 Size> void operator>>(InputArchive& arch, TStaticArray<T, Size>& v)
+	template<typename T, uint32 Size, uint32 Align> 
+    void operator>>(InputArchive& Arch, TStaticArray<T, Size, Align>& V)
 	{
-		for (uint32 i = 0; i < Size; ++i)
+		for (uint32 I = 0; I < Size; ++I)
 		{
-			arch >> v[i];
+			Arch >> V[I];
 		}
 	}
 	
diff --git a/Engine/Plugins/Experimental/Mutable/Source/MutableRuntime/Private/MuR/CodeRunner.cpp b/Engine/Plugins/Experimental/Mutable/Source/MutableRuntime/Private/MuR/CodeRunner.cpp
--- a/Engine/Plugins/Experimental/Mutable/Source/MutableRuntime/Private/MuR/CodeRunner.cpp
+++ b/Engine/Plugins/Experimental/Mutable/Source/MutableRuntime/Private/MuR/CodeRunner.cpp
@@ -3283,15 +3283,15 @@ namespace mu
                         switch (EBlendType(args.blendType))
                         {
 						case EBlendType::BT_NORMAL_COMBINE: check(false); break;
-                        case EBlendType::BT_SOFTLIGHT: BufferLayer<SoftLightChannelMasked, SoftLightChannel, false>( Base->GetData(), Base.get(), Mask.get(), Blended.get(), bApplyColorBlendToAlpha, bBlendOnlyOneMip); break;
-                        case EBlendType::BT_HARDLIGHT: BufferLayer<HardLightChannelMasked, HardLightChannel, false>(Base->GetData(), Base.get(), Mask.get(), Blended.get(), bApplyColorBlendToAlpha, bBlendOnlyOneMip); break;
-                        case EBlendType::BT_BURN: BufferLayer<BurnChannelMasked, BurnChannel, false>(Base->GetData(), Base.get(), Mask.get(), Blended.get(), bApplyColorBlendToAlpha, bBlendOnlyOneMip); break;
-                        case EBlendType::BT_DODGE: BufferLayer<DodgeChannelMasked, DodgeChannel, false>(Base->GetData(), Base.get(), Mask.get(), Blended.get(), bApplyColorBlendToAlpha, bBlendOnlyOneMip); break;
-                        case EBlendType::BT_SCREEN: BufferLayer<ScreenChannelMasked, ScreenChannel, false>(Base->GetData(), Base.get(), Mask.get(), Blended.get(), bApplyColorBlendToAlpha, bBlendOnlyOneMip); break;
-                        case EBlendType::BT_OVERLAY: BufferLayer<OverlayChannelMasked, OverlayChannel, false>(Base->GetData(), Base.get(), Mask.get(), Blended.get(), bApplyColorBlendToAlpha, bBlendOnlyOneMip); break;
-                        case EBlendType::BT_LIGHTEN: BufferLayer<LightenChannelMasked, LightenChannel, false>(Base->GetData(), Base.get(), Mask.get(), Blended.get(), bApplyColorBlendToAlpha, bBlendOnlyOneMip); break;
-                        case EBlendType::BT_MULTIPLY: BufferLayer<MultiplyChannelMasked, MultiplyChannel, false>(Base->GetData(), Base.get(), Mask.get(), Blended.get(), bApplyColorBlendToAlpha, bBlendOnlyOneMip); break;
-                        case EBlendType::BT_BLEND: BufferLayer<BlendChannelMasked, BlendChannel, false>(Base->GetData(), Base.get(), Mask.get(), Blended.get(), bApplyColorBlendToAlpha, bBlendOnlyOneMip); break;
+                        case EBlendType::BT_SOFTLIGHT: BufferLayer<SoftLightChannelMasked, SoftLightChannel, false>(Base.get(), Base.get(), Mask.get(), Blended.get(), bApplyColorBlendToAlpha, bBlendOnlyOneMip); break;
+                        case EBlendType::BT_HARDLIGHT: BufferLayer<HardLightChannelMasked, HardLightChannel, false>(Base.get(), Base.get(), Mask.get(), Blended.get(), bApplyColorBlendToAlpha, bBlendOnlyOneMip); break;
+                        case EBlendType::BT_BURN: BufferLayer<BurnChannelMasked, BurnChannel, false>(Base.get(), Base.get(), Mask.get(), Blended.get(), bApplyColorBlendToAlpha, bBlendOnlyOneMip); break;
+                        case EBlendType::BT_DODGE: BufferLayer<DodgeChannelMasked, DodgeChannel, false>(Base.get(), Base.get(), Mask.get(), Blended.get(), bApplyColorBlendToAlpha, bBlendOnlyOneMip); break;
+                        case EBlendType::BT_SCREEN: BufferLayer<ScreenChannelMasked, ScreenChannel, false>(Base.get(), Base.get(), Mask.get(), Blended.get(), bApplyColorBlendToAlpha, bBlendOnlyOneMip); break;
+                        case EBlendType::BT_OVERLAY: BufferLayer<OverlayChannelMasked, OverlayChannel, false>(Base.get(), Base.get(), Mask.get(), Blended.get(), bApplyColorBlendToAlpha, bBlendOnlyOneMip); break;
+                        case EBlendType::BT_LIGHTEN: BufferLayer<LightenChannelMasked, LightenChannel, false>(Base.get(), Base.get(), Mask.get(), Blended.get(), bApplyColorBlendToAlpha, bBlendOnlyOneMip); break;
+                        case EBlendType::BT_MULTIPLY: BufferLayer<MultiplyChannelMasked, MultiplyChannel, false>(Base.get(), Base.get(), Mask.get(), Blended.get(), bApplyColorBlendToAlpha, bBlendOnlyOneMip); break;
+                        case EBlendType::BT_BLEND: BufferLayer<BlendChannelMasked, BlendChannel, false>(Base.get(), Base.get(), Mask.get(), Blended.get(), bApplyColorBlendToAlpha, bBlendOnlyOneMip); break;
                         default: check(false);
                         }
 
@@ -3305,15 +3305,15 @@ namespace mu
 						switch (EBlendType(args.blendType))
 						{
 						case EBlendType::BT_NORMAL_COMBINE: check(false); break;
-						case EBlendType::BT_SOFTLIGHT: BufferLayerEmbeddedMask<SoftLightChannelMasked, SoftLightChannel, false>(Base->GetData(), Base.get(), Blended.get(), bApplyColorBlendToAlpha, bBlendOnlyOneMip); break;
-						case EBlendType::BT_HARDLIGHT: BufferLayerEmbeddedMask<HardLightChannelMasked, HardLightChannel, false>(Base->GetData(), Base.get(), Blended.get(), bApplyColorBlendToAlpha, bBlendOnlyOneMip); break;
-						case EBlendType::BT_BURN: BufferLayerEmbeddedMask<BurnChannelMasked, BurnChannel, false>(Base->GetData(), Base.get(), Blended.get(), bApplyColorBlendToAlpha, bBlendOnlyOneMip); break;
-						case EBlendType::BT_DODGE: BufferLayerEmbeddedMask<DodgeChannelMasked, DodgeChannel, false>(Base->GetData(), Base.get(), Blended.get(), bApplyColorBlendToAlpha, bBlendOnlyOneMip); break;
-						case EBlendType::BT_SCREEN: BufferLayerEmbeddedMask<ScreenChannelMasked, ScreenChannel, false>(Base->GetData(), Base.get(), Blended.get(), bApplyColorBlendToAlpha, bBlendOnlyOneMip); break;
-						case EBlendType::BT_OVERLAY: BufferLayerEmbeddedMask<OverlayChannelMasked, OverlayChannel, false>(Base->GetData(), Base.get(), Blended.get(), bApplyColorBlendToAlpha, bBlendOnlyOneMip); break;
-						case EBlendType::BT_LIGHTEN: BufferLayerEmbeddedMask<LightenChannelMasked, LightenChannel, false>(Base->GetData(), Base.get(), Blended.get(), bApplyColorBlendToAlpha, bBlendOnlyOneMip); break;
-						case EBlendType::BT_MULTIPLY: BufferLayerEmbeddedMask<MultiplyChannelMasked, MultiplyChannel, false>(Base->GetData(), Base.get(), Blended.get(), bApplyColorBlendToAlpha, bBlendOnlyOneMip); break;
-						case EBlendType::BT_BLEND: BufferLayerEmbeddedMask<BlendChannelMasked, BlendChannel, false>(Base->GetData(), Base.get(), Blended.get(), bApplyColorBlendToAlpha, bBlendOnlyOneMip); break;
+						case EBlendType::BT_SOFTLIGHT: BufferLayerEmbeddedMask<SoftLightChannelMasked, SoftLightChannel, false>(Base.get(), Base.get(), Blended.get(), bApplyColorBlendToAlpha, bBlendOnlyOneMip); break;
+						case EBlendType::BT_HARDLIGHT: BufferLayerEmbeddedMask<HardLightChannelMasked, HardLightChannel, false>(Base.get(), Base.get(), Blended.get(), bApplyColorBlendToAlpha, bBlendOnlyOneMip); break;
+						case EBlendType::BT_BURN: BufferLayerEmbeddedMask<BurnChannelMasked, BurnChannel, false>(Base.get(), Base.get(), Blended.get(), bApplyColorBlendToAlpha, bBlendOnlyOneMip); break;
+						case EBlendType::BT_DODGE: BufferLayerEmbeddedMask<DodgeChannelMasked, DodgeChannel, false>(Base.get(), Base.get(), Blended.get(), bApplyColorBlendToAlpha, bBlendOnlyOneMip); break;
+						case EBlendType::BT_SCREEN: BufferLayerEmbeddedMask<ScreenChannelMasked, ScreenChannel, false>(Base.get(), Base.get(), Blended.get(), bApplyColorBlendToAlpha, bBlendOnlyOneMip); break;
+						case EBlendType::BT_OVERLAY: BufferLayerEmbeddedMask<OverlayChannelMasked, OverlayChannel, false>(Base.get(), Base.get(), Blended.get(), bApplyColorBlendToAlpha, bBlendOnlyOneMip); break;
+						case EBlendType::BT_LIGHTEN: BufferLayerEmbeddedMask<LightenChannelMasked, LightenChannel, false>(Base.get(), Base.get(), Blended.get(), bApplyColorBlendToAlpha, bBlendOnlyOneMip); break;
+						case EBlendType::BT_MULTIPLY: BufferLayerEmbeddedMask<MultiplyChannelMasked, MultiplyChannel, false>(Base.get(), Base.get(), Blended.get(), bApplyColorBlendToAlpha, bBlendOnlyOneMip); break;
+						case EBlendType::BT_BLEND: BufferLayerEmbeddedMask<BlendChannelMasked, BlendChannel, false>(Base.get(), Base.get(), Blended.get(), bApplyColorBlendToAlpha, bBlendOnlyOneMip); break;
 						default: check(false);
 						}
 					}
@@ -3422,18 +3422,19 @@ namespace mu
 				{
 					MUTABLE_CPUPROFILER_SCOPE(ImageNormalComposite_EmergencyFix);
 
-					int levelCount = Base->GetLODCount();
-					ImagePtr Dest = CreateImage(Normal->GetSizeX(), Normal->GetSizeY(), levelCount, Normal->GetFormat(), EInitializationType::NotInitialized);
+					int32 StartLevel = Normal->GetLODCount() - 1;
+					int32 LevelCount = Base->GetLODCount();
+					
+					Ptr<Image> NormalFix = CloneOrTakeOver(Normal);
 
-					FMipmapGenerationSettings mipSettings{};
-					ImOp.ImageMipmap(m_pSettings->ImageCompressionQuality, Dest.get(), Normal.get(), levelCount, mipSettings);
+					FMipmapGenerationSettings MipSettings{};
+					ImOp.ImageMipmap(m_pSettings->ImageCompressionQuality, NormalFix.get(), NormalFix.get(), StartLevel, LevelCount, MipSettings);
 
-					Release(Normal);
-					Normal = Dest;
+					Normal = NormalFix;
 				}
 
 
-                ImagePtr Result = CreateImage( Base->GetSizeX(), Base->GetSizeY(), Base->GetLODCount(), Base->GetFormat(), EInitializationType::NotInitialized);
+                Ptr<Image> Result = CreateImage(Base->GetSizeX(), Base->GetSizeY(), Base->GetLODCount(), Base->GetFormat(), EInitializationType::NotInitialized);
 				ImageNormalComposite(Result.get(), Base.get(), Normal.get(), args.mode, args.power);
 
 				Release(Base);
@@ -3521,9 +3522,9 @@ namespace mu
             switch (item.Stage)
             {
             case 0:
-                AddOp( FScheduledOp( item.At, item, 1),
-                        FScheduledOp( args.source, item),
-                        FScheduledOp( args.sizeSource, item) );
+                AddOp(FScheduledOp(item.At, item, 1),
+                      	FScheduledOp(args.source, item),
+                        FScheduledOp(args.sizeSource, item));
                 break;
 
             case 1:
@@ -3535,29 +3536,26 @@ namespace mu
 				FImageSize DestSize = SizeBase->GetSize();
 				Release(SizeBase);
 
-                if ( Base->GetSize()!=DestSize )
+                if (Base->GetSize() != DestSize)
                 {
 					int32 BaseLODCount = Base->GetLODCount();
 					Ptr<Image> Result = CreateImage(DestSize[0], DestSize[1], BaseLODCount, Base->GetFormat(), EInitializationType::NotInitialized);
-					ImOp.ImageResizeLinear( Result.get(), m_pSettings->ImageCompressionQuality, Base.get());
+					ImOp.ImageResizeLinear(Result.get(), m_pSettings->ImageCompressionQuality, Base.get());
 					Release(Base);
 
                     // If the source image had mips, generate them as well for the resized image.
                     // This shouldn't happen often since "ResizeLike" should be usually optimised out
                     // during model compilation. The mipmap generation below is not very precise with
                     // the number of mips that are needed and will probably generate too many
-                    bool bSourceHasMips = BaseLODCount>1;
-                    if (bSourceHasMips)
+                    bool bSourceHasMips = BaseLODCount > 1;
+                    
+					if (bSourceHasMips)
                     {
-                        int levelCount = Image::GetMipmapCount( Result->GetSizeX(), Result->GetSizeY() );
-                        Ptr<Image> Mipmapped = CreateImage( Result->GetSizeX(), Result->GetSizeY(), levelCount, Result->GetFormat(), EInitializationType::NotInitialized);
-
-						FMipmapGenerationSettings mipSettings{};
+						int32 LevelCount = Image::GetMipmapCount(Result->GetSizeX(), Result->GetSizeY());	
+						Result->DataStorage.SetNumLODs(LevelCount);
 
-						ImOp.ImageMipmap( m_pSettings->ImageCompressionQuality, Mipmapped.get(), Result.get(), levelCount, mipSettings );
-
-						Release(Result);
-						Result = Mipmapped;
+						FMipmapGenerationSettings MipSettings{};
+						ImOp.ImageMipmap(m_pSettings->ImageCompressionQuality, Result.get(), Result.get(), 0, LevelCount, MipSettings);
                     }				
 
 					StoreImage(item, Result);
@@ -3599,25 +3597,24 @@ namespace mu
 
             break;
         }
-
         case OP_TYPE::IM_BLANKLAYOUT:
         {
-			OP::ImageBlankLayoutArgs args = Program.GetOpArgs<OP::ImageBlankLayoutArgs>(item.At);
+			OP::ImageBlankLayoutArgs Args = Program.GetOpArgs<OP::ImageBlankLayoutArgs>(item.At);
             switch (item.Stage)
             {
             case 0:
-                AddOp( FScheduledOp( item.At, item, 1), FScheduledOp::FromOpAndOptions( args.layout, item, 0) );
+                AddOp( FScheduledOp( item.At, item, 1), FScheduledOp::FromOpAndOptions(Args.layout, item, 0));
                 break;
 
             case 1:
             {
             	MUTABLE_CPUPROFILER_SCOPE(IM_BLANKLAYOUT_1)
             		
-                Ptr<const Layout> pLayout = LoadLayout(FScheduledOp::FromOpAndOptions(args.layout, item, 0));
+                Ptr<const Layout> pLayout = LoadLayout(FScheduledOp::FromOpAndOptions(Args.layout, item, 0));
 
                 FIntPoint SizeInBlocks = pLayout->GetGridSize();
 
-				FIntPoint BlockSizeInPixels(args.blockSize[0], args.blockSize[1]);
+				FIntPoint BlockSizeInPixels(Args.blockSize[0], Args.blockSize[1]);
 
 				// Image size if we don't skip any mipmap
 				FIntPoint FullImageSizeInPixels = SizeInBlocks * BlockSizeInPixels;
@@ -3649,22 +3646,22 @@ namespace mu
 					//}
 				}
 
-                int MipsToGenerate = 1;
-                if ( args.generateMipmaps )
+                int32 MipsToGenerate = 1;
+                if (Args.generateMipmaps)
                 {
-                    if ( args.mipmapCount==0 )
+                    if (Args.mipmapCount == 0)
                     {
 						MipsToGenerate = Image::GetMipmapCount(ImageSizeInPixels.X, ImageSizeInPixels.Y);
                     }
                     else
                     {
-						MipsToGenerate = FMath::Max(args.mipmapCount-MipsToSkip,1);
+						MipsToGenerate = FMath::Max(Args.mipmapCount - MipsToSkip, 1);
                     }
                 }
 
 				// It needs to be initialized in case it has gaps.
-                ImagePtr New = CreateImage(ImageSizeInPixels.X, ImageSizeInPixels.Y, MipsToGenerate, EImageFormat(args.format), EInitializationType::Black );
-                StoreImage( item, New );
+                Ptr<Image> New = CreateImage(ImageSizeInPixels.X, ImageSizeInPixels.Y, MipsToGenerate, EImageFormat(Args.format), EInitializationType::Black );
+                StoreImage(item, New);
                 break;
             }
 
@@ -3842,25 +3839,25 @@ namespace mu
                 }
 
                 // Factor from 0 to 1 between the two targets
-                const FScheduledOpData& data = m_heapData[(size_t)item.CustomState];
-                float bifactor = data.Interpolate.Bifactor;
-                int min = data.Interpolate.Min;
-                int max = data.Interpolate.Max;
+                const FScheduledOpData& Data = m_heapData[(size_t)item.CustomState];
+                float Bifactor = Data.Interpolate.Bifactor;
+                int32 Min = Data.Interpolate.Min;
+                int32 Max = Data.Interpolate.Max;
 
-                if ( bifactor < UE_SMALL_NUMBER )
+                if (Bifactor < UE_SMALL_NUMBER)
                 {
-                    Ptr<const Image> Source = LoadImage( FCacheAddress(args.targets[min],item) );
+                    Ptr<const Image> Source = LoadImage(FCacheAddress(args.targets[Min], item));
 					StoreImage(item, Source);
 				}
-                else if ( bifactor > 1.0f-UE_SMALL_NUMBER )
+                else if (Bifactor > 1.0f - UE_SMALL_NUMBER)
                 {
-                    Ptr<const Image> Source = LoadImage( FCacheAddress(args.targets[max],item) );
+                    Ptr<const Image> Source = LoadImage(FCacheAddress(args.targets[Max], item));
 					StoreImage(item, Source);
 				}
                 else
                 {
-					Ptr<const Image> pMin = LoadImage( FCacheAddress(args.targets[min],item) );
-                    Ptr<const Image> pMax = LoadImage( FCacheAddress(args.targets[max],item) );
+					Ptr<const Image> pMin = LoadImage(FCacheAddress(args.targets[Min], item));
+                    Ptr<const Image> pMax = LoadImage(FCacheAddress(args.targets[Max], item));
 
                     if (pMin && pMax)
                     {						
@@ -3896,30 +3893,32 @@ namespace mu
 						if (pNew->GetLODCount() != LevelCount)
 						{
 							MUTABLE_CPUPROFILER_SCOPE(Mipmap_ForInterpolate);
+						
+							int32 StartLevel = pNew->GetLODCount() - 1;
+							// pNew is local owned, no need to CloneOrTakeOver.
+							pNew->DataStorage.SetNumLODs(LevelCount);
 
-							ImagePtr pDest = CreateImage(pNew->GetSizeX(), pNew->GetSizeY(), LevelCount, pNew->GetFormat(), EInitializationType::NotInitialized);
-
-							FMipmapGenerationSettings settings{};
-							ImOp.ImageMipmap(m_pSettings->ImageCompressionQuality, pDest.get(), pNew.get(), LevelCount, settings);
+							FMipmapGenerationSettings Settings{};
+							ImOp.ImageMipmap(m_pSettings->ImageCompressionQuality, pNew.get(), pNew.get(), StartLevel, LevelCount, Settings);
 
-							Release(pNew);
-							pNew = pDest;
 						}
 
 						if (pMax->GetLODCount() != LevelCount)
 						{
 							MUTABLE_CPUPROFILER_SCOPE(Mipmap_ForInterpolate);
 
-							ImagePtr pDest = CreateImage(pMax->GetSizeX(), pMax->GetSizeY(), LevelCount, pMax->GetFormat(), EInitializationType::NotInitialized);
+							int32 StartLevel = pMax->GetLODCount() - 1;
 
-							FMipmapGenerationSettings settings{};
-							ImOp.ImageMipmap(m_pSettings->ImageCompressionQuality, pDest.get(), pMax.get(), LevelCount, settings);
+							Ptr<Image> MaxFix = CloneOrTakeOver(pMax);
+							MaxFix->DataStorage.SetNumLODs(LevelCount);
+							
+							FMipmapGenerationSettings Settings{};
+							ImOp.ImageMipmap(m_pSettings->ImageCompressionQuality, MaxFix.get(), MaxFix.get(), StartLevel, LevelCount, Settings);
 
-							Release(pMax);
-							pMax = pDest;
+							pMax = MaxFix;
 						}
 
-                        ImageInterpolate( pNew.get(), pMax.get(), bifactor );
+                        ImageInterpolate(pNew.get(), pMax.get(), Bifactor);
 
 						Release(pMax);
 						StoreImage(item, pNew);
@@ -3932,7 +3931,6 @@ namespace mu
                     {
 						StoreImage(item, pMax);
 					}
-
 				}
 
                 break;
@@ -4639,15 +4637,12 @@ namespace mu
 					{
 						MUTABLE_CPUPROFILER_SCOPE(RunCode_RasterMesh_BilinearMipGen);
 
-						Ptr<Image> NewImage = CreateImage(Source->GetSizeX(), Source->GetSizeY(), 2, Source->GetFormat(), EInitializationType::NotInitialized);
-
-						check(NewImage->GetDataSize() >= Source->GetDataSize());
-						FMemory::Memcpy(NewImage->GetData(), Source->GetData(), Source->GetDataSize());
+						Ptr<Image> OwnedSource = CloneOrTakeOver(Source);
 
-						ImageMipmapInPlace(0, NewImage.get(), FMipmapGenerationSettings{});
+						OwnedSource->DataStorage.SetNumLODs(2);
+						ImageMipmapInPlace(0, OwnedSource.get(), FMipmapGenerationSettings{});
 
-						Release(Source);
-						Source = NewImage;
+						Source = OwnedSource;
 					}
 				}
 
@@ -4965,15 +4960,12 @@ namespace mu
 				{
 					MUTABLE_CPUPROFILER_SCOPE(RunCode_ImageTransform_BilinearMipGen);
 
-					Ptr<Image> NewImage = CreateImage(Source->GetSizeX(), Source->GetSizeY(), 2, Source->GetFormat(), EInitializationType::NotInitialized);
+					Ptr<Image> OwnedSource = CloneOrTakeOver(Source);
+					OwnedSource->DataStorage.SetNumLODs(2);
 
-					check(NewImage->GetDataSize() >= Source->GetDataSize());
-					FMemory::Memcpy(NewImage->GetData(), Source->GetData(), Source->GetDataSize());
+					ImageMipmapInPlace(0, OwnedSource.get(), FMipmapGenerationSettings{});
 
-					ImageMipmapInPlace(0, NewImage.get(), FMipmapGenerationSettings{});
-
-					Release(Source);
-					Source = NewImage;
+					Source = OwnedSource;
 				}
 
 				Scale.X = FMath::IsNearlyZero(Scale.X, UE_KINDA_SMALL_NUMBER) ? UE_KINDA_SMALL_NUMBER : Scale.X;
@@ -6318,25 +6310,25 @@ namespace mu
 			{
 				// Somewhat synched with Full op execution code.
 				FImageDesc BaseDesc = m_heapImageDesc[item.CustomState];
-				int levelCount = args.levels;
-				int maxLevelCount = Image::GetMipmapCount(BaseDesc.m_size[0], BaseDesc.m_size[1]);
-				if (levelCount == 0)
+				int32 LevelCount = args.levels;
+				int32 MaxLevelCount = Image::GetMipmapCount(BaseDesc.m_size[0], BaseDesc.m_size[1]);
+				if (LevelCount == 0)
 				{
-					levelCount = maxLevelCount;
+					LevelCount = MaxLevelCount;
 				}
-				else if (levelCount > maxLevelCount)
+				else if (LevelCount > MaxLevelCount)
 				{
 					// If code generation is smart enough, this should never happen.
 					// \todo But apparently it does, sometimes.
-					levelCount = maxLevelCount;
+					LevelCount = MaxLevelCount;
 				}
 
 				// At least keep the levels we already have.
-				int32 startLevel = BaseDesc.m_lods;
-				levelCount = FMath::Max(startLevel, levelCount);
+				int32 StartLevel = BaseDesc.m_lods;
+				LevelCount = FMath::Max(StartLevel, LevelCount);
 
 				// Update result.
-				m_heapImageDesc[item.CustomState].m_lods = levelCount;
+				m_heapImageDesc[item.CustomState].m_lods = LevelCount;
 				StoreValidDesc(item);
 				break;
 			}
diff --git a/Engine/Plugins/Experimental/Mutable/Source/MutableRuntime/Private/MuR/CodeRunnerTasked.cpp b/Engine/Plugins/Experimental/Mutable/Source/MutableRuntime/Private/MuR/CodeRunnerTasked.cpp
--- a/Engine/Plugins/Experimental/Mutable/Source/MutableRuntime/Private/MuR/CodeRunnerTasked.cpp
+++ b/Engine/Plugins/Experimental/Mutable/Source/MutableRuntime/Private/MuR/CodeRunnerTasked.cpp
@@ -569,7 +569,6 @@ namespace mu
 		OutDesc = m_heapImageDesc[0];
 	}
 
-
 	//---------------------------------------------------------------------------------------------
 	//---------------------------------------------------------------------------------------------
 	//---------------------------------------------------------------------------------------------
@@ -685,14 +684,16 @@ namespace mu
 			{
 				MUTABLE_CPUPROFILER_SCOPE(ImageLayer_EmergencyFix);
 
-				int32 levelCount = Base->GetLODCount();
-				Ptr<Image> Dest = Runner->CreateImage(Mask->GetSizeX(), Mask->GetSizeY(), levelCount, Mask->GetFormat(), EInitializationType::NotInitialized);
+				int32 StartLevel = Mask->GetLODCount() - 1;
+				int32 LevelCount = Base->GetLODCount();
 
-				FMipmapGenerationSettings settings{};
-				ImOp.ImageMipmap(ImageCompressionQuality, Dest.get(), Mask.get(), levelCount, settings);
+				Ptr<Image> MaskFix = Runner->CloneOrTakeOver(Mask);
+				MaskFix->DataStorage.SetNumLODs(LevelCount);
 
-				Runner->Release(Mask);
-				Mask = Dest;
+				FMipmapGenerationSettings Settings{};
+				ImOp.ImageMipmap(ImageCompressionQuality, MaskFix.get(), MaskFix.get(), StartLevel, LevelCount, Settings);
+
+				Mask = MaskFix;
 			}
 		}
 
@@ -755,15 +756,15 @@ namespace mu
 			switch (EBlendType(Args.blendType))
 			{
 			case EBlendType::BT_NORMAL_COMBINE: ImageNormalCombine(Result.get(), Result.get(), Mask.get(), Blended.get(), bOnlyOneMip); break;
-			case EBlendType::BT_SOFTLIGHT: BufferLayer<SoftLightChannelMasked, SoftLightChannel, true>(Result->GetData(), Result.get(), Mask.get(), Blended.get(), bApplyColorBlendToAlpha, bOnlyOneMip); break;
-			case EBlendType::BT_HARDLIGHT: BufferLayer<HardLightChannelMasked, HardLightChannel, true>(Result->GetData(), Result.get(), Mask.get(), Blended.get(), bApplyColorBlendToAlpha, bOnlyOneMip); break;
-			case EBlendType::BT_BURN: BufferLayer<BurnChannelMasked, BurnChannel, true>(Result->GetData(), Result.get(), Mask.get(), Blended.get(), bApplyColorBlendToAlpha, bOnlyOneMip); break;
-			case EBlendType::BT_DODGE: BufferLayer<DodgeChannelMasked, DodgeChannel, true>(Result->GetData(), Result.get(), Mask.get(), Blended.get(), bApplyColorBlendToAlpha, bOnlyOneMip); break;
-			case EBlendType::BT_SCREEN: BufferLayer<ScreenChannelMasked, ScreenChannel, true>(Result->GetData(), Result.get(), Mask.get(), Blended.get(), bApplyColorBlendToAlpha, bOnlyOneMip); break;
-			case EBlendType::BT_OVERLAY: BufferLayer<OverlayChannelMasked, OverlayChannel, true>(Result->GetData(), Result.get(), Mask.get(), Blended.get(), bApplyColorBlendToAlpha, bOnlyOneMip); break;
-			case EBlendType::BT_LIGHTEN: BufferLayer<LightenChannelMasked, LightenChannel, true>(Result->GetData(), Result.get(), Mask.get(), Blended.get(), bApplyColorBlendToAlpha, bOnlyOneMip); break;
-			case EBlendType::BT_MULTIPLY: BufferLayer<MultiplyChannelMasked, MultiplyChannel, true>(Result->GetData(), Result.get(), Mask.get(), Blended.get(), bApplyColorBlendToAlpha, bOnlyOneMip); break;
-			case EBlendType::BT_BLEND: BufferLayer<BlendChannelMasked, BlendChannel, true>(Result->GetData(), Result.get(), Mask.get(), Blended.get(), bApplyColorBlendToAlpha, bOnlyOneMip); break;
+			case EBlendType::BT_SOFTLIGHT: BufferLayer<SoftLightChannelMasked, SoftLightChannel, true>(Result.get(), Result.get(), Mask.get(), Blended.get(), bApplyColorBlendToAlpha, bOnlyOneMip); break;
+			case EBlendType::BT_HARDLIGHT: BufferLayer<HardLightChannelMasked, HardLightChannel, true>(Result.get(), Result.get(), Mask.get(), Blended.get(), bApplyColorBlendToAlpha, bOnlyOneMip); break;
+			case EBlendType::BT_BURN: BufferLayer<BurnChannelMasked, BurnChannel, true>(Result.get(), Result.get(), Mask.get(), Blended.get(), bApplyColorBlendToAlpha, bOnlyOneMip); break;
+			case EBlendType::BT_DODGE: BufferLayer<DodgeChannelMasked, DodgeChannel, true>(Result.get(), Result.get(), Mask.get(), Blended.get(), bApplyColorBlendToAlpha, bOnlyOneMip); break;
+			case EBlendType::BT_SCREEN: BufferLayer<ScreenChannelMasked, ScreenChannel, true>(Result.get(), Result.get(), Mask.get(), Blended.get(), bApplyColorBlendToAlpha, bOnlyOneMip); break;
+			case EBlendType::BT_OVERLAY: BufferLayer<OverlayChannelMasked, OverlayChannel, true>(Result.get(), Result.get(), Mask.get(), Blended.get(), bApplyColorBlendToAlpha, bOnlyOneMip); break;
+			case EBlendType::BT_LIGHTEN: BufferLayer<LightenChannelMasked, LightenChannel, true>(Result.get(), Result.get(), Mask.get(), Blended.get(), bApplyColorBlendToAlpha, bOnlyOneMip); break;
+			case EBlendType::BT_MULTIPLY: BufferLayer<MultiplyChannelMasked, MultiplyChannel, true>(Result.get(), Result.get(), Mask.get(), Blended.get(), bApplyColorBlendToAlpha, bOnlyOneMip); break;
+			case EBlendType::BT_BLEND: BufferLayer<BlendChannelMasked, BlendChannel, true>(Result.get(), Result.get(), Mask.get(), Blended.get(), bApplyColorBlendToAlpha, bOnlyOneMip); break;
 			case EBlendType::BT_NONE: break;
 			default: check(false);
 			}
@@ -778,15 +779,15 @@ namespace mu
 			switch (EBlendType(Args.blendType))
 			{
 			case EBlendType::BT_NORMAL_COMBINE: check(false); break;
-			case EBlendType::BT_SOFTLIGHT: BufferLayerEmbeddedMask<SoftLightChannelMasked, SoftLightChannel, false>(Result->GetData(), Result.get(), Blended.get(), bApplyColorBlendToAlpha, bOnlyOneMip); break;
-			case EBlendType::BT_HARDLIGHT: BufferLayerEmbeddedMask<HardLightChannelMasked, HardLightChannel, false>(Result->GetData(), Result.get(), Blended.get(), bApplyColorBlendToAlpha, bOnlyOneMip); break;
-			case EBlendType::BT_BURN: BufferLayerEmbeddedMask<BurnChannelMasked, BurnChannel, false>(Result->GetData(), Result.get(), Blended.get(), bApplyColorBlendToAlpha, bOnlyOneMip); break;
-			case EBlendType::BT_DODGE: BufferLayerEmbeddedMask<DodgeChannelMasked, DodgeChannel, false>(Result->GetData(), Result.get(), Blended.get(), bApplyColorBlendToAlpha, bOnlyOneMip); break;
-			case EBlendType::BT_SCREEN: BufferLayerEmbeddedMask<ScreenChannelMasked, ScreenChannel, false>(Result->GetData(), Result.get(), Blended.get(), bApplyColorBlendToAlpha, bOnlyOneMip); break;
-			case EBlendType::BT_OVERLAY: BufferLayerEmbeddedMask<OverlayChannelMasked, OverlayChannel, false>(Result->GetData(), Result.get(), Blended.get(), bApplyColorBlendToAlpha, bOnlyOneMip); break;
-			case EBlendType::BT_LIGHTEN: BufferLayerEmbeddedMask<LightenChannelMasked, LightenChannel, false>(Result->GetData(), Result.get(), Blended.get(), bApplyColorBlendToAlpha, bOnlyOneMip); break;
-			case EBlendType::BT_MULTIPLY: BufferLayerEmbeddedMask<MultiplyChannelMasked, MultiplyChannel, false>(Result->GetData(), Result.get(), Blended.get(), bApplyColorBlendToAlpha, bOnlyOneMip); break;
-			case EBlendType::BT_BLEND: BufferLayerEmbeddedMask<BlendChannelMasked, BlendChannel, false>(Result->GetData(), Result.get(), Blended.get(), bApplyColorBlendToAlpha, bOnlyOneMip); break;
+			case EBlendType::BT_SOFTLIGHT: BufferLayerEmbeddedMask<SoftLightChannelMasked, SoftLightChannel, false>(Result.get(), Result.get(), Blended.get(), bApplyColorBlendToAlpha, bOnlyOneMip); break;
+			case EBlendType::BT_HARDLIGHT: BufferLayerEmbeddedMask<HardLightChannelMasked, HardLightChannel, false>(Result.get(), Result.get(), Blended.get(), bApplyColorBlendToAlpha, bOnlyOneMip); break;
+			case EBlendType::BT_BURN: BufferLayerEmbeddedMask<BurnChannelMasked, BurnChannel, false>(Result.get(), Result.get(), Blended.get(), bApplyColorBlendToAlpha, bOnlyOneMip); break;
+			case EBlendType::BT_DODGE: BufferLayerEmbeddedMask<DodgeChannelMasked, DodgeChannel, false>(Result.get(), Result.get(), Blended.get(), bApplyColorBlendToAlpha, bOnlyOneMip); break;
+			case EBlendType::BT_SCREEN: BufferLayerEmbeddedMask<ScreenChannelMasked, ScreenChannel, false>(Result.get(), Result.get(), Blended.get(), bApplyColorBlendToAlpha, bOnlyOneMip); break;
+			case EBlendType::BT_OVERLAY: BufferLayerEmbeddedMask<OverlayChannelMasked, OverlayChannel, false>(Result.get(), Result.get(), Blended.get(), bApplyColorBlendToAlpha, bOnlyOneMip); break;
+			case EBlendType::BT_LIGHTEN: BufferLayerEmbeddedMask<LightenChannelMasked, LightenChannel, false>(Result.get(), Result.get(), Blended.get(), bApplyColorBlendToAlpha, bOnlyOneMip); break;
+			case EBlendType::BT_MULTIPLY: BufferLayerEmbeddedMask<MultiplyChannelMasked, MultiplyChannel, false>(Result.get(), Result.get(), Blended.get(), bApplyColorBlendToAlpha, bOnlyOneMip); break;
+			case EBlendType::BT_BLEND: BufferLayerEmbeddedMask<BlendChannelMasked, BlendChannel, false>(Result.get(), Result.get(), Blended.get(), bApplyColorBlendToAlpha, bOnlyOneMip); break;
 			case EBlendType::BT_NONE: break;
 			default: check(false);
 			}
@@ -950,14 +951,16 @@ namespace mu
 			if (Mask->GetLODCount() < Base->GetLODCount())
 			{
 				MUTABLE_CPUPROFILER_SCOPE(ImageResize_EmergencyFixMips);
-				int32 levelCount = Base->GetLODCount();
-				Ptr<Image> pDest = Runner->CreateImage(Mask->GetSizeX(), Mask->GetSizeY(), levelCount, Mask->GetFormat(), EInitializationType::NotInitialized);
+				int32 StartLevel = Mask->GetLODCount() - 1;
+				int32 LevelCount = Base->GetLODCount();
 
-				FMipmapGenerationSettings settings{};
-				ImOp.ImageMipmap(ImageCompressionQuality, pDest.get(), Mask.get(), levelCount, settings);
+				Ptr<Image> MaskFix = Runner->CloneOrTakeOver(Mask);
+				MaskFix->DataStorage.SetNumLODs(LevelCount);
 
-				Runner->Release(Mask);
-				Mask = pDest;
+				FMipmapGenerationSettings Settings{};
+				ImOp.ImageMipmap(ImageCompressionQuality, MaskFix.get(), MaskFix.get(), StartLevel, LevelCount, Settings);
+
+				Mask = MaskFix;
 			}
 		}
 
@@ -987,15 +990,15 @@ namespace mu
 				switch (EBlendType(Args.blendType))
 				{
 				case EBlendType::BT_NORMAL_COMBINE: ImageNormalCombine(Result.get(), Result.get(), Mask.get(), Color); break;
-				case EBlendType::BT_SOFTLIGHT: BufferLayerColour<SoftLightChannelMasked, SoftLightChannel>(Result->GetData(), Result.get(), Mask.get(), Color); break;
-				case EBlendType::BT_HARDLIGHT: BufferLayerColour<HardLightChannelMasked, HardLightChannel>(Result->GetData(), Result.get(), Mask.get(), Color); break;
-				case EBlendType::BT_BURN: BufferLayerColour<BurnChannelMasked, BurnChannel>(Result->GetData(), Result.get(), Mask.get(), Color); break;
-				case EBlendType::BT_DODGE: BufferLayerColour<DodgeChannelMasked, DodgeChannel>(Result->GetData(), Result.get(), Mask.get(), Color); break;
-				case EBlendType::BT_SCREEN: BufferLayerColour<ScreenChannelMasked, ScreenChannel>(Result->GetData(), Result.get(), Mask.get(), Color); break;
-				case EBlendType::BT_OVERLAY: BufferLayerColour<OverlayChannelMasked, OverlayChannel>(Result->GetData(), Result.get(), Mask.get(), Color); break;
-				case EBlendType::BT_LIGHTEN: BufferLayerColour<LightenChannelMasked, LightenChannel>(Result->GetData(), Result.get(), Mask.get(), Color); break;
-				case EBlendType::BT_MULTIPLY: BufferLayerColour<MultiplyChannelMasked, MultiplyChannel>(Result->GetData(), Result.get(), Mask.get(), Color); break;
-				case EBlendType::BT_BLEND: BufferLayerColour<BlendChannelMasked, BlendChannel>(Result->GetData(), Result.get(), Mask.get(), Color); break;
+				case EBlendType::BT_SOFTLIGHT: BufferLayerColour<SoftLightChannelMasked, SoftLightChannel>(Result.get(), Result.get(), Mask.get(), Color); break;
+				case EBlendType::BT_HARDLIGHT: BufferLayerColour<HardLightChannelMasked, HardLightChannel>(Result.get(), Result.get(), Mask.get(), Color); break;
+				case EBlendType::BT_BURN: BufferLayerColour<BurnChannelMasked, BurnChannel>(Result.get(), Result.get(), Mask.get(), Color); break;
+				case EBlendType::BT_DODGE: BufferLayerColour<DodgeChannelMasked, DodgeChannel>(Result.get(), Result.get(), Mask.get(), Color); break;
+				case EBlendType::BT_SCREEN: BufferLayerColour<ScreenChannelMasked, ScreenChannel>(Result.get(), Result.get(), Mask.get(), Color); break;
+				case EBlendType::BT_OVERLAY: BufferLayerColour<OverlayChannelMasked, OverlayChannel>(Result.get(), Result.get(), Mask.get(), Color); break;
+				case EBlendType::BT_LIGHTEN: BufferLayerColour<LightenChannelMasked, LightenChannel>(Result.get(), Result.get(), Mask.get(), Color); break;
+				case EBlendType::BT_MULTIPLY: BufferLayerColour<MultiplyChannelMasked, MultiplyChannel>(Result.get(), Result.get(), Mask.get(), Color); break;
+				case EBlendType::BT_BLEND: BufferLayerColour<BlendChannelMasked, BlendChannel>(Result.get(), Result.get(), Mask.get(), Color); break;
 				default: check(false);
 				}
 
@@ -1184,21 +1187,7 @@ namespace mu
 		bool bSuccess = false;
 		FImageOperator ImOp = FImageOperator::GetDefault(ImagePixelFormatFunc);
 		ImOp.ImagePixelFormat(bSuccess, ImageCompressionQuality, Result.get(), Base.get(), -1);
-
-		int32 OriginalDataSize = FMath::Max(Result->m_data.Num(), Base->m_data.Num());
-		int32 ExcessDataSize = OriginalDataSize * 4 + 4;
-		while (!bSuccess)
-		{
-			MUTABLE_CPUPROFILER_SCOPE(Recompression_OutOfSpace);
-
-			// Bad case, where the RLE compressed data requires more memory than the uncompressed data.
-			// We need to support it anyway for small mips or scaled images.
-			Result->m_data.SetNum(ExcessDataSize);
-			bSuccess = false;
-			ImOp.ImagePixelFormat(bSuccess,ImageCompressionQuality, Result.get(), Base.get(), -1);
-			ExcessDataSize *= 4;
-		}
-
+		check(bSuccess);
 	}
 
 
@@ -1214,10 +1203,6 @@ namespace mu
 		}
 	}
 
-
-	//---------------------------------------------------------------------------------------------
-	//---------------------------------------------------------------------------------------------
-	//---------------------------------------------------------------------------------------------
 	class FImageMipmapTask : public CodeRunner::FIssuedTask
 	{
 	public:
@@ -1231,7 +1216,8 @@ namespace mu
 		void Complete(CodeRunner*) override;
 
 	private:
-		int ImageCompressionQuality = 0;
+		int32 ImageCompressionQuality = 0;
+		int32 StartLevel = -1;
 		OP::ImageMipmapArgs Args;
 		Ptr<const Image> Base;
 		Ptr<Image> Result;
@@ -1279,8 +1265,7 @@ namespace mu
 		}
 
 		// At least keep the levels we already have.
-		int StartLevel = Base->GetLODCount();
-		LevelCount = FMath::Max(StartLevel, LevelCount);
+		LevelCount = FMath::Max(Base->GetLODCount(), LevelCount);
 		
 		if (LevelCount == Base->GetLODCount())
 		{
@@ -1288,12 +1273,15 @@ namespace mu
 			return false;
 		}
 
-		// Create destination data
-		Result = Runner->CreateImage(Base->GetSizeX(), Base->GetSizeY(), LevelCount, Base->GetFormat(), EInitializationType::NotInitialized);
-		Result->m_flags = Base->m_flags;
+		StartLevel = Base->GetLODCount() - 1;
+		
+		Result = Runner->CloneOrTakeOver(Base);
+		Base = nullptr;
+
+		Result->DataStorage.SetNumLODs(LevelCount);
 
 		FImageOperator ImOp = MakeImageOperator(Runner);
-		ImOp.ImageMipmap_PrepareScratch( Result.get(), Base.get(), LevelCount, Scratch);
+		ImOp.ImageMipmap_PrepareScratch(Result.get(), StartLevel, LevelCount, Scratch);
 
 		ImagePixelFormatFunc = Runner->m_pSystem->ImagePixelFormatOverride;
 
@@ -1301,26 +1289,29 @@ namespace mu
 	}
 
 
-	//---------------------------------------------------------------------------------------------
 	void FImageMipmapTask::DoWork()
 	{
 		// This runs in a worker thread
 		MUTABLE_CPUPROFILER_SCOPE(FImageMipmapTask);
 
-		FMipmapGenerationSettings settings{};
+		check(StartLevel >= 0);
+
+		FMipmapGenerationSettings Settings{};
 		FImageOperator ImOp = FImageOperator::GetDefault(ImagePixelFormatFunc);
-		ImOp.ImageMipmap(Scratch, ImageCompressionQuality, Result.get(), Base.get(), Result->GetLODCount(), settings);
+		ImOp.ImageMipmap(Scratch, ImageCompressionQuality, Result.get(), Result.get(), StartLevel, Result->GetLODCount(), Settings);
 	}
 
 
-	//---------------------------------------------------------------------------------------------
 	void FImageMipmapTask::Complete(CodeRunner* Runner)
 	{
 		FImageOperator ImOp = MakeImageOperator(Runner);
 		ImOp.ImageMipmap_ReleaseScratch(Scratch);
 
 		// This runs in the Runner thread
-		Runner->Release(Base);
+		if (Base)
+		{
+			Runner->Release(Base);
+		}
 
 		// If no shortcut was taken
 		if (Result)
@@ -1519,8 +1510,8 @@ namespace mu
 		// This runs on a random worker thread
 		MUTABLE_CPUPROFILER_SCOPE(FImageSaturateTask);
 
-		constexpr bool bUseVectorIntrinsics = false;
-		ImageSaturateInPlace<bUseVectorIntrinsics>(Result.get(), Factor);
+		constexpr bool bUseVectorIntrinsics = true;
+		ImageSaturate<bUseVectorIntrinsics>(Result.get(), Factor);
 	}
 
 	//---------------------------------------------------------------------------------------------
@@ -1794,7 +1785,7 @@ namespace mu
 		// This runs on a random worker thread
 		MUTABLE_CPUPROFILER_SCOPE(FImageInvertTask);
 
-		ImageInvertInPlace(Result.get());
+		ImageInvert(Result.get());
 	}
 
 
@@ -2417,7 +2408,6 @@ namespace mu
 
 		Runner->StoreImage(Op, Result);
 	}	
-
 	
 	//---------------------------------------------------------------------------------------------
 	//---------------------------------------------------------------------------------------------
