[Mutable] Store Image LODs in different buffers.

diff --git a/Engine/Plugins/Experimental/Mutable/Source/MutableRuntime/Private/MuR/Image.cpp b/Engine/Plugins/Experimental/Mutable/Source/MutableRuntime/Private/MuR/Image.cpp
--- a/Engine/Plugins/Experimental/Mutable/Source/MutableRuntime/Private/MuR/Image.cpp
+++ b/Engine/Plugins/Experimental/Mutable/Source/MutableRuntime/Private/MuR/Image.cpp
@@ -3,300 +3,233 @@
 #include "MuR/Image.h"
 
 #include "HAL/UnrealMemory.h"
-#include "HAL/IConsoleManager.h"
 #include "Math/UnrealMathSSE.h"
 #include "Math/NumericLimits.h"
 #include "MuR/ImagePrivate.h"
 #include "MuR/MutableMath.h"
 #include "MuR/MutableTrace.h"
 
-namespace
-{
-	static bool bDisableCompressedImageBlackBlockInit = false;
-	static FAutoConsoleVariableRef CVarDisableCompressedImageBlackBlockInit(
-		TEXT("mutable.DisableCompressedImageBlackBlockInit"),
-		bDisableCompressedImageBlackBlockInit,
-		TEXT("A value of 1 disables mutable compressed black block initialization."),
-		ECVF_Default);
-}
-
 namespace mu
 {
 
-    //---------------------------------------------------------------------------------------------
-    Image::Image()
-    {
-    }
-
+	Image::Image()
+	{
+	}
 
-    //---------------------------------------------------------------------------------------------
 	Image::Image(uint32 SizeX, uint32 SizeY, uint32 Lods, EImageFormat Format, EInitializationType InitType)
 	{
 		Init(SizeX, SizeY, Lods, Format, InitType);
 	}
-
 	
 	void Image::Init(uint32 SizeX, uint32 SizeY, uint32 Lods, EImageFormat Format, EInitializationType InitType)
 	{
 		MUTABLE_CPUPROFILER_SCOPE(NewImage)
 		LLM_SCOPE_BYNAME(TEXT("MutableRuntime"));
 
-        check(Format != EImageFormat::IF_NONE);
-		check(Format < EImageFormat::IF_COUNT);
-
-		check(SizeX <= TNumericLimits<uint16>::Max());
-		check(SizeY <= TNumericLimits<uint16>::Max());
-		check(Lods <= TNumericLimits<uint8>::Max());
-
-		// TODO: check that lods is sensible for the size.
-
-        m_format = Format;
-        m_size = FImageSize( (uint16)SizeX, (uint16)SizeY );
-        m_lods = (uint8)Lods;
-
-        const FImageFormatData& fdata = GetImageFormatData( Format );
-        int32 PixelsPerBlock = fdata.PixelsPerBlockX*fdata.PixelsPerBlockY;
-        if (PixelsPerBlock)
-        {
-			if (InitType == EInitializationType::Black)
-			{
-				InitToBlack();
-			}
-			else
-			{
-				int32 DataSize = CalculateDataSize();
-				m_data.SetNumUninitialized(DataSize);
-			}
-        }
-    }
+		const FImageDesc ImageDesc {FImageSize(uint16(SizeX), uint16(SizeY)), Format, static_cast<uint8>(Lods)};
+		DataStorage.Init(ImageDesc, InitType);
+	}
 
 
-	//---------------------------------------------------------------------------------------------
 	void Image::InitToBlack()
 	{
-		int32 DataSize = CalculateDataSize();
-		if (bDisableCompressedImageBlackBlockInit)
-		{
-			m_data.SetNumUninitialized(DataSize);
-			// Do it in separate steps in case we are reusing a buffer.
-			FMemory::Memzero(m_data.GetData(), DataSize);
-		}
-		else
-		{
-			const FImageFormatData& fdata = GetImageFormatData(m_format);
-			check(fdata.BytesPerBlock <= FImageFormatData::MAX_BYTES_PER_BLOCK);
-
-			constexpr uint8 ZeroedBlock[FImageFormatData::MAX_BYTES_PER_BLOCK] = { 0 };
-
-			const SIZE_T BlockSizeSanitized = FMath::Min<SIZE_T>(FImageFormatData::MAX_BYTES_PER_BLOCK, fdata.BytesPerBlock);
-			const bool bIsFormatBlackBlockZeroed = FMemory::Memcmp(fdata.BlackBlock, ZeroedBlock, BlockSizeSanitized) == 0;
-
-			if (bIsFormatBlackBlockZeroed)
-			{
-				m_data.SetNumUninitialized(DataSize);
-				// Do it in separate steps in case we are reusing a buffer.
-				FMemory::Memzero(m_data.GetData(), DataSize);
-			}
-			else
-			{
-				m_data.SetNumUninitialized(DataSize);
-
-				check(fdata.BytesPerBlock > 0);
-				for (int32 BlockDataOffset = 0; BlockDataOffset < DataSize; BlockDataOffset += fdata.BytesPerBlock)
-				{
-					FMemory::Memcpy(m_data.GetData() + BlockDataOffset, fdata.BlackBlock, fdata.BytesPerBlock);
-				}
-			}
-		}
+		const FImageDesc Desc = DataStorage.MakeImageDesc();
+		DataStorage.Init(Desc, EInitializationType::Black);
 
 		m_flags = 0;
 		RelevancyMinY = 0;
 		RelevancyMaxY = 0;
 	}
 
-	//---------------------------------------------------------------------------------------------
 	Ptr<Image> Image::CreateAsReference(uint32 ID, const FImageDesc& Desc, bool bForceLoad)
 	{
 		Ptr<Image> Result = new Image;
 		Result->ReferenceID = ID;
-		Result->m_size = Desc.m_size;
-		Result->m_lods = Desc.m_lods;
-		Result->m_format = Desc.m_format;
+	
+		Result->DataStorage.InitVoid(Desc);
+
 		Result->m_flags = EImageFlags::IF_IS_REFERENCE;
 		if (bForceLoad)
 		{
 			Result->m_flags = Result->m_flags | EImageFlags::IF_IS_FORCELOAD;
 		}
+
 		return Result;
 	}
 
 
-    //---------------------------------------------------------------------------------------------
-    void Image::Serialise( const Image* p, OutputArchive& arch )
-    {
-        arch << *p;
-    }
+	void Image::Serialise(const Image* ValuePtr, OutputArchive& Arch)
+	{
+		Arch << *ValuePtr;
+	}
 
-	void Image::Serialise(OutputArchive& arch) const
+	void Image::Serialise(OutputArchive& Arch) const
 	{
-		uint32 ver = 3;
-		arch << ver;
+		uint32 Version = 4;
+		Arch << Version;
 
-		arch << m_size;
-		arch << m_lods;
-		arch << (uint8)m_format;
-		arch << m_data;
+		Arch << DataStorage;
 
 		// Remove non-persistent flags.
 		uint8 flags = m_flags & ~IF_HAS_RELEVANCY_MAP;
-		arch << flags;
+		Arch << flags;
 	}
 
-	void Image::Unserialise(InputArchive& arch)
+	void Image::Unserialise(InputArchive& Arch)
 	{
-		uint32 ver;
-		arch >> ver;
-		check(ver == 3);
+		uint32 Version = 0;
+		Arch >> Version;
+		check(Version <= 4);
 
-		arch >> m_size;
-		arch >> m_lods;
+		if (Version <= 3)
+		{
+			FImageSize Size;
+			Arch >> Size;
+			
+			uint8 LODs;
+			Arch >> LODs;
+
+			uint8 Format;
+			Arch >> Format;
 
-		uint8 format;
-		arch >> format;
-		m_format = (EImageFormat)format;
-		arch >> m_data;
+			FImageArray OldData;
+			Arch >> OldData;
 
-		arch >> m_flags;
+
+			// This does not work for RLE compressed formats.
+			if (!OldData.IsEmpty())
+			{
+				DataStorage.Init(FImageDesc(Size, static_cast<EImageFormat>(Format), LODs), EInitializationType::NotInitialized);
+				check(!DataStorage.IsEmpty());
+				
+				int32 OldDataConsumedOffset = 0;
+				for (FImageArray& Buffer : DataStorage.Buffers)
+				{
+					check(OldData.Num() >= OldDataConsumedOffset + Buffer.Num());
+					FMemory::Memcpy(Buffer.GetData(), OldData.GetData() + OldDataConsumedOffset, Buffer.Num());
+
+					OldDataConsumedOffset += Buffer.Num();
+				}
+			}
+		}
+		else if (Version >= 4)
+		{
+			Arch >> DataStorage;
+		}
+		else
+		{
+			check(false);
+		}
+
+		Arch >> m_flags;
 	}
 
 
-    //---------------------------------------------------------------------------------------------
-    Ptr<Image> Image::StaticUnserialise( InputArchive& arch )
-    {
+	Ptr<Image> Image::StaticUnserialise(InputArchive& Arch)
+	{
 		LLM_SCOPE_BYNAME(TEXT("MutableRuntime"));
-        Ptr<Image> pResult = new Image();
-        arch >> *pResult;
-        return pResult;
-    }
 
+		Ptr<Image> ResultPtr = new Image();
+		Arch >> *ResultPtr;
+		
+		return ResultPtr;
+	}
+
+	int32 Image::GetDataSize() const
+	{
+		return DataStorage.GetAllocatedSize();
+	}
 
-	//---------------------------------------------------------------------------------------------
 	Ptr<Image> FImageOperator::ExtractMip(const Image* This, int32 Mip)
 	{
-		if (Mip == 0 && This->m_lods == 1)
+		if (Mip == 0 && This->GetLODCount() == 1)
 		{
 			return CloneImage(This);
 		}
 
 		FIntVector2 MipSize = This->CalculateMipSize(Mip);
 
-		int32 Quality = 4;
+		constexpr int32 Quality = 4;
 
-		if (This->m_lods > Mip)
+		if (This->GetLODCount() > Mip)
 		{
-			const uint8* SourceData = This->GetMipData(Mip);
-			int32 DataSize = This->CalculateDataSize(Mip);
+			TArrayView<const uint8> SrcView = This->DataStorage.GetLOD(Mip); 
 
-			if (DataSize)
-			{
-				Ptr<Image> pResult = CreateImage(MipSize[0], MipSize[1], 1, This->m_format, EInitializationType::NotInitialized);
-				pResult->m_flags = This->m_flags;
-				uint8* DestData = pResult->GetData();
-				FMemory::Memcpy(DestData, SourceData, DataSize);
-				return pResult;
-			}
-			else
+			Ptr<Image> ResultImage = CreateImage(MipSize.X, MipSize.Y, 1, This->GetFormat(), EInitializationType::NotInitialized);
+			ResultImage->m_flags = This->m_flags;
+			
+			// Probably an RLE texture, resize.
+			if (ResultImage->DataStorage.IsEmpty())
 			{
-				EImageFormat uncompressedFormat = GetUncompressedFormat(This->m_format);
-				// TODO: OnlyLOD=Mip and then ExtractMip(0)?
-				Ptr<Image> Temp = ImagePixelFormat(Quality, This, uncompressedFormat);
-				Temp = ExtractMip(Temp.get(), Mip);
-				Ptr<Image> Result = ImagePixelFormat(Quality, Temp.get(), This->m_format);
-				Result->m_flags = This->m_flags;
-				ReleaseImage(Temp);
-				return Result;
+				ResultImage->DataStorage.ResizeLOD(0, SrcView.Num());
 			}
+
+			TArrayView<uint8> DestView = ResultImage->DataStorage.GetLOD(0);
+			
+			check(DestView.Num() == SrcView.Num())
+			FMemory::Memcpy(DestView.GetData(), SrcView.GetData(), DestView.Num());
+			
+			return ResultImage;
 		}
 
 		// We need to generate the mip
 		// \TODO: optimize, quality
-		Ptr<Image> Resized = CreateImage(MipSize[0], MipSize[1],1,This->GetFormat(),EInitializationType::NotInitialized);
+
+		Ptr<Image> Resized = CreateImage(MipSize.X, MipSize.Y, 1, This->GetFormat(), EInitializationType::NotInitialized);
 		ImageResizeLinear(Resized.get(), Quality, This);
-		Ptr<Image> Result = ExtractMip( Resized.get(), 0);
+		
+		//TODO: Do we realy need to extract the mip again? The resized Image looks already to be what we want.
+		Ptr<Image> Result = ExtractMip(Resized.get(), 0);
 		ReleaseImage(Resized);
+
 		return Result;
 	}
+	
+	uint16 Image::GetSizeX() const
+	{
+		return GetSize().X;
+	}
 
+	uint16 Image::GetSizeY() const
+	{
+		return GetSize().Y;
+	}
 
-    //---------------------------------------------------------------------------------------------
-    uint16 Image::GetSizeX() const
-    {
-        return m_size[0];
-    }
-
-
-    //---------------------------------------------------------------------------------------------
-    uint16 Image::GetSizeY() const
-    {
-        return m_size[1];
-    }
-
-
-	//---------------------------------------------------------------------------------------------
 	const FImageSize& Image::GetSize() const
 	{
-		return m_size;
+		return DataStorage.ImageSize;
 	}
 
-
-    //---------------------------------------------------------------------------------------------
 	EImageFormat Image::GetFormat() const
-    {
-        return m_format;
-    }
-
-
-    //---------------------------------------------------------------------------------------------
-    int Image::GetLODCount() const
-    {
-        return FMath::Max( 1, (int)m_lods );
-    }
-
-
-    //---------------------------------------------------------------------------------------------
-    const uint8* Image::GetData() const
-    {
-        return m_data.GetData();
-    }
-
-
-    //---------------------------------------------------------------------------------------------
-    int32 Image::GetDataSize() const
-    {
-        return m_data.Num();
-    }
+	{
+		return DataStorage.ImageFormat;
+	}
 
+	int32 Image::GetLODCount() const
+	{
+		return FMath::Max(1, (int32)DataStorage.NumLODs);
+	}
 
-    //---------------------------------------------------------------------------------------------
-    int32 Image::GetLODDataSize( int lod ) const
-    {
-        return CalculateDataSize( lod );
-    }
+	const uint8* Image::GetLODData(int32 LOD) const
+	{
+		return DataStorage.GetLOD(LOD).GetData();
+	}
 
+	uint8* Image::GetLODData(int32 LOD)
+	{
+		return DataStorage.GetLOD(LOD).GetData();
+	}
 
-    //---------------------------------------------------------------------------------------------
-    uint8* Image::GetData()
-    {
-		return m_data.GetData();
-    }
+	int32 Image::GetLODDataSize(int32 LOD) const
+	{
+		return DataStorage.GetLOD(LOD).Num();
+	}
 
-	//---------------------------------------------------------------------------------------------
 	bool Image::IsReference() const
 	{
 		return m_flags & EImageFlags::IF_IS_REFERENCE;
 	}
 
-	//---------------------------------------------------------------------------------------------
 	bool Image::IsForceLoad() const
 	{
 		return m_flags & EImageFlags::IF_IS_FORCELOAD;
@@ -309,451 +242,325 @@ namespace mu
 		return ReferenceID;
 	}
 
-
-    //---------------------------------------------------------------------------------------------
-    //---------------------------------------------------------------------------------------------
-	//---------------------------------------------------------------------------------------------
-	int32 Image::CalculateDataSize( int32 SizeX, int32 SizeY, int32 LodCount, EImageFormat Format )
+	int32 Image::CalculateDataSize(int32 SizeX, int32 SizeY, int32 LODCount, EImageFormat Format)
 	{
-		int32 res = 0;
+		int32 Result = 0;
 
-		const FImageFormatData& fdata = GetImageFormatData(Format);
-		if (fdata.BytesPerBlock)
+		const FImageFormatData& FormatData = GetImageFormatData(Format);
+		if (FormatData.BytesPerBlock)
 		{
-			for (int32 LodIndex = 0; LodIndex < FMath::Max(1, LodCount); ++LodIndex)
+			for (int32 LODIndex = 0; LODIndex < FMath::Max(1, LODCount); ++LODIndex)
 			{
-				int32 blocksX = FMath::DivideAndRoundUp(SizeX, int32(fdata.PixelsPerBlockX));
-				int32 blocksY = FMath::DivideAndRoundUp(SizeY, int32(fdata.PixelsPerBlockY));
+				int32 BlocksX = FMath::DivideAndRoundUp(SizeX, int32(FormatData.PixelsPerBlockX));
+				int32 BlocksY = FMath::DivideAndRoundUp(SizeY, int32(FormatData.PixelsPerBlockY));
 
-				res += (blocksX * blocksY) * fdata.BytesPerBlock;
+				Result += (BlocksX * BlocksY) * FormatData.BytesPerBlock;
 
 				SizeX = FMath::DivideAndRoundUp(SizeX, 2);
 				SizeY = FMath::DivideAndRoundUp(SizeY, 2);
 			}
 		}
 
-		return res;
+		return Result;
 	}
 
 
-	//---------------------------------------------------------------------------------------------
-	int32 Image::CalculateDataSize() const
-	{
-		return CalculateDataSize( m_size[0], m_size[1], m_lods, m_format );
-	}
+	//int32 Image::CalculateDataSize() const
+	//{
+	//	check(false);
+	//	return DataStorage.GetDataSize();
+	//}
+
+	//int32 Image::CalculateDataSize(int32 Mip) const
+	//{
+	//	check(false);
+	//	return DataStorage.GetLOD(Mip).Num();
+	//}
+
+	//int32 Image::CalculatePixelCount() const
+	//{
+	//	check(false);
+	//	int32 Result = 0;
+
+	//	const int32 LODCount = FMath::Max(1, GetLODCount());
+	//	const FImageFormatData& FormatData = GetImageFormatData(GetFormat());
+	//	
+	//	if (FormatData.BytesPerBlock)
+	//	{
+	//		FImageSize Size = GetSize();
+
+	//		for (int32 L = 0; L < LODCount; ++L)
+	//		{
+	//			int32 BlocksX = FMath::DivideAndRoundUp(Size[0], (uint16)FormatData.PixelsPerBlockX);
+	//			int32 BlocksY = FMath::DivideAndRoundUp(Size[1], (uint16)FormatData.PixelsPerBlockY);
+
+	//			Result += (BlocksX*BlocksY) * FormatData.PixelsPerBlockX * FormatData.PixelsPerBlockY;
+
+	//			Size[0] = FMath::DivideAndRoundUp(Size[0], (uint16)2);
+	//			Size[1] = FMath::DivideAndRoundUp(Size[1], (uint16)2);
+	//		}
+	//	}
+	//	else
+	//	{
+	//		// An RLE image.
+	//		for (int32 L = 0; L < LODCount; ++L)
+	//		{
+	//			FIntVector2 MipSize = CalculateMipSize(L);
+	//			Result += MipSize[0] * MipSize[1];
+	//		}
+	//	}
+
+	//	return Result;
+	//}
+
+
+	////---------------------------------------------------------------------------------------------
+	//int32 Image::CalculatePixelCount(int32 Mip) const
+	//{
+	//	check(false);
+	//	int32 Result = 0;
+	//	
+	//	const FImageFormatData& FormatData = GetImageFormatData(GetFormat());
+	//	if (FormatData.BytesPerBlock)
+	//	{
+	//		FImageSize Size = GetSize();
+
+	//		for (int32 L = 0; L < Mip + 1; ++L)
+	//		{
+	//			int32 BlocksX = FMath::DivideAndRoundUp(Size[0], (uint16)FormatData.PixelsPerBlockX);
+	//			int32 BlocksY = FMath::DivideAndRoundUp(Size[1], (uint16)FormatData.PixelsPerBlockY);
+
+	//			if (L == Mip)
+	//			{
+	//				Result = (BlocksX*BlocksY) * FormatData.PixelsPerBlockX * FormatData.PixelsPerBlockY;
+	//				break;
+	//			}
+
+	//			Size[0] = FMath::DivideAndRoundUp(Size[0], (uint16)2);
+	//			Size[1] = FMath::DivideAndRoundUp(Size[1], (uint16)2);
+	//		}
+	//	}
+	//	else
+	//	{
+	//		// An RLE image.
+	//		FIntVector2 MipSize = CalculateMipSize(Mip);
+	//		Result += MipSize[0] * MipSize[1];
+	//	}
+
+	//	return Result;
+	//}
 
 
 	//---------------------------------------------------------------------------------------------
-	int32 Image::CalculateDataSize( int mip ) const
-    {
-		int32 res = 0;
-
-        const FImageFormatData& fdata = GetImageFormatData( m_format );
-        if (fdata.BytesPerBlock)
-        {
-            FImageSize s = m_size;
-
-            for ( int l=0; l< FMath::Max(1,(int)m_lods); ++l )
-            {
-                int blocksX = FMath::DivideAndRoundUp( s[0], (uint16)fdata.PixelsPerBlockX );
-                int blocksY = FMath::DivideAndRoundUp( s[1], (uint16)fdata.PixelsPerBlockY );
-
-                if ( mip==l )
-                {
-                    res = (blocksX*blocksY) * fdata.BytesPerBlock;
-                    break;
-                }
-
-                s[0] = FMath::DivideAndRoundUp( s[0], (uint16)2 );
-                s[1] = FMath::DivideAndRoundUp( s[1], (uint16)2 );
-            }
-        }
-
-        return res;
-    }
-
-
-    //---------------------------------------------------------------------------------------------
-    int32 Image::CalculatePixelCount() const
-    {
-		int32 res = 0;
-
-        const FImageFormatData& fdata = GetImageFormatData( m_format );
-        if (fdata.BytesPerBlock)
-        {
-            FImageSize s = m_size;
-
-            for ( int l=0; l<FMath::Max(1,(int)m_lods); ++l )
-            {
-                int blocksX = FMath::DivideAndRoundUp( s[0], (uint16)fdata.PixelsPerBlockX );
-                int blocksY = FMath::DivideAndRoundUp( s[1], (uint16)fdata.PixelsPerBlockY );
-
-                res += (blocksX*blocksY) * fdata.PixelsPerBlockX * fdata.PixelsPerBlockY;
-
-                s[0] = FMath::DivideAndRoundUp( s[0], (uint16)2 );
-                s[1] = FMath::DivideAndRoundUp( s[1], (uint16)2 );
-            }
-        }
-        else
-        {
-            // An RLE image.
-            for ( int l=0; l<FMath::Max(1,(int)m_lods); ++l )
-            {
-				FIntVector2 mipSize = CalculateMipSize(l);
-                res += mipSize[0] * mipSize[1];
-            }
-        }
-
-        return res;
-    }
-
+	FIntVector2 Image::CalculateMipSize(int32 Mip) const
+	{
+		FIntVector2 Result = FIntVector2(GetSizeX(), GetSizeY());
 
-    //---------------------------------------------------------------------------------------------
-	int32 Image::CalculatePixelCount( int mip ) const
-    {
-		int32 res = 0;
-
-        const FImageFormatData& fdata = GetImageFormatData( m_format );
-        if (fdata.BytesPerBlock)
-        {
-            FImageSize s = m_size;
-
-            for ( int l=0; l<mip+1; ++l )
-            {
-                int blocksX = FMath::DivideAndRoundUp( s[0], (uint16)fdata.PixelsPerBlockX );
-                int blocksY = FMath::DivideAndRoundUp( s[1], (uint16)fdata.PixelsPerBlockY );
-
-                if ( l==mip )
-                {
-                    res = (blocksX*blocksY) * fdata.PixelsPerBlockX * fdata.PixelsPerBlockY;
-                }
-
-                s[0] = FMath::DivideAndRoundUp( s[0], (uint16)2 );
-                s[1] = FMath::DivideAndRoundUp( s[1], (uint16)2 );
-            }
-        }
-        else
-        {
-            // An RLE image.
-			FIntVector2 mipSize = CalculateMipSize(mip);
-            res += mipSize[0] * mipSize[1];
-        }
-
-        return res;
-    }
+		for (int32 L = 0; L < Mip + 1; ++L)
+		{
+			if (L == Mip)
+			{
+				return Result;
+			}
 
+			Result[0] = FMath::DivideAndRoundUp(Result[0], 2);
+			Result[1] = FMath::DivideAndRoundUp(Result[1], 2);
+		}
 
-    //---------------------------------------------------------------------------------------------
-	FIntVector2 Image::CalculateMipSize( int mip ) const
-    {
-		FIntVector2 res(0,0);
-
-		FIntVector2 s = FIntVector2(m_size[0], m_size[1]);
-
-        for ( int l=0; l<mip+1; ++l )
-        {
-            if ( l==mip )
-            {
-                res[0] = s[0];
-                res[1] = s[1];
-                return res;
-            }
-
-            s[0] = FMath::DivideAndRoundUp( s[0], 2 );
-            s[1] = FMath::DivideAndRoundUp( s[1], 2 );
-        }
-
-        return res;
-    }
+		return FIntVector2(0, 0);
+	}
 
 	
 	//---------------------------------------------------------------------------------------------
-	int Image::GetMipmapCount(int SizeX, int SizeY)
+	int32 Image::GetMipmapCount(int32 SizeX, int32 SizeY)
 	{
 		if (SizeX <= 0 || SizeY <= 0)
 		{
 			return 0;
 		}
 
-		int MaxLevel = FMath::CeilLogTwo(FMath::Max(SizeX, SizeY)) + 1;
+		int32 MaxLevel = FMath::CeilLogTwo(FMath::Max(SizeX, SizeY)) + 1;
 		return MaxLevel;
 	}
 
 
-	//---------------------------------------------------------------------------------------------
-	const uint8_t* Image::GetMipData(int mip) const
+	const uint8* Image::GetMipData(int32 Mip) const
 	{
-		check((mip >= 0 && mip < m_lods) || (m_lods == 0));
-
-		const uint8_t* pResult = m_data.GetData();
+		return DataStorage.GetLOD(Mip).GetData();
+	}
 
-		const FImageFormatData& fdata = GetImageFormatData(m_format);
 
-		if (fdata.BytesPerBlock)
-		{
-			// Fixed-sized formats
-			FImageSize s = m_size;
+	uint8* Image::GetMipData(int32 Mip)
+	{
+		return DataStorage.GetLOD(Mip).GetData();
+	}
 
-			for (int l = 0; l < mip; ++l)
-			{
-				int blocksX = FMath::DivideAndRoundUp(s[0], (uint16)fdata.PixelsPerBlockX);
-				int blocksY = FMath::DivideAndRoundUp(s[1], (uint16)fdata.PixelsPerBlockY);
+	int32 Image::GetMipsDataSize() const
+	{
+		return DataStorage.GetDataSize();
+	}
 
-				pResult += (blocksX * blocksY) * fdata.BytesPerBlock;
+	FVector4f Image::Sample(FVector2f Coords) const
+	{
+		FVector4f Result = FVector4f(0.0f);
 
-				s[0] = FMath::DivideAndRoundUp(s[0], (uint16)2);
-				s[1] = FMath::DivideAndRoundUp(s[1], (uint16)2);
-			}
+		if (GetSizeX() == 0 || GetSizeY() == 0) 
+		{ 
+			return Result; 
 		}
 
-		else if (m_format == EImageFormat::IF_L_UBYTE_RLE ||
-			m_format == EImageFormat::IF_L_UBIT_RLE)
-		{
-			// Every mip has variable size, but it is stored in the first 4 bytes.
-			for (int l = 0; l < mip; ++l)
-			{
-				uint32_t mipSize = *(const uint32_t*)pResult;
-				pResult += mipSize;
-			}
-		}
-		else
-		{
-			checkf(false, TEXT("Trying to get mipmap pointer in an unsupported pixel format."));
-		}
+		const FImageFormatData& FormatData = GetImageFormatData(GetFormat());
 
-		return pResult;
-	}
+		FImageSize Size = GetSize();
+		EImageFormat Format = GetFormat();
 
+		// TODO: This looks like it should support block compression, but it is not implemented.
+		// Decide if it's worth implement it or simplify the ByteOffset computation assuming no
+		// compression is allowed.
 
-	//---------------------------------------------------------------------------------------------
-	uint8_t* Image::GetMipData(int mip)
-	{
-		check((mip >= 0 && mip < m_lods) || (m_lods == 0));
+		int32 PixelX = FMath::Max(0, FMath::Min(Size.X - 1, (int32)(Coords.X * Size.X)));
+		int32 BlockX = PixelX / FormatData.PixelsPerBlockX;
+		int32 BlockPixelX = PixelX % FormatData.PixelsPerBlockX;
 
-		uint8_t* pResult = m_data.GetData();
+		int32 PixelY = FMath::Max(0, FMath::Min(Size.Y - 1, (int32)(Coords.Y * Size.Y)));
+		int32 BlockY = PixelY / FormatData.PixelsPerBlockY;
+		int32 BlockPixelY = PixelY % FormatData.PixelsPerBlockY;
 
-		const FImageFormatData& fdata = GetImageFormatData(m_format);
+		int32 BlocksPerRow = FMath::DivideAndRoundUp(int32(Size.X), int32(FormatData.PixelsPerBlockX));
+		int32 BlockOffset = BlockX + BlockY * BlocksPerRow;
 
-		if (fdata.BytesPerBlock)
-		{
-			// Fixed-sized formats
-			FImageSize s = m_size;
+		const TArrayView<const uint8> Data = DataStorage.GetLOD(0);
 
-			for (int l = 0; l < mip; ++l)
-			{
-				int blocksX = FMath::DivideAndRoundUp(s[0], (uint16)fdata.PixelsPerBlockX);
-				int blocksY = FMath::DivideAndRoundUp(s[1], (uint16)fdata.PixelsPerBlockY);
-
-				pResult += (blocksX * blocksY) * fdata.BytesPerBlock;
+		// Non-generic part
+		if (Format == EImageFormat::IF_RGB_UBYTE)
+		{
+			int32 ByteOffset = BlockOffset * FormatData.BytesPerBlock
+					+ (BlockPixelY * FormatData.PixelsPerBlockX + BlockPixelX) * 3;
 
-				s[0] = FMath::DivideAndRoundUp(s[0], (uint16)2);
-				s[1] = FMath::DivideAndRoundUp(s[1], (uint16)2);
-			}
+			Result[0] = Data[ByteOffset + 0] / 255.0f;
+			Result[1] = Data[ByteOffset + 1] / 255.0f;
+			Result[2] = Data[ByteOffset + 2] / 255.0f;
+			Result[3] = 1.0f;
 		}
-
-		else if (m_format == EImageFormat::IF_L_UBYTE_RLE ||
-			m_format == EImageFormat::IF_L_UBIT_RLE)
+		else if (Format == EImageFormat::IF_RGBA_UBYTE)
 		{
-			// Every mip has variable size, but it is stored in the first 4 bytes.
-			for (int l = 0; l < mip; ++l)
-			{
-				uint32_t mipSize = *(const uint32_t*)pResult;
-				pResult += mipSize;
-			}
+			int32 ByteOffset = BlockOffset * FormatData.BytesPerBlock
+					+ (BlockPixelY * FormatData.PixelsPerBlockX + BlockPixelX) * 4;
+
+			Result[0] = Data[ByteOffset + 0] / 255.0f;
+			Result[1] = Data[ByteOffset + 1] / 255.0f;
+			Result[2] = Data[ByteOffset + 2] / 255.0f;
+			Result[3] = Data[ByteOffset + 3] / 255.0f;
 		}
-		else
+		else if (Format == EImageFormat::IF_BGRA_UBYTE)
 		{
-			checkf(false, TEXT("Trying to get mipmap pointer in an unsupported pixel format."));
-		}
+			int32 ByteOffset = BlockOffset * FormatData.BytesPerBlock
+					+ (BlockPixelY * FormatData.PixelsPerBlockX + BlockPixelX) * 4;
 
-		return pResult;
-	}
+			Result[0] = Data[ByteOffset + 2] / 255.0f;
+			Result[1] = Data[ByteOffset + 1] / 255.0f;
+			Result[2] = Data[ByteOffset + 0] / 255.0f;
+			Result[3] = Data[ByteOffset + 3] / 255.0f;
+		}
+		else if (Format== EImageFormat::IF_L_UBYTE)
+		{
+			int32 ByteOffset = BlockOffset * FormatData.BytesPerBlock
+					+ (BlockPixelY * FormatData.PixelsPerBlockX + BlockPixelX) * 1;
 
-	int32 Image::GetMipsDataSize() const
-	{
-		int32 res = 0;
-
-        const FImageFormatData& fdata = GetImageFormatData( m_format );
-        if (fdata.BytesPerBlock)
-        {
-			return CalculateDataSize();
-        }
-		else if (m_format == EImageFormat::IF_L_UBYTE_RLE ||
-				 m_format == EImageFormat::IF_L_UBIT_RLE)
-		{	
-			if (m_data.Num())
-			{ 
-			    // Every mip has variable size, but it is stored in the first 4 bytes.
-				const uint8* pMipData = m_data.GetData();
-				for (int l = 0; l < m_lods; ++l)
-				{
-					const uint32_t mipSize = *(const uint32_t*)pMipData;
-					pMipData += mipSize;
-					res += mipSize;
-				}	
-			}
+			Result[0] = Data[ByteOffset] / 255.0f;
+			Result[1] = Data[ByteOffset] / 255.0f;
+			Result[2] = Data[ByteOffset] / 255.0f;
+			Result[3] = 1.0f;
 		}
 		else
 		{
-			checkf(false, TEXT("Trying to get mips data size in an unsupported pixel format."));
+			check(false);
 		}
 
-        return res;
+		return Result;
 	}
 
-    //---------------------------------------------------------------------------------------------
-    FVector4f Image::Sample( FVector2f coords ) const
-    {
-		FVector4f result;
-
-        if (m_size[0]==0 || m_size[1]==0) { return result; }
-
-        const FImageFormatData& fdata = GetImageFormatData( m_format );
-
-        int pixelX = FMath::Max( 0, FMath::Min( m_size[0]-1, (int)(coords[0] * m_size[0] )));
-        int blockX = pixelX / fdata.PixelsPerBlockX;
-        int blockPixelX = pixelX % fdata.PixelsPerBlockX;
-
-        int pixelY = FMath::Max( 0, FMath::Min( m_size[1]-1, (int)(coords[1] * m_size[1] )));
-        int blockY = pixelY / fdata.PixelsPerBlockY;
-        int blockPixelY = pixelY % fdata.PixelsPerBlockY;
-
-        int32 blocksPerRow = FMath::DivideAndRoundUp(int32(m_size[0]), int32(fdata.PixelsPerBlockX));
-        int32 blockOffset = blockX + blockY * blocksPerRow;
-
-        // Non-generic part
-        if ( m_format== EImageFormat::IF_RGB_UBYTE )
-        {
-            int byteOffset = blockOffset * fdata.BytesPerBlock
-                    + ( blockPixelY * fdata.PixelsPerBlockX + blockPixelX ) * 3;
-
-            result[0] = m_data[ byteOffset+0 ] / 255.0f;
-            result[1] = m_data[ byteOffset+1 ] / 255.0f;
-            result[2] = m_data[ byteOffset+2 ] / 255.0f;
-            result[3] = 1;
-        }
-        else if ( m_format== EImageFormat::IF_RGBA_UBYTE )
-        {
-            int byteOffset = blockOffset * fdata.BytesPerBlock
-                    + ( blockPixelY * fdata.PixelsPerBlockX + blockPixelX ) * 4;
-
-            result[0] = m_data[ byteOffset+0 ] / 255.0f;
-            result[1] = m_data[ byteOffset+1 ] / 255.0f;
-            result[2] = m_data[ byteOffset+2 ] / 255.0f;
-            result[3] = m_data[ byteOffset+3 ] / 255.0f;
-        }
-        else if ( m_format== EImageFormat::IF_BGRA_UBYTE )
-        {
-            int byteOffset = blockOffset * fdata.BytesPerBlock
-                    + ( blockPixelY * fdata.PixelsPerBlockX + blockPixelX ) * 4;
-
-            result[0] = m_data[ byteOffset+2 ] / 255.0f;
-            result[1] = m_data[ byteOffset+1 ] / 255.0f;
-            result[2] = m_data[ byteOffset+0 ] / 255.0f;
-            result[3] = m_data[ byteOffset+3 ] / 255.0f;
-        }
-        else if ( m_format== EImageFormat::IF_L_UBYTE )
-        {
-            int byteOffset = blockOffset * fdata.BytesPerBlock
-                    + ( blockPixelY * fdata.PixelsPerBlockX + blockPixelX ) * 1;
-
-            result[0] = m_data[ byteOffset ] / 255.0f;
-            result[1] = m_data[ byteOffset ] / 255.0f;
-            result[2] = m_data[ byteOffset ] / 255.0f;
-            result[3] = 1.0f;
-        }
-        else
-        {
-            check( false );
-        }
-
-        return result;
-    }
-
-    //---------------------------------------------------------------------------------------------
-    bool Image::IsPlainColour(FVector4f& colour) const
-    {
-        bool res = true;
-
-        int pixelCount = m_size[0] * m_size[1];
+	//---------------------------------------------------------------------------------------------
+	bool Image::IsPlainColour(FVector4f& OutColor) const
+	{
+		bool Result = true;
 
-        if ( pixelCount )
-        {
-            switch( m_format )
-            {
+		const TArrayView<const uint8> DataView = DataStorage.GetLOD(0);
+		
+		if (DataView.Num())
+		{
+			switch(GetFormat())
+			{
 			case EImageFormat::IF_L_UBYTE:
 			{
-				const uint8* pData = m_data.GetData();
-				uint8 v = pData[0];
+				uint8 const * const DataPtr = DataView.GetData();
 
-				for (int p = 0; res && p < pixelCount; ++p)
+				const int32 NumElems = DataView.Num();
+				int32 I = 0;
+				for (; I < NumElems; ++I)
 				{
-					uint8 nv = *pData++;
-					res &= (nv == v);
+					if (FMemory::Memcmp(DataPtr, DataPtr + I*1, 1) != 0)
+					{
+						break;
+					}
 				}
 
-				if (res)
-				{
-					colour[0] = colour[1] = colour[2] = float(v) / 255.0f;
-					colour[3] = 1.0f;
-				}
+				Result = I >= NumElems;
+				OutColor = FVector4f(*DataPtr, *DataPtr, *DataPtr, 255) / 255.0f;
+
 				break;
 			}
 
 			case EImageFormat::IF_RGB_UBYTE:
 			{
-				const uint8* pData = m_data.GetData();
-				uint8 r = pData[0];
-				uint8 g = pData[1];
-				uint8 b = pData[2];
+				uint8 const * const DataPtr = DataView.GetData();
 
-				for (int p = 0; res && p < pixelCount; ++p)
+				const int32 NumElems = DataView.Num() / 3;
+				int32 I = 0;
+				for (; I < NumElems; ++I)
 				{
-					uint8 nr = *pData++;
-					uint8 ng = *pData++;
-					uint8 nb = *pData++;
-					res &= (nr == r) && (ng==g) && (nb==b);
+					if (FMemory::Memcmp(DataPtr, DataPtr + I*3, 3) != 0)
+					{
+						break;
+					}
 				}
 
-				if (res)
-				{
-					colour[0] = r / 255.0f;
-					colour[1] = g / 255.0f;
-					colour[2] = b / 255.0f;
-					colour[3] = 1.0f;
-				}
+				Result = I >= NumElems;
+				OutColor = FVector4f(DataPtr[0], DataPtr[1], DataPtr[2], 255) / 255.0f;
+
 				break;
 			}
 
 			case EImageFormat::IF_RGBA_UBYTE:
 			case EImageFormat::IF_BGRA_UBYTE:
 			{
-				const uint32* pData = reinterpret_cast<const uint32*>(m_data.GetData());
-				uint32 v = pData[0];
+				uint8 const * const DataPtr = DataView.GetData();
 
-				for (int p = 0; res && p < pixelCount; ++p)
+				const int32 NumElems = DataView.Num() / 4;
+				int32 I = 0;
+				for (; I < NumElems; ++I)
 				{
-					uint32 nv = *pData++;
-					res &= (nv == v);
+					if (FMemory::Memcmp(DataPtr, DataPtr + I*4, 4) != 0)
+					{
+						break;
+					}
 				}
 
-				if (res)
+				Result = I >= NumElems;
+				if (GetFormat() == EImageFormat::IF_RGBA_UBYTE)
 				{
-					const uint8* pByteData = m_data.GetData();
-					if (m_format == EImageFormat::IF_RGBA_UBYTE)
-					{
-						colour[0] = float(pByteData[0]) / 255.0f;
-						colour[1] = float(pByteData[1]) / 255.0f;
-						colour[2] = float(pByteData[2]) / 255.0f;
-					}
-					else
-					{
-						colour[0] = float(pByteData[2]) / 255.0f;
-						colour[1] = float(pByteData[1]) / 255.0f;
-						colour[2] = float(pByteData[0]) / 255.0f;
-					}
-					colour[3] = float(pByteData[3]) / 255.0f;
+					OutColor = FVector4f(DataPtr[0], DataPtr[1], DataPtr[2], DataPtr[3]) / 255.0f;
 				}
+				else if (GetFormat() == EImageFormat::IF_BGRA_UBYTE)
+				{
+					OutColor = FVector4f(DataPtr[2], DataPtr[1], DataPtr[0], DataPtr[3]) / 255.0f;
+				}
+				else
+				{
+					check(false);
+				}
+
+
 				break;
 			}
 
@@ -761,151 +568,154 @@ namespace mu
 			// the compressed block could be compared and only uncompress if all are the same to
 			// check if all pixels in the block are also equal.
 			default:
-                res = false;
-                break;
-            }
-        }
+				Result = false;
+				break;
+			}
+		}
 
-        return res;
-    }
+		return Result;
+	}
 
 
-    //---------------------------------------------------------------------------------------------
-    bool Image::IsFullAlpha() const
-    {
-        int pixelCount = m_size[0] * m_size[1];
-
-        if ( pixelCount )
-        {
-            switch( m_format )
-            {
-            case EImageFormat::IF_RGBA_UBYTE:
-            case EImageFormat::IF_BGRA_UBYTE:
-            {
-                const uint8_t* pData = &m_data[3];
-
-                for ( int p=0; p<pixelCount; ++p )
-                {
-                    uint8_t nv = *pData;
-                    if (nv!=255) return false;
-                    pData+=4;
-                }
-                return true;
-                break;
-            }
-
-            case EImageFormat::IF_RGB_UBYTE:
-            {
-                return true;
-                break;
-            }
-
-            default:
-                return false;
-            }
-        }
-
-        return true;
-    }
+	bool Image::IsFullAlpha() const
+	{
+		const TArrayView<const uint8> DataView = DataStorage.GetLOD(0);
+		
+		if (DataView.Num())
+		{
+			switch(GetFormat())
+			{
+			case EImageFormat::IF_RGBA_UBYTE:
+			case EImageFormat::IF_BGRA_UBYTE:
+			{
+				const uint8* DataPtr = DataView.GetData() + 3;
+
+				const int32 NumElems = DataView.Num() / 4;
+				for (int32 I = 0; I < NumElems; ++I)
+				{
+					if (*(DataPtr + I*4) != 255) 
+					{
+						return false;
+					}
+				}
 
+				return true;
+				break;
+			}
 
-    //---------------------------------------------------------------------------------------------
-    namespace
-    {
-        bool is_zero( const uint8_t* buff, size_t size )
-        {
-            return (*buff==0) && (FMemory::Memcmp( buff, buff + 1, size - 1 )==0);
-        }
-    }
+			case EImageFormat::IF_RGB_UBYTE:
+			{
+				return true;
+				break;
+			}
 
+			default:
+				return false;
+			}
+		}
 
-    //---------------------------------------------------------------------------------------------
-     void Image::GetNonBlackRect_Reference(FImageRect& rect) const
-     {            
-         rect.min[0] = rect.min[1] = 0;
-		 rect.size[0] = m_size[0];
-		 rect.size[1] = m_size[1];
-
-         if ( !rect.size[0] || !rect.size[1] )
-             return;
-
-         bool first = true;
-         uint16 sx = m_size[0];
-         uint16 sy = m_size[1];
-
-         // Slow reference implementation
-         uint16 top = 0;
-         uint16 left = 0;
-         uint16 right = sx - 1;
-         uint16 bottom = sy - 1;
-
-         for ( uint16 y=0; y<sy; ++y )
-         {
-             for ( uint16 x=0; x<sx; ++x )
-             {
-                 bool black = false;
-                 switch( m_format )
-                 {
-                 case EImageFormat::IF_L_UBYTE:
-                 {
-                     uint8_t nv = m_data[y*sx+x];
-                     if (nv==0)
-                     {
-                         black = true;
-                     }
-                     break;
-                 }
-
-                 case EImageFormat::IF_RGB_UBYTE:
-                 {
-                     const uint8_t* pData = &m_data[(y*sx+x)*3];
-                     if (!pData[0] && !pData[1] && !pData[2])
-                     {
-                         black = true;
-                     }
-                     break;
-                 }
-
-                 case EImageFormat::IF_RGBA_UBYTE:
-                 case EImageFormat::IF_BGRA_UBYTE:
-                 {
-                     const uint8_t* pData = &m_data[(y*sx+x)*4];
-                     if (!pData[0] && !pData[1] && !pData[2] && !pData[3])
-                     {
-                         black = true;
-                     }
-                     break;
-                 }
-
-                 default:
-                     check(false);
-                     break;
-                 }
-
-                 if (!black)
-                 {
-                     if (first)
-                     {
-                         first = false;
-                         left = right = x;
-                         top = bottom = y;
-                     }
-                     else
-                     {
-                         left = FMath::Min( left, x );
-                         right = FMath::Max( right, x );
-                         top = FMath::Min( top, y );
-                         bottom = FMath::Max( bottom, y );
-                     }
-                 }
-             }
-         }
-
-         rect.min[0] = left;
-         rect.min[1] = top;
-         rect.size[0] = right - left + 1;
-         rect.size[1] = bottom - top + 1;
-     }
+		return true;
+	}
+
+
+	namespace
+	{
+		bool IsZero(const uint8* Buff, SIZE_T Size)
+		{
+			if (Size == 0)
+			{
+				return true;
+			}
+			return (*Buff == 0) && (FMemory::Memcmp(Buff, Buff + 1, Size - 1) == 0);
+		}
+	}
+
+
+	void Image::GetNonBlackRect_Reference(FImageRect& Rect) const
+	{
+		const FImageSize Size = GetSize();
+		Rect.min[0] = Rect.min[1] = 0;
+		Rect.size[0] = Size.X;
+		Rect.size[1] = Size.Y;
+
+		if (!Rect.size[0] || !Rect.size[1])
+		{
+			return;
+		}
+
+		bool bFirst = true;
+
+		// Slow reference implementation
+		uint16 Top = 0;
+		uint16 Left = 0;
+		uint16 Right = Size.X - 1;
+		uint16 Bottom = Size.Y - 1;
+
+		const EImageFormat Format = GetFormat();
+
+		const TArrayView<const uint8> DataView = DataStorage.GetLOD(0);
+
+		const uint8 * const DataPtr = DataView.GetData();
+
+		for (uint16 Y = 0; Y < Size.Y; ++Y)
+		{
+			for (uint16 X = 0; X < Size.X; ++X)
+			{
+				bool bIsBlack = false;
+
+				switch (Format)
+				{
+				case EImageFormat::IF_L_UBYTE:
+				{
+					bIsBlack = DataPtr[Y*Size.X + X] == 0;
+					break;
+				}
+
+				case EImageFormat::IF_RGB_UBYTE:
+				{
+					constexpr uint32 ZeroValue = 0;
+					bIsBlack = FMemory::Memcmp(&ZeroValue, DataPtr + (Y*Size.X + X)*3, 3) == 0;
+
+					break;
+				}
+
+				case EImageFormat::IF_RGBA_UBYTE:
+				case EImageFormat::IF_BGRA_UBYTE:
+				{
+					constexpr uint32 ZeroValue = 0;
+					bIsBlack = FMemory::Memcmp(&ZeroValue, DataPtr + (Y*Size.X + X)*4, 4) == 0;
+					break;
+				 }
+
+				 default:
+					 check(false);
+					 break;
+				}
+
+				 if (!bIsBlack)
+				 {
+					 if (bFirst)
+					 {
+						 bFirst = false;
+						 Left = Right = X;
+						 Top = Bottom = Y;
+					 }
+					 else
+					 {
+						 Left = FMath::Min(Left, X);
+						 Right = FMath::Max(Right, X);
+						 Top = FMath::Min(Top, Y);
+						 Bottom = FMath::Max(Bottom, Y);
+					 }
+				 }
+			 }
+		 }
+
+		 Rect.min[0] = Left;
+		 Rect.min[1] = Top;
+		 Rect.size[0] = Right - Left + 1;
+		 Rect.size[1] = Bottom - Top + 1;
+	 }
 
 
     //---------------------------------------------------------------------------------------------
@@ -1068,230 +878,196 @@ namespace mu
   //       }
     }
 
-
-	//---------------------------------------------------------------------------------------------
 	void Image::ReduceLODsTo(int32 NewLODCount)
 	{
-		bool bIsBlockBased = GetImageFormatData(m_format).BytesPerBlock > 0;
-		int32 MaxLODs = GetMipmapCount(m_size[0], m_size[1]);
-		int32 LODSToSkip = MaxLODs - NewLODCount;
-		if (LODSToSkip > 0 && bIsBlockBased)
-		{
-			check(LODSToSkip < m_lods);
-			
-			uint32 DataToSkip = 0;
-			for (int32 l = 0; l < LODSToSkip; ++l)
-			{
-				uint32 MipDataSize = CalculateDataSize(l);
-				check(MipDataSize > 0);
-				DataToSkip += MipDataSize;
-			}
-
-			FIntVector2 FinalSize = CalculateMipSize(LODSToSkip);
-			m_size[0] = uint16(FinalSize[0]);
-			m_size[1] = uint16(FinalSize[1]);
-			m_lods = NewLODCount;
-
-			int32 FinalDataSize = m_data.Num() - int32(DataToSkip);
-			FMemory::Memmove(m_data.GetData(), m_data.GetData() + DataToSkip, FinalDataSize);
-		}
+		DataStorage.SetNumLODs(NewLODCount);
 	}
 
-
-	//---------------------------------------------------------------------------------------------
-	void Image::ReduceLODs(int32 LODSToSkip)
+	void Image::ReduceLODs(int32 LODsToSkip)
 	{
-		bool bIsBlockBased = GetImageFormatData(m_format).BytesPerBlock > 0;
-		if (LODSToSkip > 0 && bIsBlockBased)
-		{
-			check(LODSToSkip < m_lods);
-			
-			uint32 DataToSkip = 0;
-			for (int32 l = 0; l < LODSToSkip; ++l)
-			{
-				uint32 MipDataSize = CalculateDataSize(l);
-				check(MipDataSize>0);
-				DataToSkip += MipDataSize;
-			}
-
-			FIntVector2 FinalSize = CalculateMipSize(LODSToSkip);
-			m_size[0] = uint16(FinalSize[0]);
-			m_size[1] = uint16(FinalSize[1]);
-			m_lods -= LODSToSkip;
-
-			int32 FinalDataSize = m_data.Num() - int32(DataToSkip);
-			FMemory::Memmove(m_data.GetData(), m_data.GetData() + DataToSkip, FinalDataSize);
-		}
+		DataStorage.DropLODs(LODsToSkip);
 	}
 
-
-	//---------------------------------------------------------------------------------------------
 	void FImageOperator::FillColor(Image* Target, FVector4f Color)
 	{
-		EImageFormat Format = Target->GetFormat();
+		const EImageFormat Format = Target->GetFormat();
+
 		switch (Format)
 		{
 		case EImageFormat::IF_RGB_UBYTE:
 		{
-			int32 PixelCount = Target->CalculatePixelCount();
-			uint8* pData = Target->GetData();
-			uint32 R = uint32(FMath::Clamp(255.0f * Color[0], 0.0f, 255.0f));
-			uint32 G = uint32(FMath::Clamp(255.0f * Color[1], 0.0f, 255.0f));
-			uint32 B = uint32(FMath::Clamp(255.0f * Color[2], 0.0f, 255.0f));
-			
-			const uint32 PixelData = R | (G << 8) | (B << 16);
+			constexpr int32 BatchSizeInElems = 1 << 14; 
+			constexpr int32 ElemSizeInBytes  = 3;
+			const int32 NumBatches = Target->DataStorage.GetNumBatches(BatchSizeInElems, ElemSizeInBytes);
 
-			if (PixelData == 0)
+			for (int32 I = 0; I < NumBatches; ++I)
 			{
-				FMemory::Memzero(pData, PixelCount*3);
-				break;
-			}
+				TArrayView<uint8> DataView = Target->DataStorage.GetBatch(I, BatchSizeInElems, ElemSizeInBytes);
 
-			if ((R == G) & (R == B))
-			{
-				FMemory::Memset(pData, static_cast<uint8>(R), PixelCount*3);
-				break;
-			}
+				const int32 PixelCount = DataView.Num()/ElemSizeInBytes; 
+				uint8* DataPtr = DataView.GetData();
+				uint32 R = uint32(FMath::Clamp(255.0f * Color[0], 0.0f, 255.0f));
+				uint32 G = uint32(FMath::Clamp(255.0f * Color[1], 0.0f, 255.0f));
+				uint32 B = uint32(FMath::Clamp(255.0f * Color[2], 0.0f, 255.0f));
+				
+				const uint32 PixelData = R | (G << 8) | (B << 16);
 
-			const uint64 TwoPixelsData = (uint64(PixelData) << 24) | PixelData; 
+				if (PixelData == 0)
+				{
+					FMemory::Memzero(DataPtr, PixelCount*3);
+					continue;
+				}
 
-			for (int32 P = 0; P < PixelCount >> 1; ++P)
-			{
-				FMemory::Memcpy(&pData[P * 6], &TwoPixelsData, 6);
-			}
+				if ((R == G) & (R == B))
+				{
+					FMemory::Memset(DataPtr, static_cast<uint8>(R), PixelCount*3);
+					continue;
+				}
 
-			if (PixelCount & 1)
-			{
-				FMemory::Memcpy(&pData[(PixelCount >> 1) * 6], &PixelData, 3);
-			}
+				const uint64 TwoPixelsData = (uint64(PixelData) << 24) | PixelData; 
+
+				for (int32 P = 0; P < PixelCount >> 1; ++P)
+				{
+					FMemory::Memcpy(&DataPtr[P * 6], &TwoPixelsData, 6);
+				}
 
+				if (PixelCount & 1)
+				{
+					FMemory::Memcpy(&DataPtr[(PixelCount >> 1) * 6], &PixelData, 3);
+				}
+			}
 			break;
 		}
 
 		case EImageFormat::IF_RGBA_UBYTE:
+		case EImageFormat::IF_BGRA_UBYTE:
 		{
-			int32 PixelCount = Target->CalculatePixelCount();
-			uint8* pData = Target->GetData();
+			constexpr int32 BatchSizeInElems = 1 << 14; 
+			constexpr int32 ElemSizeInBytes  = 4;
+			const int32 NumBatches = Target->DataStorage.GetNumBatches(BatchSizeInElems, ElemSizeInBytes);
+
 			uint32 R = uint32(FMath::Clamp(255.0f * Color[0], 0.0f, 255.0f));
 			uint32 G = uint32(FMath::Clamp(255.0f * Color[1], 0.0f, 255.0f));
 			uint32 B = uint32(FMath::Clamp(255.0f * Color[2], 0.0f, 255.0f));
 			uint32 A = uint32(FMath::Clamp(255.0f * Color[3], 0.0f, 255.0f));
 		
-			const uint32 PixelData = R | (G << 8) | (B << 16) | (A << 24);  
+			const uint32 PixelData = Format == EImageFormat::IF_RGBA_UBYTE
+					? R | (G << 8) | (B << 16) | (A << 24)
+					: B | (G << 8) | (R << 16) | (A << 24);
 
-			if (PixelData == 0)
+			for (int32 I = 0; I < NumBatches; ++I)
 			{
-				FMemory::Memzero(pData, PixelCount*4);
-				break;
-			}
+				TArrayView<uint8> DataView = Target->DataStorage.GetBatch(I, BatchSizeInElems, ElemSizeInBytes);
 
-			if ((R == G) & (R == B) & (R == A))
-			{
-				FMemory::Memset(pData, static_cast<uint8>(R), PixelCount*4);
-				break;
-			}
+				const int32 PixelCount = DataView.Num()/ElemSizeInBytes;
+				uint8* DataPtr = DataView.GetData();
 
-			const uint64 TwoPixelsData = (uint64(PixelData) << 32) | PixelData; 
+				if (PixelData == 0)
+				{
+					FMemory::Memzero(DataPtr, PixelCount*4);
+					continue;
+				}
 
-			for (int32 P = 0; P < (PixelCount >> 1); ++P)
-			{
-				FMemory::Memcpy(&pData[P * 8], &TwoPixelsData, 8);
-			}
+				if ((R == G) & (R == B) & (R == A))
+				{
+					FMemory::Memset(DataPtr, static_cast<uint8>(R), PixelCount*4);
+					continue;
+				}
 
-			if (PixelCount & 1)
-			{
-				FMemory::Memcpy(&pData[(PixelCount >> 1) * 8], &PixelData, 4);
+				const uint64 TwoPixelsData = (uint64(PixelData) << 32) | PixelData; 
+
+				for (int32 P = 0; P < (PixelCount >> 1); ++P)
+				{
+					FMemory::Memcpy(&DataPtr[P * 8], &TwoPixelsData, 8);
+				}
+
+				if (PixelCount & 1)
+				{
+					FMemory::Memcpy(&DataPtr[(PixelCount >> 1) * 8], &PixelData, 4);
+				}
 			}
 
 			break;
 		}
 
-		case EImageFormat::IF_BGRA_UBYTE:
+		case EImageFormat::IF_L_UBYTE:
 		{
-			int32 PixelCount = Target->CalculatePixelCount();
-			uint8* pData = Target->GetData();
-			uint32 R = uint32(FMath::Clamp(255.0f * Color[0], 0.0f, 255.0f));
-			uint32 G = uint32(FMath::Clamp(255.0f * Color[1], 0.0f, 255.0f));
-			uint32 B = uint32(FMath::Clamp(255.0f * Color[2], 0.0f, 255.0f));
-			uint32 A = uint32(FMath::Clamp(255.0f * Color[3], 0.0f, 255.0f));
-	
-			const uint32 PixelData = B | (G << 8) | (R << 16) | (A << 24);  
-
-			if (PixelData == 0)
-			{
-				FMemory::Memzero(pData, PixelCount*4);
-				break;
-			}
+			constexpr int32 BatchSizeInElems = 1 << 14; 
+			constexpr int32 ElemSizeInBytes  = 1;
+			const int32 NumBatches = Target->DataStorage.GetNumBatches(BatchSizeInElems, ElemSizeInBytes);
 
-			if ((R == G) & (R == B) & (R == A))
+			uint32 R = uint32(FMath::Clamp(255.0f * Color[0], 0.0f, 255.0f));
+		
+			for (int32 I = 0; I < NumBatches; ++I)
 			{
-				FMemory::Memset(pData, static_cast<uint8>(R), PixelCount*4);
-				break;
-			}
+				TArrayView<uint8> DataView = Target->DataStorage.GetBatch(I, BatchSizeInElems, ElemSizeInBytes);
 
-			const uint64 TwoPixelsData = (uint64(PixelData) << 32) | PixelData; 
+				const int32 PixelCount = DataView.Num()/ElemSizeInBytes;
+				uint8* DataPtr = DataView.GetData();
 
-			for (int32 P = 0; P < (PixelCount >> 1); ++P)
-			{
-				FMemory::Memcpy(&pData[P * 8], &TwoPixelsData, 8);
-			}
+				if (R == 0)
+				{
+					FMemory::Memzero(DataPtr, PixelCount*1);
+					continue;
+				}
 
-			if (PixelCount & 1)
-			{
-				FMemory::Memcpy(&pData[(PixelCount >> 1) * 8], &PixelData, 4);
+				FMemory::Memset(DataPtr, static_cast<uint8>(R), PixelCount*1);
 			}
 
 			break;
 		}
 
-		case EImageFormat::IF_L_UBYTE:
-		{
-			int pixelCount = Target->CalculatePixelCount();
-			uint8* pData = Target->GetData();
-			uint8 v = FMath::Min<uint8>(255, FMath::Max<uint8>(0, uint8(255.0f * Color[0])));
-			FMemory::Memset(pData, v, pixelCount);
-			break;
-		}
-
 		default:
 		{
 			// Generic case that supports compressed formats.
 			const FImageFormatData& FormatData = GetImageFormatData(Format);
-			Ptr<Image> Block = CreateImage(FormatData.PixelsPerBlockX, FormatData.PixelsPerBlockY, 1, EImageFormat::IF_RGBA_UBYTE, EInitializationType::NotInitialized);
-
-			uint32 Pixel = (uint32(FMath::Clamp(255.0f * Color[0], 0.0f, 255.0f)) << 0)
-				| (uint32(FMath::Clamp(255.0f * Color[1], 0.0f, 255.0f)) << 8)
-				| (uint32(FMath::Clamp(255.0f * Color[2], 0.0f, 255.0f)) << 16)
-				| (uint32(FMath::Clamp(255.0f * Color[3], 0.0f, 255.0f)) << 24);
-			uint32* BlockData = reinterpret_cast<uint32*>(Block->GetData());
-			for (int32 I = 0; I < FormatData.PixelsPerBlockX * FormatData.PixelsPerBlockY; ++I)
+			Ptr<Image> BlockImage = CreateImage(
+						FormatData.PixelsPerBlockX, FormatData.PixelsPerBlockY, 1, EImageFormat::IF_RGBA_UBYTE, EInitializationType::NotInitialized);
+
+			const uint32 PixelData = 
+				  (uint32(FMath::Clamp(255.0f * Color[0], 0.0f, 255.0f)) << 0 ) |
+				  (uint32(FMath::Clamp(255.0f * Color[1], 0.0f, 255.0f)) << 8 ) |
+				  (uint32(FMath::Clamp(255.0f * Color[2], 0.0f, 255.0f)) << 16) |
+				  (uint32(FMath::Clamp(255.0f * Color[3], 0.0f, 255.0f)) << 24);
+
+			uint8 * const UncompressedBlockData = BlockImage->GetLODData(0);
+			const int32 UncompressedPixelCount = FormatData.PixelsPerBlockX * FormatData.PixelsPerBlockY;
+			for (int32 I = 0; I < UncompressedPixelCount; ++I)
 			{
-				*BlockData = Pixel;
-				BlockData += 1;
+				FMemory::Memcpy(UncompressedBlockData + I*4, &PixelData, 4);
 			}
-			Ptr<Image> Converted = ImagePixelFormat(0, Block.get(), Format);
-			ReleaseImage(Block);
-
-			int32 DataSize = Target->CalculateDataSize();
-			check(DataSize % FormatData.BytesPerBlock == 0);
-			int32 BlockCount = DataSize / FormatData.BytesPerBlock;
-			uint8* TargetData = Target->GetData();
-			while (BlockCount)
+
+			Ptr<Image> Converted = ImagePixelFormat(0, BlockImage.get(), Format);
+			ReleaseImage(BlockImage);
+
+			const TArrayView<const uint8> BlockView = Converted->DataStorage.GetLOD(0); 
+
+			constexpr int32 BatchSizeInElems = 1 << 14;
+			const int32 ElemSizeInBytes = FormatData.BytesPerBlock;
+			check(BlockView.Num() == ElemSizeInBytes);
+			const int32 NumBatches = Target->DataStorage.GetNumBatches(BatchSizeInElems, ElemSizeInBytes);
+
+			uint8 const * const BlockData = BlockView.GetData();
+			for (int32 I = 0; I < NumBatches; ++I)
 			{
-				FMemory::Memcpy(TargetData, Converted->GetData(), FormatData.BytesPerBlock);
-				TargetData += FormatData.BytesPerBlock;
-				--BlockCount;
+				TArrayView<uint8> DataView = Target->DataStorage.GetBatch(I, BatchSizeInElems, ElemSizeInBytes);
+
+				check(DataView.Num() % ElemSizeInBytes == 0);
+				
+				const int32 BatchNumBlocks = DataView.Num() / ElemSizeInBytes;
+				uint8* DataPtr = DataView.GetData();
+
+
+				for (int32 BlockIndex = 0; BlockIndex < BatchNumBlocks; ++BlockIndex)
+				{
+					FMemory::Memcpy(DataPtr + BlockIndex*ElemSizeInBytes, BlockData, ElemSizeInBytes);
+				}
 			}
 
 			ReleaseImage(Converted);
-			break;
+			break;	
 		}
-
-
 		}
-
 	}
-
-
 }
 
diff --git a/Engine/Plugins/Experimental/Mutable/Source/MutableRuntime/Private/MuR/ImageRLE.cpp b/Engine/Plugins/Experimental/Mutable/Source/MutableRuntime/Private/MuR/ImageRLE.cpp
--- a/Engine/Plugins/Experimental/Mutable/Source/MutableRuntime/Private/MuR/ImageRLE.cpp
+++ b/Engine/Plugins/Experimental/Mutable/Source/MutableRuntime/Private/MuR/ImageRLE.cpp
@@ -5,7 +5,6 @@
 #include "HAL/UnrealMemory.h"
 #include "Misc/AssertionMacros.h"
 
-
 namespace mu
 {
 
@@ -99,7 +98,7 @@ namespace mu
 			if (Difference)
 			{
 				// Different pos.
-				size_t Delta = 0;
+				SIZE_T Delta = 0;
 				for ( ; Delta<width*rows; ++Delta )
 				{
 					if (Temp[Delta] != InitialBaseData[Delta])
@@ -109,7 +108,9 @@ namespace mu
 				}
 
 				UncompressRLE_L(width, rows, destData, Temp.GetData());
-				CompressRLE_L( width, rows, InitialBaseData, destData, destDataSize);
+				int32 OutSize = 0;
+                CompressRLE_L(OutSize, width, rows, InitialBaseData, destData, destDataSize);
+                check(OutSize > 0)
 				check(false);
 			}
 		}
@@ -120,7 +121,6 @@ namespace mu
     }
 
 
-    //---------------------------------------------------------------------------------------------
     uint32 UncompressRLE_L( int32 width, int32 rows, const uint8* pStartBaseData, uint8* pStartDestData )
     {
 		const uint8* pBaseData = pStartBaseData;
@@ -170,8 +170,7 @@ namespace mu
         return TotalSize;
     }
 
-    //---------------------------------------------------------------------------------------------
-    void CompressRLE_L1( uint32& OutCompressedSize, int32 width, int32 rows,
+    void CompressRLE_L1(uint32& OutCompressedSize, int32 width, int32 rows,
                              const uint8* pBaseData,
                              uint8* destData,
                              uint32 destDataSize )
@@ -235,8 +234,7 @@ namespace mu
     }
 
 
-    //---------------------------------------------------------------------------------------------
-    uint32 UncompressRLE_L1( int32 width, int32 rows, const uint8* pStartBaseData, uint8* pDestData )
+    uint32 UncompressRLE_L1(int32 width, int32 rows, const uint8* pStartBaseData, uint8* pDestData)
     {
         const uint8* pBaseData = pStartBaseData;
         pBaseData += sizeof(uint32); // Total mip size
@@ -277,10 +275,9 @@ namespace mu
     }
 
 
-    //---------------------------------------------------------------------------------------------
-    void CompressRLE_RGBA( int32 width, int32 rows,
+    void CompressRLE_RGBA(uint32& OutCompressedSize, int32 width, int32 rows,
                            const uint8* pBaseDataByte,
-                           Image::ImageDataContainerType& destData )
+                           uint8* DestData, uint32 DestDataSize)
     {
         // TODO: Support for compression from compressed data size, like L_RLE formats.
         TArray<int8_t> rle;
@@ -348,17 +345,22 @@ namespace mu
 					offset += BytesToCopy;
 				}
             }
+
+            if (offset > DestDataSize)
+            {
+                break;
+            }
         }
 
-        destData.SetNum( offset );
-        if ( offset )
+        OutCompressedSize = DestDataSize >= offset ? offset : 0; 
+        
+        if (OutCompressedSize > 0)
         {
-            FMemory::Memmove( &destData[0], &rle[0], offset );
+            FMemory::Memmove(DestData, rle.GetData(), offset);
         }
     }
 
 
-    //---------------------------------------------------------------------------------------------
     void UncompressRLE_RGBA( int32 width, int32 rows, const uint8* pBaseData, uint8* pDestDataB )
     {
         uint32* pDestData = reinterpret_cast<uint32*>( pDestDataB );
@@ -411,7 +413,6 @@ namespace mu
     }
 
 
-    //---------------------------------------------------------------------------------------------
     struct UINT24
     {
         uint8 d[3];
@@ -429,16 +430,19 @@ namespace mu
     static_assert( sizeof(UINT24)==3, "Uint24SizeCheck" );
 
 
-    //---------------------------------------------------------------------------------------------
-    void CompressRLE_RGB( int32 width, int32 rows,
-                          const uint8* pBaseDataByte,
-                          Image::ImageDataContainerType& destData )
+    void CompressRLE_RGB(uint32& OutCompressedSize, int32 width, int32 rows,
+                         const uint8* pBaseDataByte,
+                         uint8* DestData, uint32 DestDataSize)
     {
+        // TODO: Optimize so no extra memory is allocated at this stage.
+
         TArray<int8_t> rle;
-        rle.Reserve(  (width*rows) );
+        rle.Reserve(width*rows);
 
         const UINT24* pBaseData = (const UINT24*)pBaseDataByte;
+
         rle.SetNum( rows*4, EAllowShrinking::No );
+
         uint32 offset = sizeof(uint32)*rows;
         for ( int32 r=0; r<rows; ++r )
         {
@@ -499,18 +503,23 @@ namespace mu
 					offset += BytesToCopy;
 				}
             }
+
+            if (offset > DestDataSize)
+            {
+                break;
+            }
         }
 
-        destData.SetNum( offset );
-        if ( offset )
+        OutCompressedSize = DestDataSize >= offset ? offset : 0; 
+        
+        if (OutCompressedSize > 0)
         {
-            FMemory::Memmove( &destData[0], &rle[0], offset );
+            FMemory::Memmove(DestData, rle.GetData(), offset);
         }
-    }
 
+    }
 
-    //---------------------------------------------------------------------------------------------
-    void UncompressRLE_RGB( int32 width, int32 rows, const uint8* pBaseData, uint8* pDestDataB )
+    void UncompressRLE_RGB(int32 width, int32 rows, const uint8* pBaseData, uint8* pDestDataB)
     {
 		check(pBaseData && pDestDataB);
 
diff --git a/Engine/Plugins/Experimental/Mutable/Source/MutableRuntime/Private/MuR/ImageRLE.h b/Engine/Plugins/Experimental/Mutable/Source/MutableRuntime/Private/MuR/ImageRLE.h
--- a/Engine/Plugins/Experimental/Mutable/Source/MutableRuntime/Private/MuR/ImageRLE.h
+++ b/Engine/Plugins/Experimental/Mutable/Source/MutableRuntime/Private/MuR/ImageRLE.h
@@ -13,7 +13,6 @@
 
 namespace mu
 {
-
     //---------------------------------------------------------------------------------------------
     //! Compress to RLE if the compressed image fits in destData, which should be preallocated.
     //! If it doesn't fit, it returns 0, and the content of destData is undefined.
@@ -28,14 +27,14 @@ namespace mu
     //!		- the value of the "equal" pixels
     //! After the header there are as many pixel values as different pixels.
     //---------------------------------------------------------------------------------------------
-	MUTABLERUNTIME_API extern void CompressRLE_L( uint32& OutCompressedSize, int32 width, int32 rows,
+	MUTABLERUNTIME_API extern void CompressRLE_L(uint32& OutCompressedSize, int32 Width, int32 Rows,
                                    const uint8* BaseData,
-                                   uint8* DestData, uint32 DestDataSize );
+                                   uint8* DestData, uint32 DestDataSize);
 
     //---------------------------------------------------------------------------------------------
-	MUTABLERUNTIME_API extern uint32 UncompressRLE_L( int32 width, int32 rows,
+	MUTABLERUNTIME_API extern uint32 UncompressRLE_L(int32 Width, int32 Rows,
                                      const uint8* BaseData,
-                                     uint8* pDestData );
+                                     uint8* pDestData);
 
     //---------------------------------------------------------------------------------------------
     //! This RLE format compresses the data per row. At the beginning of the data there are the row
@@ -44,15 +43,15 @@ namespace mu
     //!     - a UINT16 with how many 0 pixels are there
     //!     - a UINT16 with how many 1 pixels follow the 0 pixels
     //---------------------------------------------------------------------------------------------
-	MUTABLERUNTIME_API extern void CompressRLE_L1( uint32& OutCompressedSize, int32 width, int32 rows,
+	MUTABLERUNTIME_API extern void CompressRLE_L1(uint32& OutCompressedSize, int32 Width, int32 Rows,
                                     const uint8* BaseData,
                                     uint8* DestData,
-                                    uint32 DestDataSize );
+                                    uint32 DestDataSize);
 
     //---------------------------------------------------------------------------------------------
-	MUTABLERUNTIME_API extern uint32 UncompressRLE_L1(int32 width, int32 rows,
+	MUTABLERUNTIME_API extern uint32 UncompressRLE_L1(int32 Width, int32 Rows,
                                       const uint8* pBaseData,
-                                      uint8* pDestData );
+                                      uint8* pDestData);
 
 
     //---------------------------------------------------------------------------------------------
@@ -64,14 +63,14 @@ namespace mu
     //!     - 4 UINT8 with the value of the "equal" pixels
     //! After the header there are as many pixel values as different pixels.
     //---------------------------------------------------------------------------------------------
-	MUTABLERUNTIME_API extern void CompressRLE_RGBA(int32 width, int32 rows,
+	MUTABLERUNTIME_API extern void CompressRLE_RGBA(uint32& OutCompressedSize, int32 Width, int32 Rows,
 									const uint8* pBaseDataByte,
-									Image::ImageDataContainerType& destData );
+									uint8* DestData, uint32 DestDataSize);
 
     //---------------------------------------------------------------------------------------------
-	MUTABLERUNTIME_API extern void UncompressRLE_RGBA(int32 width, int32 rows,
+	MUTABLERUNTIME_API extern void UncompressRLE_RGBA(int32 Width, int32 Rows,
                                     const uint8* pBaseData,
-                                    uint8* pDestDataB );
+                                    uint8* pDestDataB);
 
 
     //---------------------------------------------------------------------------------------------
@@ -83,183 +82,222 @@ namespace mu
     //!     - 4 UINT8 with the value of the "equal" pixels
     //! After the header there are as many pixel values as different pixels.
     //---------------------------------------------------------------------------------------------
-	MUTABLERUNTIME_API extern void CompressRLE_RGB(int32 width, int32 rows,
+	MUTABLERUNTIME_API extern void CompressRLE_RGB(uint32& OutCompressedSize, int32 Width, int32 Rows,
 									const uint8* pBaseDataByte,
-									Image::ImageDataContainerType& destData );
+									uint8* DestData, uint32 DestDataSize);
 
     //---------------------------------------------------------------------------------------------
-	MUTABLERUNTIME_API extern void UncompressRLE_RGB(int32 width, int32 rows,
+	MUTABLERUNTIME_API extern void UncompressRLE_RGB(int32 Width, int32 Rows,
                                    const uint8* pBaseData,
-                                   uint8* pDestDataB );
+                                   uint8* pDestDataB);
 
 
-	//---------------------------------------------------------------------------------------------
-	//!
-	//---------------------------------------------------------------------------------------------
-	inline void UncompressRLE_L(const Image* pBase, Image* pDest)
+	inline void UncompressRLE_L(const Image* BaseImage, Image* DestImage)
 	{
-		int sizeX = pBase->GetSizeX();
-		int sizeY = pBase->GetSizeY();
-		const uint8* pData = pBase->GetData();
-		uint8* pDestData = pDest->GetData();
-		for (int32 Lod = 0; Lod < pBase->GetLODCount(); ++Lod)
+		check(DestImage->GetLODCount() <= BaseImage->GetLODCount());
+		check(DestImage->GetSize() == BaseImage->GetSize());
+		
+		int32 SizeX = BaseImage->GetSizeX();
+		int32 SizeY = BaseImage->GetSizeY();
+		
+        const int32 NumLODs = DestImage->GetLODCount();
+		for (int32 L = 0; L < NumLODs; ++L)
 		{
-			uint32_t compressedSize = UncompressRLE_L(sizeX, sizeY, pData, pDestData);
-			pData += compressedSize;
-
-			pDestData += pDest->GetLODDataSize(Lod);
-			sizeX = FMath::Max(1, FMath::DivideAndRoundUp(sizeX, 2));
-			sizeY = FMath::Max(1, FMath::DivideAndRoundUp(sizeY, 2));
+			uint32 CompressedSize = UncompressRLE_L(SizeX, SizeY, BaseImage->GetLODData(L), DestImage->GetLODData(L));
+	
+            SizeX = FMath::Max(1, FMath::DivideAndRoundUp(SizeX, 2));
+			SizeY = FMath::Max(1, FMath::DivideAndRoundUp(SizeY, 2));
 		}
 	}
 
-
-    //---------------------------------------------------------------------------------------------
-    //!
-    //---------------------------------------------------------------------------------------------
-    inline void CompressRLE_L(bool& bOutSuccess, const Image* pBase, Image* pDest )
+    inline void CompressRLE_L(const Image* BaseImage, Image* DestImage)
     {
-        int sizeX = pBase->GetSizeX();
-        int sizeY = pBase->GetSizeY();
-        const uint8* pData = pBase->GetData();
-        uint32 DestDataSize = pDest->GetDataSize();
-        uint8* pDestData = pDest->GetData();
-        for (int Lod=0;Lod<pBase->GetLODCount();++Lod)
+		check(DestImage->GetLODCount() <= BaseImage->GetLODCount());
+		check(DestImage->GetSize() == BaseImage->GetSize());
+		
+		int32 SizeX = BaseImage->GetSizeX();
+		int32 SizeY = BaseImage->GetSizeY();
+
+        const int32 NumLODs = DestImage->GetLODCount();
+        for (int32 L = 0; L < NumLODs; ++L)
         {
-			uint32 compressedSize = 0;
-			CompressRLE_L( compressedSize, sizeX, sizeY,
-				pData,
-				pDestData, DestDataSize );
-
-            if (!compressedSize)
-            {
-				bOutSuccess = false;
-                return;
-            }
-			check(compressedSize<=DestDataSize);
-
-            DestDataSize -= compressedSize;
-            pDestData += compressedSize;
-
-            pData += pBase->GetLODDataSize(Lod);
-            sizeX = FMath::Max(1, FMath::DivideAndRoundUp(sizeX,2));
-            sizeY = FMath::Max(1, FMath::DivideAndRoundUp(sizeY,2));
+			bool bSuccess = false;
+			while (!bSuccess)
+			{
+				uint32 CompressedSize = 0;
+				CompressRLE_L(CompressedSize, SizeX, SizeY, BaseImage->GetLODData(L), DestImage->GetLODData(L), DestImage->GetLODDataSize(L));
+			
+				bSuccess = CompressedSize > 0;
+				if (bSuccess)
+				{
+					DestImage->DataStorage.ResizeLOD(L, CompressedSize);
+				}
+				else
+				{
+					const int32 MipMemorySize = DestImage->DataStorage.GetLOD(L).Num();
+					DestImage->DataStorage.ResizeLOD(L, FMath::Max(4, MipMemorySize * 2));
+				}
+			}
+
+			SizeX = FMath::DivideAndRoundUp(SizeX, 2);
+			SizeY = FMath::DivideAndRoundUp(SizeY, 2);
         }
+    }
 
-		uint32 TotalDataSize = pDestData - pDest->GetData();
-        pDest->m_data.SetNum( TotalDataSize );
-
-#ifdef MUTABLE_DEBUG_RLE		
-		{
-			ImagePtr Temp = new Image(pBase->GetSizeX(),pBase->GetSizeY(),pBase->GetLODCount(),pBase->GetFormat());
-			UncompressRLE_L(pDest,Temp.get());
-			check( !FMemory::Memcmp(Temp->GetData(),pBase->GetData(),pBase->GetDataSize() ) );
-		}
-#endif
-
-		bOutSuccess = true;
-	}
-
-
-    //---------------------------------------------------------------------------------------------
-    //! This assumes the maximum data size has been preallocateds in pDest. It will return false
-    //! if the RLE-compressed data doesn't fit.
-    //---------------------------------------------------------------------------------------------
-    inline void CompressRLE_L1(bool& bOutSuccess, const Image* pBase, Image* pDest )
+    inline void CompressRLE_L1(const Image* BaseImage, Image* DestImage)
     {
-        int sizeX = pBase->GetSizeX();
-        int sizeY = pBase->GetSizeY();
-        const uint8* pData = pBase->GetData();
-        uint32 destDataSize = pDest->GetDataSize();
-        uint8* pDestData = pDest->GetData();
-        for (int lod=0;lod<pBase->GetLODCount();++lod)
+		check(DestImage->GetLODCount() <= BaseImage->GetLODCount());
+		check(DestImage->GetSize() == BaseImage->GetSize());
+		
+		int32 SizeX = BaseImage->GetSizeX();
+		int32 SizeY = BaseImage->GetSizeY();
+
+        const int32 NumLODs = DestImage->GetLODCount();
+        for (int32 L = 0; L < NumLODs; ++L)
         {
-			uint32 CompressedSize=0;
-            CompressRLE_L1( CompressedSize, sizeX, sizeY, pData, pDestData, destDataSize );
-			if (!CompressedSize)
+			bool bSuccess = false;
+			while (!bSuccess)
 			{
-				bOutSuccess = false;
-                return;
-            }
-
-            destDataSize -= CompressedSize;
-            pDestData += CompressedSize;
-
-            pData += pBase->GetLODDataSize(lod);
-			sizeX = FMath::Max(1, FMath::DivideAndRoundUp(sizeX, 2));
-			sizeY = FMath::Max(1, FMath::DivideAndRoundUp(sizeY, 2));
-		}
-
-        pDest->m_data.SetNum( pDestData-pDest->GetData() );
-
-		bOutSuccess = true;
+				uint32 CompressedSize = 0;
+				CompressRLE_L1(CompressedSize, SizeX, SizeY, BaseImage->GetLODData(L), DestImage->GetLODData(L), DestImage->GetLODDataSize(L));
+			
+				bSuccess = CompressedSize > 0;
+				if (bSuccess)
+				{
+					DestImage->DataStorage.ResizeLOD(L, CompressedSize);
+				}
+				else
+				{
+					const int32 MipMemorySize = DestImage->DataStorage.GetLOD(L).Num();
+					DestImage->DataStorage.ResizeLOD(L, FMath::Max(4, MipMemorySize * 2));
+				}
+			}
+
+			SizeX = FMath::DivideAndRoundUp(SizeX, 2);
+			SizeY = FMath::DivideAndRoundUp(SizeY, 2);
+        }
     }
 
-
-    //---------------------------------------------------------------------------------------------
-    inline void UncompressRLE_L1( const Image* pBase, Image* pDest )
+    inline void UncompressRLE_L1(const Image* BaseImage, Image* DestImage)
     {
-        int sizeX = pBase->GetSizeX();
-        int sizeY = pBase->GetSizeY();
-        const uint8* pData = pBase->GetData();
-        uint8* pDestData = pDest->GetData();
-        for (int lod=0;lod<pBase->GetLODCount();++lod)
+		check(DestImage->GetLODCount() <= BaseImage->GetLODCount());
+		check(DestImage->GetSize() == BaseImage->GetSize());
+        
+		int32 SizeX = BaseImage->GetSizeX();
+        int32 SizeY = BaseImage->GetSizeY();
+
+        const int32 NumLODs = DestImage->GetLODCount();
+        for (int32 L = 0; L < NumLODs; ++L)
         {
-            uint32_t compressedSize = UncompressRLE_L1( sizeX, sizeY, pData, pDestData );
-            pData += compressedSize;
+            uint32 CompressedSize = UncompressRLE_L1(SizeX, SizeY, BaseImage->GetLODData(L), DestImage->GetLODData(L));
 
-            pDestData += pDest->GetLODDataSize(lod);
-			sizeX = FMath::Max(1, FMath::DivideAndRoundUp(sizeX, 2));
-			sizeY = FMath::Max(1, FMath::DivideAndRoundUp(sizeY, 2));
+			SizeX = FMath::Max(1, FMath::DivideAndRoundUp(SizeX, 2));
+			SizeY = FMath::Max(1, FMath::DivideAndRoundUp(SizeY, 2));
 		}
     }
 
-
-    //---------------------------------------------------------------------------------------------
-    //!
-    //---------------------------------------------------------------------------------------------
-    inline void CompressRLE_RGB( const Image* pBase, Image* pDest )
+    inline void CompressRLE_RGB(const Image* BaseImage, Image* DestImage)
     {
-        CompressRLE_RGB( pBase->GetSizeX(),
-                         pBase->GetSizeY(),
-                         pBase->GetData(),
-                         pDest->m_data );
+		check(DestImage->GetLODCount() <= BaseImage->GetLODCount());
+		check(DestImage->GetSize() == BaseImage->GetSize());
+		
+		int32 SizeX = BaseImage->GetSizeX();
+		int32 SizeY = BaseImage->GetSizeY();
+
+        const int32 NumLODs = DestImage->GetLODCount();
+        for (int32 L = 0; L < NumLODs; ++L)
+        {
+			bool bSuccess = false;
+			while (!bSuccess)
+			{
+				uint32 CompressedSize = 0;
+				CompressRLE_RGB(CompressedSize, SizeX, SizeY, BaseImage->GetLODData(L), DestImage->GetLODData(L), DestImage->GetLODDataSize(L));
+			
+				bSuccess = CompressedSize > 0;
+				if (bSuccess)
+				{
+					DestImage->DataStorage.ResizeLOD(L, CompressedSize);
+				}
+				else
+				{
+					const int32 MipMemorySize = DestImage->DataStorage.GetLOD(L).Num();
+					DestImage->DataStorage.ResizeLOD(L, FMath::Max(4, MipMemorySize * 2));
+				}
+			}
+
+			SizeX = FMath::DivideAndRoundUp(SizeX, 2);
+			SizeY = FMath::DivideAndRoundUp(SizeY, 2);
+        }
     }
 
-
-    //---------------------------------------------------------------------------------------------
-    //!
-    //---------------------------------------------------------------------------------------------
-    inline void UncompressRLE_RGB( const Image* pBase, Image* pDest )
+    inline void UncompressRLE_RGB(const Image* BaseImage, Image* DestImage)
     {
-        UncompressRLE_RGB( pBase->GetSizeX(),
-                           pBase->GetSizeY(),
-                           pBase->GetData(),
-                           pDest->GetData() );
-    }
+		check(DestImage->GetLODCount() <= BaseImage->GetLODCount());
+		check(DestImage->GetSize() == BaseImage->GetSize());
+        
+		int32 SizeX = BaseImage->GetSizeX();
+        int32 SizeY = BaseImage->GetSizeY();
+       
+		int32 NumLODs = DestImage->GetLODCount();
+        for (int32 L = 0; L < NumLODs; ++L)
+        {
+            UncompressRLE_RGB(SizeX, SizeY, BaseImage->GetLODData(L), DestImage->GetLODData(L));
 
+			SizeX = FMath::Max(1, FMath::DivideAndRoundUp(SizeX, 2));
+			SizeY = FMath::Max(1, FMath::DivideAndRoundUp(SizeY, 2));
+        }
+    }
 
-    //---------------------------------------------------------------------------------------------
-    //!
-    //---------------------------------------------------------------------------------------------
-    inline void CompressRLE_RGBA( const Image* pBase, Image* pDest )
+    inline void CompressRLE_RGBA(const Image* BaseImage, Image* DestImage)
     {
-        CompressRLE_RGBA( pBase->GetSizeX(),
-                          pBase->GetSizeY(),
-                          pBase->GetData(),
-                          pDest->m_data );
-    }
+		check(DestImage->GetLODCount() <= BaseImage->GetLODCount());
+		check(DestImage->GetSize() == BaseImage->GetSize());
 
+		int32 SizeX = BaseImage->GetSizeX();
+		int32 SizeY = BaseImage->GetSizeY();
 
-    //---------------------------------------------------------------------------------------------
-    //!
-    //---------------------------------------------------------------------------------------------
-    inline void UncompressRLE_RGBA( const Image* pBase, Image* pDest )
+        const int32 NumLODs = DestImage->GetLODCount();
+        for (int32 L = 0; L < NumLODs; ++L)
+        {
+			bool bSuccess = false;
+			while (!bSuccess)
+			{
+				uint32 CompressedSize = 0;
+				CompressRLE_RGBA(CompressedSize, SizeX, SizeY, BaseImage->GetLODData(L), DestImage->GetLODData(L), DestImage->GetLODDataSize(L));
+			
+				bSuccess = CompressedSize > 0;
+				if (bSuccess)
+				{
+					DestImage->DataStorage.ResizeLOD(L, CompressedSize);
+				}
+				else
+				{
+					const int32 MipMemorySize = DestImage->DataStorage.GetLOD(L).Num();
+					DestImage->DataStorage.ResizeLOD(L, FMath::Max(4, MipMemorySize * 2));
+				}
+			}
+
+			SizeX = FMath::DivideAndRoundUp(SizeX, 2);
+			SizeY = FMath::DivideAndRoundUp(SizeY, 2);
+        }
+    }
+
+    inline void UncompressRLE_RGBA(const Image* BaseImage, Image* DestImage)
     {
-        UncompressRLE_RGBA( pBase->GetSizeX(),
-                            pBase->GetSizeY(),
-                            pBase->GetData(),
-                            pDest->GetData() );
+		check(DestImage->GetLODCount() <= BaseImage->GetLODCount());
+		check(DestImage->GetSize() == BaseImage->GetSize());
+
+        int32 SizeX = BaseImage->GetSizeX();
+        int32 SizeY = BaseImage->GetSizeY();
+        
+        const int32 NumLODs = DestImage->GetLODCount();
+        for (int32 L = 0; L < NumLODs; ++L)
+        {
+            UncompressRLE_RGBA(SizeX, SizeY, BaseImage->GetLODData(L), DestImage->GetLODData(L));
+
+			SizeX = FMath::Max(1, FMath::DivideAndRoundUp(SizeX, 2));
+			SizeY = FMath::Max(1, FMath::DivideAndRoundUp(SizeY, 2));
+        }
     }
 
 }
diff --git a/Engine/Plugins/Experimental/Mutable/Source/MutableRuntime/Private/MuR/OpImageBinarise.h b/Engine/Plugins/Experimental/Mutable/Source/MutableRuntime/Private/MuR/OpImageBinarise.h
--- a/Engine/Plugins/Experimental/Mutable/Source/MutableRuntime/Private/MuR/OpImageBinarise.h
+++ b/Engine/Plugins/Experimental/Mutable/Source/MutableRuntime/Private/MuR/OpImageBinarise.h
@@ -5,6 +5,7 @@
 #include "MuR/ImagePrivate.h"
 #include "MuR/Platform.h"
 
+#include "Async/ParallelFor.h"
 
 
 namespace mu
@@ -13,62 +14,83 @@ namespace mu
 	//---------------------------------------------------------------------------------------------
 	//!
 	//---------------------------------------------------------------------------------------------
-	inline void ImageBinarise( Image* pDest, const Image* pA, float threshold )
+	inline void ImageBinarise(Image* DestImage, const Image* AImage, float Threshold)
 	{
-		if (!pA || !pDest)
+		if (!AImage || !DestImage)
 		{
 			return;
 		}
 
-        uint8* pDestBuf = pDest->GetData();
-        const uint8* pABuf = pA->GetData();
+		check(DestImage->GetFormat() == EImageFormat::IF_L_UBYTE);
 
 		// Generic implementation
-		int32 pixelCount = pA->CalculatePixelCount();
+        const uint16 ThresholdUNorm = static_cast<uint16>(FMath::Clamp(Threshold, 0.0f, 1.0f) * 255.0f);
 
-        uint32 t_8 = uint32( threshold * 255.0f );
+		constexpr int32 BatchNumElems = 1 << 14;
+		const int32 BytesPerElem = GetImageFormatData(AImage->GetFormat()).BytesPerBlock;
+		
+		const int32 NumBatches = AImage->DataStorage.GetNumBatches(BatchNumElems, BytesPerElem);
+		check(DestImage->DataStorage.GetNumBatches(BatchNumElems, 1) == NumBatches);
 
-		switch ( pA->GetFormat() )
+		auto ProcessBatch = [DestImage, AImage, ThresholdUNorm, BatchNumElems, BytesPerElem](int32 BatchId)
 		{
-		case EImageFormat::IF_L_UBYTE:
-		{
-			for ( int i=0; i<pixelCount; ++i )
+			TArrayView<uint8> DestView = DestImage->DataStorage.GetBatch(BatchId, BatchNumElems, 1);	
+			TArrayView<const uint8> AView = AImage->DataStorage.GetBatch(BatchId, BatchNumElems, BytesPerElem);	
+
+			const int32 NumElems = DestView.Num();
+			check(AView.Num() / BytesPerElem == NumElems);
+
+			uint8* DestBuf = DestView.GetData();
+			const uint8* ABuf = AView.GetData();
+
+			switch(AImage->GetFormat())
 			{
-                uint32 a_8 = pABuf[i];
-				pDestBuf[i] = (a_8>=t_8) ? 255 : 0;
+			case EImageFormat::IF_L_UBYTE:
+			{
+				for (int32 I = 0; I < NumElems; ++I)
+				{
+					DestBuf[I] = ABuf[I] >= ThresholdUNorm ? 255 : 0;
+				}
+				break;
 			}
-			break;
-		}
 
-		case EImageFormat::IF_RGB_UBYTE:
-		{
-			for ( int i=0; i<pixelCount; ++i )
+			case EImageFormat::IF_RGB_UBYTE:
 			{
-                uint32 a_8 = pABuf[3*i+0];
-				a_8 += pABuf[3*i+1];
-				a_8 += pABuf[3*i+2];
-				a_8 /= 3;
-				pDestBuf[i] = (a_8>=t_8) ? 255 : 0;
+				for (int32 I = 0; I < NumElems; ++I)
+				{
+					uint16 Value = ABuf[3*I + 0];
+					Value += ABuf[3*I + 1];
+					Value += ABuf[3*I + 2];
+					DestBuf[I] = (Value / 3) >= ThresholdUNorm ? 255 : 0;
+				}
+				break;
 			}
-			break;
-		}
 
-        case EImageFormat::IF_RGBA_UBYTE:
-        case EImageFormat::IF_BGRA_UBYTE:
-        {
-			for ( int i=0; i<pixelCount; ++i )
+			case EImageFormat::IF_RGBA_UBYTE:
+			case EImageFormat::IF_BGRA_UBYTE:
 			{
-                uint32 a_8 = pABuf[4*i+0];
-				a_8 += pABuf[4*i+1];
-				a_8 += pABuf[4*i+2];
-				a_8 /= 3;
-				pDestBuf[i] = (a_8>=t_8) ? 255 : 0;
+				for (int32 I = 0; I < NumElems; ++I)
+				{
+					uint16 Value = ABuf[4*I + 0];
+					Value += ABuf[4*I + 1];
+					Value += ABuf[4*I + 2];
+					DestBuf[I] = (Value / 3) >= ThresholdUNorm ? 255 : 0;
+				}
+				break;
 			}
-			break;
-		}
 
-		default:
-			check(false);
+			default:
+				check(false);
+			}
+		};
+
+		if (NumBatches == 1)
+		{
+			ProcessBatch(0);
+		}
+		else if (NumBatches > 1)
+		{
+			ParallelFor(NumBatches, ProcessBatch);
 		}
 	}
 
diff --git a/Engine/Plugins/Experimental/Mutable/Source/MutableRuntime/Private/MuR/OpImageColourMap.h b/Engine/Plugins/Experimental/Mutable/Source/MutableRuntime/Private/MuR/OpImageColourMap.h
--- a/Engine/Plugins/Experimental/Mutable/Source/MutableRuntime/Private/MuR/OpImageColourMap.h
+++ b/Engine/Plugins/Experimental/Mutable/Source/MutableRuntime/Private/MuR/OpImageColourMap.h
@@ -5,130 +5,150 @@
 #include "MuR/ImagePrivate.h"
 #include "MuR/Platform.h"
 
+#include "Async/ParallelFor.h"
+
 namespace mu
 {
 
-	inline void ImageColourMap( Image* pDest, const Image* pSource, const Image* pMask, const Image* pMap, bool bOnlyOneMip )
+	inline void ImageColourMap(Image* DestImage, const Image* SourceImage, const Image* MaskImage, const Image* MapImage, bool bOnlyOneMip)
 	{
-		check( pSource->GetSizeX() == pMask->GetSizeX() );
-		check( pSource->GetSizeY() == pMask->GetSizeY() );
-		check( pSource->GetLODCount() == pMask->GetLODCount() || bOnlyOneMip );
-		check( pMask->GetFormat() == EImageFormat::IF_L_UBYTE );
-
-        uint8* pDestBuf = pDest->GetData();
-        const uint8* pSourceBuf = pSource->GetData();
-        const uint8* pMaskBuf = pMask->GetData();
+		check(SourceImage->GetSizeX() == MaskImage->GetSizeX());
+		check(SourceImage->GetSizeY() == MaskImage->GetSizeY());
+		check(SourceImage->GetLODCount() == MaskImage->GetLODCount() || bOnlyOneMip);
+		check(SourceImage->GetFormat() == DestImage->GetFormat());
+		check(SourceImage->GetSize() == DestImage->GetSize());
+		check(MaskImage->GetFormat() == EImageFormat::IF_L_UBYTE);
 
 		// Generic implementation
-		int32 PixelCount = 0;
-		if (bOnlyOneMip)
-		{
-			PixelCount = pSource->CalculatePixelCount(0);
-		}
-		else
-		{
-			PixelCount = pSource->CalculatePixelCount();
-		}
 
 		// Make a palette for faster conversion
-        uint8 palette[256][4];
-		for ( int i=0; i<256; ++i )
+        uint8 Palette[256][4];
+		for (int32 I = 0; I < 256; ++I)
 		{
-			FVector4f c = pMap->Sample(FVector2f( float(i)/256.0f, 0.0f ) );
-            palette[i][0] = (uint8)FMath::Max( 0, FMath::Min( 255, int(c[0]*255.0f) ) );
-            palette[i][1] = (uint8)FMath::Max( 0, FMath::Min( 255, int(c[1]*255.0f) ) );
-            palette[i][2] = (uint8)FMath::Max( 0, FMath::Min( 255, int(c[2]*255.0f) ) );
-            palette[i][3] = (uint8)FMath::Max( 0, FMath::Min( 255, int(c[3]*255.0f) ) );
+			FVector4f Color = MapImage->Sample(FVector2f(float(I)/255.0f, 0.0f));
+            Palette[I][0] = (uint8)FMath::Max(0, FMath::Min(255, int32(Color[0]*255.0f)));
+            Palette[I][1] = (uint8)FMath::Max(0, FMath::Min(255, int32(Color[1]*255.0f)));
+            Palette[I][2] = (uint8)FMath::Max(0, FMath::Min(255, int32(Color[2]*255.0f)));
+            Palette[I][3] = (uint8)FMath::Max(0, FMath::Min(255, int32(Color[3]*255.0f)));
 		}
 
-		switch ( pSource->GetFormat() )
-		{
-		case EImageFormat::IF_L_UBYTE:
+		const int32 LODBegin = 0;
+		const int32 LODEnd = bOnlyOneMip ? 1 : DestImage->GetLODCount();
+	
+		constexpr int32 NumBatchElems = 1 << 14;
+		const int32 BytesPerElem = GetImageFormatData(DestImage->GetFormat()).BytesPerBlock;
+		check(GetImageFormatData(SourceImage->GetFormat()).BytesPerBlock == BytesPerElem);
+
+		const int32 NumBatches = DestImage->DataStorage.GetNumBatchesLODRange(NumBatchElems, BytesPerElem, LODBegin, LODEnd); 
+		check(NumBatches == SourceImage->DataStorage.GetNumBatchesLODRange(NumBatchElems, BytesPerElem, LODBegin, LODEnd));
+
+		auto ProcessBatch = [DestImage, SourceImage, MaskImage, Palette, NumBatchElems, BytesPerElem, LODBegin, LODEnd](int32 BatchId)
 		{
-			for ( int i=0; i< PixelCount; ++i )
+			TArrayView<uint8> DestView = DestImage->DataStorage.GetBatchLODRange(BatchId, NumBatchElems, BytesPerElem, LODBegin, LODEnd); 
+			TArrayView<const uint8> SourceView = SourceImage->DataStorage.GetBatchLODRange(BatchId, NumBatchElems, BytesPerElem, LODBegin, LODEnd); 
+			TArrayView<const uint8> MaskView = MaskImage->DataStorage.GetBatchLODRange(BatchId, NumBatchElems, 1, LODBegin, LODEnd); 
+
+			const int32 NumElems = DestView.Num() / BytesPerElem;
+			check(SourceView.Num() / BytesPerElem == NumElems);
+			check(MaskView.Num() / 1 == NumElems);
+
+			uint8* DestBuf = DestView.GetData();
+			const uint8* SourceBuf = SourceView.GetData();
+			const uint8* MaskBuf = MaskView.GetData();
+
+			switch (SourceImage->GetFormat())
+			{
+			case EImageFormat::IF_L_UBYTE:
 			{
-                uint8 m_8 = pMaskBuf[i];
-				if (m_8>127)
+				for (int32 I = 0; I < NumElems; ++I)
 				{
-					pDestBuf[i] = palette[ pSourceBuf[i] ][0];
+					if (MaskBuf[I] > 127)
+					{
+						DestBuf[I] = Palette[SourceBuf[I]][0];
+					}
+					else
+					{
+						DestBuf[I] = SourceBuf[I];
+					}
 				}
-				else
+				break;
+			}
+			case EImageFormat::IF_RGB_UBYTE:
+			{
+				for (int32 I = 0; I < NumElems; ++I)
 				{
-					pDestBuf[i] = pSourceBuf[i];
+					if (MaskBuf[I])
+					{
+						DestBuf[3*I + 0] = Palette[SourceBuf[3*I + 0]][0];
+						DestBuf[3*I + 1] = Palette[SourceBuf[3*I + 1]][1];
+						DestBuf[3*I + 2] = Palette[SourceBuf[3*I + 2]][2];
+					}
+					else
+					{
+						DestBuf[3*I + 0] = SourceBuf[3*I + 0];
+						DestBuf[3*I + 1] = SourceBuf[3*I + 1];
+						DestBuf[3*I + 2] = SourceBuf[3*I + 2];
+					}
 				}
+				break;
 			}
-			break;
-		}
 
-		case EImageFormat::IF_RGB_UBYTE:
-		{
-			for ( int i=0; i< PixelCount; ++i )
+			case EImageFormat::IF_RGBA_UBYTE:
 			{
-                uint8 m_8 = pMaskBuf[i];
-				if (m_8>127)
+				for (int32 I = 0; I < NumElems; ++I)
 				{
-					pDestBuf[3*i+0] = palette[ pSourceBuf[3*i+0] ][0];
-					pDestBuf[3*i+1] = palette[ pSourceBuf[3*i+1] ][1];
-					pDestBuf[3*i+2] = palette[ pSourceBuf[3*i+2] ][2];
+					if (MaskBuf[I] > 127)
+					{
+						DestBuf[4*I + 0] = Palette[SourceBuf[4*I + 0]][0];
+						DestBuf[4*I + 1] = Palette[SourceBuf[4*I + 1]][1];
+						DestBuf[4*I + 2] = Palette[SourceBuf[4*I + 2]][2];
+						DestBuf[4*I + 3] = Palette[SourceBuf[4*I + 3]][3];
+					}
+					else
+					{
+						DestBuf[4*I + 0] = SourceBuf[4*I + 0];
+						DestBuf[4*I + 1] = SourceBuf[4*I + 1];
+						DestBuf[4*I + 2] = SourceBuf[4*I + 2];
+						DestBuf[4*I + 3] = SourceBuf[4*I + 3];
+					}
 				}
-				else
+				break;
+			}
+
+			case EImageFormat::IF_BGRA_UBYTE:
+			{
+				for (int32 I = 0; I < NumElems; ++I)
 				{
-					pDestBuf[3*i+0] = pSourceBuf[3*i+0];
-					pDestBuf[3*i+1] = pSourceBuf[3*i+1];
-					pDestBuf[3*i+2] = pSourceBuf[3*i+2];
+					if (MaskBuf[I] > 127)
+					{
+						DestBuf[4*I + 0] = Palette[SourceBuf[4*I + 0]][2];
+						DestBuf[4*I + 1] = Palette[SourceBuf[4*I + 1]][1];
+						DestBuf[4*I + 2] = Palette[SourceBuf[4*I + 2]][0];
+						DestBuf[4*I + 3] = Palette[SourceBuf[4*I + 3]][3];
+					}
+					else
+					{
+						DestBuf[4*I + 0] = SourceBuf[4*I + 0];
+						DestBuf[4*I + 1] = SourceBuf[4*I + 1];
+						DestBuf[4*I + 2] = SourceBuf[4*I + 2];
+						DestBuf[4*I + 3] = SourceBuf[4*I + 3];
+					}
 				}
+				break;
 			}
-			break;
-		}
 
-        case EImageFormat::IF_RGBA_UBYTE:
-        {
-            for ( int i=0; i< PixelCount; ++i )
-            {
-                uint8 m_8 = pMaskBuf[i];
-                if (m_8>127)
-                {
-                    pDestBuf[4*i+0] = palette[ pSourceBuf[4*i+0] ][0];
-                    pDestBuf[4*i+1] = palette[ pSourceBuf[4*i+1] ][1];
-                    pDestBuf[4*i+2] = palette[ pSourceBuf[4*i+2] ][2];
-                    pDestBuf[4*i+3] = palette[ pSourceBuf[4*i+3] ][3];
-                }
-                else
-                {
-                    pDestBuf[4*i+0] = pSourceBuf[4*i+0];
-                    pDestBuf[4*i+1] = pSourceBuf[4*i+1];
-                    pDestBuf[4*i+2] = pSourceBuf[4*i+2];
-                    pDestBuf[4*i+3] = pSourceBuf[4*i+3];
-                }
-            }
-            break;
-        }
-
-        case EImageFormat::IF_BGRA_UBYTE:
-        {
-            for ( int i=0; i< PixelCount; ++i )
-            {
-                uint8 m_8 = pMaskBuf[i];
-                if (m_8>127)
-                {
-                    pDestBuf[4*i+0] = palette[ pSourceBuf[4*i+0] ][2];
-                    pDestBuf[4*i+1] = palette[ pSourceBuf[4*i+1] ][1];
-                    pDestBuf[4*i+2] = palette[ pSourceBuf[4*i+2] ][0];
-                    pDestBuf[4*i+3] = palette[ pSourceBuf[4*i+3] ][3];
-                }
-                else
-                {
-                    pDestBuf[4*i+0] = pSourceBuf[4*i+0];
-                    pDestBuf[4*i+1] = pSourceBuf[4*i+1];
-                    pDestBuf[4*i+2] = pSourceBuf[4*i+2];
-                    pDestBuf[4*i+3] = pSourceBuf[4*i+3];
-                }
-            }
-            break;
-        }
-
-		default:
-			check(false);
+			default:
+				check(false);
+			}
+		};
+
+		if (NumBatches == 1)
+		{
+			ProcessBatch(0);
+		}
+		else if (NumBatches > 1)
+		{
+			ParallelFor(NumBatches, ProcessBatch);
 		}
 	}
 
diff --git a/Engine/Plugins/Experimental/Mutable/Source/MutableRuntime/Private/MuR/OpImageCompose.cpp b/Engine/Plugins/Experimental/Mutable/Source/MutableRuntime/Private/MuR/OpImageCompose.cpp
--- a/Engine/Plugins/Experimental/Mutable/Source/MutableRuntime/Private/MuR/OpImageCompose.cpp
+++ b/Engine/Plugins/Experimental/Mutable/Source/MutableRuntime/Private/MuR/OpImageCompose.cpp
@@ -7,32 +7,32 @@
 namespace mu
 {
 
-	void FImageOperator::ImageCompose( Image* pBase, const Image* pBlock, const box< UE::Math::TIntVector2<uint16> >& rect )
+	void FImageOperator::ImageCompose(Image* Base, const Image* Block, const box< UE::Math::TIntVector2<uint16> >& Rect)
 	{
-		check(pBase && pBlock);
-		check(pBase != pBlock);
-		check( rect.min[0]>=0 && rect.min[1]>=0 );
-		check( rect.size[0]>=0 && rect.size[1]>=0 );
-		check( rect.min[0]+rect.size[0]<=pBase->GetSizeX() );
-		check( rect.min[1]+rect.size[1]<=pBase->GetSizeY() );
-		check( rect.size[0]==pBlock->GetSizeX() );
-		check( rect.size[1]==pBlock->GetSizeY() );
-        check( pBase->GetFormat()==pBlock->GetFormat() );
-
-        // This assert may be acceptable if the missing lods are due to the image block format size
-        // limiting it, or disadjustment due to image scaling down. It is handled below.
-        //check( pBase->GetLODCount()<=pBlock->GetLODCount() );
-
-		EImageFormat OriginalFormat = pBase->GetFormat();
-		const FImageFormatData& finfo = GetImageFormatData( OriginalFormat );
+		check(Base && Block);
+		check(Base != Block);
+		check(Rect.min[0] >= 0 && Rect.min[1] >= 0);
+		check(Rect.size[0] >= 0 && Rect.size[1] >= 0);
+		check(Rect.min[0] + Rect.size[0] <= Base->GetSizeX());
+		check(Rect.min[1] + Rect.size[1] <= Base->GetSizeY());
+		check(Rect.size[0] == Block->GetSizeX());
+		check(Rect.size[1] == Block->GetSizeY());
+		check(Base->GetFormat() == Block->GetFormat());
+
+		// This assert may be acceptable if the missing lods are due to the image block format size
+		// limiting it, or disadjustment due to image scaling down. It is handled below.
+		//check( pBase->GetLODCount()<=pBlock->GetLODCount() );
+
+		EImageFormat OriginalFormat = Base->GetFormat();
+		const FImageFormatData& FormatInfo = GetImageFormatData(OriginalFormat);
 
 		// If this is true, we are composing coordinates or sizes that don't fall on pixel format
 		// block borders.
 		bool bSubPixelBlockCompose = 
-			rect.min[0] % finfo.PixelsPerBlockX != 0
-			|| rect.min[1] % finfo.PixelsPerBlockY != 0
-			|| rect.size[0] % finfo.PixelsPerBlockX != 0
-			|| rect.size[1] % finfo.PixelsPerBlockY != 0;
+			Rect.min[0] % FormatInfo.PixelsPerBlockX != 0
+			|| Rect.min[1] % FormatInfo.PixelsPerBlockY != 0
+			|| Rect.size[0] % FormatInfo.PixelsPerBlockX != 0
+			|| Rect.size[1] % FormatInfo.PixelsPerBlockY != 0;
 
 		if (bSubPixelBlockCompose)
 		{
@@ -40,78 +40,80 @@ namespace mu
 			MUTABLE_CPUPROFILER_SCOPE(ImageCompose_SubBlock_Fix);
 			EImageFormat UncompressedFormat = GetUncompressedFormat(OriginalFormat);
 			constexpr int32 Quality = 0;
-			Ptr<Image> TempBase = ImagePixelFormat(Quality, pBase, UncompressedFormat);
-			Ptr<Image> TempBlock = ImagePixelFormat(Quality, pBlock, UncompressedFormat);
-			ImageCompose( TempBase.get(), TempBlock.get(), rect);
+			Ptr<Image> TempBase = ImagePixelFormat(Quality, Base, UncompressedFormat);
+			Ptr<Image> TempBlock = ImagePixelFormat(Quality, Block, UncompressedFormat);
+			ImageCompose(TempBase.get(), TempBlock.get(), Rect);
 			ReleaseImage(TempBlock);
 			Ptr<Image> NewBase = ImagePixelFormat(Quality, TempBase.get(), OriginalFormat);
 			ReleaseImage(TempBase);
 			// ImageCompose above may have reduced the LODs because of the block lods.
-			pBase->CopyMove( NewBase.get() );
+			Base->CopyMove(NewBase.get());
 			ReleaseImage(NewBase);
 			return;
 		}
 		else
 		{
 			// Sizes in blocks of the current mips
-			UE::Math::TIntVector2<uint16> baseMipSize(
-				FMath::DivideAndRoundUp(pBase->GetSizeX(), uint16(finfo.PixelsPerBlockX)),
-				FMath::DivideAndRoundUp(pBase->GetSizeY(), uint16(finfo.PixelsPerBlockY)));
+			UE::Math::TIntVector2<uint16> BaseMipSize(
+				FMath::DivideAndRoundUp(Base->GetSizeX(), uint16(FormatInfo.PixelsPerBlockX)),
+				FMath::DivideAndRoundUp(Base->GetSizeY(), uint16(FormatInfo.PixelsPerBlockY)));
 
-			UE::Math::TIntVector2<uint16> blockMipPos(
-				rect.min[0] / finfo.PixelsPerBlockX,
-				rect.min[1] / finfo.PixelsPerBlockY);
+			UE::Math::TIntVector2<uint16> BlockMipPos(
+				Rect.min[0] / FormatInfo.PixelsPerBlockX,
+				Rect.min[1] / FormatInfo.PixelsPerBlockY);
 
-			UE::Math::TIntVector2<uint16> blockMipSize(
-				FMath::DivideAndRoundUp(rect.size[0], uint16(finfo.PixelsPerBlockX)),
-				FMath::DivideAndRoundUp(rect.size[1], uint16(finfo.PixelsPerBlockY)));
+			UE::Math::TIntVector2<uint16> BlockMipSize(
+				FMath::DivideAndRoundUp(Rect.size[0], uint16(FormatInfo.PixelsPerBlockX)),
+				FMath::DivideAndRoundUp(Rect.size[1], uint16(FormatInfo.PixelsPerBlockY)));
 
-			int doneMips = 0;
-			for (; doneMips < pBase->GetLODCount() && doneMips < pBlock->GetLODCount(); ++doneMips)
+			int32 DoneMips = 0;
+			for (; DoneMips < Base->GetLODCount() && DoneMips < Block->GetLODCount(); ++DoneMips)
 			{
 				//UE_LOG(LogMutableCore, Warning, "Block Mip Pos : %d, %d", blockMipPos[0], blockMipPos[1]);
 
 				// Sizes of a row in bytes
-				int baseRowSize = finfo.BytesPerBlock * baseMipSize[0];
-				int blockRowSize = finfo.BytesPerBlock * blockMipSize[0];
+				int32 BaseRowSize = FormatInfo.BytesPerBlock * BaseMipSize[0];
+				int32 BlockRowSize = FormatInfo.BytesPerBlock * BlockMipSize[0];
 
-				uint8_t* pBaseBuf = pBase->GetMipData(doneMips);
-				const uint8_t* pBlockBuf = pBlock->GetMipData(doneMips);
+				TArrayView<uint8> BaseLODView = Base->DataStorage.GetLOD(DoneMips);
+				TArrayView<const uint8> BlockLODView = Block->DataStorage.GetLOD(DoneMips);
 
-				int skipBytes = baseRowSize * blockMipPos[1] + blockMipPos[0] * finfo.BytesPerBlock;
+				uint8* BaseBuf = BaseLODView.GetData();
+				const uint8* BlockBuf = BlockLODView.GetData();
 
-				pBaseBuf += skipBytes;
-				for (int y = 0; y < blockMipSize[1]; ++y)
+				const int32 SkipBytes = BaseRowSize * BlockMipPos[1] + BlockMipPos[0] * FormatInfo.BytesPerBlock;
+
+				BaseBuf += SkipBytes;
+				for (int32 Y = 0; Y < BlockMipSize.Y; ++Y)
 				{
-					check(pBaseBuf + blockRowSize <= pBase->GetData() + pBase->GetDataSize());
-					check(pBlockBuf + blockRowSize <= pBlock->GetData() + pBlock->GetDataSize());
-					FMemory::Memcpy(pBaseBuf, pBlockBuf, blockRowSize);
-					pBlockBuf += blockRowSize;
-					pBaseBuf += baseRowSize;
+					check(BaseBuf + BlockRowSize <= BaseLODView.GetData() + BaseLODView.Num());
+					check(BlockBuf + BlockRowSize <= BlockLODView.GetData() + BlockLODView.Num());
+					FMemory::Memcpy(BaseBuf, BlockBuf, BlockRowSize);
+
+					BlockBuf += BlockRowSize;
+					BaseBuf += BaseRowSize;
 				}
 
 				// We may not proceed to the next mip, if the layout block size is smaller than the
 				// pixel format block size.
-				if (blockMipPos[0] % 2 || blockMipPos[1] % 2 || blockMipSize[0] % 2 || blockMipSize[1] % 2)
+				if (BlockMipPos[0] % 2 || BlockMipPos[1] % 2 || BlockMipSize[0] % 2 || BlockMipSize[1] % 2)
 				{
-					++doneMips;
+					++DoneMips;
 					break;
 				}
 
-				baseMipSize[0] = FMath::DivideAndRoundUp(baseMipSize[0], uint16(2));
-				baseMipSize[1] = FMath::DivideAndRoundUp(baseMipSize[1], uint16(2));
-				blockMipPos[0] = FMath::DivideAndRoundUp(blockMipPos[0], uint16(2));
-				blockMipPos[1] = FMath::DivideAndRoundUp(blockMipPos[1], uint16(2));
-				blockMipSize[0] = FMath::DivideAndRoundUp(blockMipSize[0], uint16(2));
-				blockMipSize[1] = FMath::DivideAndRoundUp(blockMipSize[1], uint16(2));
+				BaseMipSize[0] =  FMath::DivideAndRoundUp(BaseMipSize[0], uint16(2));
+				BaseMipSize[1] =  FMath::DivideAndRoundUp(BaseMipSize[1], uint16(2));
+				BlockMipPos[0] =  FMath::DivideAndRoundUp(BlockMipPos[0], uint16(2));
+				BlockMipPos[1] =  FMath::DivideAndRoundUp(BlockMipPos[1], uint16(2));
+				BlockMipSize[0] = FMath::DivideAndRoundUp(BlockMipSize[0], uint16(2));
+				BlockMipSize[1] = FMath::DivideAndRoundUp(BlockMipSize[1], uint16(2));
 			}
 
-			// Adjust the actually valid mips, maybe wasting some of the reserved memory.
-			// TODO: Review this.
-			pBase->m_lods = (uint8)doneMips;
-			pBase->m_data.SetNum(pBase->CalculateDataSize());
+			// Adjust the actually valid mips.
+			Base->DataStorage.SetNumLODs(DoneMips);
 		}
-    }
+	}
 
 }
 
