[Mutable] Store Image LODs in different buffers.

diff --git a/Engine/Plugins/Experimental/Mutable/Source/MutableRuntime/Private/MuR/OpImageCrop.cpp b/Engine/Plugins/Experimental/Mutable/Source/MutableRuntime/Private/MuR/OpImageCrop.cpp
--- a/Engine/Plugins/Experimental/Mutable/Source/MutableRuntime/Private/MuR/OpImageCrop.cpp
+++ b/Engine/Plugins/Experimental/Mutable/Source/MutableRuntime/Private/MuR/OpImageCrop.cpp
@@ -23,7 +23,7 @@ namespace mu
 			// \TODO: This may use some additional untracked memory locally in this function.
 			BaseReformat = ImagePixelFormat(CompressionQuality, Base.get(), UncompressedFormat);
 			Base = BaseReformat;
-			Cropped = CreateImage( InCropped->GetSizeX(), InCropped->GetSizeY(), InCropped->GetLODCount(), UncompressedFormat, EInitializationType::NotInitialized);
+			Cropped = CreateImage(InCropped->GetSizeX(), InCropped->GetSizeY(), InCropped->GetLODCount(), UncompressedFormat, EInitializationType::NotInitialized);
         }
 
 		const FImageFormatData& finfo = GetImageFormatData(UncompressedFormat);
@@ -34,35 +34,36 @@ namespace mu
 
 		// TODO: better error control. This happens if some layouts are corrupt.
 		bool bCorrect =
-				( Rect.min[0]>=0 && Rect.min[1]>=0 ) &&
-				( Rect.size[0]>=0 && Rect.size[1]>=0 ) &&
-				( Rect.min[0]+Rect.size[0]<=Base->GetSizeX() ) &&
-				( Rect.min[1]+Rect.size[1]<=Base->GetSizeY() );
+				(Rect.min[0] >= 0 && Rect.min[1] >= 0) &&
+				(Rect.size[0] >= 0 && Rect.size[1] >= 0) &&
+				(Rect.min[0] + Rect.size[0] <= Base->GetSizeX()) &&
+				(Rect.min[1] + Rect.size[1] <= Base->GetSizeY());
+
 		if (!bCorrect)
 		{
 			return false;
 		}
 
 		// Block images are not supported for now
-		check( finfo.PixelsPerBlockX == 1 );
-		check( finfo.PixelsPerBlockY == 1 );
+		check(finfo.PixelsPerBlockX == 1);
+		check(finfo.PixelsPerBlockY == 1);
 
-		checkf( Rect.min[0] % finfo.PixelsPerBlockX == 0, TEXT("Rect must snap to blocks.") );
-		checkf( Rect.min[1] % finfo.PixelsPerBlockY == 0, TEXT("Rect must snap to blocks.") );
-		checkf( Rect.size[0] % finfo.PixelsPerBlockX == 0, TEXT("Rect must snap to blocks.") );
-		checkf( Rect.size[1] % finfo.PixelsPerBlockY == 0, TEXT("Rect must snap to blocks.") );
+		checkf(Rect.min[0] % finfo.PixelsPerBlockX == 0, TEXT("Rect must snap to blocks."));
+		checkf(Rect.min[1] % finfo.PixelsPerBlockY == 0, TEXT("Rect must snap to blocks."));
+		checkf(Rect.size[0] % finfo.PixelsPerBlockX == 0, TEXT("Rect must snap to blocks."));
+		checkf(Rect.size[1] % finfo.PixelsPerBlockY == 0, TEXT("Rect must snap to blocks."));
 
-		int baseRowSize = finfo.BytesPerBlock * FMath::DivideAndRoundUp(Base->GetSizeX(), uint16(finfo.PixelsPerBlockX));
-		int cropRowSize = finfo.BytesPerBlock * FMath::DivideAndRoundUp(Rect.size[0], int32(finfo.PixelsPerBlockX));
+		int32 baseRowSize = finfo.BytesPerBlock * FMath::DivideAndRoundUp(Base->GetSizeX(), uint16(finfo.PixelsPerBlockX));
+		int32 cropRowSize = finfo.BytesPerBlock * FMath::DivideAndRoundUp(Rect.size[0], int32(finfo.PixelsPerBlockX));
 
-        const uint8_t* pBaseBuf = Base->GetData();
-        uint8_t* pCropBuf = Cropped->GetData();
+        const uint8* pBaseBuf = Base->GetLODData(0);
+        uint8* pCropBuf = Cropped->GetLODData(0);
 
-		int skipPixels = Base->GetSizeX() * Rect.min[1] + Rect.min[0];
-		pBaseBuf += finfo.BytesPerBlock * skipPixels / finfo.PixelsPerBlockX;
-		for ( int y=0; y<Rect.size[1]; ++y )
+		int32 SkipPixels = Base->GetSizeX() * Rect.min[1] + Rect.min[0];
+		pBaseBuf += finfo.BytesPerBlock * SkipPixels / finfo.PixelsPerBlockX;
+		for (int32 Y = 0; Y < Rect.size[1]; ++Y)
 		{
-			FMemory::Memcpy( pCropBuf, pBaseBuf, cropRowSize );
+			FMemory::Memcpy(pCropBuf, pBaseBuf, cropRowSize);
 			pCropBuf += cropRowSize;
 			pBaseBuf += baseRowSize;
 		}
@@ -72,13 +73,8 @@ namespace mu
 			ReleaseImage(BaseReformat);
 
 			bool bSuccess = false;
-			int32 DataSize = Cropped->m_data.Num();
-			while (!bSuccess)
-			{
-				InCropped->m_data.SetNumUninitialized(DataSize);
-				ImagePixelFormat(bSuccess, CompressionQuality, InCropped, Cropped.get() );
-				DataSize = (DataSize + 16) * 2;
-			}
+			ImagePixelFormat(bSuccess, CompressionQuality, InCropped, Cropped.get());
+			check(bSuccess);
 
 			ReleaseImage(Cropped);
 		}
diff --git a/Engine/Plugins/Experimental/Mutable/Source/MutableRuntime/Private/MuR/OpImageDisplace.h b/Engine/Plugins/Experimental/Mutable/Source/MutableRuntime/Private/MuR/OpImageDisplace.h
--- a/Engine/Plugins/Experimental/Mutable/Source/MutableRuntime/Private/MuR/OpImageDisplace.h
+++ b/Engine/Plugins/Experimental/Mutable/Source/MutableRuntime/Private/MuR/OpImageDisplace.h
@@ -12,109 +12,114 @@
 namespace mu
 {
 
-    inline uint8 MutableEncodeOffset( int x, int y )
+    inline uint8 MutableEncodeOffset(int32 X, int32 Y)
 	{
-        uint8 c = uint8( ((7)<<4) | 7 );
-		if ( x<8 && x>-8 && y<8 && y>-8 )
+        uint8 C = uint8((7 << 4) | 7);
+		if ((X < 8) & (X > -8) & (Y < 8) & (Y > -8))
 		{
-            c = uint8( ((x+7)<<4) | (y+7) );
+            C = uint8(((X + 7) << 4) | (Y + 7));
 		}
-		return c;
+
+		return C;
 	}
 
-    inline void MutableDecodeOffset( uint8 c, int& x, int& y )
+    inline void MutableDecodeOffset(uint8 C, int32& X, int32& Y)
 	{
-		x = int(c>>4) - 7;
-		y = int(c&0xf) - 7;
+		X = int32(C >> 4) - 7;
+		Y = int32(C & 0xF) - 7;
 	}
 
 	//---------------------------------------------------------------------------------------------
 	//!
 	//---------------------------------------------------------------------------------------------
-	inline void ImageMakeGrowMap( Image* pResult, const Image* pMask, int InBorder )
+	inline void ImageMakeGrowMap(Image* ResultImage, const Image* MaskImage, int32 InBorder)
 	{
 		MUTABLE_CPUPROFILER_SCOPE(ImageMakeGrowMap);
 
-		check(pMask->GetFormat() == EImageFormat::IF_L_UBYTE);
-		check(pResult->GetFormat() == EImageFormat::IF_L_UBYTE);
-		check(pResult->GetSizeX() == pMask->GetSizeX());
-		check(pResult->GetSizeY() == pMask->GetSizeY());
-		check(pResult->GetLODCount() == pMask->GetLODCount());
+		check(MaskImage->GetFormat() == EImageFormat::IF_L_UBYTE);
+		check(ResultImage->GetFormat() == EImageFormat::IF_L_UBYTE);
+		check(ResultImage->GetSizeX() == MaskImage->GetSizeX());
+		check(ResultImage->GetSizeY() == MaskImage->GetSizeY());
+		check(ResultImage->GetLODCount() == MaskImage->GetLODCount());
 
-		int32 MipCount = pResult->GetLODCount();
-		int32 SizeX = pResult->GetSizeX();
-		int32 SizeY = pResult->GetSizeY();
+		int32 MipCount = ResultImage->GetLODCount();
+		int32 SizeX = ResultImage->GetSizeX();
+		int32 SizeY = ResultImage->GetSizeY();
 		
 		if (SizeX <= 0 || SizeY <= 0)
 		{
 			return;
 		}
 
-		FMemory::Memset( pResult->GetData(), MutableEncodeOffset( 0, 0 ), pResult->CalculatePixelCount() );
+		const uint8 ZeroOffset = MutableEncodeOffset(0, 0);
+		for (FImageArray& Buffer : ResultImage->DataStorage.Buffers)
+		{
+			FMemory::Memset(Buffer.GetData(), ZeroOffset, Buffer.Num());
+		}
 
-		Ptr<const Image> pThisMask = pMask;
-		for ( int b=0; b<InBorder; ++b )
+		Ptr<const Image> ThisMask = MaskImage;
+		for (int32 B = 0; B < InBorder; ++B)
 		{
-			Ptr<Image> pNextMask = pThisMask->Clone();
+			Ptr<Image> NextMask = ThisMask->Clone();
 
-			for( int CurrentMip=0; CurrentMip<MipCount; ++CurrentMip)
+			for(int32 CurrentMip = 0; CurrentMip < MipCount; ++CurrentMip)
 			{ 
-				FIntVector2 MipSize = pResult->CalculateMipSize(CurrentMip);
+				FIntVector2 MipSize = ResultImage->CalculateMipSize(CurrentMip);
 
-				uint8* ResultData = pResult->GetMipData(CurrentMip);
-				const uint8* ThisMaskData = pThisMask->GetMipData(CurrentMip);
-				uint8* NextMaskData = pNextMask->GetMipData(CurrentMip);
+				uint8* ResultData = ResultImage->GetMipData(CurrentMip);
+				const uint8* ThisMaskData = ThisMask->GetMipData(CurrentMip);
+				uint8* NextMaskData = NextMask->GetMipData(CurrentMip);
 
-				//for ( int y=0; y< MipSize.Y; ++y )
+				//for (int32 Y=0; Y< MipSize.Y; ++Y)
 				const auto ProcessRow = [
 					ThisMaskData, NextMaskData, MipSize, ResultData
-				] (int32 y)
+				] (int32 Y)
 				{
-					uint8* pNextData = NextMaskData + MipSize.X * y;
+					uint8* NextData = NextMaskData + MipSize.X * Y;
 
-					const uint8* pS0 = y > 0 ? ThisMaskData + (y - 1) * MipSize.X : nullptr;
-					const uint8* pS1 = ThisMaskData + y * MipSize.X;
-					bool bNotLastRow = y < MipSize.Y - 1;
-					const uint8* pS2 = bNotLastRow ? ThisMaskData + (y + 1) * MipSize.X : nullptr;
+					const uint8* S0 = Y > 0 ? ThisMaskData + (Y - 1) * MipSize.X : nullptr;
+					const uint8* S1 = ThisMaskData + Y * MipSize.X;
+					bool bNotLastRow = Y < MipSize.Y - 1;
+					const uint8* S2 = bNotLastRow ? ThisMaskData + (Y + 1) * MipSize.X : nullptr;
 
-					const uint8* pR0 = y > 0 ? ResultData + (y - 1) * MipSize.X : 0;
-					uint8* pR1 = ResultData + y * MipSize.X;
+					const uint8* R0 = Y > 0 ? ResultData + (Y - 1) * MipSize.X : 0;
+					uint8* R1 = ResultData + Y * MipSize.X;
 
-					const uint8* pR2 = bNotLastRow ? ResultData + (y + 1) * MipSize.X : nullptr;
+					const uint8* R2 = bNotLastRow ? ResultData + (Y + 1) * MipSize.X : nullptr;
 
-					for (int x = 0; x < MipSize.X; ++x)
+					for (int32 X = 0; X < MipSize.X; ++X)
 					{
-						int dx, dy;
-						if (pS1[x])
+						int32 Dx, Dy;
+						if (S1[X])
 						{
-							*pNextData = 255;
+							*NextData = 255;
 						}
-						else if (y > 0 && pS0[x])
+						else if (Y > 0 && S0[X])
 						{
-							MutableDecodeOffset(pR0[x], dx, dy);
-							pR1[x] = MutableEncodeOffset(dx, dy - 1);
-							*pNextData = 255;
+							MutableDecodeOffset(R0[X], Dx, Dy);
+							R1[X] = MutableEncodeOffset(Dx, Dy - 1);
+							*NextData = 255;
 						}
-						else if (bNotLastRow && pS2 && pR2 && pS2[x])
+						else if (bNotLastRow && S2 && R2 && S2[X])
 						{
-							MutableDecodeOffset(pR2[x], dx, dy);
-							pR1[x] = MutableEncodeOffset(dx, dy + 1);
-							*pNextData = 255;
+							MutableDecodeOffset(R2[X], Dx, Dy);
+							R1[X] = MutableEncodeOffset(Dx, Dy + 1);
+							*NextData = 255;
 						}
-						else if (x < MipSize.X - 1 && pS1[x + 1])
+						else if (X < MipSize.X - 1 && S1[X + 1])
 						{
-							MutableDecodeOffset(pR1[x + 1], dx, dy);
-							pR1[x] = MutableEncodeOffset(dx + 1, dy);
-							*pNextData = 255;
+							MutableDecodeOffset(R1[X + 1], Dx, Dy);
+							R1[X] = MutableEncodeOffset(Dx + 1, Dy);
+							*NextData = 255;
 						}
-						else if (x > 0 && pS1[x - 1])
+						else if (X > 0 && S1[X - 1])
 						{
-							MutableDecodeOffset(pR1[x - 1], dx, dy);
-							pR1[x] = MutableEncodeOffset(dx - 1, dy);
-							*pNextData = 255;
+							MutableDecodeOffset(R1[X - 1], Dx, Dy);
+							R1[X] = MutableEncodeOffset(Dx - 1, Dy);
+							*NextData = 255;
 						}
 
-						++pNextData;
+						++NextData;
 					}
 
 				};
@@ -122,224 +127,226 @@ namespace mu
 				ParallelFor(MipSize.Y, ProcessRow);
 			}
 
-			pThisMask = pNextMask;
+			ThisMask = NextMask;
 		}
 
 	}
 
 
-	//---------------------------------------------------------------------------------------------
-	//!
-	//---------------------------------------------------------------------------------------------
-	inline void ImageDisplace( Image* pResult, const Image* pSource, const Image* pMap )
+	inline void ImageDisplace(Image* ResultImage, const Image* SourceImage, const Image* MapImage)
 	{
 		MUTABLE_CPUPROFILER_SCOPE(ImageDisplace)
 
-		check( pResult->GetFormat() == pSource->GetFormat()  );
-		check( pMap->GetFormat() == EImageFormat::IF_L_UBYTE
-						||
-						pMap->GetFormat() == EImageFormat::IF_L_UBYTE_RLE );
-		check( pResult->GetSizeX() == pSource->GetSizeX() );
-		check( pResult->GetSizeY() == pSource->GetSizeY() );
-		check( pResult->GetSizeX() == pMap->GetSizeX() );
-		check( pResult->GetSizeY() == pMap->GetSizeY() );
+		check(ResultImage->GetFormat() == SourceImage->GetFormat());
+		check(MapImage->GetFormat() == EImageFormat::IF_L_UBYTE || MapImage->GetFormat() == EImageFormat::IF_L_UBYTE_RLE);
+		check(ResultImage->GetSizeX() == SourceImage->GetSizeX());
+		check(ResultImage->GetSizeY() == SourceImage->GetSizeY());
+		check(ResultImage->GetSizeX() == MapImage->GetSizeX());
+		check(ResultImage->GetSizeY() == MapImage->GetSizeY());
 
-		int32 LODCount = pResult->GetLODCount();
-		check(LODCount <= pMap->GetLODCount());
-		check(LODCount <= pSource->GetLODCount());
+		int32 LODCount = ResultImage->GetLODCount();
+		check(LODCount <= MapImage->GetLODCount());
+		check(LODCount <= SourceImage->GetLODCount());
 
-		int32 sizeX = pResult->GetSizeX();
-		int32 sizeY = pResult->GetSizeY();
+		int32 SizeX = ResultImage->GetSizeX();
+		int32 SizeY = ResultImage->GetSizeY();
 
-        //if (sizeX<4 || sizeY<4)
+        //if (SizeX < 4 || SizeY < 4)
         //{
         //    return;
         //}
+        
+		EImageFormat MapFormat = MapImage->GetFormat();
+        bool bIsUncompressed = MapFormat == EImageFormat::IF_L_UBYTE;
 
-        const uint8* pSourceData = pSource->GetData();
-
-		EImageFormat mapFormat = pMap->GetFormat();
-        bool isUncompressed = ( mapFormat == EImageFormat::IF_L_UBYTE );
-
-        if ( isUncompressed )
+        if (bIsUncompressed)
 		{
-			switch (pResult->GetFormat())
+			switch (ResultImage->GetFormat())
 			{
-
 			case EImageFormat::IF_L_UBYTE:
 			{
-				//for ( int y=0; y<sizeY; ++y )
+				const uint8* SourceData = SourceImage->GetLODData(0);
+				const uint8* MapBuf = MapImage->GetLODData(0);
+				uint8* ResultBuf = ResultImage->GetLODData(0);
+
+				//for (int32 Y = 0; Y < SizeY; ++Y)
 				const auto ProcessRow = [
-					pSourceData, pMap, pResult, sizeX, sizeY
-				] (int32 y)
+					SourceData, MapBuf, ResultBuf, SizeX, SizeY
+				] (int32 Y)
 				{
-					const uint8* pMapData = pMap->GetData() + y * sizeX;
-					uint8* pResultData = pResult->GetData() + y * sizeX;
+					const uint8* MapData = MapBuf + Y * SizeX;
+					uint8* ResultData = ResultBuf + Y * SizeX;
 
-					for (int x = 0; x < sizeX; ++x)
+					for (int32 X = 0; X < SizeX; ++X)
 					{
-						int dx, dy;
-						MutableDecodeOffset(*pMapData, dx, dy);
+						int32 Dx, Dy;
+						MutableDecodeOffset(*MapData, Dx, Dy);
 
 						// This could actually happen since we enable the crop+displace optimization
-						//check(x + dx >= 0 && x + dx < sizeX);
-						//check(y + dy >= 0 && y + dy < sizeY);
-						if ((x + dx >= 0 && x + dx < sizeX) && (y + dy >= 0 && y + dy < sizeY))
+						//check(X + Dx >= 0 && X + Dx < SizeX);
+						//check(Y + Dy >= 0 && Y + Dy < SizeY);
+						if ((X + Dx >= 0 && X + Dx < SizeX) && (Y + Dy >= 0 && Y + Dy < SizeY))
 						{
-							int offset = ((y + dy) * sizeX + (x + dx));
-							FMemory::Memcpy(pResultData, &pSourceData[offset], 1);
+							int32 Offset = ((Y + Dy) * SizeX + (X + Dx));
+							FMemory::Memcpy(ResultData, SourceData + Offset, 1);
 						}
 
-						++pResultData;
-						++pMapData;
+						++ResultData;
+						++MapData;
 					}
 
 				};
 
-				ParallelFor(sizeY, ProcessRow);
+				ParallelFor(SizeY, ProcessRow);
 				break;
 			}
 
 			case EImageFormat::IF_RGB_UBYTE:
 			{
-				//for (int y = 0; y < sizeY; ++y)
-				const auto ProcessRow = [
-					pSourceData, pMap, pResult, sizeX, sizeY
-				] (int32 y)
+        		const uint8* SourceData = SourceImage->GetLODData(0);
+				const uint8* MapBuf = MapImage->GetLODData(0);
+				uint8* ResultBuf = ResultImage->GetLODData(0);
+				//for (int y = 0; y < SizeY; ++y)
+				const auto ProcessRow = [SourceData, MapBuf, ResultBuf, SizeX, SizeY] (int32 Y)
 				{
-					const uint8* pMapData = pMap->GetData() + y * sizeX;
-					uint8* pResultData = pResult->GetData() + y * sizeX*3;
+					const uint8* MapData = MapBuf + Y * SizeX;
+					uint8* ResultData = ResultBuf + Y * SizeX*3;
 
-					for (int x = 0; x < sizeX; ++x)
+					for (int32 X = 0; X < SizeX; ++X)
 					{
-						int dx, dy;
-						MutableDecodeOffset(*pMapData, dx, dy);
+						int32 Dx, Dy;
+						MutableDecodeOffset(*MapData, Dx, Dy);
 						// This could actually happen since we enable the crop+displace optimization
-						//check(x + dx >= 0 && x + dx < sizeX);
-						//check(y + dy >= 0 && y + dy < sizeY);
-						if ((x + dx >= 0 && x + dx < sizeX) && (y + dy >= 0 && y + dy < sizeY))
+						//check(X + Dx >= 0 && X + Dx < SizeX);
+						//check(Y + Dy >= 0 && Y + Dy < SizeY);
+						if ((X + Dx >= 0 && X + Dx < SizeX) && (Y + Dy >= 0 && Y + Dy < SizeY))
 						{
-							int offset = ((y + dy) * sizeX + (x + dx)) * 3;
-							FMemory::Memcpy(pResultData, &pSourceData[offset], 3);
+							int32 Offset = ((Y + Dy) * SizeX + (X + Dx)) * 3;
+							FMemory::Memcpy(ResultData, SourceData + Offset, 3);
 						}
 
-						pResultData += 3;
-						++pMapData;
+						ResultData += 3;
+						++MapData;
 					}
 
 				};
 
-				ParallelFor(sizeY, ProcessRow);
+				ParallelFor(SizeY, ProcessRow);
 				break;
 			}
 
 			default:
 			{
 				// Generic
-				int pixelSize = GetImageFormatData( pResult->GetFormat() ).BytesPerBlock;
-				//for ( int y=0; y<sizeY; ++y )
+				int32 PixelSize = GetImageFormatData(ResultImage->GetFormat()).BytesPerBlock;
+				
+				const uint8* SourceData = SourceImage->GetLODData(0);
+				const uint8* MapBuf = MapImage->GetLODData(0);
+				uint8* ResultBuf = ResultImage->GetLODData(0);
+				
+				//for (int32 Y = 0; Y < SizeY; ++Y)
 				const auto ProcessRow = [
-					pSourceData, pMap, pResult, sizeX, sizeY, pixelSize
-				] (int32 y)
+					SourceData, MapBuf, ResultBuf, SizeX, SizeY, PixelSize
+				] (int32 Y)
 				{
-					const uint8* pMapData = pMap->GetData() + y * sizeX;
-					uint8* pResultData = pResult->GetData() + y * sizeX * pixelSize;
+					const uint8* MapData = MapBuf + Y * SizeX;
+					uint8* ResultData = ResultBuf + Y * SizeX * PixelSize;
 
-					for (int x = 0; x < sizeX; ++x)
+					for (int32 X = 0; X < SizeX; ++X)
 					{
-						int dx, dy;
-						MutableDecodeOffset(*pMapData, dx, dy);
+						int32 Dx, Dy;
+						MutableDecodeOffset(*MapData, Dx, Dy);
 						// This could actually happen since we enable the crop+displace optimization
-						//check(x + dx >= 0 && x + dx < sizeX);
-						//check(y + dy >= 0 && y + dy < sizeY);
-						if ((x + dx >= 0 && x + dx < sizeX) && (y + dy >= 0 && y + dy < sizeY))
+						//check(X + Dx >= 0 && X + Dx < SizeX);
+						//check(Y + Dy >= 0 && Y + Dy < SizeY);
+						if ((X + Dx >= 0 && X + Dx < SizeX) && (Y + Dy >= 0 && Y + Dy < SizeY))
 						{
-							int offset = ((y + dy) * sizeX + (x + dx)) * pixelSize;
-							FMemory::Memcpy(pResultData, &pSourceData[offset], pixelSize);
+							int32 Offset = ((Y + Dy) * SizeX + (X + Dx)) * PixelSize;
+							FMemory::Memcpy(ResultData, SourceData + Offset, PixelSize);
 						}
 
-						pResultData += pixelSize;
-						++pMapData;
+						ResultData += PixelSize;
+						++MapData;
 					}
 
 				};
 
-				ParallelFor(sizeY, ProcessRow);
+				ParallelFor(SizeY, ProcessRow);
 			}
 
 			}
 		}
 
-        else if (mapFormat==EImageFormat::IF_L_UBYTE_RLE )
+        else if (MapFormat == EImageFormat::IF_L_UBYTE_RLE)
 		{
-            int pixelSize = GetImageFormatData( pResult->GetFormat() ).BytesPerBlock;
-
-			const uint8* pMapData = pMap->GetData();
-			uint8* pResultData = pResult->GetData();
-
+            int32 PixelSize = GetImageFormatData(ResultImage->GetFormat()).BytesPerBlock;
+			
 			// TODO: ParallelFor this
-
-            for (int lod=0;lod<LODCount; ++lod)
+            for (int32 LOD = 0; LOD < LODCount; ++LOD)
             {
-                pMapData += 4+sizeY*sizeof(uint32);
+				const uint8* SourceData = SourceImage->GetLODData(LOD);
+                const uint8* MapData = MapImage->GetLODData(LOD);
+				uint8* ResultData = ResultImage->GetLODData(LOD);
 
-                for ( int y=0; y<sizeY; ++y )
+				// Skip RLE header, total size and row sizes.
+				MapData += sizeof(uint32) + SizeY*sizeof(uint32);
+
+                for (int32 Y = 0; Y < SizeY; ++Y)
                 {
-                    int x=0;
-                    while ( x<sizeX )
+                    int32 X = 0;
+                    while (X < SizeX)
                     {
                         // Decode header
-                        uint16 equal = *(const uint16*)pMapData;
-                        pMapData += 2;
+                        uint16 Equal = *(const uint16*)MapData;
+                        MapData += 2;
 
-                        uint8 different = *pMapData;
-                        ++pMapData;
+                        uint8 Different = *MapData;
+                        ++MapData;
 
-                        uint8 equalPixel = *pMapData;
-                        ++pMapData;
+                        uint8 EqualPixel = *MapData;
+                        ++MapData;
 
                         // Equal pixels
                         {
-                            int dx,dy;
-                            MutableDecodeOffset( equalPixel, dx, dy );
-                            dx += x;
-                            dy += y;
+                            int32 Dx, Dy;
+                            MutableDecodeOffset(EqualPixel, Dx, Dy);
+                            Dx += X;
+                            Dy += Y;
 
-							if (dx >= 0 && dx < sizeX && dy >= 0 && dy < sizeY)
+							if (Dx >= 0 && Dx < SizeX && Dy >= 0 && Dy < SizeY)
 							{
-								int offset = (dy * sizeX + dx) * pixelSize;
-								FMemory::Memcpy(pResultData, &pSourceData[offset], equal * pixelSize);
+								int32 Offset = (Dy*SizeX + Dx) * PixelSize;
+								FMemory::Memcpy(ResultData, SourceData + Offset, Equal*PixelSize);
 							}
 
-                            pResultData += equal*pixelSize;
-                            x += equal;
+                            ResultData += Equal*PixelSize;
+                            X += Equal;
                         }
 
-
                         // Different pixels
-                        for ( int i=0; i<different; ++i )
+                        for (int32 I = 0; I < Different; ++I)
                         {
-                            int dx,dy;
-                            MutableDecodeOffset( pMapData[i], dx, dy );
-                            dx += x;
-                            dy += y;
+                            int32 Dx, Dy;
+                            MutableDecodeOffset(MapData[I], Dx, Dy);
+                            Dx += X;
+                            Dy += Y;
 
-							if (dx>=0 && dx<sizeX && dy>=0 && dy<sizeY)
+							if (Dx >= 0 && Dx < SizeX && Dy >= 0 && Dy < SizeY)
 							{
-								int offset = (dy * sizeX + dx) * pixelSize;
-								FMemory::Memcpy(pResultData, &pSourceData[offset], pixelSize);
+								int32 Offset = (Dy * SizeX + Dx) * PixelSize;
+								FMemory::Memcpy(ResultData, &SourceData[Offset], PixelSize);
 							}
 
-                            pResultData += pixelSize;
-                            ++x;
+                            ResultData += PixelSize;
+                            ++X;
                         }
 
-                        pMapData += different;
+                        MapData += Different;
                     }
                 }
 
-                sizeX = FMath::DivideAndRoundUp(sizeX,2);
-                sizeY = FMath::DivideAndRoundUp(sizeY,2);
+                SizeX = FMath::DivideAndRoundUp(SizeX, 2);
+                SizeY = FMath::DivideAndRoundUp(SizeY, 2);
             }
 		}
 
@@ -350,16 +357,14 @@ namespace mu
 
 
 		// Update the relevancy data of the image for the worst case. 
-		if (pResult->m_flags & Image::EImageFlags::IF_HAS_RELEVANCY_MAP)
+		if (ResultImage->m_flags & Image::EImageFlags::IF_HAS_RELEVANCY_MAP)
 		{
 			// Displace can encode at max MUTABLE_GROW_BORDER_VALUE pixels away according to "border" in MakeGrowMap.
-			pResult->RelevancyMinY = FMath::Max(int32(pResult->RelevancyMinY) - MUTABLE_GROW_BORDER_VALUE, 0);
-			pResult->RelevancyMaxY = FMath::Min(pResult->RelevancyMaxY + MUTABLE_GROW_BORDER_VALUE, pResult->GetSizeY() - 1);
+			ResultImage->RelevancyMinY = FMath::Max(int32(ResultImage->RelevancyMinY) - MUTABLE_GROW_BORDER_VALUE, 0);
+			ResultImage->RelevancyMaxY = FMath::Min(ResultImage->RelevancyMaxY + MUTABLE_GROW_BORDER_VALUE, ResultImage->GetSizeY() - 1);
 		}
-
 	}
 
-
 }
 
 
diff --git a/Engine/Plugins/Experimental/Mutable/Source/MutableRuntime/Private/MuR/OpImageGradient.h b/Engine/Plugins/Experimental/Mutable/Source/MutableRuntime/Private/MuR/OpImageGradient.h
--- a/Engine/Plugins/Experimental/Mutable/Source/MutableRuntime/Private/MuR/OpImageGradient.h
+++ b/Engine/Plugins/Experimental/Mutable/Source/MutableRuntime/Private/MuR/OpImageGradient.h
@@ -13,7 +13,7 @@ namespace mu
 	{
 		check(pDest && pDest->GetFormat()==EImageFormat::IF_RGBA_UBYTE);
 
-        uint8* pDestBuf = pDest->GetData();
+        uint8* pDestBuf = pDest->GetLODData(0);
 
 		int32 sizeX = pDest->GetSizeX();
 		int32 sizeY = pDest->GetSizeX();
diff --git a/Engine/Plugins/Experimental/Mutable/Source/MutableRuntime/Private/MuR/OpImageInterpolate.h b/Engine/Plugins/Experimental/Mutable/Source/MutableRuntime/Private/MuR/OpImageInterpolate.h
--- a/Engine/Plugins/Experimental/Mutable/Source/MutableRuntime/Private/MuR/OpImageInterpolate.h
+++ b/Engine/Plugins/Experimental/Mutable/Source/MutableRuntime/Private/MuR/OpImageInterpolate.h
@@ -5,76 +5,70 @@
 #include "MuR/ImagePrivate.h"
 #include "MuR/Platform.h"
 
+#include "Async/ParallelFor.h"
+
 namespace mu
 {
 
-	//---------------------------------------------------------------------------------------------
-	//!
-	//---------------------------------------------------------------------------------------------
-    inline void ImageInterpolate( Image* pDest,
-                                  const Image* pB,
-                                  float factor )
+    inline void ImageInterpolate(Image* DestImage, const Image* BImage, float Factor)
 	{
 		MUTABLE_CPUPROFILER_SCOPE(ImageInterpolate)
 
-		check(pDest && pB);
-        check(pDest->GetSizeX() == pB->GetSizeX() );
-        check(pDest->GetSizeY() == pB->GetSizeY() );
-        check(pDest->GetFormat() == pB->GetFormat() );
+		check(DestImage && BImage);
+        check(DestImage->GetSizeX() == BImage->GetSizeX());
+        check(DestImage->GetSizeY() == BImage->GetSizeY());
+        check(DestImage->GetFormat() == BImage->GetFormat());
 
 		// Clamp the factor
-		factor = FMath::Max( 0.0f, FMath::Min( 1.0f, factor ) );
+		Factor = FMath::Clamp(Factor, 0.0f, 1.0f);
+		const uint16 FactorUNorm = static_cast<uint16>(Factor * 255.0f);
 
-        uint8* pDestBuf = pDest->GetData();
-        const uint8* pBBuf = pB->GetData();
+		const int32 BytesPerElem = GetImageFormatData(DestImage->GetFormat()).BytesPerBlock;
 
-		// Generic implementation
-		int32 PixelCount = (int32)pDest->CalculatePixelCount();
+		constexpr int32 NumBatchElems = 1 << 14;
 
-		switch ( pDest->GetFormat() )
-		{
-		case EImageFormat::IF_L_UBYTE:
-		{
-            uint32 w_8 = (uint32)(factor*255);
-			for ( int i=0; i< PixelCount; ++i )
-			{
-                uint32 a_8 = *pDestBuf;
-                uint32 b_8 = *pBBuf++;
-                uint32 i_16 = a_8 * (255-w_8) + b_8 * w_8;
-                *pDestBuf++ = (uint8) ( i_16>>8 );
-			}
-			break;
-		}
+		const int32 NumBatches = DestImage->DataStorage.GetNumBatches(NumBatchElems, BytesPerElem);
+		check(BImage->DataStorage.GetNumBatches(NumBatchElems, BytesPerElem) == NumBatches);
 
-		case EImageFormat::IF_RGB_UBYTE:
+		// Generic implementation	
+		auto ProcessBatch = [DestImage, BImage, NumBatchElems, BytesPerElem, FactorUNorm](int32 BatchId)
 		{
-            uint32 w_8 = (uint32)(factor*255);
-			for ( int i=0; i< PixelCount *3; ++i )
-			{
-                uint32 a_8 = *pDestBuf;
-                uint32 b_8 = *pBBuf++;
-                uint32 i_16 = a_8 * (255-w_8) + b_8 * w_8;
-                *pDestBuf++ = (uint8) ( i_16>>8 );
-			}
-			break;
-		}
+			TArrayView<uint8> DestView = DestImage->DataStorage.GetBatch(BatchId, NumBatchElems, BytesPerElem);
+			TArrayView<const uint8> BView = BImage->DataStorage.GetBatch(BatchId, NumBatchElems, BytesPerElem);
 
-        case EImageFormat::IF_BGRA_UBYTE:
-        case EImageFormat::IF_RGBA_UBYTE:
-        {
-            uint32 w_8 = (uint32)(factor*255);
-			for ( int i=0; i< PixelCount *4; ++i )
+			const uint8* BBuf = BView.GetData();
+			uint8* DestBuf = DestView.GetData();
+
+			const int32 NumBytes = DestView.Num();
+			check(BView.Num() == NumBytes);
+
+			switch (DestImage->GetFormat())
 			{
-                uint32 a_8 = *pDestBuf;
-                uint32 b_8 = *pBBuf++;
-                uint32 i_16 = a_8 * (255-w_8) + b_8 * w_8;
-                *pDestBuf++ = (uint8)(i_16>>8);
+			case EImageFormat::IF_L_UBYTE:
+			case EImageFormat::IF_RGB_UBYTE:
+			case EImageFormat::IF_BGRA_UBYTE:
+			case EImageFormat::IF_RGBA_UBYTE:
+			{
+				for (int32 I = 0; I < NumBytes; ++I)
+				{
+					uint16 AValue = DestBuf[I];
+					uint16 BValue = BBuf[I];
+					DestBuf[I] = static_cast<uint8>((AValue * (255 - FactorUNorm) + BValue*FactorUNorm) / 255);
+				}
+				break;
 			}
-			break;
-		}
+			default:
+				check(false);
+			}
+		};
 
-		default:
-			check(false);
+		if (NumBatches == 1)
+		{
+			ProcessBatch(0);
+		}
+		else if (NumBatches > 1)
+		{
+			ParallelFor(NumBatches, ProcessBatch);
 		}
 
 	}
diff --git a/Engine/Plugins/Experimental/Mutable/Source/MutableRuntime/Private/MuR/OpImageInvert.h b/Engine/Plugins/Experimental/Mutable/Source/MutableRuntime/Private/MuR/OpImageInvert.h
--- a/Engine/Plugins/Experimental/Mutable/Source/MutableRuntime/Private/MuR/OpImageInvert.h
+++ b/Engine/Plugins/Experimental/Mutable/Source/MutableRuntime/Private/MuR/OpImageInvert.h
@@ -10,192 +10,112 @@
 namespace mu
 {
 
-	/** Create a new image inverting the colour (RGB or L) components of an image. Leave alpha untouched. */
-	inline ImagePtr ImageInvert(const Image* pA)
-	{
-		MUTABLE_CPUPROFILER_SCOPE(ImageInvert);
-
-		ImagePtr pDest = new Image(pA->GetSizeX(), pA->GetSizeY(), pA->GetLODCount(), pA->GetFormat(), EInitializationType::NotInitialized);
-
-		uint8* pDestBuf = pDest->GetData();
-		const uint8* pABuf = pA->GetData();
-
-		//Generic implementation
-		int32 pixelCount = pA->CalculatePixelCount();
-
-		switch (pA->GetFormat())
-		{
-		case EImageFormat::IF_L_UBYTE:
-		{
-			for (int32 i = 0; i < pixelCount; ++i)
-			{
-				pDestBuf[i] = 255 - pABuf[i];
-			}
-			break;
-		}
-
-		case EImageFormat::IF_RGB_UBYTE:
-		{
-			for (int32 i = 0; i < pixelCount; ++i)
-			{
-				for (int c = 0; c < 3; ++c)
-				{
-					pDestBuf[i * 3 + c] = 255 - pABuf[i * 3 + c];
-				}
-			}
-			break;
-		}
-		case EImageFormat::IF_RGBA_UBYTE:
-		case EImageFormat::IF_BGRA_UBYTE:
-		{
-			for (int32 i = 0; i < pixelCount; ++i)
-			{
-				for (int32 c = 0; c < 3; ++c)
-				{
-					pDestBuf[i * 4 + c] = 255 - pABuf[i * 4 + c];
-				}
-
-				pDestBuf[i * 4 + 3] = pABuf[i * 4 + 3];
-			}
-			break;
-		}
-		default:
-			check(false);
-		}
-
-		return pDest;
-	}
-
 	/** Invert the colour (RGB or L) components of an image. Leave alpha untouched. */
 	template<bool bUseVectorImpl = false>
-	inline void ImageInvertInPlace(Image* pA)
+	inline void ImageInvert(Image* AImage)
 	{
 		MUTABLE_CPUPROFILER_SCOPE(ImageInvertInPlace);
 			
-		if (!pA)
+		if (!AImage)
 		{
 			return;
 		}
 
-		uint8* pABuf = pA->GetData();
-
 		//Generic implementation
-		const int32 pixelCount = pA->CalculatePixelCount();
+		constexpr int32 NumBatchElems = 1 << 14;
 
-		switch (pA->GetFormat())
-		{
-		case EImageFormat::IF_L_UBYTE:
-		{
-			if constexpr (!bUseVectorImpl)
-			{
-				for (int32 i = 0; i < pixelCount; ++i)
-				{
-					pABuf[i] = 255 - pABuf[i];
-				}
-			}
-			else
-			{
-				const int32 PixelRem = pixelCount % 4;
+		const int32 BytesPerElem = GetImageFormatData(AImage->GetFormat()).BytesPerBlock;
+		const int32 NumBatches = AImage->DataStorage.GetNumBatches(NumBatchElems, BytesPerElem);
 
-				constexpr VectorRegister4Int One255 = MakeVectorRegisterIntConstant(255, 255, 255, 255);
-				for (int32 i = 0; i < pixelCount - PixelRem; i += 4)
-				{
-					VectorRegister4Int Value = MakeVectorRegisterInt(pABuf[i], pABuf[i + 1], pABuf[i + 2], pABuf[i + 3]);
-					Value = VectorIntSubtract(One255, Value);
-
-					alignas(alignof(VectorRegister4Int)) int32 IndexableValue[4];
-					VectorIntStoreAligned(Value, IndexableValue);
-
-					pABuf[i + 0] = static_cast<uint8>(IndexableValue[0]);
-					pABuf[i + 1] = static_cast<uint8>(IndexableValue[1]);
-					pABuf[i + 2] = static_cast<uint8>(IndexableValue[2]);
-					pABuf[i + 3] = static_cast<uint8>(IndexableValue[3]);
-				}
-
-				for (int32 i = pixelCount - PixelRem; i < pixelCount; ++i)
-				{
-					pABuf[i] = 255 - pABuf[i];
-				}
-			}
-			break;
-		}
-
-		case EImageFormat::IF_RGB_UBYTE:
+		auto ProcessBatch = [AImage, NumBatchElems, BytesPerElem](int32 BatchId)
 		{
-			if constexpr (!bUseVectorImpl)
+			TArrayView<uint8> AView = AImage->DataStorage.GetBatch(BatchId, NumBatchElems, BytesPerElem);
+			const int32 NumBytes = AView.Num();
+			
+			uint8* ABuf = AView.GetData();
+
+			switch (AImage->GetFormat())
 			{
-				for (int32 i = 0; i < pixelCount; ++i)
+			case EImageFormat::IF_L_UBYTE:
+			case EImageFormat::IF_RGB_UBYTE:
+			{
+				if constexpr (!bUseVectorImpl)
 				{
-					for (int32 c = 0; c < 3; ++c)
+					for (int32 I = 0; I < NumBytes; ++I)
 					{
-						pABuf[i * 3 + c] = 255 - pABuf[i * 3 + c];
+						ABuf[I] = 255 - ABuf[I];
 					}
 				}
-			}
-			else
-			{
-				const int32 ElemCount = pixelCount * 3;
-				const int32 ElemRem = ElemCount % 4;
-
-				constexpr VectorRegister4Int One255 = MakeVectorRegisterIntConstant(255, 255, 255, 255);
-				for (int32 i = 0; i < ElemCount - ElemRem; i += 4)
+				else
 				{
-					VectorRegister4Int Value = MakeVectorRegisterInt(pABuf[i], pABuf[i + 1], pABuf[i + 2], pABuf[i + 3]);
-					Value = VectorIntSubtract(One255, Value);
+					const int32 RemBytes = NumBytes % 4;
 
-					alignas(alignof(VectorRegister4Int)) int32 IndexableValue[4];
-					VectorIntStoreAligned(Value, IndexableValue);
+					constexpr VectorRegister4Int One255 = MakeVectorRegisterIntConstant(255, 255, 255, 255);
+					for (int32 I = 0; I < NumBytes - RemBytes; I += 4)
+					{
+						VectorRegister4Int Value = MakeVectorRegisterInt(ABuf[I], ABuf[I + 1], ABuf[I + 2], ABuf[I + 3]);
+						Value = VectorIntSubtract(One255, Value);
 
-					pABuf[i + 0] = static_cast<uint8>(IndexableValue[0]);
-					pABuf[i + 1] = static_cast<uint8>(IndexableValue[1]);
-					pABuf[i + 2] = static_cast<uint8>(IndexableValue[2]);
-					pABuf[i + 3] = static_cast<uint8>(IndexableValue[3]);
-				}
+						alignas(alignof(VectorRegister4Int)) int32 IndexableValue[4];
+						VectorIntStoreAligned(Value, IndexableValue);
 
-				for (int32 i = ElemCount - ElemRem; i < ElemCount; ++i)
-				{
-					pABuf[i] = 255 - pABuf[i];
+						ABuf[I + 0] = static_cast<uint8>(IndexableValue[0]);
+						ABuf[I + 1] = static_cast<uint8>(IndexableValue[1]);
+						ABuf[I + 2] = static_cast<uint8>(IndexableValue[2]);
+						ABuf[I + 3] = static_cast<uint8>(IndexableValue[3]);
+					}
+
+					for (int32 I = NumBytes - RemBytes; I < NumBytes; ++I)
+					{
+						ABuf[I] = 255 - ABuf[I];
+					}
 				}
+				break;
 			}
-			break;
-		}
-		case EImageFormat::IF_RGBA_UBYTE:
-		case EImageFormat::IF_BGRA_UBYTE:
-		{
-			if constexpr (!bUseVectorImpl)
+			case EImageFormat::IF_RGBA_UBYTE:
+			case EImageFormat::IF_BGRA_UBYTE:
 			{
-				for (int32 i = 0; i < pixelCount; ++i)
+				if constexpr (!bUseVectorImpl)
 				{
-					for (int32 c = 0; c < 3; ++c)
+					const int32 NumElems = NumBytes/4;
+					for (int32 I = 0; I < NumElems; ++I)
 					{
-						pABuf[i * 4 + c] = 255 - pABuf[i * 4 + c];
+						ABuf[I*4 + 0] = 255 - ABuf[I*4 + 0];
+						ABuf[I*4 + 1] = 255 - ABuf[I*4 + 1];
+						ABuf[I*4 + 2] = 255 - ABuf[I*4 + 2];
 					}
 				}
-			}
-			else
-			{
-				const int32 ElemCount = pixelCount * 4;
-
-				constexpr VectorRegister4Int One255 = MakeVectorRegisterIntConstant(255, 255, 255, 255);
-
-				for (int32 i = 0; i < ElemCount; i += 4)
+				else
 				{
-					VectorRegister4Int Value = MakeVectorRegisterInt(pABuf[i], pABuf[i + 1], pABuf[i + 2], pABuf[i + 3]);
-					Value = VectorIntSubtract(One255, Value);
+					constexpr VectorRegister4Int One255 = MakeVectorRegisterIntConstant(255, 255, 255, 255);
 
-					alignas(alignof(VectorRegister4Int)) int32 IndexableValue[4];
-					VectorIntStoreAligned(Value, IndexableValue);
+					for (int32 I = 0; I < NumBytes; I += 4)
+					{
+						VectorRegister4Int Value = MakeVectorRegisterInt(ABuf[I], ABuf[I + 1], ABuf[I + 2], ABuf[I + 3]);
+						Value = VectorIntSubtract(One255, Value);
 
-					pABuf[i + 0] = static_cast<uint8>(IndexableValue[0]);
-					pABuf[i + 1] = static_cast<uint8>(IndexableValue[1]);
-					pABuf[i + 2] = static_cast<uint8>(IndexableValue[2]);
-					//pABuf[i + 3] = static_cast<uint8>(IndexableValue[3]);
+						alignas(alignof(VectorRegister4Int)) int32 IndexableValue[4];
+						VectorIntStoreAligned(Value, IndexableValue);
+
+						ABuf[I + 0] = static_cast<uint8>(IndexableValue[0]);
+						ABuf[I + 1] = static_cast<uint8>(IndexableValue[1]);
+						ABuf[I + 2] = static_cast<uint8>(IndexableValue[2]);
+					}
 				}
+				break;
 			}
-			break;
+			default:
+				check(false);
+			}
+		};
+
+		if (NumBatches == 1)
+		{
+			ProcessBatch(0);
 		}
-		default:
-			check(false);
+		else if (NumBatches > 1)
+		{
+			ParallelFor(NumBatches, ProcessBatch);
 		}
 	}
 }
