[Mutable] Store Image LODs in different buffers.

diff --git a/Engine/Plugins/Experimental/Mutable/Source/MutableRuntime/Private/MuR/OpImageLayer.h b/Engine/Plugins/Experimental/Mutable/Source/MutableRuntime/Private/MuR/OpImageLayer.h
--- a/Engine/Plugins/Experimental/Mutable/Source/MutableRuntime/Private/MuR/OpImageLayer.h
+++ b/Engine/Plugins/Experimental/Mutable/Source/MutableRuntime/Private/MuR/OpImageLayer.h
@@ -11,380 +11,333 @@
 
 namespace mu
 {
-	inline bool IsAnyComponentLargerThan1( FVector4f v )
+	inline bool IsAnyComponentLargerThan1(FVector4f Value)
 	{
-		return (v[0] > 1) || (v[1] > 1) || (v[2] > 1) || (v[3] > 1);
+		return (Value[0] > 1) | (Value[1] > 1) | (Value[2] > 1) | (Value[3] > 1);
 	}
 	
-
 	/** Apply a blending function to an image with a colour source.
 	* It only affects the RGB or L channels, leaving alpha untouched.
 	*/
-	template< unsigned (*BLEND_FUNC)(unsigned,unsigned), bool CLAMP >
-    inline void BufferLayerColour( Image* pResult, const Image* pBase, FVector4f col )
+	namespace Private
 	{
-        check( pResult->GetFormat() == pBase->GetFormat() );
-        check( pResult->GetSizeX() == pBase->GetSizeX() );
-        check( pResult->GetSizeY() == pBase->GetSizeY() );
-        check( pResult->GetLODCount() == pBase->GetLODCount() );
-
-		EImageFormat baseFormat = pBase->GetFormat();
+		template<uint32 NumChannels, uint32 (*BLEND_FUNC)(uint32, uint32), bool bClamp>
+		FORCENOINLINE void BufferLayerColourGenericChannel(uint8* DestBuf, const uint8* BaseBuf, int32 NumElems, const FIntVector& Color)
+		{
+			static_assert(NumChannels > 0 && NumChannels <= 4);
 
-        uint8* pDestBuf = pResult->GetData();
-        const uint8* pBaseBuf = pBase->GetData();
+			for (int32 I = 0; I < NumElems; ++I)
+			{
+				for (uint32 C = 0; C < NumChannels; ++C)
+				{
+					uint32 Base = BaseBuf[NumChannels * I + C];
+					uint32 Result = BLEND_FUNC(Base, Color[C]);
+					if constexpr (bClamp)
+					{
+						DestBuf[NumChannels * I + C] = (uint8)FMath::Min(255u, Result);
+					}
+					else
+					{
+						DestBuf[NumChannels * I + C] = (uint8)Result;
+					}
+				}
+			}
+		}
+	}
 
-		uint32 top[3];
-		top[0] = uint32(255 * col[0]);
-		top[1] = uint32(255 * col[1]);
-		top[2] = uint32(255 * col[2]);
+	template< unsigned (*BLEND_FUNC)(unsigned,unsigned), bool CLAMP >
+    inline void BufferLayerColour(Image* ResultImage, const Image* BaseImage, FVector4f Color)
+	{
+        check(ResultImage->GetFormat() == BaseImage->GetFormat());
+        check(ResultImage->GetSizeX() == BaseImage->GetSizeX());
+        check(ResultImage->GetSizeY() == BaseImage->GetSizeY());
+        check(ResultImage->GetLODCount() == BaseImage->GetLODCount());
 
+		const EImageFormat BaseFormat = BaseImage->GetFormat();
 		// Generic implementation
-		const int32 PixelCount = pBase->CalculatePixelCount();
+		constexpr int32 NumBatchElems = 4096*4;
+		const int32 BytesPerElem = GetImageFormatData(BaseFormat).BytesPerBlock;	
 
-		constexpr int32 NumBatchElems = 4096*2;
-		const int32 NumBatches = FMath::DivideAndRoundUp(PixelCount, NumBatchElems);
+		const int32 NumBatches = BaseImage->DataStorage.GetNumBatches(NumBatchElems, BytesPerElem);
 
-		switch (baseFormat)
+		const FIntVector ColorValue = FIntVector(Color.X * 255.0f, Color.Y * 255.0f, Color.Z * 255.0f);
+		auto ProcessBatch = [NumBatchElems, BytesPerElem, ColorValue, BaseImage, ResultImage, BaseFormat](int32 BatchId)
 		{
-		case EImageFormat::IF_RGB_UBYTE:
-		{
-			auto ProcessBatch =
-			[
-				pBaseBuf, pDestBuf, top, PixelCount, NumBatchElems
-			] (int32 BatchId)
-			{
-				const int32 BatchBegin = BatchId * NumBatchElems;
-				const int32 BatchEnd = FMath::Min(BatchBegin + NumBatchElems, PixelCount);
+			TArrayView<const uint8> BaseView = BaseImage->DataStorage.GetBatch(BatchId, NumBatchElems, BytesPerElem);
+			TArrayView<uint8> ResultView = ResultImage->DataStorage.GetBatch(BatchId, NumBatchElems, BytesPerElem);
+			
+			check(BaseView.Num() == ResultView.Num());
 
-				for (int32 i = BatchBegin; i < BatchEnd; ++i)
-				{
-					for (int c = 0; c < 3; ++c)
-					{
-						unsigned base = pBaseBuf[3 * i + c];
-						unsigned result = BLEND_FUNC(base, top[c]);
-						if constexpr (CLAMP)
-						{
-							pDestBuf[3 * i + c] = (uint8)FMath::Min(255u, result);
-						}
-						else
-						{
-							pDestBuf[3 * i + c] = (uint8)result;
-						}
-					}
-				}
-			};
+			const int32 NumElems = BaseView.Num() / BytesPerElem;
 
-			if (NumBatches == 1)
+			switch (BaseFormat)
 			{
-				ProcessBatch(0);
+			case EImageFormat::IF_L_UBYTE:
+			{
+				check(BytesPerElem == 1);
+				Private::BufferLayerColourGenericChannel<1, BLEND_FUNC, CLAMP>(ResultView.GetData(), BaseView.GetData(), NumElems, ColorValue);
+				break;
 			}
-			else if (NumBatches > 1)
+			case EImageFormat::IF_RGB_UBYTE:
 			{
-				ParallelFor(NumBatches, ProcessBatch);
+				check(BytesPerElem == 3);
+				Private::BufferLayerColourGenericChannel<3, BLEND_FUNC, CLAMP>(ResultView.GetData(), BaseView.GetData(), NumElems, ColorValue);
+				break;
 			}
-
-			break;
-		}
-		case EImageFormat::IF_RGBA_UBYTE:
-		{
-			auto ProcessBatch = 
-			[
-				pBaseBuf, pDestBuf, top, PixelCount, NumBatchElems
-			] (int32 BatchId)
+			case EImageFormat::IF_RGBA_UBYTE:
 			{
-				const int32 BatchBegin = BatchId * NumBatchElems;
-				const int32 BatchEnd = FMath::Min(BatchBegin + NumBatchElems, PixelCount);
-
-				for (int32 i = BatchBegin; i < BatchEnd; ++i)
-				{
-					for (int c = 0; c < 3; ++c)
-					{
-						unsigned base = pBaseBuf[4 * i + c];
-						unsigned result = BLEND_FUNC(base, top[c]);
-						if constexpr (CLAMP)
-						{
-							pDestBuf[4 * i + c] = (uint8)FMath::Min(255u, result);
-						}
-						else
-						{
-							pDestBuf[4 * i + c] = (uint8)result;
-						}
-					}
-					pDestBuf[4 * i + 3] = pBaseBuf[4 * i + 3];
-				}
-			};
-
-			if (NumBatches == 1)
+				check(BytesPerElem == 4);
+				Private::BufferLayerColourGenericChannel<4, BLEND_FUNC, CLAMP>(ResultView.GetData(), BaseView.GetData(), NumElems, ColorValue);
+				break;
+			}
+			case EImageFormat::IF_BGRA_UBYTE:
 			{
-				ProcessBatch(0);
+				check(BytesPerElem == 4);
+				FIntVector BGRAColorValue = FIntVector(ColorValue.Z, ColorValue.Y, ColorValue.X);
+				Private::BufferLayerColourGenericChannel<4, BLEND_FUNC, CLAMP>(ResultView.GetData(), BaseView.GetData(), NumElems, BGRAColorValue);
+				break;
 			}
-			else if (NumBatches > 1)
+			default:
 			{
-				ParallelFor(NumBatches, ProcessBatch);
+				checkf(false, TEXT("Unsupported format."));
+				break;
+			}
 			}
+		};
 
-			break;
+		if (NumBatches == 1)
+		{
+			ProcessBatch(0);
 		}
-		case EImageFormat::IF_BGRA_UBYTE:
+		else if (NumBatches > 1)
 		{
-			auto ProcessBatch =
-			[
-				pBaseBuf, pDestBuf, top, PixelCount, NumBatchElems
-			] (int32 BatchId)
-			{
-				const int32 BatchBegin = BatchId * NumBatchElems;
-				const int32 BatchEnd = FMath::Min(BatchBegin + NumBatchElems, PixelCount);
+			ParallelFor(NumBatches, ProcessBatch);
+		}
+    }
+
+
+	template< uint32 (*BLEND_FUNC)(uint32, uint32) >
+	inline void BufferLayerColour(Image* ResultImage, const Image* BaseImage, FVector4f Color)
+	{
+		bool bIsClampNeeded = IsAnyComponentLargerThan1(Color);
+		if (bIsClampNeeded)
+		{
+			BufferLayerColour<BLEND_FUNC, true>(ResultImage, BaseImage, Color);
+		}
+		else
+		{
+			BufferLayerColour<BLEND_FUNC, false>(ResultImage, BaseImage, Color);
+		}
+	}
 
-				for (int32 i = BatchBegin; i < BatchEnd; ++i)
+	namespace Private
+	{
+		template<uint32 NumChannels, uint32 (*BLEND_FUNC)(uint32,uint32), bool Clamp>
+		FORCENOINLINE void BufferLayerColourFromAlphaGenericChannel(uint8* DestBuf, const uint8* BaseBuf, int32 NumElems, const FIntVector& Color)
+		{
+			static_assert(NumChannels >= 3 && NumChannels < 4);
+
+			for (int32 I = 0; I < NumElems; ++I)
+			{
+				for (uint32 C = 0; C < NumChannels; ++C)
 				{
-					for (int c = 0; c < 3; ++c)
+					const uint32 Alpha = Invoke([&]() -> uint32
 					{
-						unsigned base = pBaseBuf[4 * i + c];
-						unsigned result = BLEND_FUNC(base, top[2 - c]);
-						if (CLAMP)
+						if constexpr (NumChannels == 3)
 						{
-							pDestBuf[4 * i + c] = (uint8)FMath::Min(255u, result);
+							return 255;
 						}
 						else
 						{
-							pDestBuf[4 * i + c] = (uint8)result;
+							return BaseBuf[NumChannels * I + 3];
 						}
+					});
+					
+					uint32 Result = BLEND_FUNC(Alpha, Color[C]);
+					if constexpr (Clamp)
+					{
+						DestBuf[NumChannels * I + C] = (uint8)FMath::Min(255u, Result);
+					}
+					else
+					{
+						DestBuf[NumChannels * I + C] = (uint8)Result;
 					}
-					pDestBuf[4 * i + 3] = pBaseBuf[4 * i + 3];
 				}
-			};
+			}
+		}
+	}
 
-			if (NumBatches == 1)
+	template<
+		uint32 (*BLEND_FUNC)(uint32, uint32), 
+		bool CLAMP>
+    inline void BufferLayerColourFromAlpha(Image* ResultImage, const Image* BaseImage, FVector4f Color)
+	{
+        check(ResultImage->GetFormat() == BaseImage->GetFormat());
+        check(ResultImage->GetSizeX() == BaseImage->GetSizeX());
+        check(ResultImage->GetSizeY() == BaseImage->GetSizeY());
+        check(ResultImage->GetLODCount() == BaseImage->GetLODCount());
+
+		const EImageFormat BaseFormat = BaseImage->GetFormat();
+		// Generic implementation
+		constexpr int32 NumBatchElems = 4096*2;
+		const int32 BytesPerElem = GetImageFormatData(BaseFormat).BytesPerBlock;	
+
+		const int32 NumBatches = BaseImage->DataStorage.GetNumBatches(NumBatchElems, BytesPerElem);
+		check(NumBatches == ResultImage->DataStorage.GetNumBatches(NumBatchElems, BytesPerElem));
+
+		const FIntVector ColorValue = FIntVector(Color.X * 255.0f, Color.Y * 255.0f, Color.Z * 255.0f);
+		auto ProcessBatch = [NumBatchElems, BytesPerElem, ColorValue, BaseImage, ResultImage, BaseFormat](int32 BatchId)
+		{
+			TArrayView<const uint8> BaseView = BaseImage->DataStorage.GetBatch(BatchId, NumBatchElems, BytesPerElem);
+			TArrayView<uint8> ResultView = ResultImage->DataStorage.GetBatch(BatchId, NumBatchElems, BytesPerElem);
+			
+			check(BaseView.Num() == ResultView.Num());
+
+			const int32 NumElems = ResultView.Num() / BytesPerElem;
+			switch (BaseFormat)
 			{
-				ProcessBatch(0);
+			case EImageFormat::IF_L_UBYTE:
+			{
+				check(BytesPerElem == 1);
+				Private::BufferLayerColourGenericChannel<1, BLEND_FUNC, CLAMP>(ResultView.GetData(), BaseView.GetData(), NumElems, ColorValue);
+				break;
 			}
-			else if (NumBatches > 1)
+			case EImageFormat::IF_RGB_UBYTE:
 			{
-				ParallelFor(NumBatches, ProcessBatch);
+				check(BytesPerElem == 3);
+				Private::BufferLayerColourGenericChannel<3, BLEND_FUNC, CLAMP>(ResultView.GetData(), BaseView.GetData(), NumElems, ColorValue);
+				break;
 			}
-
-			break;
-		}
-		case EImageFormat::IF_L_UBYTE:
-		{
-			auto ProcessBatch = 
-			[
-				pBaseBuf, pDestBuf, top, PixelCount, NumBatchElems
-			] (int32 BatchId)
+			case EImageFormat::IF_RGBA_UBYTE:
 			{
-				const int32 BatchBegin = BatchId * NumBatchElems;
-				const int32 BatchEnd = FMath::Min(BatchBegin + NumBatchElems, PixelCount);
-
-				for (int32 i = BatchBegin; i < BatchEnd; ++i)
-				{
-					unsigned base = pBaseBuf[i];
-					unsigned result = BLEND_FUNC(base, top[0]);
-					if constexpr (CLAMP)
-					{
-						pDestBuf[i] = (uint8)FMath::Min(255u, result);
-					}
-					else
-					{
-						pDestBuf[i] = (uint8)result;
-					}
-				}
-			};
-
-			if (NumBatches == 1)
+				check(BytesPerElem == 4);
+				Private::BufferLayerColourGenericChannel<4, BLEND_FUNC, CLAMP>(ResultView.GetData(), BaseView.GetData(), NumElems, ColorValue);
+				break;
+			}
+			case EImageFormat::IF_BGRA_UBYTE:
 			{
-				ProcessBatch(0);
+				check(BytesPerElem == 4);
+				FIntVector BGRAColorValues = FIntVector(ColorValue.Z, ColorValue.Y, ColorValue.X);
+				Private::BufferLayerColourGenericChannel<4, BLEND_FUNC, CLAMP>(ResultView.GetData(), BaseView.GetData(), NumElems , ColorValue);
+				break;
 			}
-			else if (NumBatches > 1)
+			default:
 			{
-				ParallelFor(NumBatches, ProcessBatch);
+				checkf(false, TEXT("Unsupported format."));
+				break;
 			}
+			}
+		};
 
-			break;
+		if (NumBatches == 1)
+		{
+			ProcessBatch(0);
 		}
-
-		default:
-			checkf(false, TEXT("Unsupported format."));
-			break;
+		else if (NumBatches > 1)
+		{
+			ParallelFor(NumBatches, ProcessBatch);
 		}
     }
 
 
-	template< unsigned (*BLEND_FUNC)(unsigned, unsigned) >
-	inline void BufferLayerColour(Image* Result, const Image* Base, FVector4f Col)
+	template<uint32 (*BLEND_FUNC)(uint32, uint32)>
+	inline void BufferLayerColourFromAlpha(Image* ResultImage, const Image* BaseImage, FVector4f Color)
 	{
-		bool bIsClampNeeded = IsAnyComponentLargerThan1(Col);
+		bool bIsClampNeeded = IsAnyComponentLargerThan1(Color);
 		if (bIsClampNeeded)
 		{
-			BufferLayerColour<BLEND_FUNC, true>(Result, Base, Col);
+			BufferLayerColourFromAlpha<BLEND_FUNC, true>(ResultImage, BaseImage, Color);
 		}
 		else
 		{
-			BufferLayerColour<BLEND_FUNC, false>(Result, Base, Col);
+			BufferLayerColourFromAlpha<BLEND_FUNC, false>(ResultImage, BaseImage, Color);
 		}
 	}
 
 
-	//---------------------------------------------------------------------------------------------
-	template< unsigned (*BLEND_FUNC)(unsigned, unsigned), bool CLAMP >
-	inline void BufferLayerColourFromAlpha(Image* Result, const Image* Base, FVector4f Col)
+	template< 
+		uint32 (*BLEND_FUNC_MASKED)(uint32, uint32, uint32),
+		uint32 (*BLEND_FUNC)(uint32, uint32),
+		bool CLAMP,
+		// Number of total channels to actually process
+		uint32 CHANNELS_TO_BLEND,
+		// Number of total channels in the base image
+		int32 BASE_CHANNEL_STRIDE>
+    inline void BufferLayerColourFormat(Image* DestImage, const Image* BaseImage, const Image* MaskImage, FVector4f Col, uint32 BaseChannelOffset, uint8 ColorChannelOffset, bool bOnlyFirstLOD)
 	{
-		check(Result->GetFormat() == Base->GetFormat());
-		check(Result->GetSizeX() == Base->GetSizeX());
-		check(Result->GetSizeY() == Base->GetSizeY());
-		check(Result->GetLODCount() == Base->GetLODCount());
+		check(CHANNELS_TO_BLEND + BaseChannelOffset <= BASE_CHANNEL_STRIDE);
+		check(DestImage->GetLODCount() <= BaseImage->GetLODCount());
 
-		EImageFormat baseFormat = Base->GetFormat();
+		FUintVector4 TopColor = FUintVector4(Col.X * 255.0f, Col.Y * 255.0f, Col.Z * 255.0f, Col.W * 255);
 
-		uint8* pDestBuf = Result->GetData();
-		const uint8* pBaseBuf = Base->GetData();
+		const EImageFormat BaseFormat = BaseImage->GetFormat();
+		const EImageFormat MaskFormat = MaskImage->GetFormat();
 
-		unsigned top[3];
-		top[0] = (unsigned)(255 * Col[0]);
-		top[1] = (unsigned)(255 * Col[1]);
-		top[2] = (unsigned)(255 * Col[2]);
+		int32 NumLODs = bOnlyFirstLOD ? 1 : BaseImage->GetLODCount();
 
-		// Generic implementation
-		int32 PixelCount = Base->CalculatePixelCount();
+        const bool bIsMaskUncompressed = (MaskFormat == EImageFormat::IF_L_UBYTE);
 
-		constexpr int32 NumBatchElems = 4096*2;
-		const int32 NumBatches = FMath::DivideAndRoundUp(PixelCount, NumBatchElems);
+		constexpr uint32 NumColorChannels = FMath::Min(CHANNELS_TO_BLEND, 3u);
+		check(NumColorChannels + ColorChannelOffset < 4);
 
-		switch (baseFormat)
-		{
-		case EImageFormat::IF_RGB_UBYTE:
-		{
-			// There is no alpha so we assume 255
-			auto ProcessBatch =
-				[ pDestBuf, top, PixelCount, NumBatchElems ] (uint32 BatchId)
-			{
-				const int32 BatchBegin = BatchId * NumBatchElems;
-				const int32 BatchEnd = FMath::Min(BatchBegin + NumBatchElems, PixelCount);
+        if (bIsMaskUncompressed)
+        {
+			constexpr int32 NumBatchElems = 4096 * 2;
+			const int32 BytesPerElem = GetImageFormatData(BaseFormat).BytesPerBlock;
+			check(GetImageFormatData(MaskFormat).BytesPerBlock == 1);
+			check(GetImageFormatData(DestImage->GetFormat()).BytesPerBlock == BytesPerElem);
 
-				for (int32 i = BatchBegin; i < BatchEnd; ++i)
-				{
-					for (int c = 0; c < 3; ++c)
-					{
-						unsigned base = 255;
-						unsigned result = BLEND_FUNC(base, top[c]);
-						if constexpr (CLAMP)
-						{
-							pDestBuf[3 * i + c] = (uint8)FMath::Min(255u, result);
-						}
-						else
-						{
-							pDestBuf[3 * i + c] = (uint8)result;
-						}
-					}
-				}
-			};
+			const int32 NumBatches = 
+					BaseImage->DataStorage.GetNumBatchesLODRange(NumBatchElems, BytesPerElem, 0, NumLODs);
 
-			if (NumBatches == 1)
-			{
-				ProcessBatch(0);
-			}
-			else if (NumBatches > 1)
+			check(NumBatches == DestImage->DataStorage.GetNumBatchesLODRange(NumBatchElems, BytesPerElem, 0, NumLODs));
+			check(NumBatches == MaskImage->DataStorage.GetNumBatchesLODRange(NumBatchElems, 1, 0, NumLODs));
+
+			auto ProcessBatch = [
+				BaseImage, MaskImage, DestImage, 
+				BytesPerElem, NumBatchElems, 
+				BaseChannelOffset, TopColor, ColorChannelOffset, NumColorChannels, 
+				NumLODs
+			](int32 BatchId)
 			{
-				ParallelFor(NumBatches, ProcessBatch);
-			}
+				TArrayView<const uint8> BaseBatchView = 
+						BaseImage->DataStorage.GetBatchLODRange(BatchId, NumBatchElems, BytesPerElem, 0, NumLODs);
 
-			break;
-		}
-		case EImageFormat::IF_RGBA_UBYTE:
-		{
-			auto ProcessBatch =
-				[ pBaseBuf, pDestBuf, top, PixelCount, NumBatchElems ] (uint32 BatchId)
-			{	
-				const int32 BatchBegin = BatchId * NumBatchElems;
-				const int32 BatchEnd = FMath::Min(BatchBegin + NumBatchElems, PixelCount);
+				TArrayView<const uint8> MaskBatchView =
+						MaskImage->DataStorage.GetBatchLODRange(BatchId, NumBatchElems, 1, 0, NumLODs);
 
-				for (int32 i = BatchBegin; i < BatchEnd; ++i)
-				{
-					for (int c = 0; c < 3; ++c)
-					{
-						unsigned base = pBaseBuf[4 * i + 3];
-						unsigned result = BLEND_FUNC(base, top[c]);
-						if constexpr (CLAMP)
-						{
-							pDestBuf[4 * i + c] = (uint8)FMath::Min(255u, result);
-						}
-						else
-						{
-							pDestBuf[4 * i + c] = (uint8)result;
-						}
-					}
-					pDestBuf[4 * i + 3] = pBaseBuf[4 * i + 3];
-				}
-			};
+				TArrayView<uint8> DestBatchView =
+						DestImage->DataStorage.GetBatchLODRange(BatchId, NumBatchElems, BytesPerElem, 0, NumLODs);
 
-			if (NumBatches == 1)
-			{
-				ProcessBatch(0);
-			}
-			else if (NumBatches > 1)
-			{
-				ParallelFor(NumBatches, ProcessBatch);
-			}
+				const uint8* BaseBuf = BaseBatchView.GetData() + BaseChannelOffset;
+				const uint8* MaskBuf = MaskBatchView.GetData();
+				uint8* DestBuf = DestBatchView.GetData() + BaseChannelOffset;
 
-			break;
-		}
-		case EImageFormat::IF_BGRA_UBYTE:
-		{	
-			auto ProcessBatch =
-				[ pBaseBuf, pDestBuf, top, PixelCount, NumBatchElems ] (uint32 BatchId)
-			{
-				const int32 BatchBegin = BatchId * NumBatchElems;
-				const int32 BatchEnd = FMath::Min(BatchBegin + NumBatchElems, PixelCount);
+				const int32 NumElems = BaseBatchView.Num() / BytesPerElem;
+				check(NumElems == MaskBatchView.Num());
+				check(NumElems == DestBatchView.Num() / BytesPerElem);
 
-				for (int32 i = BatchBegin; i < BatchEnd; ++i)
+				for (int32 I = 0; I < NumElems; ++I)
 				{
-					for (int c = 0; c < 3; ++c)
+					uint32 MaskData = MaskBuf[I];
+					for (int32 C = 0; C < NumColorChannels; ++C)
 					{
-						unsigned base = pBaseBuf[4 * i + 3];
-						unsigned result = BLEND_FUNC(base, top[2 - c]);
+						const uint32 Base = BaseBuf[BASE_CHANNEL_STRIDE*I + C];
+						const uint32 Result = BLEND_FUNC_MASKED(Base, TopColor[C + ColorChannelOffset], MaskData);
 						if constexpr (CLAMP)
 						{
-							pDestBuf[4 * i + c] = (uint8)FMath::Min(255u, result);
+							DestBuf[BASE_CHANNEL_STRIDE*I + C] = (uint8)FMath::Min(255u, Result);
 						}
 						else
 						{
-							pDestBuf[4 * i + c] = (uint8)result;
+							DestBuf[BASE_CHANNEL_STRIDE*I + C] = (uint8)Result;
 						}
 					}
-					pDestBuf[4 * i + 3] = pBaseBuf[4 * i + 3];
-				}
-			};
-
-			if (NumBatches == 1)
-			{
-				ProcessBatch(0);
-			}
-			else if (NumBatches > 1)
-			{
-				ParallelFor(NumBatches, ProcessBatch);
-			}
 
-			break;
-		}
-		case EImageFormat::IF_L_UBYTE:
-		{
-			auto ProcessBatch =
-				[ pBaseBuf, pDestBuf, top, PixelCount, NumBatchElems ] (uint32 BatchId)
-			{
-				const int32 BatchBegin = BatchId * NumBatchElems;
-				const int32 BatchEnd = FMath::Min(BatchBegin + NumBatchElems, PixelCount);
-
-				for (int32 i = BatchBegin; i < BatchEnd; ++i)
-				{
-					// There is no alpha so we assume 255
-					unsigned base = 255;
-					unsigned result = BLEND_FUNC(base, top[0]);
-					if constexpr (CLAMP)
-					{
-						pDestBuf[i] = (uint8)FMath::Min(255u, result);
-					}
-					else
+					constexpr bool bIsNC4 = (BASE_CHANNEL_STRIDE == 4);
+					if constexpr (bIsNC4)
 					{
-						pDestBuf[i] = (uint8)result;
+						DestBuf[BASE_CHANNEL_STRIDE*I + 3] = BaseBuf[BASE_CHANNEL_STRIDE*I + 3];
 					}
 				}
 			};
@@ -397,217 +350,136 @@ namespace mu
 			{
 				ParallelFor(NumBatches, ProcessBatch);
 			}
-			break;
-		}
-
-		default:
-			checkf(false, TEXT("Unsupported format."));
-			break;
-		}
-	}
-
-
-	//---------------------------------------------------------------------------------------------
-	template< unsigned (*BLEND_FUNC)(unsigned, unsigned)>
-	inline void BufferLayerColourFromAlpha(Image* Result, const Image* Base, FVector4f Col)
-	{
-		bool bIsClampNeeded = IsAnyComponentLargerThan1(Col);
-		if (bIsClampNeeded)
-		{
-			BufferLayerColourFromAlpha<BLEND_FUNC,true>(Result, Base, Col);
-		}
-		else
-		{
-			BufferLayerColourFromAlpha<BLEND_FUNC,false>(Result, Base, Col);
-		}
-	}
-
-
-	//---------------------------------------------------------------------------------------------
-	template< unsigned (*BLEND_FUNC_MASKED)(unsigned,unsigned,unsigned),
-		unsigned (*BLEND_FUNC)(unsigned,unsigned),
-		bool CLAMP,
-		// Number of total channels to actually process
-		uint32 CHANNELS_TO_BLEND,
-		// Number of total channels in the base image
-		int32 BASE_CHANNEL_STRIDE >
-    inline void BufferLayerColourFormat( uint8* StartDestBuf, const Image* Base, const Image* Mask, FVector4f Col, uint32 BaseOffset, uint8 ColorOffset, bool bOnlyOneMip )
-	{
-		check(CHANNELS_TO_BLEND+BaseOffset<=BASE_CHANNEL_STRIDE);
-
-		uint32 top[4];
-		top[0] = uint32(255 * Col[0]);
-		top[1] = uint32(255 * Col[1]);
-		top[2] = uint32(255 * Col[2]);
-		top[3] = uint32(255 * Col[3]);
-
-		uint8* pDestBuf = StartDestBuf+ BaseOffset;
-        const uint8* pMaskBuf = Mask->GetData();
-        const uint8* pBaseBuf = Base->GetData()+ BaseOffset;
-		EImageFormat maskFormat = Mask->GetFormat();
-
-		int32 LODCount = Base->GetLODCount();
-
-		if (bOnlyOneMip)
-		{
-			LODCount = 1;
-		}
-
-        bool isUncompressed = ( maskFormat == EImageFormat::IF_L_UBYTE );
-
-		constexpr uint32 NumColorChannels = FMath::Min(CHANNELS_TO_BLEND,3u);
-		check(NumColorChannels+ColorOffset<4);
-
-        if ( isUncompressed )
-        {
-            int32 pixelCount = Base->CalculatePixelCount();
-            for ( int i=0; i<pixelCount; ++i )
-            {
-                unsigned mask = pMaskBuf[i];
-                for ( int32 c=0; c<NumColorChannels; ++c )
-                {
-                    unsigned base = pBaseBuf[BASE_CHANNEL_STRIDE *i+c];
-                    unsigned result = BLEND_FUNC_MASKED( base, top[c+ColorOffset], mask );
-                    if ( CLAMP )
-                    {
-                        pDestBuf[BASE_CHANNEL_STRIDE *i+c] = (uint8)FMath::Min( 255u, result );
-                    }
-                    else
-                    {
-                        pDestBuf[BASE_CHANNEL_STRIDE *i+c] = (uint8)result;
-                    }
-                }
-
-                constexpr bool isNC4 = (BASE_CHANNEL_STRIDE ==4);
-                if ( isNC4 )
-                {
-                    pDestBuf[BASE_CHANNEL_STRIDE *i+3] = pBaseBuf[BASE_CHANNEL_STRIDE *i+3];
-                }
-            }
         }
-        else if ( maskFormat==EImageFormat::IF_L_UBYTE_RLE )
+        else if (MaskFormat == EImageFormat::IF_L_UBYTE_RLE)
         {
-            int rows = Base->GetSizeY();
-            int width = Base->GetSizeX();
+            int32 Rows = BaseImage->GetSizeY();
+            int32 Width = BaseImage->GetSizeX();
 
-            for (int32 lod=0;lod<LODCount; ++lod)
+            for (int32 Lod = 0; Lod < NumLODs; ++Lod)
             {
-				// Skip mip size and row sizes.
-                pMaskBuf += sizeof(uint32) +rows*sizeof(uint32);
+                const uint8* BaseBuf = BaseImage->GetLODData(Lod);
+                const uint8* MaskBuf = MaskImage->GetLODData(Lod);
+				uint8* DestBuf = DestImage->GetLODData(Lod);
 
-                for ( int r=0; r<rows; ++r )
+				// Remove RLE header, mip size and row sizes.
+				MaskBuf += sizeof(uint32);
+				MaskBuf += Rows * sizeof(uint32);
+
+                for (int32 RowIndex = 0; RowIndex < Rows; ++RowIndex)
                 {
-                    const uint8* pDestRowEnd = pDestBuf + width* BASE_CHANNEL_STRIDE;
-                    while ( pDestBuf!=pDestRowEnd )
+                    const uint8* DestRowEnd = DestBuf + Width * BASE_CHANNEL_STRIDE;
+                    while (DestBuf != DestRowEnd)
                     {
                         // Decode header
-						uint16 equal = 0;
-						FMemory::Memmove(&equal, pMaskBuf, sizeof(uint16));
-                        pMaskBuf += 2;
+						uint16 Equal = 0;
+						FMemory::Memmove(&Equal, MaskBuf, sizeof(uint16));
+                        MaskBuf += 2;
 
-                        uint8 different = *pMaskBuf;
-                        ++pMaskBuf;
+                        uint8 Different = *MaskBuf;
+                        ++MaskBuf;
 
-                        uint8 equalPixel = *pMaskBuf;
-                        ++pMaskBuf;
+                        uint8 EqualPixel = *MaskBuf;
+                        ++MaskBuf;
 
                         // Equal pixels
-						check(pDestBuf + BASE_CHANNEL_STRIDE * equal <= StartDestBuf + Base->GetDataSize());
-                        if ( equalPixel==255 )
+						//check(DestBuf + BASE_CHANNEL_STRIDE * Equal <= BaseImage->GetDataSize(Lod));
+                        if (EqualPixel == 255)
                         {
-                            for ( int i=0; i<equal; ++i )
+                            for (int32 I = 0; I < Equal; ++I)
                             {
-                                for ( int32 c=0; c<NumColorChannels; ++c )
+                                for (int32 C = 0; C < NumColorChannels; ++C)
                                 {
-                                    unsigned base = pBaseBuf[BASE_CHANNEL_STRIDE *i+c];
-                                    unsigned result = BLEND_FUNC( base, top[c+ColorOffset] );
-                                    if ( CLAMP )
+                                    uint32 Base = BaseBuf[BASE_CHANNEL_STRIDE * I + C];
+                                    uint32 Result = BLEND_FUNC(Base, TopColor[C + ColorChannelOffset]);
+                                    if constexpr (CLAMP)
                                     {
-                                        pDestBuf[BASE_CHANNEL_STRIDE *i+c] = (uint8)FMath::Min( 255u, result );
+                                        DestBuf[BASE_CHANNEL_STRIDE * I + C] = (uint8)FMath::Min(255u, Result);
                                     }
                                     else
                                     {
-                                        pDestBuf[BASE_CHANNEL_STRIDE *i+c] = (uint8)result;
+                                        DestBuf[BASE_CHANNEL_STRIDE * I + C] = (uint8)Result;
                                     }
                                 }
 
-                                constexpr bool isNC4 = (BASE_CHANNEL_STRIDE ==4);
-                                if ( isNC4 )
+                                constexpr bool bIsNC4 = (BASE_CHANNEL_STRIDE == 4);
+                                if (bIsNC4)
                                 {
-                                    pDestBuf[BASE_CHANNEL_STRIDE *i+3] = pBaseBuf[BASE_CHANNEL_STRIDE *i+3];
+                                    DestBuf[BASE_CHANNEL_STRIDE * I + 3] = BaseBuf[BASE_CHANNEL_STRIDE * I + 3];
                                 }
                             }
                         }
-                        else if ( equalPixel>0 )
+                        else if (EqualPixel > 0)
                         {
-                            for ( int i=0; i<equal; ++i )
+                            for (int32 I = 0; I < Equal; ++I)
                             {
-                                for ( int32 c=0; c<NumColorChannels; ++c )
+                                for (int32 C = 0; C < NumColorChannels; ++C)
                                 {
-                                    unsigned base = pBaseBuf[BASE_CHANNEL_STRIDE *i+c];
-                                    unsigned result = BLEND_FUNC_MASKED( base, top[c+ColorOffset], equalPixel );
-                                    if ( CLAMP )
+                                    uint32 Base = BaseBuf[BASE_CHANNEL_STRIDE * I + C];
+                                    uint32 Result = BLEND_FUNC_MASKED(Base, TopColor[C + ColorChannelOffset], EqualPixel);
+                                    if constexpr (CLAMP)
                                     {
-                                        pDestBuf[BASE_CHANNEL_STRIDE *i+c] = (uint8)FMath::Min( 255u, result );
+                                        DestBuf[BASE_CHANNEL_STRIDE * I + C] = (uint8)FMath::Min(255u, Result);
                                     }
                                     else
                                     {
-                                        pDestBuf[BASE_CHANNEL_STRIDE *i+c] = (uint8)result;
+                                        DestBuf[BASE_CHANNEL_STRIDE * I + C] = (uint8)Result;
                                     }
                                 }
 
-                                constexpr bool isNC4 = (BASE_CHANNEL_STRIDE ==4);
-                                if ( isNC4 )
+                                constexpr bool bIsNC4 = (BASE_CHANNEL_STRIDE == 4);
+                                if (bIsNC4)
                                 {
-                                    pDestBuf[BASE_CHANNEL_STRIDE *i+3] = pBaseBuf[BASE_CHANNEL_STRIDE *i+3];
+                                    DestBuf[BASE_CHANNEL_STRIDE * I + 3] = BaseBuf[BASE_CHANNEL_STRIDE * I + 3];
                                 }
                             }
                         }
                         else
                         {
                             // It could happen if xxxxxOnBase
-                            if (pDestBuf!=pBaseBuf)
+                            if (DestBuf != BaseBuf)
                             {
-                                FMemory::Memmove( pDestBuf, pBaseBuf, BASE_CHANNEL_STRIDE*equal );
+                                FMemory::Memmove(DestBuf, BaseBuf, BASE_CHANNEL_STRIDE*Equal);
                             }
                         }
-                        pDestBuf += BASE_CHANNEL_STRIDE *equal;
-                        pBaseBuf += BASE_CHANNEL_STRIDE *equal;
+
+                        DestBuf += BASE_CHANNEL_STRIDE * Equal;
+                        BaseBuf += BASE_CHANNEL_STRIDE * Equal;
 
                         // Different pixels
-						check(pDestBuf + BASE_CHANNEL_STRIDE * different <= StartDestBuf + Base->GetDataSize());
-                        for ( int i=0; i<different; ++i )
+						//check(DestBuf + BASE_CHANNEL_STRIDE * Different <= StartDestBuf + BaseImage->GetDataSize(Lod));
+                        for (int32 I = 0; I < Different; ++I)
                         {
-                            for ( int32 c=0; c<NumColorChannels; ++c )
+                            for (int32 C = 0; C < NumColorChannels; ++C)
                             {
-                                unsigned mask = pMaskBuf[i];
-                                unsigned base = pBaseBuf[BASE_CHANNEL_STRIDE *i+c];
-                                unsigned result = BLEND_FUNC_MASKED( base, top[c+ColorOffset], mask );
-                                if ( CLAMP )
+                                uint32 Mask = MaskBuf[I];
+                                uint32 Base = BaseBuf[BASE_CHANNEL_STRIDE * I + C];
+                                uint32 Result = BLEND_FUNC_MASKED(Base, TopColor[C + ColorChannelOffset], Mask);
+                                if constexpr (CLAMP)
                                 {
-                                    pDestBuf[BASE_CHANNEL_STRIDE *i+c] = (uint8)FMath::Min( 255u, result );
+                                    DestBuf[BASE_CHANNEL_STRIDE * I + C] = (uint8)FMath::Min(255u, Result);
                                 }
                                 else
                                 {
-                                    pDestBuf[BASE_CHANNEL_STRIDE *i+c] = (uint8)result;
+                                    DestBuf[BASE_CHANNEL_STRIDE * I + C] = (uint8)Result;
                                 }
                             }
 
-                            constexpr bool isNC4 = (BASE_CHANNEL_STRIDE ==4);
-                            if ( isNC4 )
+                            constexpr bool bIsNC4 = (BASE_CHANNEL_STRIDE == 4);
+                            if (bIsNC4)
                             {
-                                pDestBuf[BASE_CHANNEL_STRIDE *i+3] = pBaseBuf[BASE_CHANNEL_STRIDE *i+3];
+                                DestBuf[BASE_CHANNEL_STRIDE * I + 3] = BaseBuf[BASE_CHANNEL_STRIDE * I + 3];
                             }
                         }
 
-                        pDestBuf += BASE_CHANNEL_STRIDE *different;
-                        pBaseBuf += BASE_CHANNEL_STRIDE *different;
-                        pMaskBuf += different;
+                        DestBuf += BASE_CHANNEL_STRIDE * Different;
+                        BaseBuf += BASE_CHANNEL_STRIDE * Different;
+                        MaskBuf += Different;
                     }
                 }
 
-                rows = FMath::DivideAndRoundUp(rows,2);
-                width = FMath::DivideAndRoundUp(width,2);
+                Rows = FMath::DivideAndRoundUp(Rows, 2);
+                Width = FMath::DivideAndRoundUp(Width, 2);
             }
         }
         else
@@ -617,134 +489,119 @@ namespace mu
 	}
 
 
-	//---------------------------------------------------------------------------------------------
-	//! Apply a blending function to an image with a colour source and a mask
-	//---------------------------------------------------------------------------------------------
-	template< unsigned (*BLEND_FUNC_MASKED)(unsigned,unsigned,unsigned),
-			  unsigned (*BLEND_FUNC)(unsigned,unsigned),
-			  bool CLAMP
-			  >
-    inline void BufferLayerColour( uint8* DestBuf, const Image* Base, const Image* Mask, FVector4f Col )
+	/**
+	* Apply a blending function to an image with a colour source and a mask
+	*/
+	template<
+		uint32 (*BLEND_FUNC_MASKED)(uint32,uint32,uint32),
+		uint32 (*BLEND_FUNC)(uint32,uint32),
+		bool CLAMP>
+    inline void BufferLayerColour(Image* ResultImage, const Image* BaseImage, const Image* MaskImage, FVector4f Col)
 	{
-		check(Base->GetSizeX() == Mask->GetSizeX());
-		check(Base->GetSizeY() == Mask->GetSizeY());
-		check(Mask->GetFormat() == EImageFormat::IF_L_UBYTE
-			||
-			Mask->GetFormat() == EImageFormat::IF_L_UBYTE_RLE);
-
-		bool bValid =
-			(Base->GetSizeX() == Mask->GetSizeX())
-			&&
-			(Base->GetSizeY() == Mask->GetSizeY())
-			&&
-			(Mask->GetFormat() == EImageFormat::IF_L_UBYTE
-			||
-			Mask->GetFormat() == EImageFormat::IF_L_UBYTE_RLE);
+		check(BaseImage->GetSizeX() == MaskImage->GetSizeX());
+		check(BaseImage->GetSizeY() == MaskImage->GetSizeY());
+		check(MaskImage->GetFormat() == EImageFormat::IF_L_UBYTE ||
+			  MaskImage->GetFormat() == EImageFormat::IF_L_UBYTE_RLE);
+
+		const bool bValid = (BaseImage->GetSizeX() == MaskImage->GetSizeX()) &&
+						    (BaseImage->GetSizeY() == MaskImage->GetSizeY()) &&
+							(MaskImage->GetFormat() == EImageFormat::IF_L_UBYTE || MaskImage->GetFormat() == EImageFormat::IF_L_UBYTE_RLE);
 		if (!bValid)
 		{
 			return;
 		}
 
-		EImageFormat baseFormat = Base->GetFormat();
-		if ( baseFormat==EImageFormat::IF_RGB_UBYTE )
+		EImageFormat BaseFormat = BaseImage->GetFormat();
+		if (BaseFormat == EImageFormat::IF_RGB_UBYTE)
 		{
-			BufferLayerColourFormat<BLEND_FUNC_MASKED, BLEND_FUNC, CLAMP, 3, 3>( DestBuf, Base, Mask, Col, 0, 0, false);
+			BufferLayerColourFormat<BLEND_FUNC_MASKED, BLEND_FUNC, CLAMP, 3, 3>(ResultImage, BaseImage, MaskImage, Col, 0, 0, false);
 		}
-        else if ( baseFormat==EImageFormat::IF_RGBA_UBYTE )
+        else if (BaseFormat == EImageFormat::IF_RGBA_UBYTE)
         {
-            BufferLayerColourFormat<BLEND_FUNC_MASKED, BLEND_FUNC, CLAMP, 3, 4>( DestBuf, Base, Mask, Col, 0, 0, false);
+            BufferLayerColourFormat<BLEND_FUNC_MASKED, BLEND_FUNC, CLAMP, 3, 4>(ResultImage, BaseImage, MaskImage, Col, 0, 0, false);
         }
-        else if ( baseFormat==EImageFormat::IF_BGRA_UBYTE )
+        else if (BaseFormat == EImageFormat::IF_BGRA_UBYTE)
         {
-            float temp = Col[0];
+            float Temp = Col[0];
             Col[0] = Col[2];
-            Col[2] = temp;
-            BufferLayerColourFormat<BLEND_FUNC_MASKED, BLEND_FUNC, CLAMP, 3, 4>( DestBuf, Base, Mask, Col, 0, 0, false);
+            Col[2] = Temp;
+            BufferLayerColourFormat<BLEND_FUNC_MASKED, BLEND_FUNC, CLAMP, 3, 4>(ResultImage, BaseImage, MaskImage, Col, 0, 0, false);
         }
-        else if ( baseFormat==EImageFormat::IF_L_UBYTE )
+        else if (BaseFormat == EImageFormat::IF_L_UBYTE)
 		{
-			BufferLayerColourFormat<BLEND_FUNC_MASKED, BLEND_FUNC, CLAMP, 1, 1>( DestBuf, Base, Mask, Col, 0, 0, false);
+			BufferLayerColourFormat<BLEND_FUNC_MASKED, BLEND_FUNC, CLAMP, 1, 1>(ResultImage, BaseImage, MaskImage, Col, 0, 0, false);
 		}
 		else
 		{
-			checkf( false, TEXT("Unsupported format.") );
+			checkf(false, TEXT("Unsupported format."));
 		}
 	}
 
 
-	//---------------------------------------------------------------------------------------------
-	template< unsigned (*BLEND_FUNC_MASKED)(unsigned, unsigned, unsigned),
-		unsigned (*BLEND_FUNC)(unsigned, unsigned)
-	>
-	inline void BufferLayerColour(uint8* DestBuf, const Image* Base, const Image* Mask, FVector4f Col)
+	template< 
+		uint32 (*BLEND_FUNC_MASKED)(uint32, uint32, uint32),
+		uint32 (*BLEND_FUNC)(uint32, uint32)>
+	inline void BufferLayerColour(Image* DestImage, const Image* BaseImage, const Image* MaskImage, FVector4f Col)
 	{
 		bool bIsClampNeeded = IsAnyComponentLargerThan1(Col);
 		if (bIsClampNeeded)
 		{
-			BufferLayerColour<BLEND_FUNC_MASKED, BLEND_FUNC, true>(DestBuf, Base, Mask, Col );
+			BufferLayerColour<BLEND_FUNC_MASKED, BLEND_FUNC, true>(DestImage, BaseImage, MaskImage, Col);
 		}
 		else
 		{
-			BufferLayerColour<BLEND_FUNC_MASKED, BLEND_FUNC, false>(DestBuf, Base, Mask, Col );
+			BufferLayerColour<BLEND_FUNC_MASKED, BLEND_FUNC, false>(DestImage, BaseImage, MaskImage, Col);
 		}
 	}
 
 
-	//---------------------------------------------------------------------------------------------
-	//! Apply a blending function to an image with another image as blending layer
-	//---------------------------------------------------------------------------------------------
-	template< unsigned (*BLEND_FUNC)(unsigned, unsigned), bool CLAMP,
+	/**
+	* Apply a blending function to an image with another image as blending layer
+	*/
+	template< 
+		uint32 (*BLEND_FUNC)(uint32, uint32), 
+		bool CLAMP,
 		// Number of total channels to actually process
 		uint32 CHANNELS_TO_BLEND,
 		// Number of total channels in the base image
 		int32 BASE_CHANNEL_STRIDE>
-
-	inline void BufferLayerColourFormatInPlace(Image* Base, FVector4f Col, uint32 BaseChannelOffset, uint8 ColOffset, bool bOnlyFirstLOD)
+	inline void BufferLayerColourFormatInPlace(Image* BaseImage, FVector4f Color, uint32 BaseChannelOffset, uint8 ColorChannelOffset, bool bOnlyFirstLOD)
 	{
-		uint8* pBaseBuf = Base->GetData() + BaseChannelOffset;
+		FUintVector4 TopColor = FUintVector4(Color.X * 255.0f, Color.Y * 255.0f, Color.Z * 255.0f, Color.W * 255.0f);
 
-		uint32 top[4];
-		top[0] = (uint32)(255 * Col[0]);
-		top[1] = (uint32)(255 * Col[1]);
-		top[2] = (uint32)(255 * Col[2]);
-		top[3] = (uint32)(255 * Col[3]);
-
-		// Generic implementation
-		int32 PixelCount = 0;
-		if (bOnlyFirstLOD)
-		{
-			PixelCount = Base->GetSizeX() * Base->GetSizeY();
-		}
-		else
-		{
-			PixelCount = Base->CalculatePixelCount();
-		}
+		int32 NumLODs = bOnlyFirstLOD ? 1 : BaseImage->GetLODCount();
 
 		constexpr int32 NumBatchElems = 4096*2;
-		const int32 NumBatches = FMath::DivideAndRoundUp(PixelCount, NumBatchElems);
+
+		const int32 BytesPerElem = GetImageFormatData(BaseImage->GetFormat()).BytesPerBlock;
+		const int32 NumBatches = 
+				BaseImage->DataStorage.GetNumBatchesLODRange(NumBatchElems, BytesPerElem, 0, NumLODs);
 
 		auto ProcessBatch =
 		[
-			pBaseBuf, top, BaseChannelOffset, ColOffset, PixelCount, NumBatchElems
+			BaseImage, TopColor, BaseChannelOffset, ColorChannelOffset, BytesPerElem, NumBatchElems, NumLODs
 		] (uint32 BatchId)
-		{
-			const int32 BatchBegin = BatchId * NumBatchElems;
-			const int32 BatchEnd = FMath::Min(BatchBegin + NumBatchElems, PixelCount);
+		{	
+			const TArrayView<uint8> BaseBatchView = 
+					BaseImage->DataStorage.GetBatchLODRange(BatchId, NumBatchElems, BytesPerElem, 0, NumLODs);
 			
-			for (int32 i = BatchBegin; i < BatchEnd; ++i)
+			uint8* BaseBuf = BaseBatchView.GetData() + BaseChannelOffset;
+			const int32 NumElems = BaseBatchView.Num() / BytesPerElem;
+			
+			for (int32 I = 0; I < NumElems; ++I)
 			{
-				for (int c = 0; c < CHANNELS_TO_BLEND; ++c)
+				for (int32 C = 0; C < CHANNELS_TO_BLEND; ++C)
 				{
-					uint32 base = pBaseBuf[BASE_CHANNEL_STRIDE * i + c];
-					uint32 blended = top[c + ColOffset];
-					uint32 result = BLEND_FUNC(base, blended);
-					if (CLAMP)
+					uint32 Base = BaseBuf[BASE_CHANNEL_STRIDE*I + C];
+					uint32 Blended = TopColor[C + ColorChannelOffset];
+					uint32 Result = BLEND_FUNC(Base, Blended);
+					if constexpr (CLAMP)
 					{
-						pBaseBuf[BASE_CHANNEL_STRIDE * i + c] = (uint8)FMath::Min(255u, result);
+						BaseBuf[BASE_CHANNEL_STRIDE*I + C] = (uint8)FMath::Min(255u, Result);
 					}
 					else
 					{
-						pBaseBuf[BASE_CHANNEL_STRIDE * i + c] = (uint8)result;
+						BaseBuf[BASE_CHANNEL_STRIDE*I + C] = (uint8)Result;
 					}
 				}
 			}
@@ -761,34 +618,40 @@ namespace mu
 	}
 
 
-	//---------------------------------------------------------------------------------------------
-	template< unsigned (*BLEND_FUNC)(unsigned, unsigned), bool CLAMP, uint32 CHANNEL_COUNT >
-	inline void BufferLayerColourInPlace(Image* Base, FVector4f Col, bool bOnlyOneMip, uint32 BaseOffset, uint8 ColOffset)
+	template<
+		uint32 (*BLEND_FUNC)(uint32, uint32), 
+		bool CLAMP, 
+		uint32 CHANNEL_COUNT>
+	inline void BufferLayerColourInPlace(Image* BaseImage, FVector4f Col, bool bOnlyOneMip, uint32 BaseOffset, uint8 ColOffset)
 	{
-		EImageFormat baseFormat = Base->GetFormat();
+		EImageFormat BaseFormat = BaseImage->GetFormat();
 
-		if (baseFormat == EImageFormat::IF_RGB_UBYTE)
+		if (BaseFormat == EImageFormat::IF_RGB_UBYTE)
 		{
 			check(BaseOffset + CHANNEL_COUNT <= 3);
-			BufferLayerColourFormatInPlace< BLEND_FUNC, CLAMP, CHANNEL_COUNT, 3 >(Base, Col, BaseOffset, ColOffset, bOnlyOneMip);
+			BufferLayerColourFormatInPlace<BLEND_FUNC, CLAMP, CHANNEL_COUNT, 3>
+					(BaseImage, Col, BaseOffset, ColOffset, bOnlyOneMip);
 		}
-		else if (baseFormat == EImageFormat::IF_RGBA_UBYTE)
+		else if (BaseFormat == EImageFormat::IF_RGBA_UBYTE)
 		{
 			check(BaseOffset + CHANNEL_COUNT <= 4);
-			BufferLayerColourFormatInPlace< BLEND_FUNC, CLAMP, CHANNEL_COUNT, 4 >(Base, Col, BaseOffset, ColOffset, bOnlyOneMip);
+			BufferLayerColourFormatInPlace<BLEND_FUNC, CLAMP, CHANNEL_COUNT, 4>
+					(BaseImage, Col, BaseOffset, ColOffset, bOnlyOneMip);
 		}
-		else if (baseFormat == EImageFormat::IF_BGRA_UBYTE)
+		else if (BaseFormat == EImageFormat::IF_BGRA_UBYTE)
 		{
-			float temp = Col[0];
+			float Temp = Col[0];
 			Col[0] = Col[2];
-			Col[2] = temp;
-			BufferLayerColourFormatInPlace< BLEND_FUNC, CLAMP, CHANNEL_COUNT, 4>(Base, Col, BaseOffset, ColOffset, bOnlyOneMip);
+			Col[2] = Temp;
+			BufferLayerColourFormatInPlace<BLEND_FUNC, CLAMP, CHANNEL_COUNT, 4>
+					(BaseImage, Col, BaseOffset, ColOffset, bOnlyOneMip);
 		}
 
-		else if (baseFormat == EImageFormat::IF_L_UBYTE)
+		else if (BaseFormat == EImageFormat::IF_L_UBYTE)
 		{
 			check(BaseOffset + CHANNEL_COUNT <= 1);
-			BufferLayerColourFormatInPlace< BLEND_FUNC, CLAMP, CHANNEL_COUNT, 1 >(Base, Col, BaseOffset, ColOffset, bOnlyOneMip);
+			BufferLayerColourFormatInPlace< BLEND_FUNC, CLAMP, CHANNEL_COUNT, 1>
+					(BaseImage, Col, BaseOffset, ColOffset, bOnlyOneMip);
 		}
 		else
 		{
@@ -797,57 +660,62 @@ namespace mu
 	}
 
 
-	//---------------------------------------------------------------------------------------------
-	template< unsigned (*BLEND_FUNC)(unsigned, unsigned), uint32 CHANNEL_COUNT >
-	inline void BufferLayerColourInPlace(Image* Base, FVector4f Col, bool bOnlyOneMip, uint32 BaseOffset, uint8 ColOffset)
+	template< 
+		uint32 (*BLEND_FUNC)(uint32, uint32), 
+		uint32 CHANNEL_COUNT>
+	inline void BufferLayerColourInPlace(Image* BaseImage, FVector4f Color, bool bOnlyOneMip, uint32 BaseOffset, uint8 ColorOffset)
 	{
-		bool bIsClampNeeded = IsAnyComponentLargerThan1(Col);
+		bool bIsClampNeeded = IsAnyComponentLargerThan1(Color);
 		if (bIsClampNeeded)
 		{
-			BufferLayerColourInPlace<BLEND_FUNC, true, CHANNEL_COUNT>(Base, Col, bOnlyOneMip, BaseOffset, ColOffset);
+			BufferLayerColourInPlace<BLEND_FUNC, true, CHANNEL_COUNT>(BaseImage, Color, bOnlyOneMip, BaseOffset, ColorOffset);
 		}
 		else
 		{
-			BufferLayerColourInPlace<BLEND_FUNC, false, CHANNEL_COUNT>(Base, Col, bOnlyOneMip, BaseOffset, ColOffset);
+			BufferLayerColourInPlace<BLEND_FUNC, false, CHANNEL_COUNT>(BaseImage, Color, bOnlyOneMip, BaseOffset, ColorOffset);
 		}
 	}
 
 
-	//---------------------------------------------------------------------------------------------
 	template< 
-		unsigned (*BLEND_FUNC_MASKED)(unsigned, unsigned, unsigned),
-		unsigned (*BLEND_FUNC)(unsigned, unsigned), 
+		uint32 (*BLEND_FUNC_MASKED)(uint32, uint32, uint32),
+		uint32 (*BLEND_FUNC)(uint32, uint32), 
 		bool CLAMP, 
-		int32 CHANNEL_COUNT >
-	inline void BufferLayerColourInPlace(Image* Base, const Image* Mask, FVector4f Col, bool bOnlyOneMip, uint32 BaseOffset, uint8 ColOffset)
+		int32 CHANNEL_COUNT>
+	inline void BufferLayerColourInPlace(Image* BaseImage, const Image* MaskImage, FVector4f Color, bool bOnlyOneMip, uint32 BaseOffset, uint8 ColorOffset)
 	{
-		check(Base->GetSizeX() == Mask->GetSizeX());
-		check(Base->GetSizeY() == Mask->GetSizeY());
+		check(BaseImage->GetSizeX() == MaskImage->GetSizeX());
+		check(BaseImage->GetSizeY() == MaskImage->GetSizeY());
 
-		EImageFormat baseFormat = Base->GetFormat();
+		EImageFormat BaseFormat = BaseImage->GetFormat();
 
-		if (baseFormat == EImageFormat::IF_RGB_UBYTE)
+		if (BaseFormat == EImageFormat::IF_RGB_UBYTE)
 		{
 			check(BaseOffset + CHANNEL_COUNT <= 3);
-			BufferLayerColourFormat< BLEND_FUNC_MASKED, BLEND_FUNC, CLAMP, CHANNEL_COUNT, 3 >(Base->GetData(), Base, Mask, Col, BaseOffset, ColOffset, bOnlyOneMip);
+			BufferLayerColourFormat<BLEND_FUNC_MASKED, BLEND_FUNC, CLAMP, CHANNEL_COUNT, 3>
+					(BaseImage, BaseImage, MaskImage, Color, BaseOffset, ColorOffset, bOnlyOneMip);
 		}
-		else if (baseFormat == EImageFormat::IF_RGBA_UBYTE)
+		else if (BaseFormat == EImageFormat::IF_RGBA_UBYTE)
 		{
 			check(BaseOffset + CHANNEL_COUNT <= 4);
-			BufferLayerColourFormat< BLEND_FUNC_MASKED, BLEND_FUNC, CLAMP, CHANNEL_COUNT, 4 >(Base->GetData(), Base, Mask, Col, BaseOffset, ColOffset, bOnlyOneMip);
+			BufferLayerColourFormat<BLEND_FUNC_MASKED, BLEND_FUNC, CLAMP, CHANNEL_COUNT, 4>
+					(BaseImage, BaseImage, MaskImage, Color, BaseOffset, ColorOffset, bOnlyOneMip);
 		}
-		else if (baseFormat == EImageFormat::IF_BGRA_UBYTE)
+		else if (BaseFormat == EImageFormat::IF_BGRA_UBYTE)
 		{
 			check(BaseOffset + CHANNEL_COUNT <= 4);
-			float temp = Col[0];
-			Col[0] = Col[2];
-			Col[2] = temp;
-			BufferLayerColourFormat< BLEND_FUNC_MASKED, BLEND_FUNC, CLAMP, CHANNEL_COUNT, 4 >(Base->GetData(), Base, Mask, Col, BaseOffset, ColOffset, bOnlyOneMip);
+			float Temp = Color[0];
+			Color[0] = Color[2];
+			Color[2] = Temp;
+
+			BufferLayerColourFormat<BLEND_FUNC_MASKED, BLEND_FUNC, CLAMP, CHANNEL_COUNT, 4>
+					(BaseImage, BaseImage, MaskImage, Color, BaseOffset, ColorOffset, bOnlyOneMip);
 		}
-		else if (baseFormat == EImageFormat::IF_L_UBYTE)
+		else if (BaseFormat == EImageFormat::IF_L_UBYTE)
 		{
 			check(BaseOffset + CHANNEL_COUNT <= 1);
-			BufferLayerColourFormat< BLEND_FUNC_MASKED, BLEND_FUNC, CLAMP, CHANNEL_COUNT, 1 >(Base->GetData(), Base, Mask, Col, BaseOffset, ColOffset, bOnlyOneMip);
+			BufferLayerColourFormat<BLEND_FUNC_MASKED, BLEND_FUNC, CLAMP, CHANNEL_COUNT, 1>
+					(BaseImage, BaseImage, MaskImage, Color, BaseOffset, ColorOffset, bOnlyOneMip);
 		}
 		else
 		{
@@ -855,88 +723,96 @@ namespace mu
 		}
 	}
 
-
-	//---------------------------------------------------------------------------------------------
-	template< unsigned (*BLEND_FUNC_MASKED)(unsigned, unsigned, unsigned), unsigned (*BLEND_FUNC)(unsigned, unsigned), uint32 CHANNEL_COUNT >
-	inline void BufferLayerColourInPlace(Image* Base, const Image* Mask, FVector4f Col, bool bOnlyOneMip, uint32 BaseOffset, uint8 ColOffset)
+	template< 
+		uint32 (*BLEND_FUNC_MASKED)(uint32, uint32, uint32), 
+		uint32 (*BLEND_FUNC)(uint32, uint32), uint32 CHANNEL_COUNT>
+	inline void BufferLayerColourInPlace(Image* BaseImage, const Image* MaskImage, FVector4f Color, bool bOnlyOneMip, uint32 BaseOffset, uint8 ColorOffset)
 	{
-		bool bIsClampNeeded = IsAnyComponentLargerThan1(Col);
+		bool bIsClampNeeded = IsAnyComponentLargerThan1(Color);
 		if (bIsClampNeeded)
 		{
-			BufferLayerColourInPlace<BLEND_FUNC_MASKED, BLEND_FUNC, true, CHANNEL_COUNT>(Base, Mask, Col, bOnlyOneMip, BaseOffset, ColOffset);
+			BufferLayerColourInPlace<BLEND_FUNC_MASKED, BLEND_FUNC, true, CHANNEL_COUNT>
+					(BaseImage, MaskImage, Color, bOnlyOneMip, BaseOffset, ColorOffset);
 		}
 		else
 		{
-			BufferLayerColourInPlace<BLEND_FUNC_MASKED, BLEND_FUNC, false, CHANNEL_COUNT>(Base, Mask, Col, bOnlyOneMip, BaseOffset, ColOffset);
+			BufferLayerColourInPlace<BLEND_FUNC_MASKED, BLEND_FUNC, false, CHANNEL_COUNT>
+					(BaseImage, MaskImage, Color, bOnlyOneMip, BaseOffset, ColorOffset);
 		}
 	}
 
 
-	//---------------------------------------------------------------------------------------------
-	//! Apply a blending function to an image with another image as blending layer
-	//---------------------------------------------------------------------------------------------
-	template< unsigned (*BLEND_FUNC)(unsigned,unsigned), bool CLAMP,
+	/**	
+	* Apply a blending function to an image with another image as blending layer
+	*/
+	template< uint32 (*BLEND_FUNC)(uint32, uint32), bool CLAMP,
 		// Number of total channels to actually process
 		int32 CHANNELS_TO_BLEND,
 		// Number of total channels in the base image
 		int32 BASE_CHANNEL_STRIDE,
 		// Number of total channels in the blend image
 		int32 BLENDED_CHANNEL_STRIDE>
-
-    inline void BufferLayerFormatInPlace(Image* pBase, const Image* pBlended,
+    inline void BufferLayerFormatInPlace(
+		Image* BaseImage, 
+		const Image* BlendedImage,
 		uint32 BaseChannelOffset,
 		uint32 BlendedChannelOffset,
-		bool bOnlyFirstLOD )
+		bool bOnlyFirstLOD)
 	{
-		check(pBase->GetSizeX() == pBlended->GetSizeX());
-		check(pBase->GetSizeY() == pBlended->GetSizeY());
+		check(BaseImage->GetSizeX() == BlendedImage->GetSizeX());
+		check(BaseImage->GetSizeY() == BlendedImage->GetSizeY());
 
 		// No longer required.
-		//check(pBase->GetFormat() == pBlended->GetFormat());
-		check(BaseChannelOffset + CHANNELS_TO_BLEND <= GetImageFormatData(pBase->GetFormat()).Channels);
-		check(BlendedChannelOffset + CHANNELS_TO_BLEND <= GetImageFormatData(pBlended->GetFormat()).Channels);
+		//check(BaseImage->GetFormat() == BlendedImage->GetFormat());
+		check(BaseChannelOffset + CHANNELS_TO_BLEND <= GetImageFormatData(BaseImage->GetFormat()).Channels);
+		check(BlendedChannelOffset + CHANNELS_TO_BLEND <= GetImageFormatData(BlendedImage->GetFormat()).Channels);
 
-		check(bOnlyFirstLOD || pBase->GetLODCount() <= pBlended->GetLODCount());
+		check(bOnlyFirstLOD || BaseImage->GetLODCount() <= BlendedImage->GetLODCount());
 
-        uint8* pBaseBuf = pBase->GetData() + BaseChannelOffset;
-        const uint8* pBlendedBuf = pBlended->GetData() + BlendedChannelOffset;
-
-		// Generic implementation
-		int32 PixelCount = 0;
-		if (bOnlyFirstLOD)
-		{
-			PixelCount = pBase->GetSizeX() * pBase->GetSizeY();
-		}
-		else
-		{
-			PixelCount = pBase->CalculatePixelCount();
-		}
+		const int32 BaseBytesPerElem = GetImageFormatData(BaseImage->GetFormat()).BytesPerBlock;
+		const int32 BlendBytesPerElem = GetImageFormatData(BlendedImage->GetFormat()).BytesPerBlock;
+	
+		const int32 NumLODs = bOnlyFirstLOD ? 1 : BaseImage->GetLODCount();
 
 		constexpr int32 NumBatchElems = 4096*2;
-		const int32 NumBatches = FMath::DivideAndRoundUp(PixelCount, NumBatchElems);
+		const int32 NumBatches =
+				BaseImage->DataStorage.GetNumBatchesLODRange(NumBatchElems, BaseBytesPerElem, 0, NumLODs);
+		
+		check(NumBatches <= BlendedImage->DataStorage.GetNumBatches(NumBatchElems, BlendBytesPerElem));
 
 		auto ProcessBatch =
 		[
-			pBaseBuf, pBlendedBuf, PixelCount, NumBatchElems
+			BaseImage, BlendedImage, BaseBytesPerElem, BlendBytesPerElem, 
+			NumBatchElems, BaseChannelOffset, BlendedChannelOffset, NumLODs
 		] (int32 BatchId)
 		{
-			const int32 BatchBegin = BatchId * NumBatchElems;
-			const int32 BatchEnd = FMath::Min(BatchBegin + NumBatchElems, PixelCount);
+			TArrayView<uint8> BaseBatchView =
+					BaseImage->DataStorage.GetBatchLODRange(BatchId, NumBatchElems, BaseBytesPerElem, 0, NumLODs);
+
+			TArrayView<const uint8> BlendedBatchView =
+					BlendedImage->DataStorage.GetBatchLODRange(BatchId, NumBatchElems, BlendBytesPerElem, 0, NumLODs);
 
-			for (int32 i = BatchBegin; i < BatchEnd; ++i)
+        	uint8* BaseBuf = BaseBatchView.GetData() + BaseChannelOffset;
+        	const uint8* BlendedBuf = BlendedBatchView.GetData() + BlendedChannelOffset;
+
+			const int32 NumElems = BaseBatchView.Num() / BaseBytesPerElem;
+			check(BlendedBatchView.Num() / BlendBytesPerElem == NumElems);
+
+			for (int32 I = 0; I < NumElems; ++I)
 			{
-				for (int c = 0; c < CHANNELS_TO_BLEND; ++c)
+				for (int32 C = 0; C < CHANNELS_TO_BLEND; ++C)
 				{
-					unsigned base = pBaseBuf[BASE_CHANNEL_STRIDE * i + c];
-					unsigned blended = pBlendedBuf[BLENDED_CHANNEL_STRIDE * i + c];
-					unsigned result = BLEND_FUNC(base, blended);
-					if (CLAMP)
+					uint32 Base = BaseBuf[BASE_CHANNEL_STRIDE * I + C];
+					uint32 Blended = BlendedBuf[BLENDED_CHANNEL_STRIDE * I + C];
+					uint32 Result = BLEND_FUNC(Base, Blended);
+					
+					if constexpr (CLAMP)
 					{
-						pBaseBuf[BASE_CHANNEL_STRIDE * i + c] = (uint8)FMath::Min(255u, result);
+						BaseBuf[BASE_CHANNEL_STRIDE * I + C] = (uint8)FMath::Min(255u, Result);
 					}
 					else
 					{
-						pBaseBuf[BASE_CHANNEL_STRIDE * i + c] = (uint8)result;
+						BaseBuf[BASE_CHANNEL_STRIDE * I + C] = (uint8)Result;
 					}
 				}
 			}
@@ -952,10 +828,12 @@ namespace mu
 		}
 	}
 	
-	//---------------------------------------------------------------------------------------------
-	//! Apply a blending function to an image with another image as blending layer
-	//---------------------------------------------------------------------------------------------
-	template< unsigned (*BLEND_FUNC)(unsigned, unsigned), bool CLAMP,
+	/**
+	* Apply a blending function to an image with another image as blending layer
+	*/
+	template<
+		uint32 (*BLEND_FUNC)(uint32, uint32), 
+		bool CLAMP,
 		// Number of total channels to actually process
 		int32 CHANNELS_TO_BLEND,
 		// Number of total channels in the base image
@@ -963,65 +841,77 @@ namespace mu
 		// Number of total channels in the blend image
 		int32 BLENDED_CHANNEL_STRIDE,
 		int32 BLENDED_CHANNEL_OFFSET>
-
-	inline void BufferLayerFormat(uint8* pDestBuf, const Image* pBase, const Image* pBlended,
-		bool bOnlyFirstLOD)
+	inline void BufferLayerFormat(Image* DestImage, const Image* BaseImage, const Image* BlendedImage, bool bOnlyFirstLOD)
 	{
-		check(pBase->GetSizeX() == pBlended->GetSizeX());
-		check(pBase->GetSizeY() == pBlended->GetSizeY());
+		check(BaseImage->GetSizeX() == BlendedImage->GetSizeX());
+		check(BaseImage->GetSizeY() == BlendedImage->GetSizeY());
 		// Not true anymore, since the BLENDED_CHANNEL_OFFSET has been added.
-		// check(pBase->GetFormat() == pBlended->GetFormat());
-		check(bOnlyFirstLOD || pBase->GetLODCount() <= pBlended->GetLODCount());
+		// check(BaseImage->GetFormat() == BlendedImage->GetFormat());
+		check(bOnlyFirstLOD || BaseImage->GetLODCount() <= BlendedImage->GetLODCount());
 
-		const uint8* pBaseBuf = pBase->GetData();
-		const uint8* pBlendedBuf = pBlended->GetData()+BLENDED_CHANNEL_OFFSET;
+		constexpr int32 UnblendedChannels = BASE_CHANNEL_STRIDE - CHANNELS_TO_BLEND;
+		constexpr int32 NumBatchElems = 4096*2;
+		
+		const int32 BaseBytesPerElem = GetImageFormatData(BaseImage->GetFormat()).BytesPerBlock;
+		const int32 BlendBytesPerElem = GetImageFormatData(BlendedImage->GetFormat()).BytesPerBlock;
+		const int32 DestBytesPerElem = GetImageFormatData(DestImage->GetFormat()).BytesPerBlock;
 
-		// Generic implementation
-		int32 PixelCount = 0;
-		if (bOnlyFirstLOD)
-		{
-			PixelCount = pBase->GetSizeX() * pBase->GetSizeY();
-		}
-		else
-		{
-			PixelCount = pBase->CalculatePixelCount();
-		}
+		const int32 NumLODs = bOnlyFirstLOD ? 1 : BaseImage->GetLODCount();
 
-		int32 UnblendedChannels = BASE_CHANNEL_STRIDE - CHANNELS_TO_BLEND;
+		const int32 NumBatches = 
+				BaseImage->DataStorage.GetNumBatchesLODRange(NumBatchElems, BaseBytesPerElem, 0, NumLODs);
+
+		check(NumBatches == BlendedImage->DataStorage.GetNumBatchesLODRange(NumBatchElems, BlendBytesPerElem, 0, NumLODs));
+		check(NumBatches == DestImage->DataStorage.GetNumBatchesLODRange(NumBatchElems, DestBytesPerElem, 0, NumLODs));
 
-		constexpr int32 NumBatchElems = 4096*2;
-		const int32 NumBatches = FMath::DivideAndRoundUp(PixelCount, NumBatchElems);
-		
 		auto ProcessBatch =
 		[
-			pBaseBuf, pBlendedBuf, pDestBuf, UnblendedChannels, PixelCount, NumBatchElems
+			BaseImage, BlendedImage, DestImage, 
+			BaseBytesPerElem, BlendBytesPerElem, DestBytesPerElem, 
+			UnblendedChannels, NumBatchElems, NumLODs
 		] (int32 BatchId)
 		{
-			const int32 BatchBegin = BatchId * NumBatchElems;
-			const int32 BatchEnd = FMath::Min(BatchBegin + NumBatchElems, PixelCount);
+			TArrayView<const uint8> BaseBatchView =
+					BaseImage->DataStorage.GetBatchLODRange(BatchId, NumBatchElems, BaseBytesPerElem, 0, NumLODs);
+
+			TArrayView<const uint8> BlendedBatchView =
+					BlendedImage->DataStorage.GetBatchLODRange(BatchId, NumBatchElems, BlendBytesPerElem, 0, NumLODs);
+
+			TArrayView<uint8> DestBatchView =
+					DestImage->DataStorage.GetBatchLODRange(BatchId, NumBatchElems, DestBytesPerElem, 0, NumLODs);
+	
+			const int32 NumElems = BaseBatchView.Num() / BaseBytesPerElem;
+			check(NumElems == BlendedBatchView.Num() / BlendBytesPerElem);
+			check(NumElems == DestBatchView.Num() / DestBytesPerElem);
+
+			const uint8* BaseBuf = BaseBatchView.GetData();
+			const uint8* BlendedBuf = BlendedBatchView.GetData() + BLENDED_CHANNEL_OFFSET;
+			uint8* DestBuf = DestBatchView.GetData();
 
-			for (int32 i = BatchBegin; i < BatchEnd; ++i)
+			for (int32 I = 0; I < NumElems; ++I)
 			{
-				for (int c = 0; c < CHANNELS_TO_BLEND; ++c)
+				for (int32 C = 0; C < CHANNELS_TO_BLEND; ++C)
 				{
-					unsigned base = pBaseBuf[BASE_CHANNEL_STRIDE * i + c];
-					unsigned blended = pBlendedBuf[BLENDED_CHANNEL_STRIDE * i + c];
-					unsigned result = BLEND_FUNC(base, blended);
+					uint32 Base = BaseBuf[BASE_CHANNEL_STRIDE * I + C];
+					uint32 Blended = BlendedBuf[BLENDED_CHANNEL_STRIDE * I + C];
+					uint32 Result = BLEND_FUNC(Base, Blended);
+
 					if constexpr (CLAMP)
 					{
-						pDestBuf[BASE_CHANNEL_STRIDE * i + c] = (uint8)FMath::Min(255u, result);
+						DestBuf[BASE_CHANNEL_STRIDE * I + C] = (uint8)FMath::Min(255u, Result);
 					}
 					else
 					{
-						pDestBuf[BASE_CHANNEL_STRIDE * i + c] = (uint8)result;
+						DestBuf[BASE_CHANNEL_STRIDE * I + C] = (uint8)Result;
 					}
 				}
 
 				// Copy the unblended channels
 				// \TODO: unnecessary when doing it in-place?
-				for (int32 c = 0; c < UnblendedChannels; ++c)
+				for (int32 C = 0; C < UnblendedChannels; ++C)
 				{
-					pDestBuf[BASE_CHANNEL_STRIDE * i + CHANNELS_TO_BLEND + c] = pBaseBuf[BASE_CHANNEL_STRIDE * i + CHANNELS_TO_BLEND + c];
+					DestBuf[BASE_CHANNEL_STRIDE * I + CHANNELS_TO_BLEND + C] 
+						= BaseBuf[BASE_CHANNEL_STRIDE * I + CHANNELS_TO_BLEND + C];
 				}
 			}
 		};
@@ -1036,24 +926,22 @@ namespace mu
 		}
 	}
 
-
-	//---------------------------------------------------------------------------------------------
-	//! Apply a blending function to an image with another image as blending layer
-	//---------------------------------------------------------------------------------------------
-	template< unsigned (*BLEND_FUNC)(unsigned, unsigned), bool CLAMP >
-	inline void BufferLayer(Image* pResult, const Image* pBase, const Image* pBlended, bool bApplyToAlpha, bool bOnlyOneMip, bool bUseBlendSourceFromBlendAlpha)
+	/**
+	* Apply a blending function to an image with another image as blending layer
+	*/
+	template< uint32 (*BLEND_FUNC)(uint32, uint32), bool CLAMP >
+	inline void BufferLayer(Image* ResultImage, const Image* BaseImage, const Image* BlendedImage, bool bApplyToAlpha, bool bOnlyOneMip, bool bUseBlendSourceFromBlendAlpha)
 	{
-		check(pResult->GetFormat() == pBase->GetFormat());
-		check(pResult->GetSizeX() == pBase->GetSizeX());
-		check(pResult->GetSizeY() == pBase->GetSizeY());
-		check(bOnlyOneMip || pResult->GetLODCount() == pBase->GetLODCount());
-		check(pBase->GetSizeX() == pBlended->GetSizeX());
-		check(pBase->GetSizeY() == pBlended->GetSizeY());
-		check(bOnlyOneMip || pResult->GetLODCount() <= pBlended->GetLODCount());
-
-		EImageFormat BaseFormat = pBase->GetFormat();
-		EImageFormat BlendedFormat = pBlended->GetFormat();
-		uint8* pDestBuf = pResult->GetData();
+		check(ResultImage->GetFormat() == BaseImage->GetFormat());
+		check(ResultImage->GetSizeX() == BaseImage->GetSizeX());
+		check(ResultImage->GetSizeY() == BaseImage->GetSizeY());
+		check(bOnlyOneMip || ResultImage->GetLODCount() == BaseImage->GetLODCount());
+		check(BaseImage->GetSizeX() == BlendedImage->GetSizeX());
+		check(BaseImage->GetSizeY() == BlendedImage->GetSizeY());
+		check(bOnlyOneMip || ResultImage->GetLODCount() <= BlendedImage->GetLODCount());
+
+		EImageFormat BaseFormat = BaseImage->GetFormat();
+		EImageFormat BlendedFormat = BlendedImage->GetFormat();
 
 		if (bUseBlendSourceFromBlendAlpha)
 		{
@@ -1061,7 +949,7 @@ namespace mu
 			{
 				if (BaseFormat == EImageFormat::IF_L_UBYTE)
 				{
-					BufferLayerFormat< BLEND_FUNC, CLAMP, 1, 1, 4, 3 >(pDestBuf, pBase, pBlended, bOnlyOneMip);
+					BufferLayerFormat<BLEND_FUNC, CLAMP, 1, 1, 4, 3>(ResultImage, BaseImage, BlendedImage, bOnlyOneMip);
 				}
 				else
 				{
@@ -1070,7 +958,7 @@ namespace mu
 			}
 			else if (BlendedFormat == EImageFormat::IF_L_UBYTE)
 			{
-				BufferLayerFormat< BLEND_FUNC, CLAMP, 1, 1, 1, 0 >(pDestBuf, pBase, pBlended, bOnlyOneMip);
+				BufferLayerFormat<BLEND_FUNC, CLAMP, 1, 1, 1, 0>(ResultImage, BaseImage, BlendedImage, bOnlyOneMip);
 			}
 		}
 		else
@@ -1079,23 +967,23 @@ namespace mu
 			if (BaseFormat == EImageFormat::IF_RGB_UBYTE)
 			{
 				check(!bUseBlendSourceFromBlendAlpha);
-				BufferLayerFormat< BLEND_FUNC, CLAMP, 3, 3, 3, 0 >(pDestBuf, pBase, pBlended, bOnlyOneMip);
+				BufferLayerFormat<BLEND_FUNC, CLAMP, 3, 3, 3, 0>(ResultImage, BaseImage, BlendedImage, bOnlyOneMip);
 			}
 			else if (BaseFormat == EImageFormat::IF_RGBA_UBYTE || BaseFormat == EImageFormat::IF_BGRA_UBYTE)
 			{
 				check(!bUseBlendSourceFromBlendAlpha);
 				if (bApplyToAlpha)
 				{
-					BufferLayerFormat< BLEND_FUNC, CLAMP, 4, 4, 4, 0 >(pDestBuf, pBase, pBlended, bOnlyOneMip);
+					BufferLayerFormat<BLEND_FUNC, CLAMP, 4, 4, 4, 0>(ResultImage, BaseImage, BlendedImage, bOnlyOneMip);
 				}
 				else
 				{
-					BufferLayerFormat< BLEND_FUNC, CLAMP, 3, 4, 4, 0 >(pDestBuf, pBase, pBlended, bOnlyOneMip);
+					BufferLayerFormat<BLEND_FUNC, CLAMP, 3, 4, 4, 0>(ResultImage, BaseImage, BlendedImage, bOnlyOneMip);
 				}
 			}
 			else if (BaseFormat == EImageFormat::IF_L_UBYTE)
 			{
-				BufferLayerFormat< BLEND_FUNC, CLAMP, 1, 1, 1, 0 >(pDestBuf, pBase, pBlended, bOnlyOneMip);
+				BufferLayerFormat<BLEND_FUNC, CLAMP, 1, 1, 1, 0>(ResultImage, BaseImage, BlendedImage, bOnlyOneMip);
 			}
 			else
 			{
@@ -1105,44 +993,44 @@ namespace mu
 	}
 
 
-	//---------------------------------------------------------------------------------------------
-	template< unsigned (*BLEND_FUNC)(unsigned, unsigned), bool CLAMP, int32 CHANNEL_COUNT >
-	inline void BufferLayerInPlace(Image* pBase, const Image* pBlended, bool bOnlyOneMip, uint32 BaseOffset, uint32 BlendedOffset)
+	template< uint32 (*BLEND_FUNC)(uint32, uint32), bool CLAMP, int32 CHANNEL_COUNT >
+	inline void BufferLayerInPlace(Image* BaseImage, const Image* BlendedImage, bool bOnlyOneMip, uint32 BaseOffset, uint32 BlendedOffset)
 	{
-		check(pBase->GetSizeX() == pBlended->GetSizeX());
-		check(pBase->GetSizeY() == pBlended->GetSizeY());
+		check(BaseImage->GetSizeX() == BlendedImage->GetSizeX());
+		check(BaseImage->GetSizeY() == BlendedImage->GetSizeY());
 		// Not required since we have the CHANNEL_COUNT and offsets. 
-		// check(pBase->GetFormat() == pBlended->GetFormat());
+		// check(BaseImage->GetFormat() == BlendedImage->GetFormat());
 
-		EImageFormat BaseFormat = pBase->GetFormat();
-		EImageFormat BlendFormat = pBlended->GetFormat();
+		EImageFormat BaseFormat = BaseImage->GetFormat();
+		EImageFormat BlendFormat = BlendedImage->GetFormat();
 
-		if (BaseFormat==EImageFormat::IF_RGB_UBYTE && BlendFormat==EImageFormat::IF_RGB_UBYTE)
+		if (BaseFormat == EImageFormat::IF_RGB_UBYTE && BlendFormat==EImageFormat::IF_RGB_UBYTE)
 		{
-			check(BaseOffset+CHANNEL_COUNT <= 3);
-			check(BlendedOffset+CHANNEL_COUNT <= 3);
-			BufferLayerFormatInPlace< BLEND_FUNC, CLAMP, CHANNEL_COUNT, 3, 3 >(pBase, pBlended, BaseOffset, BlendedOffset, bOnlyOneMip);
+			check(BaseOffset + CHANNEL_COUNT <= 3);
+			check(BlendedOffset + CHANNEL_COUNT <= 3);
+			BufferLayerFormatInPlace<BLEND_FUNC, CLAMP, CHANNEL_COUNT, 3, 3>(BaseImage, BlendedImage, BaseOffset, BlendedOffset, bOnlyOneMip);
 		}
-		else if ( (BaseFormat == EImageFormat::IF_RGBA_UBYTE || BaseFormat == EImageFormat::IF_BGRA_UBYTE)
-			&&
-			(BlendFormat == EImageFormat::IF_RGBA_UBYTE || BlendFormat == EImageFormat::IF_BGRA_UBYTE) )
+		else if ((BaseFormat == EImageFormat::IF_RGBA_UBYTE || BaseFormat == EImageFormat::IF_BGRA_UBYTE) &&
+				 (BlendFormat == EImageFormat::IF_RGBA_UBYTE || BlendFormat == EImageFormat::IF_BGRA_UBYTE) )
 		{
 			check(BaseOffset + CHANNEL_COUNT <= 4);
 			check(BlendedOffset + CHANNEL_COUNT <= 4);
-			BufferLayerFormatInPlace< BLEND_FUNC, CLAMP, CHANNEL_COUNT, 4, 4 >(pBase, pBlended, BaseOffset, BlendedOffset, bOnlyOneMip);
+			BufferLayerFormatInPlace< BLEND_FUNC, CLAMP, CHANNEL_COUNT, 4, 4>
+					(BaseImage, BlendedImage, BaseOffset, BlendedOffset, bOnlyOneMip);
 		}
-		else if ( (BaseFormat == EImageFormat::IF_RGBA_UBYTE || BaseFormat == EImageFormat::IF_BGRA_UBYTE)
-			&&
-			BlendFormat == EImageFormat::IF_L_UBYTE )
+		else if ((BaseFormat == EImageFormat::IF_RGBA_UBYTE || BaseFormat == EImageFormat::IF_BGRA_UBYTE) &&
+				 BlendFormat == EImageFormat::IF_L_UBYTE )
 		{
 			check(BaseOffset + CHANNEL_COUNT <= 4);
 			check(BlendedOffset + CHANNEL_COUNT <= 1);
-			BufferLayerFormatInPlace< BLEND_FUNC, CLAMP, CHANNEL_COUNT, 4, 1 >(pBase, pBlended, BaseOffset, BlendedOffset, bOnlyOneMip);
+			BufferLayerFormatInPlace< BLEND_FUNC, CLAMP, CHANNEL_COUNT, 4, 1>
+					(BaseImage, BlendedImage, BaseOffset, BlendedOffset, bOnlyOneMip);
 		}
-		else if (BaseFormat==EImageFormat::IF_L_UBYTE && BlendFormat==EImageFormat::IF_L_UBYTE)
+		else if (BaseFormat == EImageFormat::IF_L_UBYTE && BlendFormat==EImageFormat::IF_L_UBYTE)
 		{
-			check(BaseOffset+CHANNEL_COUNT <= 1);
-			BufferLayerFormatInPlace< BLEND_FUNC, CLAMP, CHANNEL_COUNT, 1, 1 >(pBase, pBlended, BaseOffset, BlendedOffset, bOnlyOneMip);
+			check(BaseOffset + CHANNEL_COUNT <= 1);
+			BufferLayerFormatInPlace< BLEND_FUNC, CLAMP, CHANNEL_COUNT, 1, 1>
+					(BaseImage, BlendedImage, BaseOffset, BlendedOffset, bOnlyOneMip);
 		}
 		else
 		{
@@ -1151,9 +1039,8 @@ namespace mu
 	}
 
 
-	//---------------------------------------------------------------------------------------------
-	template< unsigned (*BLEND_FUNC_MASKED)(unsigned, unsigned, unsigned),
-		unsigned (*BLEND_FUNC)(unsigned, unsigned),
+	template< uint32 (*BLEND_FUNC_MASKED)(uint32, uint32, uint32),
+		uint32 (*BLEND_FUNC)(uint32, uint32),
 		bool CLAMP,
 		// Number of total channels to actually process
 		int32 CHANNELS_TO_BLEND,
@@ -1161,80 +1048,94 @@ namespace mu
 		int32 BASE_CHANNEL_STRIDE,
 		// Number of total channels in the blend image
 		int32 BLENDED_CHANNEL_STRIDE>
-	inline void BufferLayerFormat(uint8* pStartDestBuf,
-		const Image* pBase,
-		const Image* pMask,
-		const Image* pBlend,
+	inline void BufferLayerFormat(
+		Image* DestImage,
+		const Image* BaseImage,
+		const Image* MaskImage,
+		const Image* BlendImage,
+		uint32 DestOffset,
 		uint32 BaseChannelOffset,
 		uint32 BlendedChannelOffset,
 		bool bOnlyFirstLOD)
 	{
-		check(pBase->GetSizeX() == pMask->GetSizeX() && pBase->GetSizeY() == pMask->GetSizeY());
-		check(pBase->GetSizeX() == pBlend->GetSizeX() && pBase->GetSizeY() == pBlend->GetSizeY());
-		check(bOnlyFirstLOD || pBase->GetLODCount() <= pMask->GetLODCount());
-		check(bOnlyFirstLOD || pBase->GetLODCount() <= pBlend->GetLODCount());
-		check(pBase->GetFormat() == pBlend->GetFormat());
-		check(pMask->GetFormat() == EImageFormat::IF_L_UBYTE
-			||
-			pMask->GetFormat() == EImageFormat::IF_L_UBYTE_RLE);
-
-		uint8* pDestBuf = pStartDestBuf;
-		const uint8* pMaskBuf = pMask->GetData();
-		const uint8* pBaseBuf = pBase->GetData() + BaseChannelOffset;
-		const uint8* pBlendedBuf = pBlend->GetData() + BlendedChannelOffset;
-
-		EImageFormat maskFormat = pMask->GetFormat();
-		bool bIsMaskUncompressed = (maskFormat == EImageFormat::IF_L_UBYTE);
-		int32 UnblendedChannels = BASE_CHANNEL_STRIDE - CHANNELS_TO_BLEND;
-
-		// The base determines the number of lods to process.
-		int32 LODCount = bOnlyFirstLOD ? 1 : pBase->GetLODCount();
+		check(BaseImage->GetSizeX() == MaskImage->GetSizeX() && BaseImage->GetSizeY() == MaskImage->GetSizeY());
+		check(BaseImage->GetSizeX() == BlendImage->GetSizeX() && BaseImage->GetSizeY() == BlendImage->GetSizeY());
+		check(bOnlyFirstLOD || BaseImage->GetLODCount() <= MaskImage->GetLODCount());
+		check(bOnlyFirstLOD || BaseImage->GetLODCount() <= BlendImage->GetLODCount());
+		check(BaseImage->GetFormat() == BlendImage->GetFormat());
+		check(MaskImage->GetFormat() == EImageFormat::IF_L_UBYTE ||
+			  MaskImage->GetFormat() == EImageFormat::IF_L_UBYTE_RLE);
 
-		if (bIsMaskUncompressed)
-		{
-			int32 PixelCount = 0;
-			if (bOnlyFirstLOD)
-			{
-				PixelCount = pBase->GetSizeX() * pBase->GetSizeY();
-			}
-			else
-			{
-				PixelCount = pBase->CalculatePixelCount();
-			}
+		const EImageFormat MaskFormat = MaskImage->GetFormat();
+
+		const int32 BaseBytesPerElem = GetImageFormatData(BaseImage->GetFormat()).BytesPerBlock; 
+		const int32 BlendBytesPerElem = GetImageFormatData(BlendImage->GetFormat()).BytesPerBlock; 
+
+		const int32 NumLODs = bOnlyFirstLOD ? 1 : BaseImage->GetLODCount();
 
+		const bool bIsMaskUncompressed = MaskFormat == EImageFormat::IF_L_UBYTE;
+		constexpr int32 UnblendedChannels = BASE_CHANNEL_STRIDE - CHANNELS_TO_BLEND;
+
+		if (bIsMaskUncompressed)
+		{	
 			constexpr int32 NumBatchElems = 4096*2;
-			const int32 NumBatches = FMath::DivideAndRoundUp(PixelCount, NumBatchElems);
+
+			const int32 NumBatches = 
+				BaseImage->DataStorage.GetNumBatchesLODRange(NumBatchElems, BaseBytesPerElem, 0, NumLODs);
+
+			check(NumBatches == BlendImage->DataStorage.GetNumBatchesLODRange(NumBatchElems, BlendBytesPerElem, 0, NumLODs));
+			check(NumBatches == MaskImage->DataStorage.GetNumBatchesLODRange(NumBatchElems, 1, 0, NumLODs));
+			check(NumBatches == DestImage->DataStorage.GetNumBatchesLODRange(NumBatchElems, BaseBytesPerElem, 0, NumLODs));
 
 			auto ProcessBatch = 
 			[
-				pBaseBuf, pBlendedBuf, pMaskBuf, pDestBuf, UnblendedChannels, PixelCount, NumBatchElems
+				BaseImage, BlendImage, MaskImage, DestImage, UnblendedChannels, BaseBytesPerElem, BlendBytesPerElem, NumLODs
 			] (uint32 BatchId)
 			{
-				const int32 BatchBegin = BatchId * NumBatchElems;
-				const int32 BatchEnd = FMath::Min(BatchBegin + NumBatchElems, PixelCount);
+				TArrayView<const uint8> BaseBatchView =
+						BaseImage->DataStorage.GetBatchLODRange(BatchId, NumBatchElems, BaseBytesPerElem, 0, NumLODs);
+
+				TArrayView<const uint8> BlendBatchView =
+						BlendImage->DataStorage.GetBatchLODRange(BatchId, NumBatchElems, BlendBytesPerElem, 0, NumLODs);
+
+				TArrayView<const uint8> MaskBatchView =
+						MaskImage->DataStorage.GetBatchLODRange(BatchId, NumBatchElems, 1, 0, NumLODs);
+
+				TArrayView<uint8> DestBatchView = 
+						DestImage->DataStorage.GetBatchLODRange(BatchId, NumBatchElems, BaseBytesPerElem, 0, NumLODs);
+
+				const uint8* BaseBuf = BaseBatchView.GetData();
+				const uint8* BlendedBuf = BlendBatchView.GetData();
+				const uint8* MaskBuf = MaskBatchView.GetData();
+				uint8* DestBuf = DestBatchView.GetData();
+
+				const int32 NumElems = BaseBatchView.Num() / BaseBytesPerElem;
+				check(NumElems == BlendBatchView.Num() / BlendBytesPerElem);
+				check(NumElems == MaskBatchView.Num() / 1);
+				check(NumElems == DestBatchView.Num() / BaseBytesPerElem);
 
-				for (int32 i = BatchBegin; i < BatchEnd; ++i)
+				for (int32 I = 0; I < NumElems; ++I)
 				{
-					unsigned mask = pMaskBuf[i];
-					for (int c = 0; c < CHANNELS_TO_BLEND; ++c)
+					uint32 Mask = MaskBuf[I];
+					for (int32 C = 0; C < CHANNELS_TO_BLEND; ++C)
 					{
-						uint32 base = pBaseBuf[BASE_CHANNEL_STRIDE * i + c];
-						uint32 blended = pBlendedBuf[BLENDED_CHANNEL_STRIDE * i + c];
-						uint32 result = BLEND_FUNC_MASKED(base, blended, mask);
+						const uint32 Base = BaseBuf[BASE_CHANNEL_STRIDE * I + C];
+						const uint32 Blended = BlendedBuf[BLENDED_CHANNEL_STRIDE * I + C];
+						const uint32 Result = BLEND_FUNC_MASKED(Base, Blended, Mask);
 						if constexpr (CLAMP)
 						{
-							pDestBuf[BASE_CHANNEL_STRIDE * i + c] = (uint8)FMath::Min(255u, result);
+							DestBuf[BASE_CHANNEL_STRIDE * I + C] = (uint8)FMath::Min(255u, Result);
 						}
 						else
 						{
-							pDestBuf[BASE_CHANNEL_STRIDE * i + c] = (uint8)result;
+							DestBuf[BASE_CHANNEL_STRIDE * I + C] = (uint8)Result;
 						}
 					}
 					// Copy the unblended channels
 					// \TODO: unnecessary when doing it in-place?
-					for (int c = 0; c < UnblendedChannels; ++c)
+					for (int32 C = 0; C < UnblendedChannels; ++C)
 					{
-						pDestBuf[BASE_CHANNEL_STRIDE * i + CHANNELS_TO_BLEND + c] = pBaseBuf[BASE_CHANNEL_STRIDE * i + CHANNELS_TO_BLEND + c];
+						DestBuf[BASE_CHANNEL_STRIDE * I + CHANNELS_TO_BLEND + C] = BaseBuf[BASE_CHANNEL_STRIDE * I + CHANNELS_TO_BLEND + C];
 					}
 				}
 			};
@@ -1248,140 +1149,144 @@ namespace mu
 				ParallelFor(NumBatches, ProcessBatch);
 			}
 		}
-		else if (maskFormat == EImageFormat::IF_L_UBYTE_RLE)
+		else if (MaskFormat == EImageFormat::IF_L_UBYTE_RLE)
 		{
-			int rows = pBase->GetSizeY();
-			int width = pBase->GetSizeX();
+			int32 Rows = BaseImage->GetSizeY();
+			int32 Width = BaseImage->GetSizeX();
 
-			for (int lod = 0; lod < LODCount; ++lod)
+			for (int32 LOD = 0; LOD < NumLODs; ++LOD)
 			{
-				pMaskBuf += 4 + rows * sizeof(uint32);
-
-				// \todo: See how to handle mask buf here
-				//ParallelFor(rows,
-				//	[
-				//		pBaseBuf, pBlendedBuf, pMaskBuf, pDestBuf, width, applyToAlpha
-				//	] (uint32 r)
-				for (int r = 0; r < rows; ++r)
+				const uint8* BaseBuf = BaseImage->GetLODData(LOD);
+				uint8* DestBuf = DestImage->GetLODData(LOD);
+				const uint8* BlendedBuf = BlendImage->GetLODData(LOD);
+				const uint8* MaskBuf = MaskImage->GetLODData(LOD);
+
+				// Remove RLE header, mip size and row sizes.
+				MaskBuf += sizeof(uint32) + Rows * sizeof(uint32);
+
+				for (int32 RowIndex = 0; RowIndex < Rows; ++RowIndex)
 				{
-					const uint8* pDestRowEnd = pDestBuf + width * BASE_CHANNEL_STRIDE;
-					while (pDestBuf != pDestRowEnd)
+					const uint8* DestRowEnd = DestBuf + Width * BASE_CHANNEL_STRIDE;
+					while (DestBuf != DestRowEnd)
 					{
 						// Decode header
-						uint16 equal = *(const uint16*)pMaskBuf;
-						pMaskBuf += 2;
+						uint16 Equal;
+						FMemory::Memcpy(&Equal, MaskBuf, 2);
+						MaskBuf += 2;
 
-						uint8 different = *pMaskBuf;
-						++pMaskBuf;
+						uint8 Different = *MaskBuf;
+						++MaskBuf;
 
-						uint8 equalPixel = *pMaskBuf;
-						++pMaskBuf;
+						uint8 EqualPixel = *MaskBuf;
+						++MaskBuf;
 
 						// Equal pixels
-						check(pDestBuf + BASE_CHANNEL_STRIDE * equal <= pStartDestBuf + pBase->GetDataSize());
-						if (equalPixel == 255)
+						//check(DestBuf + BASE_CHANNEL_STRIDE * Equal <= BaseImage->GetLODDataSize(0));
+						if (EqualPixel == 255)
 						{
-							for (int i = 0; i < equal; ++i)
+							for (int32 I = 0; I < Equal; ++I)
 							{
-								for (int c = 0; c < CHANNELS_TO_BLEND; ++c)
+								for (int32 C = 0; C < CHANNELS_TO_BLEND; ++C)
 								{
-									uint32 base = pBaseBuf[BASE_CHANNEL_STRIDE * i + c];
-									uint32 blended = pBlendedBuf[BLENDED_CHANNEL_STRIDE * i + c];
-									uint32 result = BLEND_FUNC(base, blended);
-									if (CLAMP)
+									uint32 Base = BaseBuf[BASE_CHANNEL_STRIDE * I + C];
+									uint32 Blended = BlendedBuf[BLENDED_CHANNEL_STRIDE * I + C];
+									uint32 Result = BLEND_FUNC(Base, Blended);
+									if constexpr (CLAMP)
 									{
-										pDestBuf[BASE_CHANNEL_STRIDE * i + c] = (uint8)FMath::Min(255u, result);
+										DestBuf[BASE_CHANNEL_STRIDE * I + C] = (uint8)FMath::Min(255u, Result);
 									}
 									else
 									{
-										pDestBuf[BASE_CHANNEL_STRIDE * i + c] = (uint8)result;
+										DestBuf[BASE_CHANNEL_STRIDE * I + C] = (uint8)Result;
 									}
 								}
 
 								// Copy the unblended channels
 								// \TODO: unnecessary when doing it in-place?
-								for (int32 c = 0; c < UnblendedChannels; ++c)
+								for (int32 C = 0; C < UnblendedChannels; ++C)
 								{
-									pDestBuf[BASE_CHANNEL_STRIDE * i + CHANNELS_TO_BLEND + c] = pBaseBuf[BASE_CHANNEL_STRIDE * i + CHANNELS_TO_BLEND + c];
+									DestBuf[BASE_CHANNEL_STRIDE * I + CHANNELS_TO_BLEND + C] = BaseBuf[BASE_CHANNEL_STRIDE * I + CHANNELS_TO_BLEND + C];
 								}
 							}
 						}
-						else if (equalPixel > 0)
+						else if (EqualPixel > 0)
 						{
-							for (int i = 0; i < equal; ++i)
+							for (int32 I = 0; I < Equal; ++I)
 							{
-								for (int c = 0; c < CHANNELS_TO_BLEND; ++c)
+								for (int32 C = 0; C < CHANNELS_TO_BLEND; ++C)
 								{
-									unsigned base = pBaseBuf[BASE_CHANNEL_STRIDE * i + c];
-									unsigned blended = pBlendedBuf[BLENDED_CHANNEL_STRIDE * i + c];
-									unsigned result = BLEND_FUNC_MASKED(base, blended, equalPixel);
-									if (CLAMP)
+									uint32 Base = BaseBuf[BASE_CHANNEL_STRIDE * I + C];
+									uint32 Blended = BlendedBuf[BLENDED_CHANNEL_STRIDE * I + C];
+									uint32 Result = BLEND_FUNC_MASKED(Base, Blended, EqualPixel);
+
+									if constexpr (CLAMP)
 									{
-										pDestBuf[BASE_CHANNEL_STRIDE * i + c] = (uint8)FMath::Min(255u, result);
+										DestBuf[BASE_CHANNEL_STRIDE * I + C] = (uint8)FMath::Min(255u, Result);
 									}
 									else
 									{
-										pDestBuf[BASE_CHANNEL_STRIDE * i + c] = (uint8)result;
+										DestBuf[BASE_CHANNEL_STRIDE * I + C] = (uint8)Result;
 									}
 								}
 
 								// Copy the unblended channels
 								// \TODO: unnecessary when doing it in-place?
-								for (int32 c = 0; c < UnblendedChannels; ++c)
+								for (int32 C = 0; C < UnblendedChannels; ++C)
 								{
-									pDestBuf[BASE_CHANNEL_STRIDE * i + CHANNELS_TO_BLEND + c] = pBaseBuf[BASE_CHANNEL_STRIDE * i + CHANNELS_TO_BLEND + c];
+									DestBuf[BASE_CHANNEL_STRIDE * I + CHANNELS_TO_BLEND + C] = BaseBuf[BASE_CHANNEL_STRIDE * I + CHANNELS_TO_BLEND + C];
 								}
 							}
 						}
 						else
 						{
 							// It could happen if xxxxxOnBase
-							if (pDestBuf != pBaseBuf)
+							if (DestBuf != BaseBuf)
 							{
-								FMemory::Memmove(pDestBuf, pBaseBuf, BASE_CHANNEL_STRIDE * equal);
+								FMemory::Memmove(DestBuf, BaseBuf, BASE_CHANNEL_STRIDE * Equal);
 							}
 						}
-						pDestBuf += BASE_CHANNEL_STRIDE * equal;
-						pBaseBuf += BASE_CHANNEL_STRIDE * equal;
-						pBlendedBuf += BLENDED_CHANNEL_STRIDE * equal;
+
+						DestBuf += BASE_CHANNEL_STRIDE * Equal;
+						BaseBuf += BASE_CHANNEL_STRIDE * Equal;
+						BlendedBuf += BLENDED_CHANNEL_STRIDE * Equal;
 
 						// Different pixels
-						check(pDestBuf + BASE_CHANNEL_STRIDE * different <= pStartDestBuf + pBase->GetDataSize());
-						for (int i = 0; i < different; ++i)
+						//check(DestBuf + BASE_CHANNEL_STRIDE * Different <= BaseImage->GetDataSize(0));
+						for (int32 I = 0; I < Different; ++I)
 						{
-							for (int c = 0; c < CHANNELS_TO_BLEND; ++c)
+							for (int32 C = 0; C < CHANNELS_TO_BLEND; ++C)
 							{
-								unsigned mask = pMaskBuf[i];
-								unsigned base = pBaseBuf[BASE_CHANNEL_STRIDE * i + c];
-								unsigned blended = pBlendedBuf[BLENDED_CHANNEL_STRIDE * i + c];
-								unsigned result = BLEND_FUNC_MASKED(base, blended, mask);
+								uint32 Mask = MaskBuf[I];
+								uint32 Base = BaseBuf[BASE_CHANNEL_STRIDE * I + C];
+								uint32 Blended = BlendedBuf[BLENDED_CHANNEL_STRIDE * I + C];
+								uint32 Result = BLEND_FUNC_MASKED(Base, Blended, Mask);
 								if (CLAMP)
 								{
-									pDestBuf[BASE_CHANNEL_STRIDE * i + c] = (uint8)FMath::Min(255u, result);
+									DestBuf[BASE_CHANNEL_STRIDE * I + C] = (uint8)FMath::Min(255u, Result);
 								}
 								else
 								{
-									pDestBuf[BASE_CHANNEL_STRIDE * i + c] = (uint8)result;
+									DestBuf[BASE_CHANNEL_STRIDE * I + C] = (uint8)Result;
 								}
 							}
 
 							// Copy the unblended channels
 							// \TODO: unnecessary when doing it in-place?
-							for (int32 c = 0; c < UnblendedChannels; ++c)
+							for (int32 C = 0; C < UnblendedChannels; ++C)
 							{
-								pDestBuf[BASE_CHANNEL_STRIDE * i + CHANNELS_TO_BLEND + c] = pBaseBuf[BASE_CHANNEL_STRIDE * i + CHANNELS_TO_BLEND + c];
+								DestBuf[BASE_CHANNEL_STRIDE * I + CHANNELS_TO_BLEND + C] = BaseBuf[BASE_CHANNEL_STRIDE * I + CHANNELS_TO_BLEND + C];
 							}
 						}
 
-						pDestBuf += BASE_CHANNEL_STRIDE * different;
-						pBaseBuf += BASE_CHANNEL_STRIDE * different;
-						pBlendedBuf += BLENDED_CHANNEL_STRIDE * different;
-						pMaskBuf += different;
+						DestBuf += BASE_CHANNEL_STRIDE * Different;
+						BaseBuf += BASE_CHANNEL_STRIDE * Different;
+						BlendedBuf += BLENDED_CHANNEL_STRIDE * Different;
+						MaskBuf += Different;
 					}
 				}
 
-				rows = FMath::DivideAndRoundUp(rows, 2);
-				width = FMath::DivideAndRoundUp(width, 2);
+				Rows = FMath::DivideAndRoundUp(Rows, 2);
+				Width = FMath::DivideAndRoundUp(Width, 2);
 			}
 		}
 		else
@@ -1390,9 +1295,8 @@ namespace mu
 		}
 	}
 
-	//---------------------------------------------------------------------------------------------
-	template< unsigned (*BLEND_FUNC_MASKED)(unsigned, unsigned, unsigned),
-		unsigned (*BLEND_FUNC)(unsigned, unsigned),
+	template< uint32 (*BLEND_FUNC_MASKED)(uint32, uint32, uint32),
+		uint32 (*BLEND_FUNC)(uint32, uint32),
 		bool CLAMP,
 		// Number of total channels to actually process
 		int32 CHANNELS_TO_BLEND,
@@ -1400,226 +1304,217 @@ namespace mu
 		int32 BASE_CHANNEL_STRIDE,
 		// Number of total channels in the blend image
 		int32 BLENDED_CHANNEL_STRIDE>
-	inline void BufferLayerFormatEmbeddedMask(uint8* pStartDestBuf,
-		const Image* pBase,
-		const Image* pBlend,
+	inline void BufferLayerFormatEmbeddedMask(
+		Image* DestImage,
+		const Image* BaseImage,
+		const Image* BlendImage,
+		uint32 DestOffset,
 		uint32 BaseChannelOffset,
 		bool bOnlyFirstLOD)
 	{
-		check(pBase->GetSizeX() == pBlend->GetSizeX() && pBase->GetSizeY() == pBlend->GetSizeY());
-		check(bOnlyFirstLOD || pBase->GetLODCount() <= pBlend->GetLODCount());
+		check(BaseImage->GetSizeX() == BlendImage->GetSizeX() && BaseImage->GetSizeY() == BlendImage->GetSizeY());
+		check(bOnlyFirstLOD || BaseImage->GetLODCount() <= BlendImage->GetLODCount());
 
-		// The base determines the number of lods to process.
-		int32 LODCount = bOnlyFirstLOD ? 1 : pBase->GetLODCount();
+		constexpr int32 UnblendedChannels = BASE_CHANNEL_STRIDE - CHANNELS_TO_BLEND;
 
-		int32 PixelCount = 0;
-		if (bOnlyFirstLOD)
-		{
-			PixelCount = pBase->GetSizeX() * pBase->GetSizeY();
-		}
-		else
-		{
-			PixelCount = pBase->CalculatePixelCount();
-		}
+		const int32 BytesPerElem = GetImageFormatData(BaseImage->GetFormat()).BytesPerBlock; 
+	
+		const int32 NumLODs = bOnlyFirstLOD ? 1 : BaseImage->GetLODCount();
 
-		int32 UnblendedChannels = BASE_CHANNEL_STRIDE - CHANNELS_TO_BLEND;
+		constexpr int32 NumBatchElems = 4096*2;
+		const int32 NumBatches = BaseImage->DataStorage.GetNumBatchesLODRange(NumBatchElems, BytesPerElem, 0, NumLODs);
 
-		uint8* pDestBuf = pStartDestBuf;
-		const uint8* pBaseBuf = pBase->GetData() + BaseChannelOffset;
-		const uint8* pBlendedBuf = pBlend->GetData() + BaseChannelOffset;
+		check(BytesPerElem == GetImageFormatData(DestImage->GetFormat()).BytesPerBlock);
+		check(BytesPerElem == GetImageFormatData(BlendImage->GetFormat()).BytesPerBlock);
+		
+		check(NumBatches == BlendImage->DataStorage.GetNumBatchesLODRange(NumBatchElems, BytesPerElem, 0, NumLODs));
+		check(NumBatches == DestImage->DataStorage.GetNumBatchesLODRange(NumBatchElems, BytesPerElem, 0, NumLODs));
 
-		if (pBlend->GetFormat() == EImageFormat::IF_RGBA_UBYTE)
+		auto ProcessBatch = 
+		[
+			BaseImage, BlendImage, DestImage, UnblendedChannels, NumBatchElems, BytesPerElem, NumLODs
+		] (int32 BatchId)
 		{
-			const uint8* pMaskBuf = pBlend->GetData() + 3;
+			TArrayView<const uint8> BlendedView = 
+					BlendImage->DataStorage.GetBatchLODRange(BatchId, NumBatchElems, BytesPerElem, 0, NumLODs);
 
-			constexpr int32 NumBatchElems = 4096*2;
-			const int32 NumBatches = FMath::DivideAndRoundUp(PixelCount, NumBatchElems);
+			TArrayView<const uint8> BaseView =
+					BaseImage->DataStorage.GetBatchLODRange(BatchId, NumBatchElems, BytesPerElem, 0, NumLODs);
 
-			auto ProcessBatch = 
-			[
-				pBaseBuf, pBlendedBuf, pDestBuf, pMaskBuf, UnblendedChannels, PixelCount, NumBatchElems
-			] (int32 BatchId)
+			TArrayView<uint8> DestView =
+					DestImage->DataStorage.GetBatchLODRange(BatchId, NumBatchElems, BytesPerElem, 0, NumLODs);
+
+			uint8* const DestBuf = DestView.GetData();
+			uint8 const * const BlendedBuf = BlendedView.GetData();
+			uint8 const * const BaseBuf = BaseView.GetData(); 
+
+			const int32 NumElems = BaseView.Num() / BytesPerElem;
+			check(NumElems == DestView.Num() / BytesPerElem);
+			check(NumElems == BlendedView.Num() / BytesPerElem);
+
+			if (BlendImage->GetFormat() == EImageFormat::IF_RGBA_UBYTE)
 			{
-				const int32 BatchBegin = BatchId * NumBatchElems;
-				const int32 BatchEnd = FMath::Min(BatchBegin + NumBatchElems, PixelCount);
+				const uint8* MaskBuf = BlendedBuf + 3;
 
-				for (int32 i = BatchBegin; i < BatchEnd; ++i)
+				for (int32 I = 0; I < NumElems; ++I)
 				{
-					uint32 mask = pMaskBuf[BLENDED_CHANNEL_STRIDE * i];
-					for (int c = 0; c < CHANNELS_TO_BLEND; ++c)
+					uint32 Mask = MaskBuf[BLENDED_CHANNEL_STRIDE * I];
+					for (int32 C = 0; C < CHANNELS_TO_BLEND; ++C)
 					{
-						uint32 base = pBaseBuf[BASE_CHANNEL_STRIDE * i + c];
-						uint32 blended = pBlendedBuf[BLENDED_CHANNEL_STRIDE * i + c];
-						uint32 result = BLEND_FUNC_MASKED(base, blended, mask);
+						uint32 Base = BaseBuf[BASE_CHANNEL_STRIDE * I + C];
+						uint32 Blended = BlendedBuf[BLENDED_CHANNEL_STRIDE * I + C];
+						uint32 Result = BLEND_FUNC_MASKED(Base, Blended, Mask);
 						if constexpr (CLAMP)
 						{
-							pDestBuf[BASE_CHANNEL_STRIDE * i + c] = (uint8)FMath::Min(255u, result);
+							DestBuf[BASE_CHANNEL_STRIDE * I + C] = (uint8)FMath::Min(255u, Result);
 						}
 						else
 						{
-							pDestBuf[BASE_CHANNEL_STRIDE * i + c] = (uint8)result;
+							DestBuf[BASE_CHANNEL_STRIDE * I + C] = (uint8)Result;
 						}
 					}
 					// Copy the unblended channels
 					// \TODO: unnecessary when doing it in-place?
-					for (int c = 0; c < UnblendedChannels; ++c)
+					for (int32 C = 0; C < UnblendedChannels; ++C)
 					{
-						pDestBuf[BASE_CHANNEL_STRIDE * i + CHANNELS_TO_BLEND + c] = pBaseBuf[BASE_CHANNEL_STRIDE * i + CHANNELS_TO_BLEND + c];
+						DestBuf[BASE_CHANNEL_STRIDE * I + CHANNELS_TO_BLEND + C] = BaseBuf[BASE_CHANNEL_STRIDE * I + CHANNELS_TO_BLEND + C];
 					}
 				}
-			};
-
-			if (NumBatches == 1)
-			{
-				ProcessBatch(0);
-			}
-			else if (NumBatches > 1)
-			{
-				ParallelFor(NumBatches, ProcessBatch);
 			}
-		}
-		else
-		{
-			constexpr int32 NumBatchElems = 4096*2;
-			const int32 NumBatches = FMath::DivideAndRoundUp(PixelCount, NumBatchElems);
-
-			auto ProcessBatch = 
-			[
-				pBaseBuf, pBlendedBuf, pDestBuf, UnblendedChannels, PixelCount, NumBatchElems
-			] (uint32 BatchId)
+			else
 			{
-				const int32 BatchBegin = BatchId * NumBatchElems;
-				const int32 BatchEnd = FMath::Min(BatchBegin + NumBatchElems, PixelCount);
-
-				for (int32 i = BatchBegin; i < BatchEnd; ++i)
+				for (int32 I = 0; I < NumElems; ++I)
 				{
-					for (int c = 0; c < CHANNELS_TO_BLEND; ++c)
+					for (int32 C = 0; C < CHANNELS_TO_BLEND; ++C)
 					{
-						uint32 base = pBaseBuf[BASE_CHANNEL_STRIDE * i + c];
-						uint32 blended = pBlendedBuf[BLENDED_CHANNEL_STRIDE * i + c];
-						uint32 result = BLEND_FUNC(base, blended);
+						uint32 Base = BaseBuf[BASE_CHANNEL_STRIDE * I + C];
+						uint32 Blended = BlendedBuf[BLENDED_CHANNEL_STRIDE * I + C];
+						uint32 Result = BLEND_FUNC(Base, Blended);
 						if constexpr (CLAMP)
 						{
-							pDestBuf[BASE_CHANNEL_STRIDE * i + c] = (uint8)FMath::Min(255u, result);
+							DestBuf[BASE_CHANNEL_STRIDE * I + C] = (uint8)FMath::Min(255u, Result);
 						}
 						else
 						{
-							pDestBuf[BASE_CHANNEL_STRIDE * i + c] = (uint8)result;
+							DestBuf[BASE_CHANNEL_STRIDE * I + C] = (uint8)Result;
 						}
 					}
 					// Copy the unblended channels
 					// \TODO: unnecessary when doing it in-place?
-					for (int c = 0; c < UnblendedChannels; ++c)
+					for (int32 C = 0; C < UnblendedChannels; ++C)
 					{
-						pDestBuf[BASE_CHANNEL_STRIDE * i + CHANNELS_TO_BLEND + c] = pBaseBuf[BASE_CHANNEL_STRIDE * i + CHANNELS_TO_BLEND + c];
+						DestBuf[BASE_CHANNEL_STRIDE * I + CHANNELS_TO_BLEND + C] = BaseBuf[BASE_CHANNEL_STRIDE * I + CHANNELS_TO_BLEND + C];
 					}
 				}
-			};
-
-			if (NumBatches == 1)
-			{
-				ProcessBatch(0);
-			}
-			else if (NumBatches > 1)
-			{
-				ParallelFor(NumBatches, ProcessBatch);
 			}
+		};
+
+		if (NumBatches == 1)
+		{
+			ProcessBatch(0);
+		}
+		else if (NumBatches > 1)
+		{
+			ParallelFor(NumBatches, ProcessBatch);
 		}
 	}
 
-
-	//---------------------------------------------------------------------------------------------
-	template< unsigned (*BLEND_FUNC)(unsigned,unsigned),
-			  bool CLAMP,
-			  int NC >
-    inline void BufferLayerFormatStrideNoAlpha( uint8* pDestBuf,
-		int stride,
-		const Image* pMask,
-		const Image* pBlend, int32 LODCount )
+	template< 
+		uint32 (*BLEND_FUNC)(uint32, uint32),
+		bool CLAMP,
+		uint32 NC>
+    inline void BufferLayerFormatStrideNoAlpha(
+		Image* DestImage,
+		int32 DestOffset,
+		int32 Stride,
+		const Image* MaskImage,
+		const Image* BlendImage/*, int32 LODCount*/)
 	{
-        const uint8* pMaskBuf = pMask->GetData();
-        const uint8* pBlendedBuf = pBlend->GetData();
+        const uint8* MaskBuf = MaskImage->GetLODData(0);
+        const uint8* BlendedBuf = BlendImage->GetLODData(0);
+		uint8* DestBuf = DestImage->GetLODData(0) + DestOffset;
 
-		EImageFormat maskFormat = pMask->GetFormat();
-        bool isUncompressed = ( maskFormat == EImageFormat::IF_L_UBYTE );
+		EImageFormat MaskFormat = MaskImage->GetFormat();
+        bool bIsUncompressed = (MaskFormat == EImageFormat::IF_L_UBYTE);
 
-        if ( isUncompressed )
+        if (bIsUncompressed)
 		{
-			int rowCount = pBlend->GetSizeY();
-			int pixelCount = pBlend->GetSizeX();
-			for ( int j=0; j<rowCount; ++j )
+			int32 RowCount = BlendImage->GetSizeY();
+			int32 PixelCount = BlendImage->GetSizeX();
+			for (int32 RowIndex = 0; RowIndex < RowCount; ++RowIndex)
 			{
-				for ( int i=0; i<pixelCount; ++i )
+				for (int32 PixelIndex = 0; PixelIndex < PixelCount; ++PixelIndex)
 				{
-					unsigned mask = *pMaskBuf;
-					if (mask)
+					uint32 Mask = *MaskBuf;
+					if (Mask)
 					{
-						for ( int c=0; c<NC; ++c )
+						for (int32 C = 0; C < NC; ++C)
 						{
-							unsigned base = *pDestBuf;
-							unsigned blended = *pBlendedBuf;
-							unsigned result = BLEND_FUNC( base, blended );
-							if ( CLAMP )
+							uint32 Base = *DestBuf;
+							uint32 Blended = *BlendedBuf;
+							uint32 Result = BLEND_FUNC(Base, Blended);
+							if constexpr (CLAMP)
 							{
-								*pDestBuf = (uint8)FMath::Min( 255u, result );
+								*DestBuf = (uint8)FMath::Min(255u, Result);
 							}
 							else
 							{
-								*pDestBuf = (uint8)result;
+								*DestBuf = (uint8)Result;
 							}
-							++pDestBuf;
-							++pBlendedBuf;
+							++DestBuf;
+							++BlendedBuf;
 						}
 					}
 					else
 					{
-						pDestBuf+=NC;
-						pBlendedBuf+=NC;
+						DestBuf += NC;
+						BlendedBuf += NC;
 					}
-					++pMaskBuf;
+					++MaskBuf;
 				}
 
-				pDestBuf+=stride;
+				DestBuf += Stride;
 			}
 		}
-        else if ( maskFormat==EImageFormat::IF_L_UBIT_RLE )
+        else if (MaskFormat == EImageFormat::IF_L_UBIT_RLE)
 		{
-			int rows = pMask->GetSizeY();
-			int width = pMask->GetSizeX();
+			int32 Rows = MaskImage->GetSizeY();
+			int32 Width = MaskImage->GetSizeX();
 
-            for (int lod=0;lod< LODCount; ++lod)
-            {
-                pMaskBuf += 4+rows*sizeof(uint32);
+            //for (int32 lod = 0; lod < LODCount; ++lod)
+            //{
+				// Remove RLE header.
+                MaskBuf += 4 + Rows*sizeof(uint32);
 
-                for ( int r=0; r<rows; ++r )
+                for (int32 RowIndex = 0; RowIndex < Rows; ++RowIndex)
                 {
-                    const uint8* pDestRowEnd = pDestBuf + width*NC;
-                    while ( pDestBuf!=pDestRowEnd )
+                    const uint8* DestRowEnd = DestBuf + Width*NC;
+                    while (DestBuf != DestRowEnd)
                     {
                         // Decode header
-                        uint16 zeros = *(const uint16*)pMaskBuf;
-                        pMaskBuf += 2;
+                        uint16 Zeros = *(const uint16*)MaskBuf;
+                        MaskBuf += 2;
 
-                        uint16 ones = *(const uint16*)pMaskBuf;
-                        pMaskBuf += 2;
+                        uint16 Ones = *(const uint16*)MaskBuf;
+                        MaskBuf += 2;
 
                         // Skip
-                        pDestBuf += zeros*NC;
-                        pBlendedBuf += zeros*NC;
+                        DestBuf += Zeros*NC;
+                        BlendedBuf += Zeros*NC;
 
                         // Copy
-                        FMemory::Memmove( pDestBuf, pBlendedBuf, ones*NC );
+                        FMemory::Memmove(DestBuf, BlendedBuf, Ones*NC);
 
-                        pDestBuf += NC*ones;
-                        pBlendedBuf += NC*ones;
+                        DestBuf += NC*Ones;
+                        BlendedBuf += NC*Ones;
                     }
 
-                    pDestBuf += stride;
+                    DestBuf += Stride;
                 }
 
-                rows = FMath::DivideAndRoundUp(rows,2);
-                width = FMath::DivideAndRoundUp(width,2);
-            }
+                //Rows = FMath::DivideAndRoundUp(Rows, 2);
+                //Width = FMath::DivideAndRoundUp(Width, 2);
+            //}
 		}
 		else
 		{
@@ -1627,74 +1522,74 @@ namespace mu
 		}
 	}
 
-
-	//---------------------------------------------------------------------------------------------
-	template< unsigned (*BLEND_FUNC_MASKED)(unsigned,unsigned,unsigned),
-			  unsigned (*BLEND_FUNC)(unsigned,unsigned),
-			  bool CLAMP >
-    inline void BufferLayer( uint8* pDestBuf,
-							 const Image* pBase,
-							 const Image* pMask,
-                             const Image* pBlend,
-                             bool bApplyToAlpha,
-							 bool bOnlyFirstLOD )
+	template<uint32(*BLEND_FUNC_MASKED)(uint32, uint32, uint32),
+			 uint32(*BLEND_FUNC)(uint32, uint32),
+			 bool CLAMP>
+    inline void BufferLayer(Image* DestImage,
+							const Image* BaseImage,
+							const Image* MaskImage,
+                            const Image* BlendImage,
+                            bool bApplyToAlpha,
+							bool bOnlyFirstLOD)
 	{
-		if ( pBase->GetFormat()==EImageFormat::IF_RGB_UBYTE )
+		if (BaseImage->GetFormat() == EImageFormat::IF_RGB_UBYTE)
 		{
-            BufferLayerFormat< BLEND_FUNC_MASKED, BLEND_FUNC, CLAMP, 3, 3, 3 >
-                    ( pDestBuf, pBase, pMask, pBlend, 0, 0, bOnlyFirstLOD) ;
+            BufferLayerFormat<BLEND_FUNC_MASKED, BLEND_FUNC, CLAMP, 3, 3, 3>
+                    (DestImage, BaseImage, MaskImage, BlendImage, 0, 0, 0, bOnlyFirstLOD);
 		}
-        else if ( pBase->GetFormat()==EImageFormat::IF_RGBA_UBYTE || pBase->GetFormat()==EImageFormat::IF_BGRA_UBYTE )
+        else if (BaseImage->GetFormat() == EImageFormat::IF_RGBA_UBYTE || 
+				 BaseImage->GetFormat() == EImageFormat::IF_BGRA_UBYTE)
 		{
 			if (bApplyToAlpha)
 			{
-				BufferLayerFormat< BLEND_FUNC_MASKED, BLEND_FUNC, CLAMP, 4, 4, 4 >
-					(pDestBuf, pBase, pMask, pBlend, 0, 0, bOnlyFirstLOD);
+				BufferLayerFormat<BLEND_FUNC_MASKED, BLEND_FUNC, CLAMP, 4, 4, 4>
+					(DestImage, BaseImage, MaskImage, BlendImage, 0, 0, 0, bOnlyFirstLOD);
 			}
 			else
 			{
-				BufferLayerFormat< BLEND_FUNC_MASKED, BLEND_FUNC, CLAMP, 3, 4, 4 >
-					(pDestBuf, pBase, pMask, pBlend, 0, 0, bOnlyFirstLOD);
+				BufferLayerFormat<BLEND_FUNC_MASKED, BLEND_FUNC, CLAMP, 3, 4, 4>
+					(DestImage, BaseImage, MaskImage, BlendImage, 0, 0, 0, bOnlyFirstLOD);
 			}
 		}
-		else if ( pBase->GetFormat()==EImageFormat::IF_L_UBYTE )
+		else if (BaseImage->GetFormat() == EImageFormat::IF_L_UBYTE)
 		{
-            BufferLayerFormat< BLEND_FUNC_MASKED, BLEND_FUNC, CLAMP, 1, 1, 1 >
-                    ( pDestBuf, pBase, pMask, pBlend, 0, 0, bOnlyFirstLOD) ;
+            BufferLayerFormat<BLEND_FUNC_MASKED, BLEND_FUNC, CLAMP, 1, 1, 1>
+                    (DestImage, BaseImage, MaskImage, BlendImage, 0, 0, 0, bOnlyFirstLOD);
 		}
         else
 		{
-			checkf( false, TEXT("Unsupported format.") );
+			checkf(false, TEXT("Unsupported format."));
 		}
 	}
 
 
-	//---------------------------------------------------------------------------------------------
-	template< unsigned (*BLEND_FUNC_MASKED)(unsigned, unsigned, unsigned),
-		unsigned (*BLEND_FUNC)(unsigned, unsigned),
+	template< uint32 (*BLEND_FUNC_MASKED)(uint32, uint32, uint32),
+		uint32 (*BLEND_FUNC)(uint32, uint32),
 		bool CLAMP >
-	inline void BufferLayerEmbeddedMask(uint8* pDestBuf,
-		const Image* pBase,
-		const Image* pBlend,
+	inline void BufferLayerEmbeddedMask(
+		Image* DestImage,
+		const Image* BaseImage,
+		const Image* BlendImage,
 		bool bApplyToAlpha,
 		bool bOnlyFirstLOD)
 	{
-		if (pBase->GetFormat() == EImageFormat::IF_RGB_UBYTE)
+		if (BaseImage->GetFormat() == EImageFormat::IF_RGB_UBYTE)
 		{
-			BufferLayerFormatEmbeddedMask< BLEND_FUNC_MASKED, BLEND_FUNC, CLAMP, 3, 3, 3 >
-				(pDestBuf, pBase, pBlend, 0, bOnlyFirstLOD);
+			BufferLayerFormatEmbeddedMask<BLEND_FUNC_MASKED, BLEND_FUNC, CLAMP, 3, 3, 3>
+				(DestImage, BaseImage, BlendImage, 0, 0, bOnlyFirstLOD);
 		}
-		else if (pBase->GetFormat() == EImageFormat::IF_RGBA_UBYTE || pBase->GetFormat() == EImageFormat::IF_BGRA_UBYTE)
+		else if (BaseImage->GetFormat() == EImageFormat::IF_RGBA_UBYTE || 
+				 BaseImage->GetFormat() == EImageFormat::IF_BGRA_UBYTE)
 		{
 			if (bApplyToAlpha)
 			{
-				BufferLayerFormatEmbeddedMask< BLEND_FUNC_MASKED, BLEND_FUNC, CLAMP, 4, 4, 4 >
-					(pDestBuf, pBase, pBlend, 0, bOnlyFirstLOD);
+				BufferLayerFormatEmbeddedMask<BLEND_FUNC_MASKED, BLEND_FUNC, CLAMP, 4, 4, 4>
+					(DestImage, BaseImage, BlendImage, 0, 0, bOnlyFirstLOD);
 			}
 			else
 			{
-				BufferLayerFormatEmbeddedMask< BLEND_FUNC_MASKED, BLEND_FUNC, CLAMP, 3, 4, 4 >
-					(pDestBuf, pBase, pBlend, 0, bOnlyFirstLOD);
+				BufferLayerFormatEmbeddedMask<BLEND_FUNC_MASKED, BLEND_FUNC, CLAMP, 3, 4, 4>
+					(DestImage, BaseImage, BlendImage, 0, 0, bOnlyFirstLOD);
 			}
 		}
 		else
@@ -1704,251 +1599,287 @@ namespace mu
 	}
 
 
-	//---------------------------------------------------------------------------------------------
-	template< unsigned (*RGB_FUNC_MASKED)(unsigned, unsigned, unsigned),
-		unsigned (*A_FUNC)(unsigned, unsigned),
+	template< 
+		uint32 (*RGB_FUNC_MASKED)(uint32, uint32, uint32),
+		uint32 (*A_FUNC)(uint32, uint32),
 		bool CLAMP >
 	inline void BufferLayerComposite(
-		Image* pBase,
-		const Image* pBlend,
+		Image* BaseImage,
+		const Image* BlendImage,
 		bool bOnlyFirstLOD,
-		uint8 BlendAlphaSourceChannel )
+		uint8 BlendAlphaSourceChannel)
 	{
-		check(pBase->GetFormat() == EImageFormat::IF_RGBA_UBYTE);
-		check(pBlend->GetFormat() == EImageFormat::IF_RGBA_UBYTE);
-		check(pBase->GetSizeX() == pBlend->GetSizeX() && pBase->GetSizeY() == pBlend->GetSizeY());
-		check(bOnlyFirstLOD || pBase->GetLODCount() <= pBlend->GetLODCount());
-
-		uint8* pBaseBuf = pBase->GetData();
-		const uint8* pBlendedBuf = pBlend->GetData();
+		check(BaseImage->GetFormat() == EImageFormat::IF_RGBA_UBYTE);
+		check(BlendImage->GetFormat() == EImageFormat::IF_RGBA_UBYTE);
+		check(BaseImage->GetSizeX() == BlendImage->GetSizeX() && BaseImage->GetSizeY() == BlendImage->GetSizeY());
+		check(bOnlyFirstLOD || BaseImage->GetLODCount() <= BlendImage->GetLODCount());
 
-		int32 PixelCount = 0;
-		uint16 SizeX = pBase->GetSizeX();
+		int32 FirstLODDataOffset = 0;	
+		int32 NumRelevantElems = -1;
 
-		if (pBlend->m_flags & Image::IF_HAS_RELEVANCY_MAP 
-			&& 
-			(bOnlyFirstLOD || pBase->GetLODCount()==1 ) )
+		if (BlendImage->m_flags & Image::IF_HAS_RELEVANCY_MAP && (bOnlyFirstLOD || BaseImage->GetLODCount() == 1))
 		{
-			check(pBlend->RelevancyMaxY < pBase->GetSizeY());
-			check(pBlend->RelevancyMaxY >= pBlend->RelevancyMinY);
+			check(BlendImage->RelevancyMaxY < BaseImage->GetSizeY());
+			check(BlendImage->RelevancyMaxY >= BlendImage->RelevancyMinY);
 
-			int32 NumRelevantRows = pBlend->RelevancyMaxY - pBlend->RelevancyMinY + 1;
-			PixelCount = SizeX * NumRelevantRows;
-			int32 DataOffset = pBlend->RelevancyMinY * SizeX * 4;
-			pBaseBuf += DataOffset;
-			pBlendedBuf += DataOffset;
-		}
-		else
-		{ 
-			if (bOnlyFirstLOD)
-			{
-				PixelCount = SizeX * pBase->GetSizeY();
-			}
-			else
-			{
-				PixelCount = pBase->CalculatePixelCount();
-			}
+			uint16 SizeX = BaseImage->GetSizeX();
+			NumRelevantElems = (BlendImage->RelevancyMaxY - BlendImage->RelevancyMinY + 1) * SizeX * 4;
+			
+			FirstLODDataOffset = BlendImage->RelevancyMinY * SizeX * 4;
 		}
 
-		// \TODO: Ensure batches operate on 32Kb aligned buffers?
+		const int32 BytesPerElem = GetImageFormatData(BaseImage->GetFormat()).BytesPerBlock;
+		check(BytesPerElem == GetImageFormatData(BlendImage->GetFormat()).BytesPerBlock);
+
+		const int32 LODBegin = 0;
+		const int32 LODEnd = BaseImage->GetLODCount();
+
 		constexpr int32 NumBatchElems = 4096*2;
-		const int32 NumBatches = FMath::DivideAndRoundUp(PixelCount, NumBatchElems);
+
+		const int32 NumBatches = bOnlyFirstLOD
+				? BaseImage->DataStorage.GetNumBatchesFirstLODOffset(NumBatchElems, BytesPerElem, FirstLODDataOffset)
+				: BaseImage->DataStorage.GetNumBatchesLODRange(NumBatchElems, BytesPerElem, LODBegin, LODEnd);
+
+		// This will always be an upper-bound for bOnlyFirtsLODs, check if it performs as expected or it needs more fine tune.
+		const int32 NumRelevantBatches = bOnlyFirstLOD && NumRelevantElems != -1
+				? FMath::DivideAndRoundUp(NumRelevantElems, NumBatchElems)
+				: NumBatches;
+
+		// This check only applies if bOnlyFirstLOD is false. The other case could give false negatives
+		// but that is already checked above.
+		check(NumBatches <= BlendImage->DataStorage.GetNumBatches(NumBatchElems, BytesPerElem));
 
 		auto ProcessBatch = 
 		[
-			pBaseBuf, pBlendedBuf, BlendAlphaSourceChannel, PixelCount, NumBatchElems
-		] (uint32 BatchId)
+			BaseImage, BlendImage, BlendAlphaSourceChannel, BytesPerElem, NumBatchElems, bOnlyFirstLOD, LODBegin, LODEnd, FirstLODDataOffset
+		] (int32 BatchId)
 		{
-			const int32 BatchBegin = BatchId * NumBatchElems;
-			const int32 BatchEnd = FMath::Min(BatchBegin + NumBatchElems, PixelCount);
+			TArrayView<uint8> BaseBatchView = bOnlyFirstLOD 
+					? BaseImage->DataStorage.GetBatchFirstLODOffset(BatchId, NumBatchElems, BytesPerElem, FirstLODDataOffset)
+					: BaseImage->DataStorage.GetBatchLODRange(BatchId, NumBatchElems, BytesPerElem, LODBegin, LODEnd);
+
+			TArrayView<const uint8> BlendBatchView = bOnlyFirstLOD 
+					? BlendImage->DataStorage.GetBatchFirstLODOffset(BatchId, NumBatchElems, BytesPerElem, FirstLODDataOffset)
+					: BlendImage->DataStorage.GetBatchLODRange(BatchId, NumBatchElems, BytesPerElem, LODBegin, LODEnd);
 
-			for (int32 i = BatchBegin; i < BatchEnd; ++i)
+			const int32 NumElems = BaseBatchView.Num() / BytesPerElem; 
+			check(BlendBatchView.Num() / BytesPerElem == NumElems);
+
+			uint8* BaseBuf = BaseBatchView.GetData();
+			const uint8* BlendBuf = BlendBatchView.GetData();
+
+			for (int32 I = 0; I < NumElems; ++I)
 			{
 				// TODO: Optimize this (SIMD?)
-				uint32 mask = pBlendedBuf[4 * i + 3];
+				uint32 Mask = BlendBuf[4 * I + 3];
 
 				// RGB
-				for (int c = 0; c < 3; ++c)
+				for (int32 C = 0; C < 3; ++C)
 				{
-					uint32 base = pBaseBuf[4 * i + c];
-					uint32 blended = pBlendedBuf[4 * i + c];
-					uint32 result = RGB_FUNC_MASKED(base, blended, mask);
+					uint32 Base = BaseBuf[4 * I + C];
+					uint32 Blended = BlendBuf[4 * I + C];
+					uint32 Result = RGB_FUNC_MASKED(Base, Blended, Mask);
 					if constexpr (CLAMP)
 					{
-						pBaseBuf[4 * i + c] = (uint8)FMath::Min(255u, result);
+						BaseBuf[4 * I + C] = (uint8)FMath::Min(255u, Result);
 					}
 					else
 					{
-						pBaseBuf[4 * i + c] = (uint8)result;
+						BaseBuf[4 * I + C] = (uint8)Result;
 					}
 				}
 
 				// A
 				{
-					uint32 base = pBaseBuf[4 * i + 3];
-					uint32 blended = pBlendedBuf[4 * i + BlendAlphaSourceChannel];
-					uint32 result = A_FUNC(base, blended);
+					uint32 Base = BaseBuf[4 * I + 3];
+					uint32 Blended = BlendBuf[4 * I + BlendAlphaSourceChannel];
+					uint32 Result = A_FUNC(Base, Blended);
 					if constexpr (CLAMP)
 					{
-						pBaseBuf[4 * i + 3] = (uint8)FMath::Min(255u, result);
+						BaseBuf[4 * I + 3] = (uint8)FMath::Min(255u, Result);
 					}
 					else
 					{
-						pBaseBuf[4 * i + 3] = (uint8)result;
+						BaseBuf[4 * I + 3] = (uint8)Result;
 					}
 				}
 			}
 		};
 
-		if (NumBatches == 1)
+		if (NumRelevantBatches == 1)
 		{
 			ProcessBatch(0);
 		}
-		else if (NumBatches > 1)
+		else if (NumRelevantBatches > 1)
 		{
-			ParallelFor(NumBatches, ProcessBatch);
+			ParallelFor(NumRelevantBatches, ProcessBatch);
 		}
 	}
 
-
-
-	//---------------------------------------------------------------------------------------------
 	template<>
 	inline void BufferLayerComposite<BlendChannelMasked, LightenChannel, false>
 	(
-		Image* pBase,
-		const Image* pBlend,
+		Image* BaseImage,
+		const Image* BlendImage,
 		bool bOnlyFirstLOD,
 		uint8 BlendAlphaSourceChannel)
 	{
-		check(pBase->GetFormat() == EImageFormat::IF_RGBA_UBYTE);
-		check(pBlend->GetFormat() == EImageFormat::IF_RGBA_UBYTE);
-		check(pBase->GetSizeX() == pBlend->GetSizeX() && pBase->GetSizeY() == pBlend->GetSizeY());
-		check(bOnlyFirstLOD || pBase->GetLODCount() <= pBlend->GetLODCount());
-		check(BlendAlphaSourceChannel==3);
-
-		uint32* pBaseBuf = reinterpret_cast<uint32*>(pBase->GetData());
-		const uint32* pBlendedBuf = reinterpret_cast<const uint32*>(pBlend->GetData());
-
-		int32 PixelCount = 0;
-		uint16 SizeX = pBase->GetSizeX();
-
-		if (pBlend->m_flags & Image::IF_HAS_RELEVANCY_MAP
-			&&
-			(bOnlyFirstLOD || pBase->GetLODCount() == 1))
+		check(BaseImage->GetFormat() == EImageFormat::IF_RGBA_UBYTE);
+		check(BlendImage->GetFormat() == EImageFormat::IF_RGBA_UBYTE);
+		check(BaseImage->GetSizeX() == BlendImage->GetSizeX() && BaseImage->GetSizeY() == BlendImage->GetSizeY());
+		check(bOnlyFirstLOD || BaseImage->GetLODCount() <= BlendImage->GetLODCount());
+		check(BlendAlphaSourceChannel == 3);
+
+		int32 FirstLODDataOffset = 0;
+		int32 NumRelevantElems = -1;
+		if (BlendImage->m_flags & Image::IF_HAS_RELEVANCY_MAP 
+			&& 
+			(bOnlyFirstLOD || BaseImage->GetLODCount()==1 ) )
 		{
-			check(pBlend->RelevancyMaxY < pBase->GetSizeY());
-			check(pBlend->RelevancyMaxY >= pBlend->RelevancyMinY);
+			check(BlendImage->RelevancyMaxY < BaseImage->GetSizeY());
+			check(BlendImage->RelevancyMaxY >= BlendImage->RelevancyMinY);
 
-			int32 NumRelevantRows = pBlend->RelevancyMaxY - pBlend->RelevancyMinY + 1;
-			PixelCount = SizeX * NumRelevantRows;
-			int32 DataOffset = pBlend->RelevancyMinY * SizeX;
-			pBaseBuf += DataOffset;
-			pBlendedBuf += DataOffset;
-		}
-		else
-		{
-			if (bOnlyFirstLOD)
-			{
-				PixelCount = SizeX * pBase->GetSizeY();
-			}
-			else
-			{
-				PixelCount = pBase->CalculatePixelCount();
-			}
+			uint16 SizeX = BaseImage->GetSizeX();
+			NumRelevantElems = (BlendImage->RelevancyMaxY - BlendImage->RelevancyMinY + 1) * SizeX;
+			
+			FirstLODDataOffset = BlendImage->RelevancyMinY * SizeX * 4;
 		}
 
 		// \TODO: Ensure batches operate on 32Kb aligned buffers?
 		constexpr int32 NumBatchElems = 4096 * 2;
-		const int32 NumBatches = FMath::DivideAndRoundUp(PixelCount, NumBatchElems);
 
+		constexpr int32 BytesPerElem = 4;
+
+		const int32 NumLODs = BaseImage->GetLODCount();
+		const int32 NumBatches = bOnlyFirstLOD
+				? BaseImage->DataStorage.GetNumBatchesFirstLODOffset(NumBatchElems, BytesPerElem, FirstLODDataOffset)
+				: BaseImage->DataStorage.GetNumBatchesLODRange(NumBatchElems, BytesPerElem, 0, NumLODs);
+
+		// This will always be an upper-bound for bOnlyFirtsLODs, check if it performs as expected or it needs more fine tune.
+		const int32 NumRelevantBatches = bOnlyFirstLOD && NumRelevantElems != -1
+				? FMath::DivideAndRoundUp(NumRelevantElems, NumBatchElems)
+				: NumBatches;
+
+		// This check only applies if bOnlyFirstLOD is false. The other case could give false negatives
+		// but that is already checked above.
+		check(NumBatches <= BlendImage->DataStorage.GetNumBatchesLODRange(NumBatchElems, BytesPerElem, 0, NumLODs));
+		
 		auto ProcessBatch =
 			[
-				pBaseBuf, pBlendedBuf, PixelCount, NumBatchElems
+				BaseImage, BlendImage, NumBatchElems, bOnlyFirstLOD, FirstLODDataOffset, NumLODs
 			] (uint32 BatchId)
 		{
-			const int32 BatchBegin = BatchId * NumBatchElems;
-			const int32 BatchEnd = FMath::Min(BatchBegin + NumBatchElems, PixelCount);
+			TArrayView<uint8> BaseBatchView = bOnlyFirstLOD 
+					? BaseImage->DataStorage.GetBatchFirstLODOffset(BatchId, NumBatchElems, BytesPerElem, FirstLODDataOffset)
+					: BaseImage->DataStorage.GetBatchLODRange(BatchId, NumBatchElems, BytesPerElem, 0, NumLODs);
 
-			for (int32 i = BatchBegin; i < BatchEnd; ++i)
+			TArrayView<const uint8> BlendBatchView = bOnlyFirstLOD 
+					? BlendImage->DataStorage.GetBatchFirstLODOffset(BatchId, NumBatchElems, BytesPerElem, FirstLODDataOffset)
+					: BlendImage->DataStorage.GetBatchLODRange(BatchId, NumBatchElems, BytesPerElem, 0, NumLODs);
+
+			const int32 NumElems = BaseBatchView.Num() / BytesPerElem; 
+			check(NumElems == BlendBatchView.Num() / BytesPerElem);
+			
+			uint8* BaseBuf = BaseBatchView.GetData();
+			const uint8* BlendBuf = BlendBatchView.GetData();
+			
+			for (int32 I = 0; I < NumElems; ++I)
 			{
 				// TODO: Optimize this (SIMD?)
-				uint32 FullBase = pBaseBuf[i];
-				uint32 FullBlended = pBlendedBuf[i];
-				uint32 mask = (FullBlended & 0xff000000) >> 24;
+				uint32 FullBase;
+				FMemory::Memcpy(&FullBase, BaseBuf + I*sizeof(uint32), sizeof(uint32)); 
+
+				uint32 FullBlended;
+				FMemory::Memcpy(&FullBlended, BlendBuf + I*sizeof(uint32), sizeof(uint32)); 
+				uint32 Mask = (FullBlended & 0xff000000) >> 24;
 
 				uint32 FullResult = 0;
-				FullResult |= BlendChannelMasked((FullBase >>  0) & 0xff, (FullBlended >>  0) & 0xff, mask) << 0;
-				FullResult |= BlendChannelMasked((FullBase >>  8) & 0xff, (FullBlended >>  8) & 0xff, mask) << 8;
-				FullResult |= BlendChannelMasked((FullBase >> 16) & 0xff, (FullBlended >> 16) & 0xff, mask) << 16;
+				FullResult |= BlendChannelMasked((FullBase >>  0) & 0xff, (FullBlended >>  0) & 0xff, Mask) << 0;
+				FullResult |= BlendChannelMasked((FullBase >>  8) & 0xff, (FullBlended >>  8) & 0xff, Mask) << 8;
+				FullResult |= BlendChannelMasked((FullBase >> 16) & 0xff, (FullBlended >> 16) & 0xff, Mask) << 16;
 				FullResult |= LightenChannel	((FullBase >> 24) & 0xff, (FullBlended >> 24) & 0xff) << 24;
 
-				pBaseBuf[i] = FullResult;
+				FMemory::Memcpy(BaseBuf + I*sizeof(uint32), &FullResult, sizeof(uint32));
 			}
 		};
 
-		if (NumBatches == 1)
+		if (NumRelevantBatches == 1)
 		{
 			ProcessBatch(0);
 		}
-		else if (NumBatches > 1)
+		else if (NumRelevantBatches > 1)
 		{
 			ParallelFor(NumBatches, ProcessBatch);
 		}
 	}
 
 
-	//---------------------------------------------------------------------------------------------
-	template< VectorRegister4Int (*RGB_FUNC_MASKED)(const VectorRegister4Int&, const VectorRegister4Int&, const VectorRegister4Int&),
+	template< 
+		VectorRegister4Int (*RGB_FUNC_MASKED)(const VectorRegister4Int&, const VectorRegister4Int&, const VectorRegister4Int&),
 		int32 (*A_FUNC)(int32, int32),
 		bool CLAMP >
 	inline void BufferLayerCompositeVector(
-		Image* pBase,
-		const Image* pBlend,
+		Image* BaseImage,
+		const Image* BlendImage,
 		bool bOnlyFirstLOD,
 		uint8 BlendAlphaSourceChannel)
 	{
-		check(pBase->GetFormat() == EImageFormat::IF_RGBA_UBYTE);
-		check(pBlend->GetFormat() == EImageFormat::IF_RGBA_UBYTE);
-		check(pBase->GetSizeX() == pBlend->GetSizeX() && pBase->GetSizeY() == pBlend->GetSizeY());
-		check(bOnlyFirstLOD || pBase->GetLODCount() <= pBlend->GetLODCount());
-
-		uint8* pBaseBuf = pBase->GetData();
-		const uint8* pBlendedBuf = pBlend->GetData();
+		check(BaseImage->GetFormat() == EImageFormat::IF_RGBA_UBYTE);
+		check(BlendImage->GetFormat() == EImageFormat::IF_RGBA_UBYTE);
+		check(BaseImage->GetSizeX() == BlendImage->GetSizeX() && BaseImage->GetSizeY() == BlendImage->GetSizeY());
+		check(bOnlyFirstLOD || BaseImage->GetLODCount() <= BlendImage->GetLODCount());
 
-		// The base determines the number of lods to process.
-		int32 LODCount = bOnlyFirstLOD ? 1 : pBase->GetLODCount();
+		constexpr int32 NumBatchElems = 4096*2;
+		constexpr int32 BytesPerElem = 4;
 
-		int32 PixelCount = 0;
-		if (bOnlyFirstLOD)
-		{
-			PixelCount = pBase->GetSizeX() * pBase->GetSizeY();
-		}
-		else
-		{
-			PixelCount = pBase->CalculatePixelCount();
-		}
+		const int32 NumLODs = bOnlyFirstLOD ? 1 : BaseImage->GetLODCount();
 
-		constexpr int32 NumBatchElems = 4096*2;
-		const int32 NumBatches = FMath::DivideAndRoundUp(PixelCount, NumBatchElems);
+		const int32 NumBatches = BaseImage->DataStorage.GetNumBatchesLODRange(NumBatchElems, BytesPerElem, 0, NumLODs);
+	
+		check(NumBatches == BlendImage->DataStorage.GetNumBatchesLODRange(NumBatchElems, BytesPerElem, 0, NumLODs));
 
 		auto ProcessBatch = 
 		[
-			pBaseBuf, pBlendedBuf, PixelCount, NumBatchElems, BlendAlphaSourceChannel
+			BaseImage, BlendImage, BytesPerElem, NumBatchElems, BlendAlphaSourceChannel, NumLODs
 		] (uint32 BatchId)
 		{
+			const TArrayView<uint8> BaseBatchView =
+					BaseImage->DataStorage.GetBatchLODRange(BatchId, NumBatchElems, BytesPerElem, 0, NumLODs);
 
-			const int32 BatchBegin = BatchId * NumBatchElems;
-			const int32 BatchEnd = FMath::Min(BatchBegin + NumBatchElems, PixelCount);
-			for (int32 i = BatchBegin; i < BatchEnd; ++i)
+			const TArrayView<const uint8> BlendBatchView =
+					BlendImage->DataStorage.GetBatchLODRange(BatchId, NumBatchElems, BytesPerElem, 0, NumLODs);
+
+			uint8* BaseBuf = BaseBatchView.GetData();
+			const uint8* BlendBuf = BlendBatchView.GetData();
+
+			const int32 NumElems = BaseBatchView.Num() / BytesPerElem;
+			check(NumElems == BlendBatchView.Num()/BytesPerElem);
+
+			for (int32 I = 0; I < NumElems; ++I)
 			{
 				// TODO: Optimize this (SIMD?)
-				const int32 BaseAlpha = pBaseBuf[4 * i + BlendAlphaSourceChannel];
-				const int32 BlendedAlpha = pBlendedBuf[4 * i + BlendAlphaSourceChannel];
+				const int32 BaseAlpha = BaseBuf[4 * I + BlendAlphaSourceChannel];
+				const int32 BlendedAlpha = BlendBuf[4 * I + BlendAlphaSourceChannel];
 
 				const VectorRegister4Int Mask = VectorIntSet1(BlendedAlpha);
-				const VectorRegister4Int Blended = MakeVectorRegisterInt(pBlendedBuf[4 * i + 0], pBlendedBuf[4 * i + 1], pBlendedBuf[4 * i + 2], pBlendedBuf[4 * i + 3]);
-				const VectorRegister4Int Base = MakeVectorRegisterInt(pBaseBuf[4 * i + 0], pBaseBuf[4 * i + 1], pBaseBuf[4 * i + 2], pBaseBuf[4 * i + 3]);
+
+				uint32 BlendPixel;
+				FMemory::Memcpy(&BlendPixel, BlendBuf, sizeof(uint32));
+
+				uint32 BasePixel;
+				FMemory::Memcpy(&BasePixel, BaseBuf, sizeof(uint32));
+
+				const VectorRegister4Int Blended = MakeVectorRegisterInt(
+						(BlendPixel >> 0) & 0xFF, 
+						(BlendPixel >> 8) & 0xFF, 
+						(BlendPixel >> 16) & 0xFF, 
+						(BlendPixel >> 24) & 0xFF);
+
+				const VectorRegister4Int Base = MakeVectorRegisterInt(
+						(BasePixel >> 0) & 0xFF, 
+						(BasePixel >> 8) & 0xFF, 
+						(BasePixel >> 16) & 0xFF, 
+						(BasePixel >> 24) & 0xFF);
 
 				VectorRegister4Int Result = RGB_FUNC_MASKED(Base, Blended, Mask);
 				if constexpr (CLAMP)
@@ -1963,15 +1894,14 @@ namespace mu
 				}
 
 				alignas(VectorRegister4Int) int32 IndexableRegister[4];
-
 				VectorIntStoreAligned(Result, &IndexableRegister);
 
-				pBaseBuf[4 * i + 0] = static_cast<uint8>(IndexableRegister[0]);
-				pBaseBuf[4 * i + 1] = static_cast<uint8>(IndexableRegister[1]);
-				pBaseBuf[4 * i + 2] = static_cast<uint8>(IndexableRegister[2]);
-				pBaseBuf[4 * i + 3] = static_cast<uint8>(IndexableRegister[3]);
+				BaseBuf[4 * I + 0] = static_cast<uint8>(IndexableRegister[0]);
+				BaseBuf[4 * I + 1] = static_cast<uint8>(IndexableRegister[1]);
+				BaseBuf[4 * I + 2] = static_cast<uint8>(IndexableRegister[2]);
+				BaseBuf[4 * I + 3] = static_cast<uint8>(IndexableRegister[3]);
 
-				pBaseBuf[4 * i + BlendAlphaSourceChannel] = static_cast<uint8>(AlphaResult);
+				BaseBuf[4 * I + BlendAlphaSourceChannel] = static_cast<uint8>(AlphaResult);
 			}
 		};
 
@@ -1985,25 +1915,25 @@ namespace mu
 		}
 	}
 
-	//---------------------------------------------------------------------------------------------
-	template< unsigned (*BLEND_FUNC)(unsigned,unsigned),
-			  bool CLAMP >
-    inline void BufferLayerStrideNoAlpha( uint8* pDestBuf, int stride, const Image* pMask, const Image* pBlend, int32 LODCount )
+	template<uint32 (*BLEND_FUNC)(uint32,uint32),
+			 bool CLAMP>
+    inline void BufferLayerStrideNoAlpha(Image* DestImage, int32 DestOffset, int32 Stride, const Image* MaskImage, const Image* BlendImage/*, int32 LODCount*/)
 	{
-		if ( pBlend->GetFormat()==EImageFormat::IF_RGB_UBYTE )
+		if (BlendImage->GetFormat() == EImageFormat::IF_RGB_UBYTE)
 		{
-			BufferLayerFormatStrideNoAlpha< BLEND_FUNC, CLAMP, 3 >
-					( pDestBuf, stride, pMask, pBlend, LODCount ) ;
+			BufferLayerFormatStrideNoAlpha<BLEND_FUNC, CLAMP, 3>
+					(DestImage, DestOffset, Stride, MaskImage, BlendImage/*, LODCount*/);
 		}
-        else if ( pBlend->GetFormat()==EImageFormat::IF_RGBA_UBYTE || pBlend->GetFormat()==EImageFormat::IF_BGRA_UBYTE )
+        else if (BlendImage->GetFormat() == EImageFormat::IF_RGBA_UBYTE || 
+				 BlendImage->GetFormat() == EImageFormat::IF_BGRA_UBYTE)
 		{
-			BufferLayerFormatStrideNoAlpha< BLEND_FUNC, CLAMP, 4 >
-					( pDestBuf, stride, pMask, pBlend, LODCount) ;
+			BufferLayerFormatStrideNoAlpha<BLEND_FUNC, CLAMP, 4>
+					(DestImage, DestOffset, Stride, MaskImage, BlendImage/*, LODCount*/);
 		}
-		else if ( pBlend->GetFormat()==EImageFormat::IF_L_UBYTE )
+		else if (BlendImage->GetFormat() == EImageFormat::IF_L_UBYTE)
 		{
-			BufferLayerFormatStrideNoAlpha< BLEND_FUNC, CLAMP, 1 >
-					( pDestBuf, stride, pMask, pBlend, LODCount) ;
+			BufferLayerFormatStrideNoAlpha<BLEND_FUNC, CLAMP, 1>
+					(DestImage, DestOffset, Stride, MaskImage, BlendImage/*, LODCount*/);
 		}
 		else
 		{
@@ -2012,428 +1942,330 @@ namespace mu
 	}
 
 
-	//---------------------------------------------------------------------------------------------
-	//! Apply a blending function to an image with another image as blending layer, on a subrect of
-	//! the base image.
-	//! \warning this method applies the blending function to the alpha channel too
-	//! \warning this method uses the mask as a binary mask (>0)
-	//---------------------------------------------------------------------------------------------
-	template< unsigned (*BLEND_FUNC)(unsigned,unsigned),
-			  bool CLAMP >
-	inline void ImageLayerOnBaseNoAlpha( Image* pBase,
-										 const Image* pMask,
-										 const Image* pBlended,
-										 const box<UE::Math::TIntVector2<uint16>>& rect )
+	/**
+	* Apply a blending function to an image with another image as blending layer, on a subrect of
+	* the base image.
+	* \warning this method applies the blending function to the alpha channel too
+	* \warning this method uses the mask as a binary mask (>0)
+	*/
+	template<uint32 (*BLEND_FUNC)(uint32, uint32), bool CLAMP>
+	inline void ImageLayerOnBaseNoAlpha(
+			Image* BaseImage,
+			const Image* MaskImage,
+			const Image* BlendedImage,
+			const box<UE::Math::TIntVector2<uint16>>& Rect)
 	{
-		check( pBase->GetSizeX() >= rect.min[0]+rect.size[0] );
-		check( pBase->GetSizeY() >= rect.min[1]+rect.size[1] );
-		check( pMask->GetSizeX() == pBlended->GetSizeX() );
-		check( pMask->GetSizeY() == pBlended->GetSizeY() );
-		check( pBase->GetFormat() == pBlended->GetFormat() );
-		check( pMask->GetFormat() == EImageFormat::IF_L_UBYTE
-						||
-						pMask->GetFormat() == EImageFormat::IF_L_UBYTE_RLE
-						||
-						pMask->GetFormat() == EImageFormat::IF_L_UBIT_RLE );
-        check( pBase->GetLODCount() <= pMask->GetLODCount() );
-        check( pBase->GetLODCount() <= pBlended->GetLODCount() );
-
-		int pixelSize = GetImageFormatData( pBase->GetFormat() ).BytesPerBlock;
-
-		int start = ( pBase->GetSizeX()*rect.min[1]+rect.min[0] ) * pixelSize;
-		int stride = ( pBase->GetSizeX() - rect.size[0] ) * pixelSize;
-		BufferLayerStrideNoAlpha<BLEND_FUNC,CLAMP>( pBase->GetData()+start, stride, pMask, pBlended, pBase->GetLODCount());
+		check(BaseImage->GetSizeX() >= Rect.min[0] + Rect.size[0]);
+		check(BaseImage->GetSizeY() >= Rect.min[1] + Rect.size[1]);
+		check(MaskImage->GetSizeX() == BlendedImage->GetSizeX());
+		check(MaskImage->GetSizeY() == BlendedImage->GetSizeY());
+		check(BaseImage->GetFormat() == BlendedImage->GetFormat());
+		check(MaskImage->GetFormat() == EImageFormat::IF_L_UBYTE ||
+			  //UBYTE_RLE does not look to be supported.
+			  //MaskImage->GetFormat() == EImageFormat::IF_L_UBYTE_RLE || 
+			  MaskImage->GetFormat() == EImageFormat::IF_L_UBIT_RLE);
+        check(BaseImage->GetLODCount() <= MaskImage->GetLODCount());
+        check(BaseImage->GetLODCount() <= BlendedImage->GetLODCount());
+
+		int32 PixelSize = GetImageFormatData(BaseImage->GetFormat()).BytesPerBlock;
+
+		int32 Start = (BaseImage->GetSizeX() * Rect.min[1] + Rect.min[0]) * PixelSize;
+		int32 Stride = (BaseImage->GetSizeX() - Rect.size[0]) * PixelSize;
+
+		// Stride is only valid for LOD 0, BufferLayerStride variants cannot operate on multiple lods.
+		// TODO: review if this needs to be supported, and implement using a rect lod reducction at this level.
+		BufferLayerStrideNoAlpha<BLEND_FUNC, CLAMP>(BaseImage, Start, Stride, MaskImage, BlendedImage/*, BaseImage->GetLODCount()*/);
 	}
 
-
-	//---------------------------------------------------------------------------------------------
-	// \TODO: Unused?
-	template< int NC >
-    inline void BufferTableColourFormat( uint8* pStartDestBuf,
-										 const Image* pBase,
-										 const Image* pMask,
-                                         const uint8* pTable0,
-                                         const uint8* pTable1,
-                                         const uint8* pTable2 )
+	template<uint32 NC>
+	FORCEINLINE uint32 PackPixel(const uint8* PixelPtr)
 	{
-		uint8* pDestBuf = pStartDestBuf;
-        const uint8* pMaskBuf = pMask->GetData();
-        const uint8* pBaseBuf = pBase->GetData();
-		EImageFormat maskFormat = pMask->GetFormat();
+		static_assert(NC > 0 && NC <= 4);
 
-        bool isUncompressed = ( maskFormat == EImageFormat::IF_L_UBYTE );
+		uint32 PixelPack = 0;
 
-		// The base determines the lods to process.
-		int32 LODCount = pBase->GetLODCount();
+		// The compiler should be able to optimize this given that NC is a constant expression.
+		FMemory::Memcpy(&PixelPack, PixelPtr, NC);
 
-        if ( isUncompressed )
-		{
-			int pixelCount = (int)pBase->CalculatePixelCount();
-			for ( int i=0; i<pixelCount; ++i )
-			{
-				unsigned mask = pMaskBuf[i];
+		return PixelPack;
+	}
 
-				if (mask)
-				{
-					if (mask<255)
-					{
-                        pDestBuf[NC*i+0] = (uint8)( ( ( ( 255 - mask ) * pBaseBuf[NC*i+0] )
-                                             + ( mask * pTable0[ pBaseBuf[NC*i+0] ] ) ) >> 8 );
-                        pDestBuf[NC*i+1] = (uint8)( ( ( ( 255 - mask ) * pBaseBuf[NC*i+1] )
-                                             + ( mask * pTable1[ pBaseBuf[NC*i+1] ] ) ) >> 8 );
-                        pDestBuf[NC*i+2] = (uint8)( ( ( ( 255 - mask ) * pBaseBuf[NC*i+2] )
-                                             + ( mask * pTable2[ pBaseBuf[NC*i+2] ] ) ) >> 8 );
-					}
-					else
-					{
-						pDestBuf[NC*i+0] = pTable0[ pBaseBuf[NC*i+0] ];
-						pDestBuf[NC*i+1] = pTable1[ pBaseBuf[NC*i+1] ];
-						pDestBuf[NC*i+2] = pTable2[ pBaseBuf[NC*i+2] ];
-					}
-				}
-				else
-				{
-					pDestBuf[NC*i+0] = pBaseBuf[NC*i+0];
-					pDestBuf[NC*i+1] = pBaseBuf[NC*i+1];
-					pDestBuf[NC*i+2] = pBaseBuf[NC*i+2];
-				}
+	template<uint32 NC>
+	FORCEINLINE void UnpackPixel(uint8* PixelPtr, uint32 PixelData)
+	{
+		static_assert(NC > 0 && NC <= 4);
 
-                constexpr bool isNC4 = (NC==4);
-                if ( isNC4 )
-				{
-					pDestBuf[NC*i+3] = pBaseBuf[NC*i+3];
-				}
-			}
-		}
-		else if ( maskFormat==EImageFormat::IF_L_UBYTE_RLE )
-		{
-			int rows = pBase->GetSizeY();
-			int width = pBase->GetSizeX();
+		// The compiler should be able to optimize this given that NC is a constant expression
+		FMemory::Memcpy(PixelPtr, &PixelData, NC);
+	}
 
-            for (int32 lod=0;lod< LODCount; ++lod)
-            {
-                pMaskBuf += 4+rows*sizeof(uint32);
+	template< 
+		uint32 (*BLEND_FUNC)(uint32, uint32), 
+		uint32 NC>
+	inline void BufferLayerCombineColour(Image* ResultImage, const Image* BaseImage, FVector4f Color, bool bOnlyFirstLOD = false)
+	{
+		static_assert(NC > 0 && NC <= 4);
 
-                for ( int r=0; r<rows; ++r )
-                {
-                    const uint8* pDestRowEnd = pDestBuf + width*NC;
-                    while ( pDestBuf!=pDestRowEnd )
-                    {
-                        // Decode header
-                        uint16 equal = *(const uint16*)pMaskBuf;
-                        pMaskBuf += 2;
+		check(BaseImage->GetSizeX() == ResultImage->GetSizeX());
+		check(BaseImage->GetSizeY() == ResultImage->GetSizeY());
 
-                        uint8 different = *pMaskBuf;
-                        ++pMaskBuf;
+		const uint32 TopColor = 
+			static_cast<uint32>(255.0f * Color[0]) << 0 |  
+			static_cast<uint32>(255.0f * Color[1]) << 8 |  
+			static_cast<uint32>(255.0f * Color[2]) << 16;  
 
-                        unsigned equalPixel = *pMaskBuf;
-                        ++pMaskBuf;
+		const int32 NumLODs = bOnlyFirstLOD ? 1 : ResultImage->GetLODCount();
+		
+		constexpr int32 BatchNumElems = 4098*2;
+		const int32 NumBatches = BaseImage->DataStorage.GetNumBatchesLODRange(BatchNumElems, NC, 0, NumLODs);
 
-                        // Equal pixels
-						check(pDestBuf + NC * equal <= pStartDestBuf + pBase->GetDataSize());
-                        if ( equalPixel==255 )
-                        {
-                            for ( int i=0; i<equal; ++i )
-                            {
-                                pDestBuf[NC*i+0] = pTable0[ pBaseBuf[NC*i+0] ];
-                                pDestBuf[NC*i+1] = pTable1[ pBaseBuf[NC*i+1] ];
-                                pDestBuf[NC*i+2] = pTable2[ pBaseBuf[NC*i+2] ];
+		check(NumBatches == ResultImage->DataStorage.GetNumBatchesLODRange(BatchNumElems, NC, 0, NumLODs));
 
-                                constexpr bool isNC4 = (NC==4);
-                                if ( isNC4 )
-                                {
-                                    pDestBuf[NC*i+3] = pBaseBuf[NC*i+3];
-                                }
-                            }
-                        }
-                        else if ( equalPixel>0 )
-                        {
-                            for ( int i=0; i<equal; ++i )
-                            {
-                                pDestBuf[NC*i+0] = (uint8)( ( ( ( 255 - equalPixel ) * pBaseBuf[NC*i+0] )
-                                                     + ( equalPixel * pTable0[ pBaseBuf[NC*i+0] ] ) ) >> 8 );
-                                pDestBuf[NC*i+1] = (uint8)( ( ( ( 255 - equalPixel ) * pBaseBuf[NC*i+1] )
-                                                     + ( equalPixel * pTable1[ pBaseBuf[NC*i+1] ] ) ) >> 8 );
-                                pDestBuf[NC*i+2] = (uint8)( ( ( ( 255 - equalPixel ) * pBaseBuf[NC*i+2] )
-                                                     + ( equalPixel * pTable2[ pBaseBuf[NC*i+2] ] ) ) >> 8 );
-
-                                constexpr bool isNC4 = (NC==4);
-                                if ( isNC4 )
-                                {
-                                    pDestBuf[NC*i+3] = pBaseBuf[NC*i+3];
-                                }
-                            }
-                        }
-                        else
-                        {
-                            // It could happen if xxxxxOnBase
-                            if (pDestBuf!=pBaseBuf)
-                            {
-								FMemory::Memmove( pDestBuf, pBaseBuf, NC*equal );
-                            }
-                        }
-                        pDestBuf += NC*equal;
-                        pBaseBuf += NC*equal;
+		auto ProcessBatch = [ResultImage, BaseImage, BatchNumElems, TopColor, NumLODs](int32 BatchId)
+		{
+			const TArrayView<const uint8> BaseBatchView = 
+					BaseImage->DataStorage.GetBatchLODRange(BatchId, BatchNumElems, NC, 0, NumLODs);
+			
+			const TArrayView<uint8> ResultBatchView =
+					ResultImage->DataStorage.GetBatchLODRange(BatchId, BatchNumElems, NC, 0, NumLODs);
 
-                        // Different pixels
-						check(pDestBuf + NC * different <= pStartDestBuf + pBase->GetDataSize());
-                        for ( int i=0; i<different; ++i )
-                        {
-                            unsigned mask = pMaskBuf[i];
-                            pDestBuf[NC*i+0] = (uint8)( ( ( ( 255 - mask ) * pBaseBuf[NC*i+0] )
-                                                 + ( mask * pTable0[ pBaseBuf[NC*i+0] ] ) ) >> 8 );
-                            pDestBuf[NC*i+1] = (uint8)( ( ( ( 255 - mask ) * pBaseBuf[NC*i+1] )
-                                                 + ( mask * pTable1[ pBaseBuf[NC*i+1] ] ) ) >> 8 );
-                            pDestBuf[NC*i+2] = (uint8)( ( ( ( 255 - mask ) * pBaseBuf[NC*i+2] )
-                                                 + ( mask * pTable2[ pBaseBuf[NC*i+2] ] ) ) >> 8 );
-
-                            constexpr bool isNC4 = (NC==4);
-                            if ( isNC4 )
-                            {
-                                pDestBuf[NC*i+3] = pBaseBuf[NC*i+3];
-                            }
-                        }
+			const int32 NumElems = BaseBatchView.Num() / NC;
 
-                        pDestBuf += NC*different;
-                        pBaseBuf += NC*different;
-                        pMaskBuf += different;
-                    }
-                }
+			check(NumElems == ResultBatchView.Num() / NC);
+			
+			const uint8* BaseBuf = BaseBatchView.GetData();
+			uint8* ResultBuf = ResultBatchView.GetData();
+			
+			for (int32 I = 0; I < NumElems; ++I)
+			{
+				const uint32 Base = PackPixel<NC>(&BaseBuf[NC * I]);
 
-                rows = FMath::DivideAndRoundUp(rows,2);
-                width = FMath::DivideAndRoundUp(width,2);
-            }
-		}
-		else
+				const uint32 Result = BLEND_FUNC(Base, TopColor);
+				UnpackPixel<NC>(&ResultBuf[NC * I], Result);
+			}
+		};
+
+		if (NumBatches == 1)
 		{
-			checkf( false, TEXT("Unsupported mask format.") );
+			ProcessBatch(0);
 		}
-	}
-
-
-	//---------------------------------------------------------------------------------------------
-	// \TODO: Unused?
-    inline void ImageTable( Image* pResult,
-                                const Image* pBase,
-								const Image* pMask,
-                                const uint8* pTable0,
-                                const uint8* pTable1,
-                                const uint8* pTable2
-								)
-	{
-        check( pResult->GetFormat() == pBase->GetFormat() );
-        check( pResult->GetSizeX() == pBase->GetSizeX() );
-        check( pResult->GetSizeY() == pBase->GetSizeY() );
-        check( pResult->GetLODCount() == pBase->GetLODCount() );
-        check( pBase->GetSizeX() == pMask->GetSizeX() );
-		check( pBase->GetSizeY() == pMask->GetSizeY() );
-		check( pMask->GetFormat() == EImageFormat::IF_L_UBYTE
-						||
-						pMask->GetFormat() == EImageFormat::IF_L_UBYTE_RLE );
-        check( pBase->GetLODCount() <= pMask->GetLODCount() );
-
-		EImageFormat baseFormat = pBase->GetFormat();
-		if ( baseFormat==EImageFormat::IF_RGB_UBYTE )
-		{
-            BufferTableColourFormat<3>( pResult->GetData(), pBase, pMask, pTable0, pTable1, pTable2 );
-		}
-        else if ( baseFormat==EImageFormat::IF_RGBA_UBYTE )
-        {
-            BufferTableColourFormat<4>( pResult->GetData(), pBase, pMask, pTable0, pTable1, pTable2 );
-        }
-        else if ( baseFormat==EImageFormat::IF_BGRA_UBYTE )
-        {
-            BufferTableColourFormat<4>( pResult->GetData(), pBase, pMask, pTable2, pTable1, pTable0 );
-        }
-        else
+		else if (NumBatches > 1)
 		{
-			checkf( false, TEXT("Unsupported format.") );
+			ParallelFor(NumBatches, ProcessBatch);
 		}
 	}
 
-
-	template<size_t NC>
-	inline uint32 PackPixel(const uint8* pPixel)
+	template< 
+		uint32 (*BLEND_FUNC)(uint32, uint32), 
+		uint32 NC>
+	inline void BufferLayerCombine(Image* ResultImage, const Image* BaseImage, const Image* BlendImage, bool bOnlyFirstLOD)
 	{
 		static_assert(NC > 0 && NC <= 4);
 
-		uint32 PixelPack = 0;
+		check(BaseImage->GetSizeX() == BlendImage->GetSizeX());
+		check(BaseImage->GetSizeY() == BlendImage->GetSizeY());
+		check(BaseImage->GetSizeX() == ResultImage->GetSizeX());
+		check(BaseImage->GetSizeY() == ResultImage->GetSizeY());
+		check(BaseImage->GetFormat() == BlendImage->GetFormat());
+		check(BaseImage->GetFormat() == ResultImage->GetFormat());
 
-		// The compiler should be able to optimize this given that NC is a constant expression
-		FMemory::Memcpy( &PixelPack, pPixel, NC );
+		const int32 NumLODs = bOnlyFirstLOD ? 1 : ResultImage->GetLODCount();
 
-		return PixelPack;
-	}
+		constexpr int32 BatchNumElems = 4098*2;
+		const int32 NumBatches = BaseImage->DataStorage.GetNumBatchesLODRange(BatchNumElems, NC, 0, NumLODs);
 
-	template<size_t NC>
-	inline void UnpackPixel(uint8* const pPixel, uint32 pixelData )
-	{
-		static_assert(NC > 0 && NC <= 4);
+		check(NumBatches == BlendImage->DataStorage.GetNumBatchesLODRange(BatchNumElems, NC, 0, NumLODs));
+		check(NumBatches == ResultImage->DataStorage.GetNumBatchesLODRange(BatchNumElems, NC, 0, NumLODs));
 
-		// The compiler should be able to optimize this given that NC is a constant expression
-		FMemory::Memcpy( pPixel, &pixelData, NC );
-	}
+		auto ProcessBatch = [ResultImage, BaseImage, BlendImage, BatchNumElems, NumLODs](int32 BatchId)
+		{
+			const TArrayView<const uint8> BaseBatchView = 
+					BaseImage->DataStorage.GetBatchLODRange(BatchId, BatchNumElems, NC, 0, NumLODs);
 
-	template< 
-			uint32 (*BLEND_FUNC)(uint32, uint32), 
-			size_t NC >
-	inline void BufferLayerCombineColour(uint8* pDestBuf, const Image* pBase, FVector4f col)
-	{
-		//static_assert(NC > 0 && NC <= 4);
+			const TArrayView<const uint8> BlendBatchView = 
+					BlendImage->DataStorage.GetBatchLODRange(BatchId, BatchNumElems, NC, 0, NumLODs);
 
-		const uint8* pBaseBuf = pBase->GetData();
+			const TArrayView<uint8> ResultBatchView =
+					ResultImage->DataStorage.GetBatchLODRange(BatchId, BatchNumElems, NC, 0, NumLODs);
 
-		uint32 top = 
-			(uint32)(255.0f * col[0]) << 0 |
-			(uint32)(255.0f * col[1]) << 8 |
-			(uint32)(255.0f * col[2]) << 16;
+			const int32 NumElems = BaseBatchView.Num() / NC;
+			check(NumElems == BlendBatchView.Num() / NC);
+			check(NumElems == ResultBatchView.Num() / NC);
+			
+			const uint8* BaseBuf = BaseBatchView.GetData();
+			const uint8* BlendBuf = BlendBatchView.GetData();
+			uint8* ResultBuf = ResultBatchView.GetData();
+			
+			for (int32 I = 0; I < NumElems; ++I)
+			{
+				const uint32 Base = PackPixel<NC>(&BaseBuf[NC * I]);
+				const uint32 Blend = PackPixel<NC>(&BlendBuf[NC * I]);
 
-		int32 pixelCount = pBase->CalculatePixelCount();
+				const uint32 Result = BLEND_FUNC(Base, Blend);
+				UnpackPixel<NC>(&ResultBuf[NC * I], Result);
+			}
+		};
 
-		for (int i = 0; i < pixelCount; ++i)
+		if (NumBatches == 1)
 		{
-			uint32 base = PackPixel<NC>(&pBaseBuf[NC * i]);
-
-			uint32 result = BLEND_FUNC(base, top);
-
-			UnpackPixel<NC>(&pDestBuf[NC * i], result);
+			ProcessBatch(0);
+		}
+		else if (NumBatches > 1)
+		{
+			ParallelFor(NumBatches, ProcessBatch);
 		}
 	}
 
 	template< 
-			uint32 (*BLEND_FUNC)(uint32, uint32), 
-			size_t NC >
-	inline void BufferLayerCombine(uint8* pDestBuf, const Image* pBase, const Image* pBlended, bool bOnlyFirstLOD)
+		uint32 (*BLEND_FUNC_MASKED)(uint32, uint32, uint32), 
+		uint32 NC>
+	inline void BufferLayerCombine(Image* ResultImage, const Image* BaseImage, const Image* MaskImage, const Image* BlendImage, bool bOnlyFirstLOD)
 	{
-		//static_assert(NC > 0 && NC <= 4);
+		static_assert(NC > 0 && NC <= 4);
 
-		check(pBase->GetSizeX() == pBlended->GetSizeX());
-		check(pBase->GetSizeY() == pBlended->GetSizeY());
-		check(pBase->GetFormat() == pBlended->GetFormat());
+		check(BaseImage->GetSizeX() == BlendImage->GetSizeX());
+		check(BaseImage->GetSizeY() == BlendImage->GetSizeY());
+		check(BaseImage->GetSizeX() == MaskImage->GetSizeX());
+		check(BaseImage->GetSizeY() == MaskImage->GetSizeY());
+		check(BaseImage->GetFormat() == BlendImage->GetFormat());
 
-		const uint8* pBaseBuf = pBase->GetData();
-		const uint8* pBlendedBuf = pBlended->GetData();
+		const int32 NumLODs = bOnlyFirstLOD ? 1 : ResultImage->GetLODCount();
 
-		int32 PixelCount = 0;
-		if (bOnlyFirstLOD)
-		{
-			PixelCount = pBase->GetSizeX() * pBase->GetSizeX();
-		}
-		else
-		{
-			PixelCount = pBase->CalculatePixelCount();
-		}		
+		constexpr int32 BatchNumElems = 4098*2;
+		const int32 NumBatches = BaseImage->DataStorage.GetNumBatchesLODRange(BatchNumElems, NC, 0, NumLODs);
 
-		for (int i = 0; i < PixelCount; ++i)
+		check(NumBatches <= MaskImage->DataStorage.GetNumBatchesLODRange(BatchNumElems, 1, 0, NumLODs));
+		check(NumBatches <= BlendImage->DataStorage.GetNumBatchesLODRange(BatchNumElems, NC, 0, NumLODs));
+		check(NumBatches <= ResultImage->DataStorage.GetNumBatchesLODRange(BatchNumElems, NC, 0, NumLODs));
+
+		auto ProcessBatch = [ResultImage, BaseImage, MaskImage, BlendImage, BatchNumElems, NumLODs](int32 BatchId)
 		{
-			const uint32 base = PackPixel<NC>(&pBaseBuf[NC * i]);
-			const uint32 blend = PackPixel<NC>(&pBlendedBuf[NC * i]);
+			const TArrayView<const uint8> BaseBatchView = 
+					BaseImage->DataStorage.GetBatchLODRange(BatchId, BatchNumElems, NC, 0, NumLODs);
 
-			const uint32 result = BLEND_FUNC(base, blend);
+			const TArrayView<const uint8> BlendBatchView = 
+					BlendImage->DataStorage.GetBatchLODRange(BatchId, BatchNumElems, NC, 0, NumLODs);
 
-			UnpackPixel<NC>(&pDestBuf[NC * i], result);
-		}
-	}
+			const TArrayView<const uint8> MaskBatchView =
+					MaskImage->DataStorage.GetBatchLODRange(BatchId, BatchNumElems, 1, 0, NumLODs);
 
-	template< 
-			uint32 (*BLEND_FUNC_MASKED)(uint32, uint32, uint32), 
-			size_t NC >
-	inline void BufferLayerCombine(uint8* pDestBuf, const Image* pBase, const Image* pMask, const Image* pBlended, bool bOnlyFirstLOD)
-	{
-		//static_assert(NC > 0 && NC <= 4);
+			const TArrayView<uint8> ResultBatchView =
+					ResultImage->DataStorage.GetBatchLODRange(BatchId, BatchNumElems, NC, 0, NumLODs);
+
+			const int32 NumElems = BaseBatchView.Num();
 
-		check(pBase->GetSizeX() == pBlended->GetSizeX());
-		check(pBase->GetSizeY() == pBlended->GetSizeY());
-		check(pBase->GetSizeX() == pMask->GetSizeX());
-		check(pBase->GetSizeY() == pMask->GetSizeY());
-		check(pBase->GetFormat() == pBlended->GetFormat());
+			check(NumElems == BlendBatchView.Num() / NC);
+			check(NumElems == MaskBatchView.Num() / 1);
+			check(NumElems == ResultBatchView.Num() / NC);
+			
+			const uint8* BaseBuf = BaseBatchView.GetData();
+			const uint8* BlendBuf = BlendBatchView.GetData();
+			const uint8* MaskBuf = MaskBatchView.GetData();
+			uint8* ResultBuf = ResultBatchView.GetData();
+			
+			for (int32 I = 0; I < NumElems; ++I)
+			{
+				const uint32 Base = PackPixel<NC>(&BaseBuf[NC * I]);
+				const uint32 Blend = PackPixel<NC>(&BlendBuf[NC * I]);
+				const uint32 Mask = PackPixel<1>(&MaskBuf[1 * I]);
 
-		const uint8* pBaseBuf = pBase->GetData();
-		const uint8* pMaskBuf = pMask->GetData();
-		const uint8* pBlendedBuf = pBlended->GetData();
+				const uint32 Result = BLEND_FUNC_MASKED(Base, Blend, Mask);
+				UnpackPixel<NC>(&ResultBuf[NC * I], Result);
+			}
+		};
 
-		int32 PixelCount = 0;
-		if (bOnlyFirstLOD)
-		{
-			PixelCount = pBase->GetSizeX() * pBase->GetSizeX();
-		}
-		else
+		if (NumBatches == 1)
 		{
-			PixelCount = pBase->CalculatePixelCount();
+			ProcessBatch(0);
 		}
-
-		for (int i = 0; i < PixelCount; ++i)
+		else if (NumBatches > 1)
 		{
-			const uint32 base = PackPixel<NC>(&pBaseBuf[NC * i]);
-			const uint32 blend = PackPixel<NC>(&pBlendedBuf[NC * i]);
-			const uint32 mask = PackPixel<1>(&pMaskBuf[i]);
-
-			const uint32 result = BLEND_FUNC_MASKED(base, blend, mask);
-
-			UnpackPixel<NC>(&pDestBuf[NC * i], result);
+			ParallelFor(NumBatches, ProcessBatch);
 		}
 	}
 
 	template< 
-			uint32 (*BLEND_FUNC_MASKED)(uint32, uint32, uint32), 
-			size_t NC >
-	inline void BufferLayerCombineColour(uint8* pDestBuf, const Image* pBase, const Image* pMask, FVector4f col)
+		uint32 (*BLEND_FUNC_MASKED)(uint32, uint32, uint32), 
+		uint32 NC>
+	inline void BufferLayerCombineColour(Image* ResultImage, const Image* BaseImage, const Image* MaskImage, FVector4f Color)
 	{
-		//static_assert(NC > 0 && NC <= 4);
+		static_assert(NC > 0 && NC <= 4);
 
-		check(pBase->GetSizeX() == pMask->GetSizeX());
-		check(pBase->GetSizeY() == pMask->GetSizeY());
+		check(BaseImage->GetSizeX() == MaskImage->GetSizeX());
+		check(BaseImage->GetSizeY() == MaskImage->GetSizeY());
 
+		const uint32 TopColor = 
+			static_cast<uint32>(255.0f * Color[0]) << 0 |  
+			static_cast<uint32>(255.0f * Color[1]) << 8 |  
+			static_cast<uint32>(255.0f * Color[2]) << 16;  
 
-		const uint8* pBaseBuf = pBase->GetData();
-		const uint8* pMaskBuf = pMask->GetData();
+		constexpr int32 BatchNumElems = 4098*2;
+		int32 NumBatches = BaseImage->DataStorage.GetNumBatches(BatchNumElems, NC);
 
-		uint32 top = 
-			(uint32)(255.0f * col[0]) << 0 |
-			(uint32)(255.0f * col[1]) << 8 |
-			(uint32)(255.0f * col[2]) << 16;
+		check(NumBatches == MaskImage->DataStorage.GetNumBatches(BatchNumElems, 1));
+		check(NumBatches == ResultImage->DataStorage.GetNumBatches(BatchNumElems, NC));
 
+		auto ProcessBatch = [ResultImage, BaseImage, MaskImage, BatchNumElems, TopColor](int32 BatchId)
+		{
+			const TArrayView<const uint8> BaseBatchView = BaseImage->DataStorage.GetBatch(BatchId, BatchNumElems, NC);
+			const TArrayView<const uint8> MaskBatchView = MaskImage->DataStorage.GetBatch(BatchId, BatchNumElems, 1);
+			const TArrayView<uint8> ResultBatchView = ResultImage->DataStorage.GetBatch(BatchId, BatchNumElems, NC);
 
-		int32 pixelCount = pBase->CalculatePixelCount();
+			const int32 NumElems = BaseBatchView.Num() / NC;
 
-		for (int i = 0; i < pixelCount; ++i)
-		{
-			const uint32 base = PackPixel<NC>(&pBaseBuf[NC * i]);
-			const uint32 mask = PackPixel<1>(&pMaskBuf[i]);
+			check(NumElems == MaskBatchView.Num() / 1);
+			check(NumElems == ResultBatchView.Num() / NC);
+			
+			const uint8* BaseBuf = BaseBatchView.GetData();
+			const uint8* MaskBuf = MaskBatchView.GetData();
+			uint8* ResultBuf = ResultBatchView.GetData();
+			
+			for (int32 I = 0; I < NumElems; ++I)
+			{
+				const uint32 Base = PackPixel<NC>(&BaseBuf[NC * I]);
+				const uint32 Mask = PackPixel<1>(&MaskBuf[1 * I]);
 
-			const uint32 result = BLEND_FUNC_MASKED(base, top, mask);
+				const uint32 Result = BLEND_FUNC_MASKED(Base, TopColor, Mask);
+				UnpackPixel<NC>(&ResultBuf[NC * I], Result);
+			}
+		};
 
-			UnpackPixel<NC>(&pDestBuf[NC * i], result);
+		if (NumBatches == 1)
+		{
+			ProcessBatch(0);
+		}
+		else if (NumBatches > 1)
+		{
+			ParallelFor(NumBatches, ProcessBatch);
 		}
 	}
 
-	template< uint32 (*BLEND_FUNC)(uint32, uint32) >
-	inline void ImageLayerCombine(Image* pResult, const Image* pBase, const Image* pBlended, bool bOnlyFirstLOD)
+	template<uint32 (*BLEND_FUNC)(uint32, uint32)>
+	inline void ImageLayerCombine(Image* ResultImage, const Image* BaseImage, const Image* BlendedImage, bool bOnlyFirstLOD)
 	{
-		check(pResult->GetFormat() == pBase->GetFormat());
-		check(pResult->GetSizeX() == pBase->GetSizeX());
-		check(pResult->GetSizeY() == pBase->GetSizeY());
-		check(bOnlyFirstLOD || pResult->GetLODCount() == pBase->GetLODCount());
-		check(pBase->GetSizeX() == pBlended->GetSizeX());
-		check(pBase->GetSizeY() == pBlended->GetSizeY());
-		check(pBase->GetFormat() == pBlended->GetFormat());
-		check(bOnlyFirstLOD || pResult->GetLODCount() <= pBlended->GetLODCount());
+		check(ResultImage->GetFormat() == BaseImage->GetFormat());
+		check(ResultImage->GetSizeX() == BaseImage->GetSizeX());
+		check(ResultImage->GetSizeY() == BaseImage->GetSizeY());
+		check(bOnlyFirstLOD || ResultImage->GetLODCount() == BaseImage->GetLODCount());
+		check(BaseImage->GetSizeX() == BlendedImage->GetSizeX());
+		check(BaseImage->GetSizeY() == BlendedImage->GetSizeY());
+		check(BaseImage->GetFormat() == BlendedImage->GetFormat());
+		check(bOnlyFirstLOD || ResultImage->GetLODCount() <= BlendedImage->GetLODCount());
 
-		const EImageFormat baseFormat = pBase->GetFormat();
+		const EImageFormat BaseFormat = BaseImage->GetFormat();
 
-		uint8* pDestBuf = pResult->GetData();
-		if (baseFormat == EImageFormat::IF_L_UBYTE)
+		if (BaseFormat == EImageFormat::IF_L_UBYTE)
 		{
-			BufferLayerCombine<BLEND_FUNC, 1 >(pDestBuf, pBase, pBlended, bOnlyFirstLOD);
+			BufferLayerCombine<BLEND_FUNC, 1>(ResultImage, BaseImage, BlendedImage, bOnlyFirstLOD);
 		}
-		else if (baseFormat == EImageFormat::IF_RGB_UBYTE)
+		else if (BaseFormat == EImageFormat::IF_RGB_UBYTE)
 		{
-			BufferLayerCombine< BLEND_FUNC, 3 >(pDestBuf, pBase, pBlended, bOnlyFirstLOD);
+			BufferLayerCombine<BLEND_FUNC, 3>(ResultImage, BaseImage, BlendedImage, bOnlyFirstLOD);
 		}
-		else if (baseFormat == EImageFormat::IF_RGBA_UBYTE || baseFormat == EImageFormat::IF_BGRA_UBYTE)
+		else if (BaseFormat == EImageFormat::IF_RGBA_UBYTE || BaseFormat == EImageFormat::IF_BGRA_UBYTE)
 		{
 			// \todo: pass swizzle template argument if BGRA_UBYTE, not yet supported.
-			BufferLayerCombine< BLEND_FUNC, 4 >(pDestBuf, pBase, pBlended, bOnlyFirstLOD);
+			BufferLayerCombine<BLEND_FUNC, 4>(ResultImage, BaseImage, BlendedImage, bOnlyFirstLOD);
 		}
 		else
 		{
@@ -2443,41 +2275,39 @@ namespace mu
 
 	template< 
 		uint32 (*BLEND_FUNC)(uint32, uint32), 
-		uint32 (*BLEND_FUNC_MASKED)(uint32, uint32, uint32) >
-	inline void ImageLayerCombine(Image* pResult, const Image* pBase, const Image* pMask, const Image* pBlended, bool bOnlyFirstLOD)
+		uint32 (*BLEND_FUNC_MASKED)(uint32, uint32, uint32)>
+	inline void ImageLayerCombine(Image* ResultImage, const Image* BaseImage, const Image* MaskImage, const Image* BlendedImage, bool bOnlyFirstLOD)
 	{
-		check(pResult->GetFormat() == pBase->GetFormat());
-		check(pResult->GetSizeX() == pBase->GetSizeX());
-		check(pResult->GetSizeY() == pBase->GetSizeY());
-		check(bOnlyFirstLOD || pResult->GetLODCount() == pBase->GetLODCount());
-		check(pBase->GetSizeX() == pBlended->GetSizeX());
-		check(pBase->GetSizeY() == pBlended->GetSizeY());
-		check(pBase->GetFormat() == pBlended->GetFormat());
-		check(bOnlyFirstLOD || pResult->GetLODCount() <= pBlended->GetLODCount());
-
-		const EImageFormat baseFormat = pBase->GetFormat();
+		check(ResultImage->GetFormat() == BaseImage->GetFormat());
+		check(ResultImage->GetSizeX() == BaseImage->GetSizeX());
+		check(ResultImage->GetSizeY() == BaseImage->GetSizeY());
+		check(bOnlyFirstLOD || ResultImage->GetLODCount() == BaseImage->GetLODCount());
+		check(BaseImage->GetSizeX() == BlendedImage->GetSizeX());
+		check(BaseImage->GetSizeY() == BlendedImage->GetSizeY());
+		check(BaseImage->GetFormat() == BlendedImage->GetFormat());
+		check(bOnlyFirstLOD || ResultImage->GetLODCount() <= BlendedImage->GetLODCount());
 
-		uint8* pDestBuf = pResult->GetData();
+		const EImageFormat BaseFormat = BaseImage->GetFormat();
 
-		if (pMask->GetFormat() != EImageFormat::IF_L_UBYTE)
+		if (MaskImage->GetFormat() != EImageFormat::IF_L_UBYTE)
 		{
 			checkf(false, TEXT("Unsupported mask format."));
 
-			BufferLayerCombine< BLEND_FUNC, 1 >(pDestBuf, pBase, pBlended, bOnlyFirstLOD);
+			BufferLayerCombine<BLEND_FUNC, 1>(ResultImage, BaseImage, BlendedImage, bOnlyFirstLOD);
 		}
 
-		if (baseFormat == EImageFormat::IF_L_UBYTE)
+		if (BaseFormat == EImageFormat::IF_L_UBYTE)
 		{
-			BufferLayerCombine< BLEND_FUNC_MASKED, 1 >(pDestBuf, pBase, pMask, pBlended, bOnlyFirstLOD);
+			BufferLayerCombine<BLEND_FUNC_MASKED, 1>(ResultImage, BaseImage, MaskImage, BlendedImage, bOnlyFirstLOD);
 		}
-		else if (baseFormat == EImageFormat::IF_RGB_UBYTE)
+		else if (BaseFormat == EImageFormat::IF_RGB_UBYTE)
 		{
-			BufferLayerCombine< BLEND_FUNC_MASKED, 3 >(pDestBuf, pBase, pMask, pBlended, bOnlyFirstLOD);
+			BufferLayerCombine<BLEND_FUNC_MASKED, 3>(ResultImage, BaseImage, MaskImage, BlendedImage, bOnlyFirstLOD);
 		}
-		else if (baseFormat == EImageFormat::IF_RGBA_UBYTE || baseFormat == EImageFormat::IF_BGRA_UBYTE)
+		else if (BaseFormat == EImageFormat::IF_RGBA_UBYTE || BaseFormat == EImageFormat::IF_BGRA_UBYTE)
 		{
 			// \todo: pass swizzle template argument if BGRA_UBYTE, not yet supported.
-			BufferLayerCombine< BLEND_FUNC_MASKED, 4 >(pDestBuf, pBase, pMask, pBlended, bOnlyFirstLOD);
+			BufferLayerCombine<BLEND_FUNC_MASKED, 4>(ResultImage, BaseImage, MaskImage, BlendedImage, bOnlyFirstLOD);
 		}
 		else
 		{
@@ -2485,30 +2315,28 @@ namespace mu
 		}
 	}
 
-	template< uint32 (*BLEND_FUNC)(uint32, uint32)>
-	inline void ImageLayerCombineColour(Image* pResult, const Image* pBase, FVector4f col)
+	template<uint32 (*BLEND_FUNC)(uint32, uint32)>
+	inline void ImageLayerCombineColour(Image* ResultImage, const Image* BaseImage, FVector4f Color)
 	{
-		check(pResult->GetFormat() == pBase->GetFormat());
-		check(pResult->GetSizeX() == pBase->GetSizeX());
-		check(pResult->GetSizeY() == pBase->GetSizeY());
-		check(pResult->GetLODCount() == pBase->GetLODCount());
+		check(ResultImage->GetFormat() == BaseImage->GetFormat());
+		check(ResultImage->GetSizeX() == BaseImage->GetSizeX());
+		check(ResultImage->GetSizeY() == BaseImage->GetSizeY());
+		check(ResultImage->GetLODCount() == BaseImage->GetLODCount());
 
-		const EImageFormat baseFormat = pBase->GetFormat();
+		const EImageFormat BaseFormat = BaseImage->GetFormat();
 
-		uint8* pDestBuf = pResult->GetData();
-
-		if (baseFormat == EImageFormat::IF_L_UBYTE)
+		if (BaseFormat == EImageFormat::IF_L_UBYTE)
 		{
-			BufferLayerCombineColour< BLEND_FUNC, 1 >(pDestBuf, pBase, col);
+			BufferLayerCombineColour<BLEND_FUNC, 1>(ResultImage, BaseImage, Color);
 		}
-		else if (baseFormat == EImageFormat::IF_RGB_UBYTE)
+		else if (BaseFormat == EImageFormat::IF_RGB_UBYTE)
 		{
-			BufferLayerCombineColour< BLEND_FUNC, 3 >(pDestBuf, pBase, col);
+			BufferLayerCombineColour<BLEND_FUNC, 3>(ResultImage, BaseImage, Color);
 		}
-		else if (baseFormat == EImageFormat::IF_RGBA_UBYTE || baseFormat == EImageFormat::IF_BGRA_UBYTE)
+		else if (BaseFormat == EImageFormat::IF_RGBA_UBYTE || BaseFormat == EImageFormat::IF_BGRA_UBYTE)
 		{
 			// \todo: pass swizzle template argument if BGRA_UBYTE, not yet supported.
- 			BufferLayerCombineColour< BLEND_FUNC, 4 >(pDestBuf, pBase, col);
+ 			BufferLayerCombineColour<BLEND_FUNC, 4>(ResultImage, BaseImage, Color);
 		}
 		else
 		{
@@ -2518,37 +2346,35 @@ namespace mu
 
 	template< 
 		uint32 (*BLEND_FUNC)(uint32, uint32),
-		uint32 (*BLEND_FUNC_MASKED)(uint32, uint32, uint32) >
-	inline void ImageLayerCombineColour(Image* pResult, const Image* pBase, const Image* pMask, FVector4f col)
+		uint32 (*BLEND_FUNC_MASKED)(uint32, uint32, uint32)>
+	inline void ImageLayerCombineColour(Image* ResultImage, const Image* BaseImage, const Image* MaskImage, FVector4f Color)
 	{
-		check(pResult->GetFormat() == pBase->GetFormat());
-		check(pResult->GetSizeX() == pBase->GetSizeX());
-		check(pResult->GetSizeY() == pBase->GetSizeY());
-		check(pResult->GetLODCount() == pBase->GetLODCount());
-
-		const EImageFormat baseFormat = pBase->GetFormat();
+		check(ResultImage->GetFormat() == BaseImage->GetFormat());
+		check(ResultImage->GetSizeX() == BaseImage->GetSizeX());
+		check(ResultImage->GetSizeY() == BaseImage->GetSizeY());
+		check(ResultImage->GetLODCount() == BaseImage->GetLODCount());
 
-		uint8* pDestBuf = pResult->GetData();
+		const EImageFormat BaseFormat = BaseImage->GetFormat();
 
-		if (pMask->GetFormat() != EImageFormat::IF_L_UBYTE)
+		if (MaskImage->GetFormat() != EImageFormat::IF_L_UBYTE)
 		{
 			checkf(false, TEXT("Unsupported mask format."));
 
-			BufferLayerCombineColour< BLEND_FUNC, 1 >(pDestBuf, pBase, col);
+			BufferLayerCombineColour<BLEND_FUNC, 1>(ResultImage, BaseImage, Color);
 		}
 
-		if (baseFormat == EImageFormat::IF_L_UBYTE)
+		if (BaseFormat == EImageFormat::IF_L_UBYTE)
 		{
-			BufferLayerCombineColour< BLEND_FUNC_MASKED, 1 >(pDestBuf, pBase, pMask, col);
+			BufferLayerCombineColour<BLEND_FUNC_MASKED, 1>(ResultImage, BaseImage, MaskImage, Color);
 		}
-		else if (baseFormat == EImageFormat::IF_RGB_UBYTE)
+		else if (BaseFormat == EImageFormat::IF_RGB_UBYTE)
 		{
-			BufferLayerCombineColour< BLEND_FUNC_MASKED, 3 >(pDestBuf, pBase, pMask, col);
+			BufferLayerCombineColour<BLEND_FUNC_MASKED, 3>(ResultImage, BaseImage, MaskImage, Color);
 		}
-		else if (baseFormat == EImageFormat::IF_RGBA_UBYTE || baseFormat == EImageFormat::IF_BGRA_UBYTE)
+		else if (BaseFormat == EImageFormat::IF_RGBA_UBYTE || BaseFormat == EImageFormat::IF_BGRA_UBYTE)
 		{
 			// \todo: pass swizzle template argument if BGRA_UBYTE, not yet supported.
- 			BufferLayerCombineColour< BLEND_FUNC_MASKED, 4 >(pDestBuf, pBase, pMask, col);
+ 			BufferLayerCombineColour<BLEND_FUNC_MASKED, 4>(ResultImage, BaseImage, MaskImage, Color);
 		}
 		else
 		{
@@ -2556,98 +2382,127 @@ namespace mu
 		}
 	}
 
-	template< size_t NCBase, size_t NCBlend, class ImageCombineFn >
-	inline void BufferLayerCombineFunctor(uint8* pDestBuf, const Image* pBase, const Image* pBlended, ImageCombineFn&& ImageCombine)
+	template<uint32 NCBase, uint32 NCBlend, class ImageCombineFn>
+	inline void BufferLayerCombineFunctor(Image* DestImage, const Image* BaseImage, const Image* BlendImage, ImageCombineFn&& ImageCombine)
 	{
-		//static_assert(NCBase > 0 && NCBase <= 4);
-		//static_assert(NCBlend > 0 && NCBlend <= 4);
+		static_assert(NCBase > 0 && NCBase <= 4);
+		static_assert(NCBlend > 0 && NCBlend <= 4);
+
+		check(BaseImage->GetFormat() == DestImage->GetFormat());
+		check(BaseImage->GetSizeX() == BlendImage->GetSizeX());
+		check(BaseImage->GetSizeY() == BlendImage->GetSizeY());
+		check(BaseImage->GetSizeX() == DestImage->GetSizeX());
+		check(BaseImage->GetSizeY() == DestImage->GetSizeY());
+		check(BaseImage->GetLODCount() <= BlendImage->GetLODCount());
+		check(BaseImage->GetLODCount() <= DestImage->GetLODCount());
+	
+		constexpr int32 BatchNumElems = 4098*2;
+		int32 NumBatches = BaseImage->DataStorage.GetNumBatches(BatchNumElems, NCBase);
+
+		check(NumBatches == BlendImage->DataStorage.GetNumBatches(BatchNumElems, NCBlend));
+		check(NumBatches == DestImage->DataStorage.GetNumBatches(BatchNumElems, NCBase));
 
-		check(pBase->GetSizeX() == pBlended->GetSizeX());
-		check(pBase->GetSizeY() == pBlended->GetSizeY());
-		check(pBase->GetLODCount() <= pBlended->GetLODCount());
+		auto ProcessBatch = [DestImage, BaseImage, BlendImage, BatchNumElems, ImageCombine](int32 BatchId)
+		{
+			const TArrayView<const uint8> BaseBatchView = BaseImage->DataStorage.GetBatch(BatchId, BatchNumElems, NCBase);
+			const TArrayView<const uint8> BlendBatchView = BlendImage->DataStorage.GetBatch(BatchId, BatchNumElems, NCBlend);
+			const TArrayView<uint8> DestBatchView = DestImage->DataStorage.GetBatch(BatchId, BatchNumElems, NCBase);
 
-		const uint8* pBaseBuf = pBase->GetData();
-		const uint8* pBlendedBuf = pBlended->GetData();
+			const int32 NumElems = BaseBatchView.Num() / NCBase;
 
-		int32 pixelCount = pBase->CalculatePixelCount();
+			const uint8* BaseBuf = BaseBatchView.GetData();
+			const uint8* BlendBuf = BlendBatchView.GetData();
+			uint8* DestBuf = DestBatchView.GetData();
 
-		for (int i = 0; i < pixelCount; ++i)
-		{
-			const uint32 base = PackPixel<NCBase>(&pBaseBuf[NCBase * i]);
-			const uint32 blend = PackPixel<NCBlend>(&pBlendedBuf[NCBlend * i]);
+			check(NumElems == BlendBatchView.Num() / NCBlend);
+			check(NumElems == DestBatchView.Num() / NCBase);
+			
+			for (int32 I = 0; I < NumElems; ++I)
+			{
+				const uint32 Base = PackPixel<NCBase>(&BaseBuf[NCBase * I]);
+				const uint32 Blend = PackPixel<NCBlend>(&BlendBuf[NCBlend * I]);
 
-			const uint32 result = ImageCombine(base, blend);
+				const uint32 Result = ImageCombine(Base, Blend);
+				UnpackPixel<NCBase>(&DestBuf[NCBase* I], Result);
+			}
+		};
 
-			UnpackPixel<NCBase>(&pDestBuf[NCBase* i], result);
+		if (NumBatches == 1)
+		{
+			ProcessBatch(0);
+		}
+		else if (NumBatches > 1)
+		{
+			ParallelFor(NumBatches, ProcessBatch);
 		}
 	}
+
 	// Same functionality as above, in this case we use a functor which allows to pass user data. 
-	template< class ImageCombineFn >
-	inline void ImageLayerCombineFunctor(Image* pResult, const Image* pBase, const Image* pBlended, ImageCombineFn&& ImageCombine)
+	template<class ImageCombineFn>
+	inline void ImageLayerCombineFunctor(Image* ResultImage, const Image* BaseImage, const Image* BlendedImage, ImageCombineFn&& ImageCombine)
 	{
-		check(pResult->GetFormat() == pBase->GetFormat());
-		check(pResult->GetSizeX() == pBase->GetSizeX());
-		check(pResult->GetSizeY() == pBase->GetSizeY());
-		check(pResult->GetLODCount() == pBase->GetLODCount());
-		check(pBase->GetSizeX() == pBlended->GetSizeX());
-		check(pBase->GetSizeY() == pBlended->GetSizeY());
-		check(pResult->GetLODCount() <= pBlended->GetLODCount());
+		check(ResultImage->GetFormat() == BaseImage->GetFormat());
+		check(ResultImage->GetSizeX() == BaseImage->GetSizeX());
+		check(ResultImage->GetSizeY() == BaseImage->GetSizeY());
+		check(ResultImage->GetLODCount() == BaseImage->GetLODCount());
+		check(BaseImage->GetSizeX() == BlendedImage->GetSizeX());
+		check(BaseImage->GetSizeY() == BlendedImage->GetSizeY());
+		check(ResultImage->GetLODCount() <= BlendedImage->GetLODCount());
 
-		const EImageFormat baseFormat = pBase->GetFormat();
-		const EImageFormat blendFormat = pBlended->GetFormat();
+		const EImageFormat BaseFormat = BaseImage->GetFormat();
+		const EImageFormat BlendFormat = BlendedImage->GetFormat();
 
-		uint8* pDestBuf = pResult->GetData();
-		if (baseFormat == EImageFormat::IF_L_UBYTE )
+		if (BaseFormat == EImageFormat::IF_L_UBYTE )
 		{
-			if (blendFormat == EImageFormat::IF_L_UBYTE )
+			if (BlendFormat == EImageFormat::IF_L_UBYTE )
 			{
-				BufferLayerCombineFunctor< 1, 1 >(pDestBuf, pBase, pBlended, Forward<ImageCombineFn>(ImageCombine));
+				BufferLayerCombineFunctor<1, 1>(ResultImage, BaseImage, BlendedImage, Forward<ImageCombineFn>(ImageCombine));
 			}
-			else if (blendFormat == EImageFormat::IF_RGB_UBYTE )
+			else if (BlendFormat == EImageFormat::IF_RGB_UBYTE )
 			{
-				BufferLayerCombineFunctor< 1, 3 >(pDestBuf, pBase, pBlended, Forward<ImageCombineFn>(ImageCombine));
+				BufferLayerCombineFunctor<1, 3>(ResultImage, BaseImage, BlendedImage, Forward<ImageCombineFn>(ImageCombine));
 			}
-			else if (blendFormat == EImageFormat::IF_RGBA_UBYTE || blendFormat == EImageFormat::IF_BGRA_UBYTE )
+			else if (BlendFormat == EImageFormat::IF_RGBA_UBYTE || BlendFormat == EImageFormat::IF_BGRA_UBYTE)
 			{
-				BufferLayerCombineFunctor< 1, 4 >(pDestBuf, pBase, pBlended, Forward<ImageCombineFn>(ImageCombine));
+				BufferLayerCombineFunctor<1, 4>(ResultImage, BaseImage, BlendedImage, Forward<ImageCombineFn>(ImageCombine));
 			}
 			else
 			{
 				checkf(false, TEXT("Unsupported format."));
 			}
 		}
-		else if (baseFormat == EImageFormat::IF_RGB_UBYTE)
+		else if (BaseFormat == EImageFormat::IF_RGB_UBYTE)
 		{
-			if (blendFormat == EImageFormat::IF_L_UBYTE )
+			if (BlendFormat == EImageFormat::IF_L_UBYTE )
 			{
-				BufferLayerCombineFunctor< 3, 1 >(pDestBuf, pBase, pBlended, Forward<ImageCombineFn>(ImageCombine));
+				BufferLayerCombineFunctor<3, 1>(ResultImage, BaseImage, BlendedImage, Forward<ImageCombineFn>(ImageCombine));
 			}
-			else if (blendFormat == EImageFormat::IF_RGB_UBYTE )
+			else if (BlendFormat == EImageFormat::IF_RGB_UBYTE )
 			{
-				BufferLayerCombineFunctor< 3, 3 >(pDestBuf, pBase, pBlended, Forward<ImageCombineFn>(ImageCombine));
+				BufferLayerCombineFunctor<3, 3>(ResultImage, BaseImage, BlendedImage, Forward<ImageCombineFn>(ImageCombine));
 			}
-			else if (blendFormat == EImageFormat::IF_RGBA_UBYTE || blendFormat == EImageFormat::IF_BGRA_UBYTE )
+			else if (BlendFormat == EImageFormat::IF_RGBA_UBYTE || BlendFormat == EImageFormat::IF_BGRA_UBYTE )
 			{
-				BufferLayerCombineFunctor< 3, 4 >(pDestBuf, pBase, pBlended, Forward<ImageCombineFn>(ImageCombine));
+				BufferLayerCombineFunctor<3, 4>(ResultImage, BaseImage, BlendedImage, Forward<ImageCombineFn>(ImageCombine));
 			}
 			else
 			{
 				checkf(false, TEXT("Unsupported format."));
 			}
 		}
-		else if (baseFormat == EImageFormat::IF_RGBA_UBYTE || baseFormat == EImageFormat::IF_BGRA_UBYTE)
+		else if (BaseFormat == EImageFormat::IF_RGBA_UBYTE || BaseFormat == EImageFormat::IF_BGRA_UBYTE)
 		{
-			if (blendFormat == EImageFormat::IF_L_UBYTE )
+			if (BlendFormat == EImageFormat::IF_L_UBYTE )
 			{
-				BufferLayerCombineFunctor< 4, 1 >(pDestBuf, pBase, pBlended, Forward<ImageCombineFn>(ImageCombine));
+				BufferLayerCombineFunctor<4, 1>(ResultImage, BaseImage, BlendedImage, Forward<ImageCombineFn>(ImageCombine));
 			}
-			else if (blendFormat == EImageFormat::IF_RGB_UBYTE )
+			else if (BlendFormat == EImageFormat::IF_RGB_UBYTE )
 			{
-				BufferLayerCombineFunctor< 4, 3 >(pDestBuf, pBase, pBlended, Forward<ImageCombineFn>(ImageCombine));
+				BufferLayerCombineFunctor<4, 3>(ResultImage, BaseImage, BlendedImage, Forward<ImageCombineFn>(ImageCombine));
 			}
-			else if (blendFormat == EImageFormat::IF_RGBA_UBYTE || blendFormat == EImageFormat::IF_BGRA_UBYTE )
+			else if (BlendFormat == EImageFormat::IF_RGBA_UBYTE || BlendFormat == EImageFormat::IF_BGRA_UBYTE )
 			{
-				BufferLayerCombineFunctor< 4, 4 >(pDestBuf, pBase, pBlended, Forward<ImageCombineFn>(ImageCombine));
+				BufferLayerCombineFunctor<4, 4>(ResultImage, BaseImage, BlendedImage, Forward<ImageCombineFn>(ImageCombine));
 			}
 			else
 			{
@@ -2662,13 +2517,12 @@ namespace mu
 
 
 	//! Blend a subimage on the base using a mask.
-	inline void ImageBlendOnBaseNoAlpha(Image* pBase,
-		const Image* pMask,
-		const Image* pBlended,
-		const box< UE::Math::TIntVector2<uint16> >& rect)
+	inline void ImageBlendOnBaseNoAlpha(Image* BaseImage,
+		const Image* MaskImage,
+		const Image* BlendedImage,
+		const box<UE::Math::TIntVector2<uint16>>& Rect)
 	{
-		ImageLayerOnBaseNoAlpha< BlendChannel, false>
-			(pBase, pMask, pBlended, rect);
+		ImageLayerOnBaseNoAlpha<BlendChannel, false>(BaseImage, MaskImage, BlendedImage, Rect);
 	}
 
 }
