[Mutable] Store Image LODs in different buffers.

diff --git a/Engine/Plugins/Experimental/Mutable/Source/MutableRuntime/Private/MuR/OpImageLuminance.h b/Engine/Plugins/Experimental/Mutable/Source/MutableRuntime/Private/MuR/OpImageLuminance.h
--- a/Engine/Plugins/Experimental/Mutable/Source/MutableRuntime/Private/MuR/OpImageLuminance.h
+++ b/Engine/Plugins/Experimental/Mutable/Source/MutableRuntime/Private/MuR/OpImageLuminance.h
@@ -5,55 +5,77 @@
 #include "MuR/ImagePrivate.h"
 #include "MuR/Platform.h"
 
+#include "Async/ParallelFor.h"
+
 namespace mu
 {
 
-	inline void ImageLuminance( Image* pDest, const Image* pA )
+	inline void ImageLuminance(Image* DestImage, const Image* AImage)
 	{
-		check(pDest && pA && pDest->GetFormat()==EImageFormat::IF_L_UBYTE);
-
-        uint8_t* pDestBuf = pDest->GetData();
-        const uint8_t* pABuf = pA->GetData();
-
-		// Generic implementation
-		int pixelCount = (int)pA->CalculatePixelCount();
-
-		switch ( pA->GetFormat() )
-		{
-		case EImageFormat::IF_RGB_UBYTE:
-		{
-			for ( int i=0; i<pixelCount; ++i )
-			{
-                uint32_t l_16 = 76 * pABuf[3*i+0] + 150 * pABuf[3*i+1] + 29 * pABuf[3*i+2];
-                pDestBuf[i] = (uint8_t)FMath::Min( l_16>>8, 255u );
-			}
-			break;
-		}
-
-        case EImageFormat::IF_RGBA_UBYTE:
+		check(DestImage && AImage && DestImage->GetFormat() == EImageFormat::IF_L_UBYTE);
+
+		const int32 BytesPerElem = GetImageFormatData(AImage->GetFormat()).BytesPerBlock;
+		constexpr int32 NumBatchElems = 1 << 14;
+
+		const int32 NumBatches = DestImage->DataStorage.GetNumBatches(NumBatchElems, 1); 
+		check(NumBatches == AImage->DataStorage.GetNumBatches(NumBatchElems, BytesPerElem)); 
+
+        auto ProcessBatch = [DestImage, AImage, NumBatchElems, BytesPerElem](int32 BatchId)
         {
-            for ( int i=0; i<pixelCount; ++i )
+			TArrayView<uint8> DestView = DestImage->DataStorage.GetBatch(BatchId, NumBatchElems, 1);
+			TArrayView<const uint8> AView = AImage->DataStorage.GetBatch(BatchId, NumBatchElems, BytesPerElem);
+
+			uint8* DestBuf = DestView.GetData(); 
+			const uint8* ABuf = AView.GetData(); 
+			
+            const int32 NumElems = DestView.Num() / 1;
+
+            switch (AImage->GetFormat())
+            {
+            case EImageFormat::IF_RGB_UBYTE:
             {
-                uint32_t l_16 = 76 * pABuf[4*i+0] + 150 * pABuf[4*i+1] + 29 * pABuf[4*i+2];
-                pDestBuf[i] = (uint8_t)FMath::Min( l_16>>8, 255u );
+                for (int32 I = 0; I < NumElems; ++I)
+                {
+                    const uint16 L = (76*ABuf[3*I + 0] + 150*ABuf[3*I + 1] + 29*ABuf[3*I + 2]);
+                    DestBuf[I] = uint8(L / 255);
+                }
+                break;
             }
-            break;
-        }
 
-        case EImageFormat::IF_BGRA_UBYTE:
-        {
-            for ( int i=0; i<pixelCount; ++i )
+            case EImageFormat::IF_RGBA_UBYTE:
             {
-                uint32_t l_16 = 76 * pABuf[4*i+2] + 150 * pABuf[4*i+1] + 29 * pABuf[4*i+0];
-                pDestBuf[i] = (uint8_t)FMath::Min( l_16>>8, 255u );
+                for (int32 I = 0; I < NumElems; ++I)
+                {
+                    const uint16 L = (76*ABuf[4*I + 0] + 150*ABuf[4*I + 1] + 29*ABuf[4*I + 2]);
+                    DestBuf[I] = uint8(L / 255);
+                }
+                break;
             }
-            break;
-        }
 
-		default:
-			check(false);
-			break;
-		}
+            case EImageFormat::IF_BGRA_UBYTE:
+            {
+                for (int32 I = 0; I < NumElems; ++I)
+                {
+                    const uint16 L = (76*ABuf[4*I + 2] + 150*ABuf[4*I + 1] + 29*ABuf[4*I + 0]);
+                    DestBuf[I] = uint8(L / 255);
+                }
+                break;
+            }
+
+            default:
+                check(false);
+                break;
+            }
+        };
+
+        if (NumBatches == 1)
+        {
+            ProcessBatch(0);
+        }
+        else
+        {
+            ParallelFor(NumBatches, ProcessBatch);
+        }
 	}
 
 }
diff --git a/Engine/Plugins/Experimental/Mutable/Source/MutableRuntime/Private/MuR/OpImageMipmap.cpp b/Engine/Plugins/Experimental/Mutable/Source/MutableRuntime/Private/MuR/OpImageMipmap.cpp
--- a/Engine/Plugins/Experimental/Mutable/Source/MutableRuntime/Private/MuR/OpImageMipmap.cpp
+++ b/Engine/Plugins/Experimental/Mutable/Source/MutableRuntime/Private/MuR/OpImageMipmap.cpp
@@ -20,220 +20,182 @@ namespace mu
 {
 
 
-namespace OpImageMipmap_Detail
+namespace OpImageMipmapInternal
 {
 
 	template<int32 PIXEL_SIZE>
-	inline void GenerateMipmapUint8Unfiltered(
-		int mips,
-		const uint8* pSource, uint8* Dest,
-		FIntVector2 sourceSize)
+	inline void GenerateNextMipmapUint8Unfiltered(const uint8* SourceData, uint8* DestData, FIntVector2 SourceSize)
 	{
-		for (; mips >= 0; --mips)
-		{
-			FIntVector2 destSize = FIntVector2(FMath::DivideAndRoundUp(sourceSize.X, 2), FMath::DivideAndRoundUp(sourceSize.Y, 2));
+		FIntVector2 DestSize = FIntVector2(
+				FMath::DivideAndRoundUp(SourceSize.X, 2), 
+				FMath::DivideAndRoundUp(SourceSize.Y, 2));
 
-			for (int32 y = 0; y < destSize.Y; ++y)
+		for (int32 Y = 0; Y < DestSize.Y; ++Y)
+		{
+			for (int32 X = 0; X < DestSize.X; ++X)
 			{
-				for (int32 x = 0; x < destSize.X; ++x)
+				for (int32 C = 0; C < PIXEL_SIZE; ++C)
 				{
-					for (int32 c = 0; c < PIXEL_SIZE; ++c)
-					{
-						Dest[(y * destSize.X + x) * PIXEL_SIZE + c] =
-							pSource[((y << 1) * sourceSize.X + (x << 1)) * PIXEL_SIZE + c];
-					}
+					DestData[(Y * DestSize.X + X) * PIXEL_SIZE + C] =
+						SourceData[((Y << 1) * SourceSize.X + (X << 1)) * PIXEL_SIZE + C];
 				}
 			}
-
-			sourceSize = destSize;
-			pSource = Dest;
-			Dest = Dest + destSize.X * destSize.Y * PIXEL_SIZE;
 		}
 	}
 
 	template<int32 PIXEL_SIZE>
-	inline void GenerateMipmapsUint8SimpleAverage(
-		int mips,
-		const uint8* pSource, uint8* Dest,
-		FIntVector2 sourceSize)
+	inline void GenerateNextMipmapUint8SimpleAverage(const uint8* SourceData, uint8* DestData, FIntVector2 SourceSize)
 	{
-		
-		const uint8* pMipSource = pSource;
-		uint8* pMipDest = Dest;
-
-		FIntVector2 destSize = sourceSize;
+		check(SourceSize[0] > 1 || SourceSize[1] > 1);
 
-		for (int m = 0; m < mips; ++m)
-		{
-			check(destSize[0] > 1 || destSize[1] > 1);
+		FIntVector2 DestSize = FIntVector2(
+				FMath::DivideAndRoundUp(SourceSize.X, 2),
+				FMath::DivideAndRoundUp(SourceSize.Y, 2));
 
-			sourceSize = destSize;
+		int32 FullColumns = SourceSize.X / 2;
+		int32 FullRows = SourceSize.Y / 2;
+		bool bStrayColumn = (SourceSize.X % 2) != 0;
+		bool bStrayRow = (SourceSize.Y % 2) != 0;
 
-			int fullColumns = destSize[0] / 2;
-			bool strayColumn = (destSize[0] % 2) != 0;
-			int fullRows = destSize[1] / 2;
-			bool strayRow = (destSize[1] % 2) != 0;
+		int32 SourceStride = SourceSize.X * PIXEL_SIZE;
+		int32 DestStride = DestSize.X * PIXEL_SIZE;
 
-			destSize[0] = FMath::DivideAndRoundUp(destSize[0], 2);
-			destSize[1] = FMath::DivideAndRoundUp(destSize[1], 2);
-
-			int sourceStride = sourceSize[0] * PIXEL_SIZE;
-			int destStride = destSize[0] * PIXEL_SIZE;
+		const auto ProcessRow = [
+			DestData, SourceData, FullColumns, bStrayColumn, SourceStride, DestStride
+		] (uint32 Y)
+		{
+			const uint8* SourceRow0 = SourceData + 2*Y*SourceStride;
+			const uint8* SourceRow1 = SourceRow0 + SourceStride;
+			uint8* DestRow = DestData + Y * DestStride;
 
-			const auto ProcessRow = [
-				pMipDest, pMipSource, fullColumns, strayColumn, sourceStride, destStride
-			] (uint32 y)
+			for (int32 X = 0; X < FullColumns; ++X)
 			{
-				const uint8* pSourceRow0 = pMipSource + 2 * y * sourceStride;
-				const uint8* pSourceRow1 = pSourceRow0 + sourceStride;
-				uint8* pDestRow = pMipDest + y * destStride;
-
-				for (int x = 0; x < fullColumns; ++x)
+				if constexpr (PIXEL_SIZE == 4)
 				{
-					if constexpr (PIXEL_SIZE == 4)
-					{
-						//const uint64 Row0Bits = *reinterpret_cast<const uint64*>(pSourceRow0);
-						//const uint64 Row1Bits = *reinterpret_cast<const uint64*>(pSourceRow1);
-
-						// Use memcpy to avoid any possible but improbable UB. memcpy should be optimized away by the compiler.
-						uint64 Row0Bits; 
-						uint64 Row1Bits;
-
-						FMemory::Memcpy(&Row0Bits, pSourceRow0, sizeof(uint64));
-						FMemory::Memcpy(&Row1Bits, pSourceRow1, sizeof(uint64));
-						
-						const uint64 XorRow0Row1Bits = Row0Bits ^ Row1Bits;
-
-						// Average of 2 unsigned integers without overflow extended to work on multiple bytes.
-						constexpr uint64 ShiftMask = 0xFEFEFEFEFEFEFEFE;
-						const uint64 ErrorCorrection = XorRow0Row1Bits & 0x0101010101010101;
-						const uint64 AvgLowBits = (Row0Bits & Row1Bits) + ((XorRow0Row1Bits & ShiftMask) >> 1) + ErrorCorrection;
-						const uint64 AvgHighBits = AvgLowBits >> 32;
-						const uint32 Result = (AvgLowBits & AvgHighBits) + (((AvgLowBits ^ AvgHighBits) & ShiftMask) >> 1);
+					// Use memcpy to avoid any possible but improbable UB. memcpy should be optimized away by the compiler.
+					uint64 Row0Bits; 
+					uint64 Row1Bits;
+
+					FMemory::Memcpy(&Row0Bits, SourceRow0, sizeof(uint64));
+					FMemory::Memcpy(&Row1Bits, SourceRow1, sizeof(uint64));
 					
-						//*reinterpret_cast<uint32*>(pDestRow) = Result;
-						FMemory::Memcpy(pDestRow, &Result, sizeof(uint32));
-					}
-					else
+					const uint64 XorRow0Row1Bits = Row0Bits ^ Row1Bits;
+
+					// Average of 2 unsigned integers without overflow extended to work on multiple bytes.
+					constexpr uint64 ShiftMask = 0xFEFEFEFEFEFEFEFE;
+					const uint64 ErrorCorrection = XorRow0Row1Bits & 0x0101010101010101;
+					const uint64 AvgLowBits = (Row0Bits & Row1Bits) + ((XorRow0Row1Bits & ShiftMask) >> 1) + ErrorCorrection;
+					const uint64 AvgHighBits = AvgLowBits >> 32;
+					const uint32 Result = (AvgLowBits & AvgHighBits) + (((AvgLowBits ^ AvgHighBits) & ShiftMask) >> 1);
+				
+					FMemory::Memcpy(DestRow, &Result, sizeof(uint32));
+				}
+				else
+				{
+					for (int32 C = 0; C < PIXEL_SIZE; ++C)
 					{
-						for (int32 C = 0; C < PIXEL_SIZE; ++C)
-						{
-							int32 PixelSum = pSourceRow0[C] + pSourceRow0[PIXEL_SIZE + C] + pSourceRow1[C] + pSourceRow1[PIXEL_SIZE + C];
-							pDestRow[C] = (uint8)(PixelSum >> 2);
-						}
+						int32 PixelSum = SourceRow0[C] + SourceRow0[PIXEL_SIZE + C] + SourceRow1[C] + SourceRow1[PIXEL_SIZE + C];
+						DestRow[C] = (uint8)(PixelSum >> 2);
 					}
-
-					pSourceRow0 += 2 * PIXEL_SIZE;
-					pSourceRow1 += 2 * PIXEL_SIZE;
-					pDestRow += PIXEL_SIZE;
 				}
 
-				if (strayColumn)
-				{
-					if constexpr (PIXEL_SIZE == 4)
-					{	
-						//const uint32 Row0Bits = *reinterpret_cast<const uint32*>(pSourceRow0);
-						//const uint32 Row1Bits = *reinterpret_cast<const uint32*>(pSourceRow1);
+				SourceRow0 += 2*PIXEL_SIZE;
+				SourceRow1 += 2*PIXEL_SIZE;
+				DestRow += PIXEL_SIZE;
+			}
 
-						uint32 Row0Bits; 
-						uint32 Row1Bits;
+			if (bStrayColumn)
+			{
+				if constexpr (PIXEL_SIZE == 4)
+				{	
+					uint32 Row0Bits; 
+					uint32 Row1Bits;
 
-						FMemory::Memcpy(&Row0Bits, pSourceRow0, sizeof(uint32));
-						FMemory::Memcpy(&Row1Bits, pSourceRow1, sizeof(uint32));
+					FMemory::Memcpy(&Row0Bits, SourceRow0, sizeof(uint32));
+					FMemory::Memcpy(&Row1Bits, SourceRow1, sizeof(uint32));
 
-						// Average of 2 unsigned integers without overflow extended to work on multiple bytes.
-						constexpr uint32 ShiftMask = 0xFEFEFEFE;
-						const uint32 Result = (Row0Bits & Row1Bits) + (((Row0Bits ^ Row1Bits) & ShiftMask) >> 1);
+					// Average of 2 unsigned integers without overflow extended to work on multiple bytes.
+					constexpr uint32 ShiftMask = 0xFEFEFEFE;
+					const uint32 Result = (Row0Bits & Row1Bits) + (((Row0Bits ^ Row1Bits) & ShiftMask) >> 1);
 
-						//*reinterpret_cast<uint32*>(pDestRow) = Result;
-						FMemory::Memcpy(pDestRow, &Result, sizeof(uint32));
-					}
-					else
-					{
-						for (int32 C = 0; C < PIXEL_SIZE; ++C)
-						{
-							int32 PixelSum = pSourceRow0[C] + pSourceRow1[C];
-							pDestRow[C] = (uint8)(PixelSum >> 1);
-						}
-					}
+					FMemory::Memcpy(DestRow, &Result, sizeof(uint32));
 				}
-			};
-
-			constexpr int PixelConcurrencyThreshold = 0xffff;
-			if (destSize[0] * destSize[1] < PixelConcurrencyThreshold)
-			{
-				for (int y = 0; y < fullRows; ++y)
+				else
 				{
-					ProcessRow(y);
+					for (int32 C = 0; C < PIXEL_SIZE; ++C)
+					{
+						int32 PixelSum = SourceRow0[C] + SourceRow1[C];
+						DestRow[C] = (uint8)(PixelSum >> 1);
+					}
 				}
 			}
-			else
+		};
+
+		constexpr int32 PixelConcurrencyThreshold = 0xffff;
+		if (DestSize[0] * DestSize[1] < PixelConcurrencyThreshold)
+		{
+			for (int32 Y = 0; Y < FullRows; ++Y)
 			{
-				ParallelFor(fullRows, ProcessRow);
+				ProcessRow(Y);
 			}
+		}
+		else
+		{
+			ParallelFor(FullRows, ProcessRow);
+		}
 
-			if (strayRow)
-			{
-				const uint8* pSourceRow0 = pMipSource + 2 * fullRows * sourceStride;
-				const uint8* pSourceRow1 = pSourceRow0 + sourceStride;
-				uint8* pDestRow = pMipDest + fullRows * destStride;
+		if (bStrayRow)
+		{
+			const uint8* SourceRow0 = SourceData + 2 * FullRows * SourceStride;
+			const uint8* SourceRow1 = SourceRow0 + SourceStride;
+			uint8* DestRow = DestData + FullRows * DestStride;
 
-				for (int x = 0; x < fullColumns; ++x)
+			for (int32 X = 0; X < FullColumns; ++X)
+			{
+				if constexpr (PIXEL_SIZE == 4)
 				{
-					if constexpr (PIXEL_SIZE == 4)
-					{
-						//const uint32 Col0Bits = *reinterpret_cast<const uint32*>(pSourceRow0);
-						//const uint32 Col1Bits = *reinterpret_cast<const uint32*>(pSourceRow0 + 4);
+					uint32 Col0Bits; 
+					uint32 Col1Bits;
 
-						uint32 Col0Bits; 
-						uint32 Col1Bits;
+					FMemory::Memcpy(&Col0Bits, SourceRow0, sizeof(uint32));
+					FMemory::Memcpy(&Col1Bits, SourceRow0 + 4, sizeof(uint32));
 
-						FMemory::Memcpy(&Col0Bits, pSourceRow0, sizeof(uint32));
-						FMemory::Memcpy(&Col1Bits, pSourceRow0 + 4, sizeof(uint32));
+					// Average of 2 unsigned integers without overflow extended to work on multiple bytes. 
+					// In this case we use the ceil variant to be consistent with the method used for 4 pixel average. 
+					constexpr uint32 ShiftMask = 0xFEFEFEFE;
+					const uint32 Result = (Col0Bits & Col1Bits) + (((Col0Bits ^ Col1Bits) & ShiftMask) >> 1);
 
-						// Average of 2 unsigned integers without overflow extended to work on multiple bytes. 
-						// In this case we use the ceil variant to be consistent with the method used for 4 pixel average. 
-						constexpr uint32 ShiftMask = 0xFEFEFEFE;
-						const uint32 Result = (Col0Bits & Col1Bits) + (((Col0Bits ^ Col1Bits) & ShiftMask) >> 1);
-
-						//*reinterpret_cast<uint32*>(pDestRow) = Result;
-						FMemory::Memcpy(pDestRow, &Result, sizeof(uint32));
-					}
-					else
+					FMemory::Memcpy(DestRow, &Result, sizeof(uint32));
+				}
+				else
+				{
+					for (int32 C = 0; C < PIXEL_SIZE; ++C)
 					{
-						for (int32 C = 0; C < PIXEL_SIZE; ++C)
-						{
-							int32 p = pSourceRow0[C] + pSourceRow0[PIXEL_SIZE + C];
-							pDestRow[C] = (uint8)(p >> 1);
-						}
+						int32 P = SourceRow0[C] + SourceRow0[PIXEL_SIZE + C];
+						DestRow[C] = (uint8)(P >> 1);
 					}
-
-					pSourceRow0 += 2 * PIXEL_SIZE;
-					pDestRow += PIXEL_SIZE;
 				}
 
-				if (strayColumn)
+				SourceRow0 += 2*PIXEL_SIZE;
+				DestRow += PIXEL_SIZE;
+			}
+
+			if (bStrayColumn)
+			{
+				if constexpr (PIXEL_SIZE == 4)
 				{
-					if constexpr (PIXEL_SIZE == 4)
-					{
-						//*reinterpret_cast<uint32*>(pDestRow) = *reinterpret_cast<const uint32*>(pSourceRow0);
-						FMemory::Memcpy(pDestRow, pSourceRow0, 4);
-					}
-					else
+					FMemory::Memcpy(DestRow, SourceRow0, 4);
+				}
+				else
+				{
+					for (int32 C = 0; C < PIXEL_SIZE; ++C)
 					{
-						for (int32 C = 0; C < PIXEL_SIZE; ++C)
-						{
-							pDestRow[C] = pSourceRow0[C];
-						}
+						DestRow[C] = SourceRow0[C];
 					}
 				}
 			}
-
-			// Reset the source pointer for the next mip, to use the dest that we have just
-			// generated.
-			pMipSource = pMipDest;
-			pMipDest += destSize[0] * destSize[1] * PIXEL_SIZE;
 		}
+	
 	}
 
 	// Generate next mip decompressed form block compressed image. 
@@ -401,30 +363,65 @@ namespace OpImageMipmap_Detail
 			ParallelFor(NumParallelJobs, ProcessJob);
 		}
 	}
-} // namespace OpImageMipmap_Detail
+} // namespace OpImageMipmapInternal
 
 
 	/** Generate the mipmaps for byte-based images of whatever number of channels.
 	* \param mips number of additional levels to build from the source.
 	*/
 	template<int32 PIXEL_SIZE>
-	inline void GenerateMipmapsUint8(int32 mips,
-		const uint8* pSource, uint8* Dest,
-		FIntVector2 sourceSize,
-		const FMipmapGenerationSettings& settings)
+	inline void GenerateMipmapUint8LODRange(
+		int32 SrcLOD, int32 DestLODBegin, int32 DestLODEnd,
+		const Image* SourceImage, Image* DestImage,
+		const FMipmapGenerationSettings& Settings)
 	{
-		using namespace OpImageMipmap_Detail;
+		using namespace OpImageMipmapInternal;
+
+		FIntVector2 SourceSize = SourceImage->CalculateMipSize(SrcLOD);
 
-		switch (settings.m_filterType)
+		check(Invoke([&]() -> bool
+		{
+			FIntVector2 DestBeginLODSize = DestImage->CalculateMipSize(DestLODBegin);
+
+			FIntVector2 SrcNextLODSize = FIntVector2(
+					FMath::DivideAndRoundUp(SourceSize.X, 2),
+					FMath::DivideAndRoundUp(SourceSize.Y, 2));
+
+			return DestBeginLODSize == SrcNextLODSize;
+		}));
+
+		switch (Settings.m_filterType)
 		{
 		case EMipmapFilterType::MFT_SimpleAverage:
 		{
-			GenerateMipmapsUint8SimpleAverage<PIXEL_SIZE>(mips, pSource, Dest, sourceSize);
+			const uint8* SrcData = SourceImage->GetLODData(SrcLOD);
+
+			for (int32 L = DestLODBegin; L < DestLODEnd; ++L)
+			{
+				uint8* DestData = DestImage->GetLODData(L);
+				GenerateNextMipmapUint8SimpleAverage<PIXEL_SIZE>(SrcData, DestData, SourceSize);
+				
+				SrcData = DestData;
+				SourceSize = FIntVector2(
+						FMath::DivideAndRoundUp(SourceSize.X, 2),
+						FMath::DivideAndRoundUp(SourceSize.Y, 2));
+			}
 			break;
 		}
 		case EMipmapFilterType::MFT_Unfiltered:
 		{
-			GenerateMipmapUint8Unfiltered<PIXEL_SIZE>(mips, pSource, Dest, sourceSize);
+			const uint8* SrcData = SourceImage->GetLODData(SrcLOD);
+
+			for (int32 L = DestLODBegin; L < DestLODEnd; ++L)
+			{
+				uint8* DestData = DestImage->GetLODData(L);
+				GenerateNextMipmapUint8SimpleAverage<PIXEL_SIZE>(SrcData, DestData, SourceSize);
+				
+				SrcData = DestData;
+				SourceSize = FIntVector2(
+						FMath::DivideAndRoundUp(SourceSize.X, 2),
+						FMath::DivideAndRoundUp(SourceSize.Y, 2));
+			}
 			break;
 		}
 		default:
@@ -435,52 +432,92 @@ namespace OpImageMipmap_Detail
 		}
 	}
 
-	/** Generate the mipmaps for Block Comporessed images of whatever number of channels.
-	 *  The result is a non compressed image of the next mip with its tail.
-	* \param mips number of additional levels to build from the source.
-	*/
+	/** 
+	 * Generate the mipmaps for Block Comporessed images of whatever number of channels.
+	 * The result is a non compressed image of the next mip with its tail.
+	 * \param mips number of additional levels to build from the source.
+	 */
 	template<int32 PixelSize>
-	inline void GenerateMipmapsBlockCompressed(int32 Mips,
-		const uint8* SourceData, uint8* DestData,
-		FIntVector2 SourceSize, EImageFormat SrcFormat, EImageFormat DestFormat,
-		const FMipmapGenerationSettings& Settings)
+	inline void GenerateMipmapsBlockCompressedLODRange(
+			int32 SrcLOD, int32 DestLODBegin, int32 DestLODEnd,
+			const Image* SourceImage, Image* DestImage,
+			const FMipmapGenerationSettings& Settings)
 	{
-		using namespace OpImageMipmap_Detail;
+		using namespace OpImageMipmapInternal;
+	
+		const FIntVector2 SourceSize = SourceImage->CalculateMipSize(SrcLOD);
+
+		check(DestImage->GetLODCount() >= DestLODEnd);
+		check(Invoke([&]() -> bool
+		{
+			FIntVector2 DestBeginLODSize = DestImage->CalculateMipSize(DestLODBegin);
+
+			FIntVector2 SrcNextLODSize = FIntVector2(
+					FMath::DivideAndRoundUp(SourceSize.X, 2),
+					FMath::DivideAndRoundUp(SourceSize.Y, 2));
+
+			return DestBeginLODSize == SrcNextLODSize;
+		}));
+
+		const EImageFormat SrcFormat = SourceImage->GetFormat();
+		const EImageFormat DestFormat = DestImage->GetFormat();
 
 		switch (Settings.m_filterType)
 		{
 		case EMipmapFilterType::MFT_SimpleAverage:
 		{
+			const uint8* SourceData = SourceImage->GetLODData(SrcLOD);
+			uint8* DestData = DestImage->GetLODData(DestLODBegin);
+
 			GenerateNextMipBlockCompressed<PixelSize, EMipmapFilterType::MFT_SimpleAverage>(
-				SourceData, DestData, SourceSize, SrcFormat, DestFormat);
+					SourceData, DestData, SourceSize, SrcFormat, DestFormat);
 
-			const FIntVector2 CurrentMipSize = FIntVector2(
-				FMath::DivideAndRoundUp(SourceSize[0], 2),
-				FMath::DivideAndRoundUp(SourceSize[1], 2));
+			FIntVector2 CurrentMipSize = FIntVector2(
+					FMath::DivideAndRoundUp(SourceSize[0], 2),
+					FMath::DivideAndRoundUp(SourceSize[1], 2));
 
 			if (CurrentMipSize.X > 1 || CurrentMipSize.Y > 1)
 			{
-				uint8* CurrentMipData = DestData;
-				uint8* NextMipData = CurrentMipData + CurrentMipSize.X*CurrentMipSize.Y*PixelSize;
-				GenerateMipmapsUint8SimpleAverage<PixelSize>(Mips - 1, CurrentMipData, NextMipData, CurrentMipSize);
+				for (int32 L = DestLODBegin + 1; L < DestLODEnd - 1; ++L)
+				{
+					const uint8* CurrentMipData = DestData;
+					uint8* NextMipData = DestImage->GetLODData(L);
+
+					GenerateNextMipmapUint8SimpleAverage<PixelSize>(CurrentMipData, NextMipData, CurrentMipSize);
+
+					CurrentMipSize = FIntVector2(
+							FMath::DivideAndRoundUp(CurrentMipSize.X, 2),
+							FMath::DivideAndRoundUp(CurrentMipSize.Y, 2));
+				}
 			}
 
 			break;
 		}
 		case EMipmapFilterType::MFT_Unfiltered:
 		{
+			const uint8* SourceData = SourceImage->GetLODData(SrcLOD);
+			uint8* DestData = DestImage->GetLODData(DestLODBegin);
+			
 			GenerateNextMipBlockCompressed<PixelSize, EMipmapFilterType::MFT_Unfiltered>(
 					SourceData, DestData, SourceSize, SrcFormat, DestFormat);
 
-			const FIntVector2 CurrentMipSize = FIntVector2(
-					FMath::DivideAndRoundUp(SourceSize[0], 2), 
-					FMath::DivideAndRoundUp(SourceSize[1], 2));
+			FIntVector2 CurrentMipSize = FIntVector2(
+					FMath::DivideAndRoundUp(SourceSize.X, 2), 
+					FMath::DivideAndRoundUp(SourceSize.Y, 2));
 
 			if (CurrentMipSize.X > 1 || CurrentMipSize.Y > 1)
 			{
-				uint8* CurrentMipData = DestData;
-				uint8* NextMipData = CurrentMipData + CurrentMipSize.Y*CurrentMipSize.X*PixelSize;
-				GenerateMipmapUint8Unfiltered<PixelSize>(Mips + 1, CurrentMipData, NextMipData, CurrentMipSize);
+				for (int32 L = DestLODBegin + 1; L < DestLODEnd - 1; ++L)
+				{
+					const uint8* CurrentMipData = DestData;
+					uint8* NextMipData = DestImage->GetLODData(L);
+
+					GenerateNextMipmapUint8SimpleAverage<PixelSize>(CurrentMipData, NextMipData, CurrentMipSize);
+
+					CurrentMipSize = FIntVector2(
+							FMath::DivideAndRoundUp(CurrentMipSize.X, 2),
+							FMath::DivideAndRoundUp(CurrentMipSize.Y, 2));
+				}
 			}
 			break;
 		}
@@ -493,33 +530,27 @@ namespace OpImageMipmap_Detail
 	}
 
 
-    //---------------------------------------------------------------------------------------------
-    void FImageOperator::ImageMipmap_PrepareScratch(Image* Dest, const Image* Base, int32 LevelCount, FScratchImageMipmap& Scratch )
+    void FImageOperator::ImageMipmap_PrepareScratch(const Image* DestImage, int32 StartLevel, int32 LevelCount, FScratchImageMipmap& Scratch )
     {
-        int StartLevel = Base->GetLODCount() - 1;
-
-        check(Dest->GetLODCount() == LevelCount);
-        check(Dest->GetSizeX() == Base->GetSizeX());
-        check(Dest->GetSizeY() == Base->GetSizeY());
-        check(Dest->GetFormat() == Base->GetFormat());
+        check(DestImage->GetLODCount() == LevelCount);
 
-		EImageFormat BaseFormat = Base->GetFormat();
-		if (mu::IsCompressedFormat(BaseFormat))
+		EImageFormat DestFormat = DestImage->GetFormat();
+		if (mu::IsCompressedFormat(DestFormat))
 		{
 			// Is it a block format?
-			if (mu::GetImageFormatData(BaseFormat).PixelsPerBlockX > 1)
+			if (mu::GetImageFormatData(DestFormat).PixelsPerBlockX > 1)
 			{
 				if (!bEnableCompressedMipGenerationMemoryOptimizations)
 				{
 					// Uncompress the last mip that we already have
-					FIntVector2 UncompressedSize = Base->CalculateMipSize(StartLevel);
+					FIntVector2 UncompressedSize = DestImage->CalculateMipSize(StartLevel);
 					Scratch.Uncompressed = CreateImage(
 						(uint16)UncompressedSize[0], (uint16)UncompressedSize[1],
 						1,
 						EImageFormat::IF_RGBA_UBYTE, EInitializationType::NotInitialized);
 				}
 
-				FIntVector2 UncompressedMipsSize = Base->CalculateMipSize(StartLevel + 1);
+				FIntVector2 UncompressedMipsSize = DestImage->CalculateMipSize(StartLevel + 1);
 				// Generate the mipmaps from there on
 				Scratch.UncompressedMips = CreateImage(
 					(uint16)UncompressedMipsSize[0], (uint16)UncompressedMipsSize[1],
@@ -527,24 +558,24 @@ namespace OpImageMipmap_Detail
 					EImageFormat::IF_RGBA_UBYTE, EInitializationType::NotInitialized);
 
 				// Compress the mipmapped image
-				Scratch.CompressedMips = CreateImage(
-					(uint16)UncompressedMipsSize[0], (uint16)UncompressedMipsSize[1],
-					Scratch.UncompressedMips->GetLODCount(),
-					Base->GetFormat(), EInitializationType::NotInitialized);
+			//	Scratch.CompressedMips = CreateImage(
+			//		(uint16)UncompressedMipsSize[0], (uint16)UncompressedMipsSize[1],
+			//		Scratch.UncompressedMips->GetLODCount(),
+			//		DestImage->GetFormat(), EInitializationType::NotInitialized);
 			}
 			else
 			{
 				// It's probably an RLE compressed format
 
 				// Uncompress the last mip that we already have
-				FIntVector2 UncompressedSize = Base->CalculateMipSize(StartLevel);
+				FIntVector2 UncompressedSize = DestImage->CalculateMipSize(StartLevel);
 				Scratch.Uncompressed = CreateImage(
 					(uint16)UncompressedSize[0], (uint16)UncompressedSize[1],
 					1,
 					EImageFormat::IF_L_UBYTE, EInitializationType::NotInitialized);
 
 
-				FIntVector2 UncompressedMipsSize = Base->CalculateMipSize(StartLevel + 1);
+				FIntVector2 UncompressedMipsSize = DestImage->CalculateMipSize(StartLevel + 1);
 				// Generate the mipmaps from there on
 				Scratch.UncompressedMips = CreateImage(
 					(uint16)UncompressedMipsSize[0], (uint16)UncompressedMipsSize[1],
@@ -553,18 +584,18 @@ namespace OpImageMipmap_Detail
 
 
 				// Compress the mipmapped image
-				Scratch.CompressedMips = CreateImage(
-					(uint16)UncompressedMipsSize[0], (uint16)UncompressedMipsSize[1],
-					Scratch.UncompressedMips->GetLODCount(),
-					Base->GetFormat(), EInitializationType::NotInitialized);
+			//	Scratch.CompressedMips = CreateImage(
+			//		(uint16)UncompressedMipsSize[0], (uint16)UncompressedMipsSize[1],
+			//		Scratch.UncompressedMips->GetLODCount(),
+			//		DestImage->GetFormat(), EInitializationType::NotInitialized);
 
 				// Preallocate ample memory for the compressed data
-				uint32 TotalMemory = Scratch.UncompressedMips->GetDataSize();
-				Scratch.CompressedMips->m_data.SetNumUninitialized(TotalMemory);
-
-				// Preallocate ample memory for the destination data
-				TotalMemory = Base->GetDataSize() + Scratch.UncompressedMips->GetDataSize();
-				Dest->m_data.SetNumUninitialized(TotalMemory);
+				int32 UncompressedNumMips = Scratch.UncompressedMips->GetLODCount();
+				for (int32 L = 0; L < UncompressedNumMips; ++L)
+				{	
+					const FIntVector2 LODSize = Scratch.Uncompressed->CalculateMipSize(L);
+					Scratch.CompressedMips->DataStorage.ResizeLOD(L, LODSize.X*LODSize.Y);
+				}
 			}
 		}
     }
@@ -577,58 +608,42 @@ namespace OpImageMipmap_Detail
 		ReleaseImage(Scratch.CompressedMips);
 	}
 
-
-	void FImageOperator::ImageMipmap(FScratchImageMipmap& Scratch, int32 CompressionQuality, Image* Dest, const Image* Base,
-		int32 LevelCount,
-		const FMipmapGenerationSettings& Settings, bool bGenerateOnlyTail)
+	void FImageOperator::ImageMipmap(FScratchImageMipmap& Scratch, int32 CompressionQuality, Image* DestImage, const Image* BaseImage,
+		int32 StartLevel, int32 NumLODs, const FMipmapGenerationSettings& Settings, bool bGenerateOnlyTail)
 	{
-		int32 StartLevel = Base->GetLODCount() - 1;
-
-		check(!(Base->m_flags & Image::IF_CANNOT_BE_SCALED));
-
-
+		check(!(BaseImage->m_flags & Image::IF_CANNOT_BE_SCALED));
+		check(DestImage->GetFormat() == BaseImage->GetFormat());
+		
 		if (!bGenerateOnlyTail)
 		{
-			check(Dest->GetLODCount() == LevelCount);
-			check(Dest->GetSizeX() == Base->GetSizeX());
-			check(Dest->GetSizeY() == Base->GetSizeY());
+			check(DestImage->GetLODCount() == NumLODs);
+			check(DestImage->GetSizeX() == BaseImage->GetSizeX());
+			check(DestImage->GetSizeY() == BaseImage->GetSizeY());
 		}
 		else
 		{
-			check(Dest->GetLODCount() + Base->GetLODCount() == LevelCount);
+			check(DestImage->GetLODCount() + BaseImage->GetLODCount() == NumLODs);
 
-			check(
-				[&]() -> bool
-				{
-					const FIntVector2 BaseImageNextMipSize = Base->CalculateMipSize(StartLevel + 1);
-					return BaseImageNextMipSize.X == Dest->GetSizeX() && BaseImageNextMipSize.Y == Dest->GetSizeY();
-				}());
+			checkCode
+			(
+				const FIntVector2 BaseNextMipSize = BaseImage->CalculateMipSize(StartLevel + 1);
+				check(BaseNextMipSize.X == DestImage->GetSizeX() && BaseNextMipSize.Y == DestImage->GetSizeY());
+			);
 		}
 
-		check(Dest->GetFormat() == Base->GetFormat());
-
-		// Calculate the data size, since the base could have more data allocated.
-		int32 BaseMipDataSize = 0;
-		if (!bGenerateOnlyTail)
+		if (!bGenerateOnlyTail && DestImage != BaseImage)
 		{
-			BaseMipDataSize = Base->GetMipsDataSize();
-			check(Dest->GetDataSize() >= BaseMipDataSize);
-			FMemory::Memcpy(Dest->GetData(), Base->GetData(), BaseMipDataSize);
-		}
+			for (int32 L = 0; L <= StartLevel; ++L)
+			{
+				TArrayView<uint8> DestView = DestImage->DataStorage.GetLOD(L);
+				TArrayView<const uint8> SrcView = BaseImage->DataStorage.GetLOD(L);
 
-		int mipsToBuild = LevelCount - StartLevel - 1;
-		if (!mipsToBuild)
-		{
-			return;
+				check(DestView.Num() == SrcView.Num());
+				FMemory::Memcpy(DestView.GetData(), SrcView.GetData(), DestView.Num());
+			}
 		}
 
-		const uint8* pSourceBuf = !bGenerateOnlyTail ? Dest->GetMipData(StartLevel) : Base->GetMipData(StartLevel);
-
-		uint8* pDestBuf = !bGenerateOnlyTail ? Dest->GetMipData(StartLevel + 1) : Dest->GetMipData(0);
-
-		FIntVector2 SourceSize = Base->CalculateMipSize(StartLevel);
-
-		EImageFormat BaseFormat = Base->GetFormat();
+		EImageFormat BaseFormat = BaseImage->GetFormat();
 		const bool bIsBlockCompressedFormat = mu::IsBlockCompressedFormat(BaseFormat);
 		const bool bIsCompressedFormat = mu::IsCompressedFormat(BaseFormat);
 		
@@ -637,98 +652,91 @@ namespace OpImageMipmap_Detail
 		if (bIsBlockCompressedFormat && bEnableCompressedMipGenerationMemoryOptimizations)
 		{
 			const EImageFormat DestFormat = Scratch.UncompressedMips->GetFormat();
-			uint8* DestData = Scratch.UncompressedMips->GetData();
+			Image* UncompressedImage = Scratch.UncompressedMips.get();
+	
+			const int32 UncompressedNumLODs = UncompressedImage->GetLODCount();
+
 			switch (DestFormat)
 			{
 			case EImageFormat::IF_L_UBYTE:
 			{
 				constexpr int32 PixelSize = 1;
-				GenerateMipmapsBlockCompressed<PixelSize>(
-						LevelCount - StartLevel - 1, pSourceBuf, DestData, SourceSize, BaseFormat, DestFormat, Settings);
+				GenerateMipmapsBlockCompressedLODRange<PixelSize>(
+						StartLevel, 0, UncompressedNumLODs, BaseImage, UncompressedImage, Settings);
 				break;
 			}
 			case EImageFormat::IF_RGB_UBYTE:
 			{
 				constexpr int32 PixelSize = 3;
-				GenerateMipmapsBlockCompressed<PixelSize>(
-						LevelCount - StartLevel - 1, pSourceBuf, DestData, SourceSize, BaseFormat, DestFormat, Settings);
+				GenerateMipmapsBlockCompressedLODRange<PixelSize>(
+						StartLevel, 0, UncompressedNumLODs, BaseImage, UncompressedImage, Settings);
 				break;
 			}
 			case EImageFormat::IF_RGBA_UBYTE:
 			{
 				constexpr int32 PixelSize = 4;
-				GenerateMipmapsBlockCompressed<PixelSize>(
-						LevelCount - StartLevel - 1, pSourceBuf, DestData, SourceSize, BaseFormat, DestFormat, Settings);
+				GenerateMipmapsBlockCompressedLODRange<PixelSize>(
+						StartLevel, 0, UncompressedNumLODs, BaseImage, UncompressedImage, Settings);
 				break;
 			}
 			default: check(false);
 			}
 
+			int32 DestLODBegin = bGenerateOnlyTail ? 0 : StartLevel + 1;
+
 			bool bSuccess = false;
-			constexpr int32 OnlyLod = -1;
-			ImagePixelFormat(bSuccess, CompressionQuality, Scratch.CompressedMips.get(), Scratch.UncompressedMips.get(), OnlyLod); 
-			check(bSuccess);
-			
-			FMemory::Memcpy(pDestBuf, Scratch.CompressedMips->GetData(), Scratch.CompressedMips->GetDataSize());
+			ImagePixelFormat(
+					bSuccess, CompressionQuality, DestImage, UncompressedImage, DestLODBegin, 0, UncompressedNumLODs); 
+			check(bSuccess);	
 		}
 		else if (bIsCompressedFormat)
 		{
+			const int32 DestLODBegin = bGenerateOnlyTail ? 0 : StartLevel + 1;
+			const int32 DestLODEnd = DestImage->GetLODCount();
+
 			// Bad case.
 			// Uncompress the last mip that we already have
 			bool bSuccess = false;
-			ImagePixelFormat(bSuccess, CompressionQuality, Scratch.Uncompressed.get(), Base, StartLevel);
+			ImagePixelFormat(bSuccess, CompressionQuality, Scratch.Uncompressed.get(), BaseImage, StartLevel);
 			check(bSuccess);
 
 			// Generate the mipmaps from there on
-
 			constexpr bool bGenerateOnlyTailForCompressed = true;
+
+			const int32 NumScratchMips = Scratch.UncompressedMips->GetLODCount() + Scratch.Uncompressed->GetLODCount();
 			ImageMipmap(Scratch, CompressionQuality, Scratch.UncompressedMips.get(),
-				Scratch.Uncompressed.get(), LevelCount - StartLevel, Settings, bGenerateOnlyTailForCompressed);
+					Scratch.Uncompressed.get(), 0, NumScratchMips, Settings, bGenerateOnlyTailForCompressed);
 
 			// Compress the mipmapped image
 			bSuccess = false;
-			ImagePixelFormat(bSuccess, CompressionQuality, Scratch.CompressedMips.get(), Scratch.UncompressedMips.get());
-			int32 ExcessDataSize = FMath::Max(2, Scratch.CompressedMips->GetDataSize());
-			while (!bSuccess)
-			{
-				// Bad case: this should almost never happen.
-				MUTABLE_CPUPROFILER_SCOPE(Mipmap_Recompression_OutOfSpace);
-
-				Scratch.CompressedMips->m_data.SetNumUninitialized(ExcessDataSize);
-				bSuccess = false;
-				ImagePixelFormat(bSuccess, CompressionQuality, Scratch.CompressedMips.get(), Scratch.UncompressedMips.get());
-				ExcessDataSize *= 4;
-			}
-
-			check(!bGenerateOnlyTail || BaseMipDataSize == 0);
-			const int32 FinalDestSize = BaseMipDataSize + Scratch.CompressedMips->GetDataSize();
-
-			if (FinalDestSize > Dest->GetDataSize())
-			{
-				// Bad case: this should almost never happen.
-				Dest->m_data.SetNumUninitialized(FinalDestSize);
-				pDestBuf = !bGenerateOnlyTail ? Dest->GetMipData(StartLevel + 1) : Dest->GetMipData(0);
-			}
-
-			FMemory::Memcpy(pDestBuf, Scratch.CompressedMips->GetData(), Scratch.CompressedMips->GetDataSize());
+			ImagePixelFormat(
+					bSuccess, CompressionQuality, DestImage, Scratch.UncompressedMips.get(), 
+					StartLevel + 1, 0, Scratch.UncompressedMips->GetLODCount()); 
+			check(bSuccess);
 		}
 		else
-		{
-			switch (Base->GetFormat())
+		{	
+			const int32 DestLODBegin = bGenerateOnlyTail ? 0 : StartLevel + 1;
+			const int32 DestLODEnd = DestImage->GetLODCount();
+
+			switch (BaseImage->GetFormat())
 			{
 			case EImageFormat::IF_L_UBYTE:
-				GenerateMipmapsUint8<1>(LevelCount - StartLevel - 1, pSourceBuf, pDestBuf, SourceSize, Settings);
+			{
+				GenerateMipmapUint8LODRange<1>(StartLevel, DestLODBegin, DestLODEnd, BaseImage, DestImage, Settings);
 				break;
-
+			}
 			case EImageFormat::IF_RGB_UBYTE:
-				GenerateMipmapsUint8<3>(LevelCount - StartLevel - 1, pSourceBuf, pDestBuf, SourceSize, Settings);
+			{
+				GenerateMipmapUint8LODRange<3>(StartLevel, DestLODBegin, DestLODEnd, BaseImage, DestImage, Settings);
 				break;
-
+			}
 			case EImageFormat::IF_BGRA_UBYTE:
 			case EImageFormat::IF_RGBA_UBYTE:
-				GenerateMipmapsUint8<4>(LevelCount - StartLevel - 1, pSourceBuf, pDestBuf, SourceSize, Settings);
+			{
+				GenerateMipmapUint8LODRange<4>(StartLevel, DestLODBegin, DestLODEnd, BaseImage, DestImage, Settings);
 				break;
-
+			}
 			default:
 				checkf(false, TEXT("Format not implemented in mipmap generation."));
 			}
@@ -736,52 +744,43 @@ namespace OpImageMipmap_Detail
 	}
 
 	void FImageOperator::ImageMipmap(int32 CompressionQuality, Image* Dest, const Image* Base,
-		int32 LevelCount,
-		const FMipmapGenerationSettings& Settings, bool bGenerateOnlyTail)
+		int32 StartLevel, int32 LevelCount, const FMipmapGenerationSettings& Settings, bool bGenerateOnlyTail)
 	{
 		FScratchImageMipmap Scratch;
-		ImageMipmap_PrepareScratch(Dest, Base, LevelCount, Scratch);
-
-		ImageMipmap(Scratch, CompressionQuality, Dest, Base, LevelCount, Settings, bGenerateOnlyTail);
 
+		ImageMipmap_PrepareScratch(Dest, StartLevel, LevelCount, Scratch);
+		ImageMipmap(Scratch, CompressionQuality, Dest, Base, StartLevel, LevelCount, Settings, bGenerateOnlyTail);
 		ImageMipmap_ReleaseScratch(Scratch);
 	}
 
-
-	/** Update all the mipmaps in the image from the data in the base one. 
-	* Only the mipmaps already existing in the image are updated.
-	*/
+	/** 
+	 * Update all the mipmaps in the image from the data in the base one. 
+	 * Only the mipmaps already existing in the image are updated.
+	 */
 	void ImageMipmapInPlace(int32 InImageCompressionQuality, Image* InBase, const FMipmapGenerationSettings& InSettings)
 	{
-		int32 StartLevel = 0;
-		int32 LevelCount = InBase->GetLODCount();
-
 		check(!(InBase->m_flags & Image::IF_CANNOT_BE_SCALED));
 
-		int32 MipsToBuild = InBase->GetLODCount()-1;
-		if (!MipsToBuild)
+		int32 LevelCount = InBase->GetLODCount();
+		
+		if (LevelCount - 1 <= 0)
 		{
 			return;
 		}
 
-		const uint8* pSourceBuf = InBase->GetMipData(StartLevel);
-		uint8* pDestBuf = InBase->GetMipData(StartLevel + 1);
-
-		FIntVector2 sourceSize = InBase->CalculateMipSize(StartLevel);
-
 		switch (InBase->GetFormat())
 		{
 		case EImageFormat::IF_L_UBYTE:
-			GenerateMipmapsUint8<1>(MipsToBuild, pSourceBuf, pDestBuf, sourceSize, InSettings);
+			GenerateMipmapUint8LODRange<1>(0, 1, LevelCount, InBase, InBase, InSettings);
 			break;
 
 		case EImageFormat::IF_RGB_UBYTE:
-			GenerateMipmapsUint8<3>(MipsToBuild, pSourceBuf, pDestBuf, sourceSize, InSettings);
+			GenerateMipmapUint8LODRange<3>(0, 1, LevelCount, InBase, InBase, InSettings);
 			break;
 
 		case EImageFormat::IF_BGRA_UBYTE:
 		case EImageFormat::IF_RGBA_UBYTE:
-			GenerateMipmapsUint8<4>(MipsToBuild, pSourceBuf, pDestBuf, sourceSize, InSettings);
+			GenerateMipmapUint8LODRange<4>(0, 1, LevelCount, InBase, InBase, InSettings);
 			break;
 
 		default:
