[Mutable] Store Image LODs in different buffers.

diff --git a/Engine/Plugins/Experimental/Mutable/Source/MutableRuntime/Private/MuR/OpImagePixelFormat.cpp b/Engine/Plugins/Experimental/Mutable/Source/MutableRuntime/Private/MuR/OpImagePixelFormat.cpp
--- a/Engine/Plugins/Experimental/Mutable/Source/MutableRuntime/Private/MuR/OpImagePixelFormat.cpp
+++ b/Engine/Plugins/Experimental/Mutable/Source/MutableRuntime/Private/MuR/OpImagePixelFormat.cpp
@@ -19,1633 +19,1380 @@ namespace mu
 {
 
 
-Ptr<Image> FImageOperator::ImagePixelFormat( int32 CompressionQuality, const Image* Base, EImageFormat TargetFormat, int32 OnlyLOD )
+Ptr<Image> FImageOperator::ImagePixelFormat(int32 CompressionQuality, const Image* Base, EImageFormat TargetFormat, int32 OnlyLOD)
 {
 	MUTABLE_CPUPROFILER_SCOPE(ImagePixelFormat);
 	TRACE_CPUPROFILER_EVENT_SCOPE_TEXT(*FString::Printf(TEXT("From %d to %d"), int32(Base->GetFormat()), int32(TargetFormat)));
 
-	FIntVector2 resultSize;
-    int resultLODCount = 0;
-
-    if (OnlyLOD==-1)
-    {
-        resultSize[0] = Base->GetSizeX();
-        resultSize[1] = Base->GetSizeY();
-        resultLODCount = Base->GetLODCount();
-    }
-    else
-    {
-        resultSize = Base->CalculateMipSize( OnlyLOD );
-        resultLODCount = 1;
-    }
-
-    Ptr<Image> Result = CreateImage( resultSize[0], resultSize[1], resultLODCount, TargetFormat, EInitializationType::NotInitialized );
+	FIntVector2 ResultSize;
+	int32 ResultLODCount = 0;
+
+	if (OnlyLOD == -1)
+	{
+		ResultSize[0] = Base->GetSizeX();
+		ResultSize[1] = Base->GetSizeY();
+		ResultLODCount = Base->GetLODCount();
+	}
+	else
+	{
+		ResultSize = Base->CalculateMipSize(OnlyLOD);
+		ResultLODCount = 1;
+	}
+
+	Ptr<Image> Result = CreateImage(ResultSize[0], ResultSize[1], ResultLODCount, TargetFormat, EInitializationType::NotInitialized);
 	Result->m_flags = Base->m_flags;
 
-	if (Base->GetSizeX()<=0 ||Base->GetSizeY()<=0)
+	if (Base->GetSizeX() <= 0 || Base->GetSizeY() <= 0)
 	{
 		return Result;
 	}
 
 	bool bSuccess = false;
-	ImagePixelFormat( bSuccess, CompressionQuality, Result.get(), Base, OnlyLOD );
-	int32 OriginalDataSize = FMath::Max(Result->m_data.Num(), Base->m_data.Num());
-	int32 ExcessDataSize = OriginalDataSize * 4 + 4;
-	while (!bSuccess)
-    {
-		MUTABLE_CPUPROFILER_SCOPE(Recompression_OutOfSpace);
-
-        // Bad case, where the RLE compressed data requires more memory than the uncompressed data.
-		// We need to support it anyway for small mips or scaled images.
-		Result->m_data.SetNum(ExcessDataSize);
-		ImagePixelFormat(bSuccess, CompressionQuality, Result.get(), Base, OnlyLOD);
-		ExcessDataSize *= 4;
-	}
-    return Result;
-}
+	ImagePixelFormat(bSuccess, CompressionQuality, Result.get(), Base, OnlyLOD);
+	check(bSuccess);
 
+	return Result;
+}
 
-void FImageOperator::ImagePixelFormat( bool& bOutSuccess, int32 CompressionQuality, Image* pResult, const Image* Base, int32 OnlyLOD)
+namespace ImagePixelFormatInternal
 {
-	MUTABLE_CPUPROFILER_SCOPE(ImagePixelFormatInPlace);
+	template<class PixelFormatFuncType>
+	void ProcessBatchedHelper(
+			Image* Dest, const Image* Src, 
+			int32 DestElemSizeInBytes, int32 SrcElemSizeInBytes, 
+			int32 DestLODBegin, int32 SrcLODBegin, int32 NumLODs, 
+			PixelFormatFuncType&& BatchFunc)
+	{
+		check(Dest->GetLODCount() >= DestLODBegin + NumLODs);
+		check(Src->GetLODCount() >= SrcLODBegin + NumLODs);
+		
+		const int32 DestLODEnd = DestLODBegin + NumLODs;
+		const int32 SrcLODEnd = SrcLODBegin + NumLODs;
+
+		constexpr int32 NumBatchElems = 1 << 15;
+
+		const int32 NumBatches = Dest->DataStorage.GetNumBatchesLODRange(NumBatchElems, DestElemSizeInBytes, DestLODBegin, DestLODEnd);
+		check(NumBatches == Src->DataStorage.GetNumBatchesLODRange(NumBatchElems, SrcElemSizeInBytes, SrcLODBegin, SrcLODEnd));
+
+		for (int32 BatchIndex = 0; BatchIndex < NumBatches; ++BatchIndex)
+		{
+			TArrayView<const uint8> SrcView = Src->DataStorage.GetBatchLODRange(BatchIndex, NumBatchElems, SrcElemSizeInBytes, SrcLODBegin, SrcLODEnd);
+			TArrayView<uint8> DestView = Dest->DataStorage.GetBatchLODRange(BatchIndex, NumBatchElems, DestElemSizeInBytes, DestLODBegin, DestLODEnd);
+
+			const int32 NumElems = DestView.Num() / DestElemSizeInBytes;
+			check(DestView.Num() % DestElemSizeInBytes == 0);
+			check(NumElems == SrcView.Num() / SrcElemSizeInBytes);
+
+			BatchFunc(DestView.GetData(), SrcView.GetData(), NumElems);
+		}
+	}
+
+	template<class PixelFormatFuncType>
+	void ProcessBatchedHelper(
+			Image* Dest, int32 DestElemSizeInBytes, int32 DestLODBegin, int32 NumLODs, 
+			PixelFormatFuncType&& BatchFunc)
+	{
+		check(Dest->GetLODCount() >= DestLODBegin + NumLODs);
+		
+		const int32 DestLODEnd = DestLODBegin + NumLODs;
 
+		constexpr int32 NumBatchElems = 1 << 15;
+
+		const int32 NumBatches = Dest->DataStorage.GetNumBatchesLODRange(NumBatchElems, DestElemSizeInBytes, DestLODBegin, DestLODEnd);
+
+		for (int32 BatchIndex = 0; BatchIndex < NumBatches; ++BatchIndex)
+		{
+			TArrayView<uint8> DestView = Dest->DataStorage.GetBatchLODRange(BatchIndex, NumBatchElems, DestElemSizeInBytes, DestLODBegin, DestLODEnd);
+
+			const int32 NumElems = DestView.Num() / DestElemSizeInBytes;
+			check(DestView.Num() % DestElemSizeInBytes == 0);
+
+			BatchFunc(DestView.GetData(), NumElems);
+		}
+	}
+
+	template<class BlockDecompressionFuncType> 
+	void BlockDecompressionLODRangeHelper(
+		Image* Result, const Image* Base, 
+		int32 ResultLODBegin, int32 BaseLODBegin, int32 NumLODs, 
+		BlockDecompressionFuncType&& DecFunc)
+	{
+		check(Base->GetLODCount() >= BaseLODBegin + NumLODs);
+		check(Result->GetLODCount() >= ResultLODBegin + NumLODs);
+		for (int32 L = 0; L < NumLODs; ++L)
+		{
+			const int32 BaseLOD = BaseLODBegin + L;
+			const int32 ResultLOD = ResultLODBegin + L;
+			FIntVector2 MipSize = Result->CalculateMipSize(ResultLOD);
+			
+			DecFunc(MipSize.X, MipSize.Y, Base->GetLODData(BaseLOD), Result->GetLODData(ResultLOD));
+		}
+	}
+
+	template<class BlockCompressionFuncType> 
+	void BlockCompressionLODRangeHelper(
+		Image* Result, const Image* Base, 
+		int32 ResultLODBegin, int32 BaseLODBegin, int32 NumLODs, int32 Quality,
+		BlockCompressionFuncType&& CompFunc)
+	{
+		check(Base->GetLODCount() >= BaseLODBegin + NumLODs);
+		check(Result->GetLODCount() >= ResultLODBegin + NumLODs);
+		for (int32 L = 0; L < NumLODs; ++L)
+		{
+			const int32 BaseLOD = BaseLODBegin + L;
+			const int32 ResultLOD = ResultLODBegin + L;
+			FIntVector2 MipSize = Result->CalculateMipSize(ResultLOD);
+			
+			CompFunc(MipSize.X, MipSize.Y, Base->GetLODData(BaseLOD), Result->GetLODData(ResultLOD), Quality);
+		}
+	}
+
+}
+
+void FImageOperator::ImagePixelFormat(bool& bOutSuccess, int32 CompressionQuality, Image* Result, const Image* Base, int32 OnlyLOD)
+{
 	if (FormatImageOverride)
 	{
 		MUTABLE_CPUPROFILER_SCOPE(ImagePixelFormatOverride);
-		FormatImageOverride(bOutSuccess, CompressionQuality, pResult, Base, OnlyLOD);
+		FormatImageOverride(bOutSuccess, CompressionQuality, Result, Base, OnlyLOD);
 		if (bOutSuccess)
 		{
 			return;
 		}
 	}
+	
+	bOutSuccess = true;
+	
+	const int32 BaseLODBegin = OnlyLOD != -1 ? OnlyLOD : 0;
+	const int32 ResultLODBegin = 0;
+	const int32 NumLODs = OnlyLOD != -1 ? 1 : Base->GetLODCount();
+
+	ImagePixelFormat(bOutSuccess, CompressionQuality, Result, Base, ResultLODBegin, BaseLODBegin, NumLODs);
+	check(bOutSuccess);
+}
+
+void FImageOperator::ImagePixelFormat(bool& bOutSuccess, int32 CompressionQuality, Image* Result, const Image* Base, 
+		int32 ResultLODBegin, int32 BaseLODBegin, int32 NumLODs)
+{
+	MUTABLE_CPUPROFILER_SCOPE(ImagePixelFormatInPlace);
+
+	check(Result->GetLODCount() >= ResultLODBegin + NumLODs);
+	check(Base->GetLODCount() >= BaseLODBegin + NumLODs);
+	check(Base->CalculateMipSize(BaseLODBegin) == Result->CalculateMipSize(ResultLODBegin));
+
+	const bool bFormatAllLODs = Result->GetLODCount() == NumLODs && 
+								BaseLODBegin == 0 && ResultLODBegin == 0;
+	
+	check(!bFormatAllLODs || Base->GetLODCount() >= NumLODs);
 
 	bOutSuccess = true;
 
-	FIntVector2 resultSize;
-    int resultLODCount = 0;
-    const uint8_t* pBaseBuf = 0;
-    int pixelCount = 0;
-    int baseLOD = 0;
-
-    if (OnlyLOD==-1)
-    {
-        pixelCount = (int)Base->CalculatePixelCount();
-        pBaseBuf = Base->GetData();
-        resultSize[0] = Base->GetSizeX();
-        resultSize[1] = Base->GetSizeY();
-        resultLODCount = Base->GetLODCount();
-        baseLOD = 0;
-    }
-    else
-    {
-        pixelCount = (int)Base->CalculatePixelCount( OnlyLOD );
-        pBaseBuf = Base->GetMipData( OnlyLOD );
-        resultSize = Base->CalculateMipSize( OnlyLOD );
-        resultLODCount = 1;
-        baseLOD = OnlyLOD;
-    }
-
-    check( pResult->GetSizeX()==resultSize[0] && pResult->GetSizeY()==resultSize[1] );
-    check( pResult->GetLODCount()==resultLODCount );
-
-    if ( pResult->GetFormat() == Base->GetFormat()
-         && OnlyLOD==-1 )
-    {
-        // Shouldn't really happen
-        pResult->m_data = Base->m_data;
-    }
-    else
-    {
-        switch ( pResult->GetFormat() )
-        {
-
-        case EImageFormat::IF_L_UBYTE:
-        {
-            uint8_t* pDestBuf = pResult->GetData();
-
-            switch ( Base->GetFormat() )
-            {
-
-            case EImageFormat::IF_L_UBYTE:
-            {
-				FMemory::Memcpy( pDestBuf, pBaseBuf, pixelCount );
-                break;
-            }
-
-            case EImageFormat::IF_L_UBYTE_RLE:
-            {
-                if( OnlyLOD == -1 )
-                {
-                    UncompressRLE_L( Base, pResult );
-                }
-                else
-                {
-                    UncompressRLE_L( resultSize[0], resultSize[1], pBaseBuf, pDestBuf );
-                }
-                break;
-            }
-
-            case EImageFormat::IF_L_UBIT_RLE:
-            {
-                if( OnlyLOD == -1 )
-                {
-                    UncompressRLE_L1( Base, pResult );
-                }
-                else
-                {
-                    UncompressRLE_L1( resultSize[0], resultSize[1], pBaseBuf, pDestBuf );
-                }
-                break;
-            }
-
-            case EImageFormat::IF_RGB_UBYTE:
-            {
-                for ( int i=0; i<pixelCount; ++i )
-                {
-                    unsigned Result = 76 * pBaseBuf[3*i+0]
-                            + 150 * pBaseBuf[3*i+1]
-                            + 29 * pBaseBuf[3*i+2];
-                    pDestBuf[i] = (uint8_t)FMath::Min( 255u, Result >> 8 );
-                }
-                break;
-            }
-
-            case EImageFormat::IF_RGB_UBYTE_RLE:
-            {
-                ImagePtr pTempBase = ImagePixelFormat( CompressionQuality, Base, EImageFormat::IF_RGB_UBYTE );
-                if (OnlyLOD==-1)
-                {
-                    pBaseBuf = pTempBase->GetData();
-                }
-                else
-                {
-                    pBaseBuf = pTempBase->GetMipData( OnlyLOD );
-                }
-                for ( int i=0; i<pixelCount; ++i )
-                {
-                    unsigned Result = 76 * pBaseBuf[3*i+0]
-                            + 150 * pBaseBuf[3*i+1]
-                            + 29 * pBaseBuf[3*i+2];
-                    pDestBuf[i] = (uint8_t)FMath::Min( 255u, Result >> 8 );
-                }
-                break;
-            }
-
-            case EImageFormat::IF_RGBA_UBYTE:
-            {
-                for ( int i=0; i<pixelCount; ++i )
-                {
-                    unsigned Result = 76 * pBaseBuf[4*i+0]
-                            + 150 * pBaseBuf[4*i+1]
-                            + 29 * pBaseBuf[4*i+2];
-                    pDestBuf[i] = (uint8_t)FMath::Min( 255u, Result >> 8 );
-                }
-                break;
-            }
-
-            case EImageFormat::IF_RGBA_UBYTE_RLE:
-            {
-                ImagePtr pTempBase = ImagePixelFormat( CompressionQuality, Base, EImageFormat::IF_RGBA_UBYTE );
-                if (OnlyLOD==-1)
-                {
-                    pBaseBuf = pTempBase->GetData();
-                }
-                else
-                {
-                    pBaseBuf = pTempBase->GetMipData( OnlyLOD );
-                }
-                for ( int i=0; i<pixelCount; ++i )
-                {
-                    unsigned Result = 76 * pBaseBuf[4*i+0]
-                            + 150 * pBaseBuf[4*i+1]
-                            + 29 * pBaseBuf[4*i+2];
-                    pDestBuf[i] = (uint8_t)FMath::Min( 255u, Result >> 8 );
-                }
-                break;
-            }
-
-            case EImageFormat::IF_BGRA_UBYTE:
-            {
-                for ( int i=0; i<pixelCount; ++i )
-                {
-                    unsigned Result = 76 * pBaseBuf[4*i+2]
-                            + 150 * pBaseBuf[4*i+1]
-                            + 29 * pBaseBuf[4*i+0];
-                    pDestBuf[i] = (uint8_t)FMath::Min( 255u, Result >> 8 );
-                }
-                break;
-            }
-
-            case EImageFormat::IF_BC1:
-            case EImageFormat::IF_BC2:
-            case EImageFormat::IF_BC3:
-            {
-                for (int m = 0; m < resultLODCount; ++m)
-                {
-                    ImagePtr pTempBase = ImagePixelFormat( CompressionQuality, Base, EImageFormat::IF_RGB_UBYTE );
-                    if (OnlyLOD == -1)
-                    {
-                        pBaseBuf = pTempBase->GetData();
-                    }
-                    else
-                    {
-                        pBaseBuf = pTempBase->GetMipData(OnlyLOD);
-                    }
-                    for (int i = 0; i < pixelCount; ++i)
-                    {
-                        unsigned Result = 76 * pBaseBuf[3 * i + 0] + 150 * pBaseBuf[3 * i + 1] +
-                                          29 * pBaseBuf[3 * i + 2];
-                        pDestBuf[i] = (uint8_t)FMath::Min(255u, Result >> 8);
-                    }
-                }
-                break;
-            }
-			case EImageFormat::IF_BC4:
+	const bool bIsPixelFormatNeeded = 
+			Result->GetFormat()   == Base->GetFormat()   &&
+			Result->GetLODCount() == Base->GetLODCount() &&
+			Result->GetLODCount() == NumLODs;
+
+	if (bIsPixelFormatNeeded)
+	{
+		// Shouldn't really happen
+		Result->DataStorage = Base->DataStorage;
+	}
+	else
+	{
+		switch (Result->GetFormat())
+		{
+		case EImageFormat::IF_L_UBYTE:
+		{
+			switch (Base->GetFormat())
+			{
+			case EImageFormat::IF_L_UBYTE:
+			{
+				ImagePixelFormatInternal::ProcessBatchedHelper(Result, Base, 1, 1, ResultLODBegin, BaseLODBegin, NumLODs,
+				[](uint8* DestBuf, const uint8* BaseBuf, int32 NumElems)
+				{
+					FMemory::Memcpy(DestBuf, BaseBuf, NumElems * 1);
+				});
+
+				break;
+			}
+			case EImageFormat::IF_L_UBYTE_RLE:
 			{
-                for ( int m=0; m<resultLODCount; ++m )
-                {
-					FIntVector2 mipSize = pResult->CalculateMipSize( m );
-                    miro::BC4_to_L( mipSize[0], mipSize[1], Base->GetMipData(baseLOD+m), pResult->GetMipData(m) );
-                }
-                break;
+				if (bFormatAllLODs)
+				{
+					UncompressRLE_L(Base, Result);
+				}
+				else
+				{
+					check(NumLODs == 1);
+					FImageSize ResultSize = Result->GetSize();
+					UncompressRLE_L(ResultSize.X, ResultSize.Y, Base->GetLODData(BaseLODBegin), Result->GetLODData(ResultLODBegin));
+				}
+				break;
+			}
+			case EImageFormat::IF_L_UBIT_RLE:
+			{
+				if(bFormatAllLODs)
+				{
+					UncompressRLE_L1(Base, Result);
+				}
+				else
+				{
+					check(NumLODs == 1);
+					FImageSize ResultSize = Result->GetSize();
+					UncompressRLE_L1(ResultSize.X, ResultSize.Y, Base->GetLODData(BaseLODBegin), Result->GetLODData(ResultLODBegin));
+				}
+				break;
+			}
+			case EImageFormat::IF_RGB_UBYTE:
+			{
+				ImagePixelFormatInternal::ProcessBatchedHelper(Result, Base, 1, 3, ResultLODBegin, BaseLODBegin, NumLODs,
+				[](uint8* DestBuf, const uint8* BaseBuf, int32 NumElems)
+				{
+					for (int32 I = 0; I < NumElems; ++I)
+					{
+						uint32 ResultValue = 76 * BaseBuf[3*I + 0] + 150 * BaseBuf[3*I + 1] + 29 * BaseBuf[3*I + 2];
+						DestBuf[I] = (uint8)FMath::Min(255u, ResultValue >> 8);
+					}
+				});
+
+				break;
 			}
+			case EImageFormat::IF_RGB_UBYTE_RLE:
+			{
+				Ptr<Image> TempBase =
+						ImagePixelFormat(CompressionQuality, Base, EImageFormat::IF_RGB_UBYTE);
 
-            default:
-                // Case not implemented
-                check( false );
+				ImagePixelFormatInternal::ProcessBatchedHelper(Result, TempBase.get(), 1, 3, ResultLODBegin, BaseLODBegin, NumLODs,
+				[](uint8* DestBuf, const uint8* BaseBuf, int32 NumElems)
+				{
+					for (int32 I = 0; I < NumElems; ++I)
+					{
+						uint32 ResultValue = 76 * BaseBuf[3*I + 0] + 150 * BaseBuf[3*I + 1] + 29 * BaseBuf[3*I + 2];
+						DestBuf[I] = (uint8)FMath::Min(255u, ResultValue >> 8);
+					}
+				});
 
-            }
+				ReleaseImage(TempBase);
 
-            break;
-        }
+				break;
+			}
+			case EImageFormat::IF_RGBA_UBYTE:
+			{
+				ImagePixelFormatInternal::ProcessBatchedHelper(Result, Base, 1, 4, ResultLODBegin, BaseLODBegin, NumLODs,
+				[](uint8* DestBuf, const uint8* BaseBuf, int32 NumElems)
+				{
+					for (int32 I = 0; I < NumElems; ++I)
+					{
+						uint32 ResultValue = 76 * BaseBuf[4*I + 0] + 150 * BaseBuf[4*I + 1] + 29 * BaseBuf[4*I + 2];
+						DestBuf[I] = (uint8)FMath::Min(255u, ResultValue >> 8);
+					}
+				});
 
-        case EImageFormat::IF_L_UBYTE_RLE:
-        {
-            switch ( Base->GetFormat() )
-            {
+				break;
+			}
+			case EImageFormat::IF_RGBA_UBYTE_RLE:
+			{
+				Ptr<Image> TempBase = ImagePixelFormat(CompressionQuality, Base, EImageFormat::IF_RGBA_UBYTE);
+				
+				ImagePixelFormatInternal::ProcessBatchedHelper(Result, TempBase.get(), 1, 4, ResultLODBegin, BaseLODBegin, NumLODs,
+				[](uint8* DestBuf, const uint8* BaseBuf, int32 NumElems)
+				{
+					for (int32 I = 0; I < NumElems; ++I)
+					{
+						uint32 ResultValue = 76 * BaseBuf[4*I + 0] + 150 * BaseBuf[4*I + 1] + 29 * BaseBuf[4*I + 2];
+						DestBuf[I] = (uint8)FMath::Min(255u, ResultValue >> 8);
+					}
+				});
 
-            case EImageFormat::IF_L_UBYTE:
-            {
-                check( OnlyLOD == -1 );
+				ReleaseImage(TempBase);
 
-                // Try to compress
-				if (!pResult->m_data.Num())
+				break;
+			}
+			case EImageFormat::IF_BGRA_UBYTE:
+			{
+				ImagePixelFormatInternal::ProcessBatchedHelper(Result, Base, 1, 4, ResultLODBegin, BaseLODBegin, NumLODs,
+				[](uint8* DestBuf, const uint8* BaseBuf, int32 NumElems)
 				{
-					// Allocate memory for the compressed data. TODO: Smaller?
-					uint32 TotalMemory = Base->m_data.Num();
-					pResult->m_data.SetNum(TotalMemory);
-				}
+					for (int32 I = 0; I < NumElems; ++I)
+					{
+						uint32 ResultValue = 76 * BaseBuf[4*I + 2] + 150 * BaseBuf[4*I + 1] + 29 * BaseBuf[4*I + 0];
+						DestBuf[I] = (uint8)FMath::Min(255u, ResultValue >> 8);
+					}
+				});
 
-                CompressRLE_L(bOutSuccess, Base, pResult);
+				break;
+			}
+			case EImageFormat::IF_BC1:
+			case EImageFormat::IF_BC2:
+			case EImageFormat::IF_BC3:
+			{
+				Ptr<Image> TempBase = ImagePixelFormat(CompressionQuality, Base, EImageFormat::IF_RGB_UBYTE);
+				
+				ImagePixelFormatInternal::ProcessBatchedHelper(Result, TempBase.get(), 1, 3, ResultLODBegin, BaseLODBegin, NumLODs,
+				[](uint8* DestBuf, const uint8* BaseBuf, int32 NumElems)
+				{
+					for (int32 I = 0; I < NumElems; ++I)
+					{
+						uint32 ResultValue = 76 * BaseBuf[3*I + 0] + 150 * BaseBuf[3*I + 1] + 29 * BaseBuf[3*I + 2];
+						DestBuf[I] = (uint8)FMath::Min(255u, ResultValue >> 8);
+					}
+				});
+
+				ReleaseImage(TempBase);
+
+				break;
+			}
+			case EImageFormat::IF_BC4:
+			{
+				ImagePixelFormatInternal::BlockDecompressionLODRangeHelper(Result, Base, ResultLODBegin, BaseLODBegin, NumLODs, miro::BC4_to_L);
+				break;
+			}
+			default:
+				// Case not implemented
+				check(false);
+			}
+
+			break;
+		}
+
+		case EImageFormat::IF_L_UBYTE_RLE:
+		{
+			switch (Base->GetFormat())
+			{
+
+			case EImageFormat::IF_L_UBYTE:
+			{
+				check(bFormatAllLODs);
+
+				// Try to compress
+				if (Result->DataStorage.IsEmpty())
+				{
+					// Allocate memory for the compressed data. TODO: Smaller?	
+					check(Base->DataStorage.Buffers.Num() == Result->DataStorage.Buffers.Num());
+
+					for (int32 I = 0; I < Base->DataStorage.Buffers.Num(); ++I)
+					{
+						int32 BufferSize = Base->DataStorage.Buffers[I].Num();
+						Result->DataStorage.Buffers[I].SetNumUninitialized(BufferSize);
+					}
+					Result->DataStorage.CompactedTailOffsets = Base->DataStorage.CompactedTailOffsets;
+				}
+				
+				CompressRLE_L(Base, Result);
 
 #ifdef MUTABLE_DEBUG_RLE		
 				if (bOutSuccess)
-                {
-                    // verify
-                    ImagePtr pTest = ImagePixelFormat( CompressionQuality, pResult, EImageFormat::IF_L_UBYTE );
-					check(!FMemory::Memcmp(pTest->GetData(), Base->GetData(), Base->GetDataSize()));
-                }
+				{
+					// Verify
+					Ptr<Image> Test = ImagePixelFormat(CompressionQuality, Result, EImageFormat::IF_L_UBYTE);
+					check(Test->DataStorage == Base->DataStorage);
+				}
 #endif
+				break;
+			}
 
-                break;
-            }
-
-            case EImageFormat::IF_RGB_UBYTE:
-            case EImageFormat::IF_RGBA_UBYTE:
-            case EImageFormat::IF_BGRA_UBYTE:
-            {
-                check( OnlyLOD == -1 );
-                ImagePtr Temp = ImagePixelFormat( CompressionQuality, Base, EImageFormat::IF_L_UBYTE, -1 );
+			case EImageFormat::IF_RGB_UBYTE:
+			case EImageFormat::IF_RGBA_UBYTE:
+			case EImageFormat::IF_BGRA_UBYTE:
+			{
+				check(bFormatAllLODs);
+				Ptr<Image> TempBase = ImagePixelFormat(CompressionQuality, Base, EImageFormat::IF_L_UBYTE, -1);
  
 				// Try to compress
-				if (!pResult->m_data.Num())
+				if (Result->DataStorage.IsEmpty())
 				{
-					// Allocate memory for the compressed data. TODO: Smaller?
-					pResult->m_data.SetNum(Temp->m_data.Num());
+					// Allocate memory for the compressed data. TODO: Smaller?	
+					check(Base->DataStorage.Buffers.Num() == Result->DataStorage.Buffers.Num());
+
+					for (int32 I = 0; I < Base->DataStorage.Buffers.Num(); ++I)
+					{
+						int32 BufferSize = Base->DataStorage.Buffers[I].Num();
+						Result->DataStorage.Buffers[I].SetNumUninitialized(BufferSize);
+					}
+					Result->DataStorage.CompactedTailOffsets = Base->DataStorage.CompactedTailOffsets;
 				}
 
-				CompressRLE_L(bOutSuccess, Temp.get(), pResult );
+				CompressRLE_L(TempBase.get(), Result);
 
-#ifdef MUTABLE_DEBUG_RLE		
+#ifdef MUTABLE_DEBUG_RLE
 				if (bOutSuccess)
-                {
-                    // verify
-                    ImagePtr pTest = ImagePixelFormat( CompressionQuality, pResult, EImageFormat::IF_L_UBYTE );
-                    check( !FMemory::Memcmp(pTest->GetData(), pTemp->GetData(), pTemp->GetDataSize() ) );
-                }
+				{
+					// Verify
+					Ptr<Image> Test = ImagePixelFormat(CompressionQuality, Result, EImageFormat::IF_L_UBYTE);
+					check(Test->DataStorage == TempBase->DataStorage);
+				}
 #endif
-				ReleaseImage(Temp);
+				ReleaseImage(TempBase);
 
-                break;
-            }
+				break;
+			}
 
-            default:
-                // Case not implemented
-                check( false );
+			default:
+				// Case not implemented
+				check(false);
 
-            }
+			}
 
-            break;
-        }
+			break;
+		}
 
-        case EImageFormat::IF_L_UBIT_RLE:
-        {
-            switch ( Base->GetFormat() )
-            {
+		case EImageFormat::IF_L_UBIT_RLE:
+		{
+			switch (Base->GetFormat())
+			{
 
-            case EImageFormat::IF_L_UBYTE:
-            {
-                check( OnlyLOD == -1 );
+			case EImageFormat::IF_L_UBYTE:
+			{
+				check(bFormatAllLODs);
 
-				// Allocate memory for the compressed data. TODO: Smaller?
-				if (!pResult->m_data.Num())
+				// Try to compress
+				if (Result->DataStorage.IsEmpty())
 				{
-					// Allocate memory for the compressed data. TODO: Smaller?
-					pResult->m_data.SetNum(Base->m_data.Num());
+					// Allocate memory for the compressed data. TODO: Smaller?	
+					check(Base->DataStorage.Buffers.Num() == Result->DataStorage.Buffers.Num());
+
+					for (int32 I = 0; I < Base->DataStorage.Buffers.Num(); ++I)
+					{
+						int32 BufferSize = Base->DataStorage.Buffers[I].Num();
+						Result->DataStorage.Buffers[I].SetNumUninitialized(BufferSize);
+					}
+					Result->DataStorage.CompactedTailOffsets = Base->DataStorage.CompactedTailOffsets;
 				}
 				
-				CompressRLE_L1(bOutSuccess, Base, pResult );
-                
+				CompressRLE_L1(Base, Result);
+				
 #ifdef MUTABLE_DEBUG_RLE		
 				if (bOutSuccess)
-                {
+				{
 					// Verify
-                    ImagePtr pTest = ImagePixelFormat( CompressionQuality, pResult, EImageFormat::IF_L_UBYTE );
-                    check( !FMemory::Memcmp(pTest->GetData(), Base->GetData(), Base->GetDataSize() ) );
-                }
+					Ptr<Image> Test = ImagePixelFormat(CompressionQuality, Result, EImageFormat::IF_L_UBYTE);
+					check(Test->DataStorage == Base->DataStorage);
+				}
 #endif
 
-                break;
-            }
+				break;
+			}
 
-            case EImageFormat::IF_RGB_UBYTE:
-            case EImageFormat::IF_RGBA_UBYTE:
-            case EImageFormat::IF_BGRA_UBYTE:
-            {
-                check( OnlyLOD == -1 );
-                ImagePtr Temp = ImagePixelFormat( CompressionQuality, Base, EImageFormat::IF_L_UBYTE, -1);
-				CompressRLE_L1(bOutSuccess, Temp.get(), pResult );
+			case EImageFormat::IF_RGB_UBYTE:
+			case EImageFormat::IF_RGBA_UBYTE:
+			case EImageFormat::IF_BGRA_UBYTE:
+			{
+				check(bFormatAllLODs);
+				Ptr<Image> TempBase = ImagePixelFormat(CompressionQuality, Base, EImageFormat::IF_L_UBYTE, -1);
+				CompressRLE_L1(TempBase.get(), Result);
 
-#ifdef bOutSuccess		
+#ifdef MUTABLE_DEBUG_RLE
 				if (bOutSuccess)
 				{
 					// Verify
-					ImagePtr pTest = ImagePixelFormat(CompressionQuality, pResult, EImageFormat::IF_L_UBYTE);
-					check(!FMemory::Memcmp(pTest->GetData(), pTemp->GetData(), pTemp->GetDataSize()));
+					Ptr<Image> Test = ImagePixelFormat(CompressionQuality, Result, EImageFormat::IF_L_UBYTE);
+					check(Test->DataStorage == TempBase->DataStorage);
 				}
 #endif
+				ReleaseImage(TempBase);
+				break;
+			}
 
-				ReleaseImage(Temp);
-                break;
-            }
-
-            default:
-                // Case not implemented
-                check( false );
-
-            }
-
-            break;
-        }
-
-        case EImageFormat::IF_RGB_UBYTE:
-        {
-            uint8_t* pDestBuf = pResult->GetData();
-
-            switch ( Base->GetFormat() )
-            {
-
-            case EImageFormat::IF_L_UBYTE:
-            {
-                for ( int i=0; i<pixelCount; ++i )
-                {
-                    pDestBuf[i*3+0] = pBaseBuf[i];
-                    pDestBuf[i*3+1] = pBaseBuf[i];
-                    pDestBuf[i*3+2] = pBaseBuf[i];
-                }
-                break;
-            }
-
-            case EImageFormat::IF_RGB_UBYTE:
-            {
-				FMemory::Memcpy( pDestBuf, pBaseBuf, 3*pixelCount );
-                break;
-            }
-
-            case EImageFormat::IF_RGBA_UBYTE:
-            {
-                for ( int i=0; i<pixelCount; ++i )
-                {
-                    pDestBuf[i*3+0] = pBaseBuf[i*4+0];
-                    pDestBuf[i*3+1] = pBaseBuf[i*4+1];
-                    pDestBuf[i*3+2] = pBaseBuf[i*4+2];
-                }
-                break;
-            }
-
-            case EImageFormat::IF_BGRA_UBYTE:
-            {
-                for ( int i=0; i<pixelCount; ++i )
-                {
-                    pDestBuf[i*3+0] = pBaseBuf[i*4+2];
-                    pDestBuf[i*3+1] = pBaseBuf[i*4+1];
-                    pDestBuf[i*3+2] = pBaseBuf[i*4+0];
-                }
-                break;
-            }
-
-            case EImageFormat::IF_RGB_UBYTE_RLE:
-            {
-                UncompressRLE_RGB( Base, pResult );
-                break;
-            }
-
-            case EImageFormat::IF_RGBA_UBYTE_RLE:
-            {
-                check( OnlyLOD == -1 );
-                ImagePtr Temp = ImagePixelFormat( CompressionQuality, Base, EImageFormat::IF_RGBA_UBYTE, -1);
-                const uint8* TempBuf = Temp->GetData();
-                for ( int i=0; i<pixelCount; ++i )
-                {
-                    pDestBuf[i*3+0] = TempBuf[i*4+0];
-                    pDestBuf[i*3+1] = TempBuf[i*4+1];
-                    pDestBuf[i*3+2] = TempBuf[i*4+2];
-                }
-				ReleaseImage(Temp);
-                break;
-            }
-
-            case EImageFormat::IF_BC1:
-            {
-                for ( int m=0; m<resultLODCount; ++m )
-                {
-					FIntVector2 mipSize = pResult->CalculateMipSize( m );
-                    miro::BC1_to_RGB( mipSize[0], mipSize[1],
-                                        Base->GetMipData(baseLOD+m),
-                                        pResult->GetMipData(m) );
-                }
-                break;
-            }
-
-            case EImageFormat::IF_BC2:
-            {
-                for ( int m=0; m<resultLODCount; ++m )
-                {
-					FIntVector2 mipSize = pResult->CalculateMipSize( m );
-                    miro::BC2_to_RGB( mipSize[0], mipSize[1],
-                                        Base->GetMipData(baseLOD+m),
-                                        pResult->GetMipData(m) );
-                }
-                break;
-            }
-
-            case EImageFormat::IF_BC3:
-            {
-                for ( int m=0; m<resultLODCount; ++m )
-                {
-					FIntVector2 mipSize = pResult->CalculateMipSize( m );
-                    miro::BC3_to_RGB( mipSize[0], mipSize[1],
-                                        Base->GetMipData(baseLOD+m),
-                                        pResult->GetMipData(m) );
-                }
-                break;
-            }
-
-            case EImageFormat::IF_BC4:
-            {
-                for ( int m=0; m<resultLODCount; ++m )
-                {
-					FIntVector2 mipSize = pResult->CalculateMipSize( m );
-                    miro::BC4_to_RGB( mipSize[0], mipSize[1],
-                                      Base->GetMipData(baseLOD+m),
-                                      pResult->GetMipData(m) );
-                }
-                break;
-            }
-
-            case EImageFormat::IF_BC5:
-            {
-                for ( int m=0; m<resultLODCount; ++m )
-                {
-					FIntVector2 mipSize = pResult->CalculateMipSize( m );
-                    miro::BC5_to_RGB( mipSize[0], mipSize[1],
-                                      Base->GetMipData(baseLOD+m),
-                                      pResult->GetMipData(m) );
-                }
-                break;
-            }
-
-            case EImageFormat::IF_ASTC_4x4_RGB_LDR:
-            {
-                for ( int m=0; m<resultLODCount; ++m )
-                {
-					FIntVector2 mipSize = pResult->CalculateMipSize( m );
-                    miro::ASTC4x4RGBL_to_RGB( mipSize[0], mipSize[1],
-                                      Base->GetMipData(baseLOD+m),
-                                      pResult->GetMipData(m) );
-                }
-                break;
-            }
-
-            case EImageFormat::IF_ASTC_4x4_RGBA_LDR:
-            {
-                for ( int m=0; m<resultLODCount; ++m )
-                {
-					FIntVector2 mipSize = pResult->CalculateMipSize( m );
-                    miro::ASTC4x4RGBAL_to_RGB( mipSize[0], mipSize[1],
-                                      Base->GetMipData(baseLOD+m),
-                                      pResult->GetMipData(m) );
-                }
-                break;
-            }
-
-            case EImageFormat::IF_ASTC_4x4_RG_LDR:
-            {
-                for ( int m=0; m<resultLODCount; ++m )
-                {
-					FIntVector2 mipSize = pResult->CalculateMipSize( m );
-                    miro::ASTC4x4RGL_to_RGB( mipSize[0], mipSize[1],
-                                      Base->GetMipData(baseLOD+m),
-                                      pResult->GetMipData(m) );
-                }
-                break;
-            }
+			default:
+				// Case not implemented
+				check(false);
 
-			case EImageFormat::IF_ASTC_6x6_RGB_LDR:
+			}
+
+			break;
+		}
+
+		case EImageFormat::IF_RGB_UBYTE:
+		{
+			switch (Base->GetFormat())
 			{
-				for (int m = 0; m < resultLODCount; ++m)
+
+			case EImageFormat::IF_L_UBYTE:
+			{
+				ImagePixelFormatInternal::ProcessBatchedHelper(Result, Base, 3, 1, ResultLODBegin, BaseLODBegin, NumLODs,
+				[](uint8* DestBuf, const uint8* BaseBuf, int32 NumElems)
 				{
-					FIntVector2 mipSize = pResult->CalculateMipSize(m);
-					miro::ASTC6x6RGBL_to_RGB(mipSize[0], mipSize[1],
-						Base->GetMipData(baseLOD + m),
-						pResult->GetMipData(m));
-				}
+					for (int32 I = 0; I < NumElems; ++I)
+					{
+						DestBuf[I*3 + 0] = BaseBuf[I];
+						DestBuf[I*3 + 1] = BaseBuf[I];
+						DestBuf[I*3 + 2] = BaseBuf[I];
+					}
+				});
+
 				break;
 			}
-
-			case EImageFormat::IF_ASTC_6x6_RGBA_LDR:
+			case EImageFormat::IF_RGB_UBYTE:
 			{
-				for (int m = 0; m < resultLODCount; ++m)
+				ImagePixelFormatInternal::ProcessBatchedHelper(Result, Base, 3, 3, ResultLODBegin, BaseLODBegin, NumLODs,
+				[](uint8* DestBuf, const uint8* BaseBuf, int32 NumElems)
 				{
-					FIntVector2 mipSize = pResult->CalculateMipSize(m);
-					miro::ASTC6x6RGBAL_to_RGB(mipSize[0], mipSize[1],
-						Base->GetMipData(baseLOD + m),
-						pResult->GetMipData(m));
-				}
+					FMemory::Memcpy(DestBuf, BaseBuf, 3*NumElems);
+				});
+
 				break;
 			}
+			case EImageFormat::IF_RGBA_UBYTE:
+			{
+				ImagePixelFormatInternal::ProcessBatchedHelper(Result, Base, 3, 4, ResultLODBegin, BaseLODBegin, NumLODs,
+				[](uint8* DestBuf, const uint8* BaseBuf, int32 NumElems)
+				{
+					for (int32 I = 0; I < NumElems; ++I)
+					{
+						DestBuf[I*3 + 0] = BaseBuf[I*4 + 0];
+						DestBuf[I*3 + 1] = BaseBuf[I*4 + 1];
+						DestBuf[I*3 + 2] = BaseBuf[I*4 + 2];
+					}
+				});
 
-			case EImageFormat::IF_ASTC_6x6_RG_LDR:
+				break;
+			}
+			case EImageFormat::IF_BGRA_UBYTE:
 			{
-				for (int m = 0; m < resultLODCount; ++m)
+				ImagePixelFormatInternal::ProcessBatchedHelper(Result, Base, 3, 4, ResultLODBegin, BaseLODBegin, NumLODs,
+				[](uint8* DestBuf, const uint8* BaseBuf, int32 NumElems)
 				{
-					FIntVector2 mipSize = pResult->CalculateMipSize(m);
-					miro::ASTC6x6RGL_to_RGB(mipSize[0], mipSize[1],
-						Base->GetMipData(baseLOD + m),
-						pResult->GetMipData(m));
-				}
+					for (int32 I = 0; I < NumElems; ++I)
+					{
+						DestBuf[I*3 + 0] = BaseBuf[I*4 + 2];
+						DestBuf[I*3 + 1] = BaseBuf[I*4 + 1];
+						DestBuf[I*3 + 2] = BaseBuf[I*4 + 0];
+					}
+				});
+
 				break;
 			}
+			case EImageFormat::IF_RGB_UBYTE_RLE:
+			{
+				UncompressRLE_RGB(Base, Result);
+				break;
+			}
+			case EImageFormat::IF_RGBA_UBYTE_RLE:
+			{
+				check(bFormatAllLODs);
+
+				Ptr<Image> TempBase = ImagePixelFormat(CompressionQuality, Base, EImageFormat::IF_RGBA_UBYTE, -1);
+
+				ImagePixelFormatInternal::ProcessBatchedHelper(Result, TempBase.get(), 3, 4, ResultLODBegin, BaseLODBegin, NumLODs,
+				[](uint8* DestBuf, const uint8* BaseBuf, int32 NumElems)
+				{
+					for (int32 I = 0; I < NumElems; ++I)
+					{
+						DestBuf[I*3 + 0] = BaseBuf[I*4 + 2];
+						DestBuf[I*3 + 1] = BaseBuf[I*4 + 1];
+						DestBuf[I*3 + 2] = BaseBuf[I*4 + 0];
+					}
+				});
 
+				ReleaseImage(TempBase);
+				break;
+			}
+			case EImageFormat::IF_BC1:
+			{
+				ImagePixelFormatInternal::BlockDecompressionLODRangeHelper(
+						Result, Base, ResultLODBegin, BaseLODBegin, NumLODs, miro::BC1_to_RGB);
+				break;
+			}
+			case EImageFormat::IF_BC2:
+			{
+				ImagePixelFormatInternal::BlockDecompressionLODRangeHelper(
+						Result, Base, ResultLODBegin, BaseLODBegin, NumLODs, miro::BC2_to_RGB);
+				break;
+			}
+			case EImageFormat::IF_BC3:
+			{
+				ImagePixelFormatInternal::BlockDecompressionLODRangeHelper(
+						Result, Base, ResultLODBegin, BaseLODBegin, NumLODs, miro::BC3_to_RGB);
+				break;
+			}
+			case EImageFormat::IF_BC4:
+			{
+				ImagePixelFormatInternal::BlockDecompressionLODRangeHelper(
+						Result, Base, ResultLODBegin, BaseLODBegin, NumLODs, miro::BC4_to_RGB);
+				break;
+			}
+			case EImageFormat::IF_BC5:
+			{
+				ImagePixelFormatInternal::BlockDecompressionLODRangeHelper(
+						Result, Base, ResultLODBegin, BaseLODBegin, NumLODs, miro::BC5_to_RGB);
+				break;
+			}
+			case EImageFormat::IF_ASTC_4x4_RGB_LDR:
+			{
+				ImagePixelFormatInternal::BlockDecompressionLODRangeHelper(
+						Result, Base, ResultLODBegin, BaseLODBegin, NumLODs, miro::ASTC4x4RGBL_to_RGB);
+				break;
+			}
+			case EImageFormat::IF_ASTC_4x4_RGBA_LDR:
+			{
+				ImagePixelFormatInternal::BlockDecompressionLODRangeHelper(
+						Result, Base, ResultLODBegin, BaseLODBegin, NumLODs, miro::ASTC4x4RGBAL_to_RGB);
+				break;
+			}
+			case EImageFormat::IF_ASTC_4x4_RG_LDR:
+			{
+				ImagePixelFormatInternal::BlockDecompressionLODRangeHelper(
+						Result, Base, ResultLODBegin, BaseLODBegin, NumLODs, miro::ASTC4x4RGL_to_RGB);
+				break;
+			}
+			case EImageFormat::IF_ASTC_6x6_RGB_LDR:
+			{
+				ImagePixelFormatInternal::BlockDecompressionLODRangeHelper(
+						Result, Base, ResultLODBegin, BaseLODBegin, NumLODs, miro::ASTC6x6RGBL_to_RGB);
+				break;
+			}
+			case EImageFormat::IF_ASTC_6x6_RGBA_LDR:
+			{
+				ImagePixelFormatInternal::BlockDecompressionLODRangeHelper(
+						Result, Base, ResultLODBegin, BaseLODBegin, NumLODs, miro::ASTC6x6RGBAL_to_RGB);
+				break;
+			}
+			case EImageFormat::IF_ASTC_6x6_RG_LDR:
+			{
+				ImagePixelFormatInternal::BlockDecompressionLODRangeHelper(
+						Result, Base, ResultLODBegin, BaseLODBegin, NumLODs, miro::ASTC6x6RGL_to_RGB);
+				break;
+			}
 			case EImageFormat::IF_ASTC_8x8_RGB_LDR:
 			{
-				for (int m = 0; m < resultLODCount; ++m)
-				{
-					FIntVector2 mipSize = pResult->CalculateMipSize(m);
-					miro::ASTC8x8RGBL_to_RGB(mipSize[0], mipSize[1],
-						Base->GetMipData(baseLOD + m),
-						pResult->GetMipData(m));
-				}
+				ImagePixelFormatInternal::BlockDecompressionLODRangeHelper(
+						Result, Base, ResultLODBegin, BaseLODBegin, NumLODs, miro::ASTC8x8RGBL_to_RGB);
 				break;
 			}
-
 			case EImageFormat::IF_ASTC_8x8_RGBA_LDR:
 			{
-				for (int m = 0; m < resultLODCount; ++m)
-				{
-					FIntVector2 mipSize = pResult->CalculateMipSize(m);
-					miro::ASTC8x8RGBAL_to_RGB(mipSize[0], mipSize[1],
-						Base->GetMipData(baseLOD + m),
-						pResult->GetMipData(m));
-				}
+				ImagePixelFormatInternal::BlockDecompressionLODRangeHelper(
+						Result, Base, ResultLODBegin, BaseLODBegin, NumLODs, miro::ASTC8x8RGBAL_to_RGB);
 				break;
 			}
-
 			case EImageFormat::IF_ASTC_8x8_RG_LDR:
 			{
-				for (int m = 0; m < resultLODCount; ++m)
-				{
-					FIntVector2 mipSize = pResult->CalculateMipSize(m);
-					miro::ASTC8x8RGL_to_RGB(mipSize[0], mipSize[1],
-						Base->GetMipData(baseLOD + m),
-						pResult->GetMipData(m));
-				}
+				ImagePixelFormatInternal::BlockDecompressionLODRangeHelper(
+						Result, Base, ResultLODBegin, BaseLODBegin, NumLODs, miro::ASTC8x8RGL_to_RGB);
 				break;
 			}
-
 			case EImageFormat::IF_ASTC_10x10_RGB_LDR:
 			{
-				for (int m = 0; m < resultLODCount; ++m)
-				{
-					FIntVector2 mipSize = pResult->CalculateMipSize(m);
-					miro::ASTC10x10RGBL_to_RGB(mipSize[0], mipSize[1],
-						Base->GetMipData(baseLOD + m),
-						pResult->GetMipData(m));
-				}
+				ImagePixelFormatInternal::BlockDecompressionLODRangeHelper(
+						Result, Base, ResultLODBegin, BaseLODBegin, NumLODs, miro::ASTC10x10RGBL_to_RGB);
 				break;
 			}
-
 			case EImageFormat::IF_ASTC_10x10_RGBA_LDR:
 			{
-				for (int m = 0; m < resultLODCount; ++m)
-				{
-					FIntVector2 mipSize = pResult->CalculateMipSize(m);
-					miro::ASTC10x10RGBAL_to_RGB(mipSize[0], mipSize[1],
-						Base->GetMipData(baseLOD + m),
-						pResult->GetMipData(m));
-				}
+				ImagePixelFormatInternal::BlockDecompressionLODRangeHelper(
+						Result, Base, ResultLODBegin, BaseLODBegin, NumLODs, miro::ASTC10x10RGBAL_to_RGB);
 				break;
 			}
-
 			case EImageFormat::IF_ASTC_10x10_RG_LDR:
 			{
-				for (int m = 0; m < resultLODCount; ++m)
-				{
-					FIntVector2 mipSize = pResult->CalculateMipSize(m);
-					miro::ASTC10x10RGL_to_RGB(mipSize[0], mipSize[1],
-						Base->GetMipData(baseLOD + m),
-						pResult->GetMipData(m));
-				}
+				ImagePixelFormatInternal::BlockDecompressionLODRangeHelper(
+						Result, Base, ResultLODBegin, BaseLODBegin, NumLODs, miro::ASTC10x10RGL_to_RGB);
 				break;
 			}
-
 			case EImageFormat::IF_ASTC_12x12_RGB_LDR:
 			{
-				for (int m = 0; m < resultLODCount; ++m)
-				{
-					FIntVector2 mipSize = pResult->CalculateMipSize(m);
-					miro::ASTC12x12RGBL_to_RGB(mipSize[0], mipSize[1],
-						Base->GetMipData(baseLOD + m),
-						pResult->GetMipData(m));
-				}
+				ImagePixelFormatInternal::BlockDecompressionLODRangeHelper(
+						Result, Base, ResultLODBegin, BaseLODBegin, NumLODs, miro::ASTC12x12RGBL_to_RGB);
 				break;
 			}
-
 			case EImageFormat::IF_ASTC_12x12_RGBA_LDR:
 			{
-				for (int m = 0; m < resultLODCount; ++m)
-				{
-					FIntVector2 mipSize = pResult->CalculateMipSize(m);
-					miro::ASTC12x12RGBAL_to_RGB(mipSize[0], mipSize[1],
-						Base->GetMipData(baseLOD + m),
-						pResult->GetMipData(m));
-				}
+				ImagePixelFormatInternal::BlockDecompressionLODRangeHelper(
+						Result, Base, ResultLODBegin, BaseLODBegin, NumLODs, miro::ASTC12x12RGBAL_to_RGB);
 				break;
 			}
-
 			case EImageFormat::IF_ASTC_12x12_RG_LDR:
 			{
-				for (int m = 0; m < resultLODCount; ++m)
-				{
-					FIntVector2 mipSize = pResult->CalculateMipSize(m);
-					miro::ASTC12x12RGL_to_RGB(mipSize[0], mipSize[1],
-						Base->GetMipData(baseLOD + m),
-						pResult->GetMipData(m));
-				}
+				ImagePixelFormatInternal::BlockDecompressionLODRangeHelper(
+						Result, Base, ResultLODBegin, BaseLODBegin, NumLODs, miro::ASTC12x12RGL_to_RGB);
 				break;
 			}
 
-            default:
-                // Case not implemented
-                check( false );
-
-            }
+			default:
+				// Case not implemented
+				check(false);
+			}
 
-            break;
-        }
+			break;
+		}
 
-        case EImageFormat::IF_BGRA_UBYTE:
-        {
-            // TODO: Optimise
-            pResult->m_format = EImageFormat::IF_RGBA_UBYTE;
-            ImagePixelFormat(bOutSuccess, CompressionQuality, pResult, Base, OnlyLOD);
+		case EImageFormat::IF_BGRA_UBYTE:
+		{
+			// TODO: Optimise
+			Result->DataStorage.ImageFormat = EImageFormat::IF_RGBA_UBYTE;
+			ImagePixelFormat(bOutSuccess, CompressionQuality, Result, Base, ResultLODBegin, BaseLODBegin, NumLODs);
 			check(bOutSuccess);
-            pResult->m_format = EImageFormat::IF_BGRA_UBYTE;
-
-            uint8* pDestBuf = pResult->GetData();
-            for ( int i=0; i<pixelCount; ++i )
-            {
-                uint8 temp = pDestBuf[i*4+0];
-                pDestBuf[i*4+0] = pDestBuf[i*4+2];
-                pDestBuf[i*4+2] = temp;
-            }
-            break;
-        }
-
-        case EImageFormat::IF_RGBA_UBYTE:
-        {
-            uint8_t* pDestBuf = pResult->GetData();
-
-            switch ( Base->GetFormat() )
-            {
-
-            case EImageFormat::IF_RGBA_UBYTE_RLE:
-            {
-                check( OnlyLOD == -1 );
-                UncompressRLE_RGBA( Base, pResult );
-                break;
-            }
-
-            case EImageFormat::IF_L_UBYTE:
-            {
-                for ( int i=0; i<pixelCount; ++i )
-                {
-                    pDestBuf[i*4+0] = pBaseBuf[i];
-                    pDestBuf[i*4+1] = pBaseBuf[i];
-                    pDestBuf[i*4+2] = pBaseBuf[i];
-                    pDestBuf[i*4+3] = 255;
-                }
-                break;
-            }
-
-            case EImageFormat::IF_RGB_UBYTE:
-            {
-                for ( int i=0; i<pixelCount; ++i )
-                {
-                    pDestBuf[i*4+0] = pBaseBuf[i*3+0];
-                    pDestBuf[i*4+1] = pBaseBuf[i*3+1];
-                    pDestBuf[i*4+2] = pBaseBuf[i*3+2];
-                    pDestBuf[i*4+3] = 255;
-                }
-                break;
-            }
-
-            case EImageFormat::IF_RGBA_UBYTE:
-            {
-				FMemory::Memcpy( pDestBuf, pBaseBuf, 4*pixelCount );
-                break;
-            }
-
-            case EImageFormat::IF_BGRA_UBYTE:
-            {
-                for ( int i=0; i<pixelCount; ++i )
-                {
-                    pDestBuf[i*4+0] = pBaseBuf[i*4+2];
-                    pDestBuf[i*4+1] = pBaseBuf[i*4+1];
-                    pDestBuf[i*4+2] = pBaseBuf[i*4+0];
-                    pDestBuf[i*4+3] = pBaseBuf[i*4+3];
-                }
-                break;
-            }
-
-            case EImageFormat::IF_BC1:
-            {
-                for ( int m=0; m<resultLODCount; ++m )
-                {
-					FIntVector2 mipSize = pResult->CalculateMipSize( m );
-                    miro::BC1_to_RGBA( mipSize[0], mipSize[1],
-                                        Base->GetMipData(baseLOD+m),
-                                        pResult->GetMipData(m) );
-                }
-                break;
-            }
-
-            case EImageFormat::IF_BC2:
-            {
-                for ( int m=0; m<resultLODCount; ++m )
-                {
-					FIntVector2 mipSize = pResult->CalculateMipSize( m );
-                    miro::BC2_to_RGBA( mipSize[0], mipSize[1],
-                                        Base->GetMipData(baseLOD+m),
-                                        pResult->GetMipData(m) );
-                }
-                break;
-            }
-
-            case EImageFormat::IF_BC3:
-            {
-                for ( int m=0; m<resultLODCount; ++m )
-                {
-					FIntVector2 mipSize = pResult->CalculateMipSize( m );
-                    miro::BC3_to_RGBA( mipSize[0], mipSize[1],
-                                        Base->GetMipData(baseLOD+m),
-                                        pResult->GetMipData(m) );
-                }
-                break;
-            }
-
-            case EImageFormat::IF_BC4:
-            {
-                for ( int m=0; m<resultLODCount; ++m )
-                {
-					FIntVector2 mipSize = pResult->CalculateMipSize( m );
-                    miro::BC4_to_RGBA( mipSize[0], mipSize[1],
-                                       Base->GetMipData(baseLOD+m),
-                                       pResult->GetMipData(m) );
-                }
-                break;
-            }
-
-            case EImageFormat::IF_BC5:
-            {
-                for ( int m=0; m<resultLODCount; ++m )
-                {
-                    FIntVector2 mipSize = pResult->CalculateMipSize( m );
-                    miro::BC5_to_RGBA( mipSize[0], mipSize[1],
-                                        Base->GetMipData(baseLOD+m),
-                                        pResult->GetMipData(m) );
-                }
-                break;
-            }
-
-            case EImageFormat::IF_ASTC_4x4_RGB_LDR:
-            {
-                for ( int m=0; m<resultLODCount; ++m )
-                {
-                    FIntVector2 mipSize = pResult->CalculateMipSize( m );
-                    miro::ASTC4x4RGBL_to_RGBA( mipSize[0], mipSize[1],
-                                        Base->GetMipData(baseLOD+m),
-                                        pResult->GetMipData(m) );
-                }
-                break;
-            }
-
-            case EImageFormat::IF_ASTC_4x4_RGBA_LDR:
-            {
-                for ( int m=0; m<resultLODCount; ++m )
-                {
-                    FIntVector2 mipSize = pResult->CalculateMipSize( m );
-                    miro::ASTC4x4RGBAL_to_RGBA( mipSize[0], mipSize[1],
-                                        Base->GetMipData(baseLOD+m),
-                                        pResult->GetMipData(m) );
-                }
-                break;
-            }
-
-            case EImageFormat::IF_ASTC_4x4_RG_LDR:
-            {
-                for ( int m=0; m<resultLODCount; ++m )
-                {
-                    FIntVector2 mipSize = pResult->CalculateMipSize( m );
-                    miro::ASTC4x4RGL_to_RGBA( mipSize[0], mipSize[1],
-                                        Base->GetMipData(baseLOD+m),
-                                        pResult->GetMipData(m) );
-                }
-                break;
-            }
+			Result->DataStorage.ImageFormat = EImageFormat::IF_BGRA_UBYTE;
 
-			case EImageFormat::IF_ASTC_6x6_RGBA_LDR:
+			ImagePixelFormatInternal::ProcessBatchedHelper(Result, 4, ResultLODBegin, NumLODs,
+			[](uint8* DestBuf, int32 NumElems)
 			{
-				for (int m = 0; m < resultLODCount; ++m)
+				for (int32 I = 0; I < NumElems; ++I)
 				{
-					FIntVector2 mipSize = pResult->CalculateMipSize(m);
-					miro::ASTC6x6RGBAL_to_RGBA(mipSize[0], mipSize[1],
-						Base->GetMipData(baseLOD + m),
-						pResult->GetMipData(m));
+					uint8 Temp = DestBuf[I*4 + 0];
+					DestBuf[I*4 + 0] = DestBuf[I*4 + 2];
+					DestBuf[I*4 + 2] = Temp;
 				}
+			});
+
+			break;
+		}
+
+		case EImageFormat::IF_RGBA_UBYTE:
+		{
+			switch (Base->GetFormat())
+			{
+
+			case EImageFormat::IF_RGBA_UBYTE_RLE:
+			{
+				check(bFormatAllLODs);
+				UncompressRLE_RGBA(Base, Result);
 				break;
 			}
-
-			case EImageFormat::IF_ASTC_6x6_RG_LDR:
+			case EImageFormat::IF_L_UBYTE:
 			{
-				for (int m = 0; m < resultLODCount; ++m)
+				ImagePixelFormatInternal::ProcessBatchedHelper(Result, Base, 4, 1, ResultLODBegin, BaseLODBegin, NumLODs,
+				[](uint8* DestBuf, const uint8* BaseBuf, int32 NumElems)
 				{
-					FIntVector2 mipSize = pResult->CalculateMipSize(m);
-					miro::ASTC6x6RGL_to_RGBA(mipSize[0], mipSize[1],
-						Base->GetMipData(baseLOD + m),
-						pResult->GetMipData(m));
-				}
+					for (int32 I = 0; I < NumElems; ++I)
+					{
+						DestBuf[I*4 + 0] = BaseBuf[I];
+						DestBuf[I*4 + 1] = BaseBuf[I];
+						DestBuf[I*4 + 2] = BaseBuf[I];
+						DestBuf[I*4 + 3] = 255;
+					}
+				});
+
 				break;
 			}
-
-			case EImageFormat::IF_ASTC_8x8_RGB_LDR:
+			case EImageFormat::IF_RGB_UBYTE:
 			{
-				for (int m = 0; m < resultLODCount; ++m)
+				ImagePixelFormatInternal::ProcessBatchedHelper(Result, Base, 4, 3, ResultLODBegin, BaseLODBegin, NumLODs,
+				[](uint8* DestBuf, const uint8* BaseBuf, int32 NumElems)
 				{
-					FIntVector2 mipSize = pResult->CalculateMipSize(m);
-					miro::ASTC8x8RGBL_to_RGBA(mipSize[0], mipSize[1],
-						Base->GetMipData(baseLOD + m),
-						pResult->GetMipData(m));
-				}
+					for (int32 I = 0; I < NumElems; ++I)
+					{
+						DestBuf[I*4 + 0] = BaseBuf[I*3 + 0];
+						DestBuf[I*4 + 1] = BaseBuf[I*3 + 1];
+						DestBuf[I*4 + 2] = BaseBuf[I*3 + 2];
+						DestBuf[I*4 + 3] = 255;
+					}
+				});
+
 				break;
 			}
-
-			case EImageFormat::IF_ASTC_8x8_RGBA_LDR:
+			case EImageFormat::IF_RGBA_UBYTE:
 			{
-				for (int m = 0; m < resultLODCount; ++m)
+				ImagePixelFormatInternal::ProcessBatchedHelper(Result, Base, 4, 4, ResultLODBegin, BaseLODBegin, NumLODs,
+				[](uint8* DestBuf, const uint8* BaseBuf, int32 NumElems)
 				{
-					FIntVector2 mipSize = pResult->CalculateMipSize(m);
-					miro::ASTC8x8RGBAL_to_RGBA(mipSize[0], mipSize[1],
-						Base->GetMipData(baseLOD + m),
-						pResult->GetMipData(m));
-				}
+					FMemory::Memcpy(DestBuf, BaseBuf, NumElems*4);
+				});
+
 				break;
 			}
-
-			case EImageFormat::IF_ASTC_8x8_RG_LDR:
+			case EImageFormat::IF_BGRA_UBYTE:
 			{
-				for (int m = 0; m < resultLODCount; ++m)
+				ImagePixelFormatInternal::ProcessBatchedHelper(Result, Base, 4, 4, ResultLODBegin, BaseLODBegin, NumLODs,
+				[](uint8* DestBuf, const uint8* BaseBuf, int32 NumElems)
 				{
-					FIntVector2 mipSize = pResult->CalculateMipSize(m);
-					miro::ASTC8x8RGL_to_RGBA(mipSize[0], mipSize[1],
-						Base->GetMipData(baseLOD + m),
-						pResult->GetMipData(m));
-				}
+					for (int32 I = 0; I < NumElems; ++I)
+					{
+						DestBuf[I*4 + 0] = BaseBuf[I*3 + 2];
+						DestBuf[I*4 + 1] = BaseBuf[I*3 + 1];
+						DestBuf[I*4 + 2] = BaseBuf[I*3 + 0];
+						DestBuf[I*4 + 3] = BaseBuf[I*3 + 3];
+					}
+				});
+
+				break;
+			}
+			case EImageFormat::IF_BC1:
+			{
+				ImagePixelFormatInternal::BlockDecompressionLODRangeHelper(
+						Result, Base, ResultLODBegin, BaseLODBegin, NumLODs, miro::BC1_to_RGBA);
+				break;
+			}
+			case EImageFormat::IF_BC2:
+			{
+				ImagePixelFormatInternal::BlockDecompressionLODRangeHelper(
+						Result, Base, ResultLODBegin, BaseLODBegin, NumLODs, miro::BC2_to_RGBA);
+				break;
+			}
+			case EImageFormat::IF_BC3:
+			{
+				ImagePixelFormatInternal::BlockDecompressionLODRangeHelper(
+						Result, Base, ResultLODBegin, BaseLODBegin, NumLODs, miro::BC3_to_RGBA);
+				break;
+			}
+			case EImageFormat::IF_BC4:
+			{
+				ImagePixelFormatInternal::BlockDecompressionLODRangeHelper(
+						Result, Base, ResultLODBegin, BaseLODBegin, NumLODs, miro::BC4_to_RGBA);
+				break;
+			}
+			case EImageFormat::IF_BC5:
+			{
+				ImagePixelFormatInternal::BlockDecompressionLODRangeHelper(
+						Result, Base, ResultLODBegin, BaseLODBegin, NumLODs, miro::BC4_to_RGBA);
+				break;
+			}
+			case EImageFormat::IF_ASTC_4x4_RGB_LDR:
+			{
+				ImagePixelFormatInternal::BlockDecompressionLODRangeHelper(
+						Result, Base, ResultLODBegin, BaseLODBegin, NumLODs, miro::ASTC4x4RGBL_to_RGBA);
+				break;
+			}
+			case EImageFormat::IF_ASTC_4x4_RGBA_LDR:
+			{
+				ImagePixelFormatInternal::BlockDecompressionLODRangeHelper(
+						Result, Base, ResultLODBegin, BaseLODBegin, NumLODs, miro::ASTC4x4RGBAL_to_RGBA);
+				break;
+			}
+			case EImageFormat::IF_ASTC_4x4_RG_LDR:
+			{
+				ImagePixelFormatInternal::BlockDecompressionLODRangeHelper(
+						Result, Base, ResultLODBegin, BaseLODBegin, NumLODs, miro::ASTC4x4RGL_to_RGBA);
 				break;
 			}
-
 			case EImageFormat::IF_ASTC_6x6_RGB_LDR:
 			{
-				for (int m = 0; m < resultLODCount; ++m)
-				{
-					FIntVector2 mipSize = pResult->CalculateMipSize(m);
-					miro::ASTC6x6RGBL_to_RGBA(mipSize[0], mipSize[1],
-						Base->GetMipData(baseLOD + m),
-						pResult->GetMipData(m));
-				}
+				ImagePixelFormatInternal::BlockDecompressionLODRangeHelper(
+						Result, Base, ResultLODBegin, BaseLODBegin, NumLODs, miro::ASTC6x6RGBL_to_RGBA);
+				break;
+			}
+			case EImageFormat::IF_ASTC_6x6_RGBA_LDR:
+			{
+				ImagePixelFormatInternal::BlockDecompressionLODRangeHelper(
+						Result, Base, ResultLODBegin, BaseLODBegin, NumLODs, miro::ASTC6x6RGBAL_to_RGBA);
+				break;
+			}
+			case EImageFormat::IF_ASTC_6x6_RG_LDR:
+			{
+				ImagePixelFormatInternal::BlockDecompressionLODRangeHelper(
+						Result, Base, ResultLODBegin, BaseLODBegin, NumLODs, miro::ASTC6x6RGL_to_RGBA);
+				break;
+			}
+			case EImageFormat::IF_ASTC_8x8_RGB_LDR:
+			{
+				ImagePixelFormatInternal::BlockDecompressionLODRangeHelper(
+						Result, Base, ResultLODBegin, BaseLODBegin, NumLODs, miro::ASTC8x8RGBL_to_RGBA);
+				break;
+			}
+			case EImageFormat::IF_ASTC_8x8_RGBA_LDR:
+			{
+				ImagePixelFormatInternal::BlockDecompressionLODRangeHelper(
+						Result, Base, ResultLODBegin, BaseLODBegin, NumLODs, miro::ASTC8x8RGBAL_to_RGBA);
+				break;
+			}
+			case EImageFormat::IF_ASTC_8x8_RG_LDR:
+			{
+				ImagePixelFormatInternal::BlockDecompressionLODRangeHelper(
+						Result, Base, ResultLODBegin, BaseLODBegin, NumLODs, miro::ASTC8x8RGL_to_RGBA);
 				break;
 			}
-
 			case EImageFormat::IF_ASTC_10x10_RGB_LDR:
 			{
-				for (int m = 0; m < resultLODCount; ++m)
-				{
-					FIntVector2 mipSize = pResult->CalculateMipSize(m);
-					miro::ASTC10x10RGBL_to_RGBA(mipSize[0], mipSize[1],
-						Base->GetMipData(baseLOD + m),
-						pResult->GetMipData(m));
-				}
+				ImagePixelFormatInternal::BlockDecompressionLODRangeHelper(
+						Result, Base, ResultLODBegin, BaseLODBegin, NumLODs, miro::ASTC10x10RGBL_to_RGBA);
 				break;
 			}
-
 			case EImageFormat::IF_ASTC_10x10_RGBA_LDR:
 			{
-				for (int m = 0; m < resultLODCount; ++m)
-				{
-					FIntVector2 mipSize = pResult->CalculateMipSize(m);
-					miro::ASTC10x10RGBAL_to_RGBA(mipSize[0], mipSize[1],
-						Base->GetMipData(baseLOD + m),
-						pResult->GetMipData(m));
-				}
+				ImagePixelFormatInternal::BlockDecompressionLODRangeHelper(
+						Result, Base, ResultLODBegin, BaseLODBegin, NumLODs, miro::ASTC10x10RGBAL_to_RGBA);
 				break;
 			}
-
 			case EImageFormat::IF_ASTC_10x10_RG_LDR:
 			{
-				for (int m = 0; m < resultLODCount; ++m)
-				{
-					FIntVector2 mipSize = pResult->CalculateMipSize(m);
-					miro::ASTC10x10RGL_to_RGBA(mipSize[0], mipSize[1],
-						Base->GetMipData(baseLOD + m),
-						pResult->GetMipData(m));
-				}
+				ImagePixelFormatInternal::BlockDecompressionLODRangeHelper(
+						Result, Base, ResultLODBegin, BaseLODBegin, NumLODs, miro::ASTC10x10RGL_to_RGBA);
 				break;
 			}
-
 			case EImageFormat::IF_ASTC_12x12_RGB_LDR:
 			{
-				for (int m = 0; m < resultLODCount; ++m)
-				{
-					FIntVector2 mipSize = pResult->CalculateMipSize(m);
-					miro::ASTC12x12RGBL_to_RGBA(mipSize[0], mipSize[1],
-						Base->GetMipData(baseLOD + m),
-						pResult->GetMipData(m));
-				}
+				ImagePixelFormatInternal::BlockDecompressionLODRangeHelper(
+						Result, Base, ResultLODBegin, BaseLODBegin, NumLODs, miro::ASTC12x12RGBL_to_RGBA);
 				break;
 			}
-
 			case EImageFormat::IF_ASTC_12x12_RGBA_LDR:
 			{
-				for (int m = 0; m < resultLODCount; ++m)
-				{
-					FIntVector2 mipSize = pResult->CalculateMipSize(m);
-					miro::ASTC12x12RGBAL_to_RGBA(mipSize[0], mipSize[1],
-						Base->GetMipData(baseLOD + m),
-						pResult->GetMipData(m));
-				}
+				ImagePixelFormatInternal::BlockDecompressionLODRangeHelper(
+						Result, Base, ResultLODBegin, BaseLODBegin, NumLODs, miro::ASTC12x12RGBAL_to_RGBA);
 				break;
 			}
-
 			case EImageFormat::IF_ASTC_12x12_RG_LDR:
 			{
-				for (int m = 0; m < resultLODCount; ++m)
-				{
-					FIntVector2 mipSize = pResult->CalculateMipSize(m);
-					miro::ASTC12x12RGL_to_RGBA(mipSize[0], mipSize[1],
-						Base->GetMipData(baseLOD + m),
-						pResult->GetMipData(m));
-				}
+				ImagePixelFormatInternal::BlockDecompressionLODRangeHelper(
+						Result, Base, ResultLODBegin, BaseLODBegin, NumLODs, miro::ASTC12x12RGL_to_RGBA);
 				break;
 			}
-
 			case EImageFormat::IF_RGB_UBYTE_RLE:
 			{
-				check(OnlyLOD == -1);
-				ImagePtr Temp = ImagePixelFormat(CompressionQuality, Base, EImageFormat::IF_RGB_UBYTE);
-				const uint8* TempData = Temp->GetData();
-				for (int i = 0; i < pixelCount; ++i)
+				check(bFormatAllLODs);
+				Ptr<Image> TempBase = ImagePixelFormat(CompressionQuality, Base, EImageFormat::IF_RGB_UBYTE);
+
+				ImagePixelFormatInternal::ProcessBatchedHelper(Result, TempBase.get(), 4, 3, ResultLODBegin, BaseLODBegin, NumLODs,
+				[](uint8* DestBuf, const uint8* BaseBuf, int32 NumElems)
 				{
-					pDestBuf[i * 4 + 0] = TempData[i * 3 + 0];
-					pDestBuf[i * 4 + 1] = TempData[i * 3 + 1];
-					pDestBuf[i * 4 + 2] = TempData[i * 3 + 2];
-					pDestBuf[i * 4 + 3] = 255;
-				}
-				ReleaseImage(Temp);
+					for (int32 I = 0; I < NumElems; ++I)
+					{
+						DestBuf[I*4 + 0] = BaseBuf[I*3 + 0];
+						DestBuf[I*4 + 1] = BaseBuf[I*3 + 1];
+						DestBuf[I*4 + 2] = BaseBuf[I*3 + 2];
+						DestBuf[I*4 + 3] = 255;
+					}
+				});
+
+				ReleaseImage(TempBase);
 				break;
 			}
-
 			case EImageFormat::IF_L_UBYTE_RLE:
 			{
-				check(OnlyLOD == -1);
-				ImagePtr Temp = ImagePixelFormat(CompressionQuality, Base, EImageFormat::IF_L_UBYTE);
-				const uint8* TempData = Temp->GetData();
-				for (int i = 0; i < pixelCount; ++i)
+				check(bFormatAllLODs);
+				Ptr<Image> TempBase = ImagePixelFormat(CompressionQuality, Base, EImageFormat::IF_L_UBYTE);
+
+				ImagePixelFormatInternal::ProcessBatchedHelper(Result, TempBase.get(), 4, 1, ResultLODBegin, BaseLODBegin, NumLODs,
+				[](uint8* DestBuf, const uint8* BaseBuf, int32 NumElems)
 				{
-					pDestBuf[i * 4 + 0] = TempData[i];
-					pDestBuf[i * 4 + 1] = TempData[i];
-					pDestBuf[i * 4 + 2] = TempData[i];
-					pDestBuf[i * 4 + 3] = 255;
-				}
-				ReleaseImage(Temp);
+					for (int32 I = 0; I < NumElems; ++I)
+					{
+						DestBuf[I*4 + 0] = BaseBuf[I];
+						DestBuf[I*4 + 1] = BaseBuf[I];
+						DestBuf[I*4 + 2] = BaseBuf[I];
+						DestBuf[I*4 + 3] = 255;
+					}
+				});
+
+				ReleaseImage(TempBase);
 				break;
 			}
 
-            default:
-                // Case not implemented
-                check( false );
+			default:
+				// Case not implemented
+				check(false);
+			}
+
+			break;
+		}
+
+		case EImageFormat::IF_RGBA_UBYTE_RLE:
+		{
+			switch (Base->GetFormat())
+			{
+
+			case EImageFormat::IF_RGBA_UBYTE:
+			{
+				check(bFormatAllLODs);
+				CompressRLE_RGBA(Base, Result);
+
+				break;
+			}
+			case EImageFormat::IF_RGB_UBYTE:
+			{
+				check(bFormatAllLODs);
+
+				// \todo: optimise
+				Ptr<Image> TempBase = ImagePixelFormat(CompressionQuality, Base, EImageFormat::IF_RGBA_UBYTE);
+				CompressRLE_RGBA(TempBase.get(), Result);
 
-            }
+				// Test
+				//Ptr<Image> Test = ImagePixelFormat(CompressionQuality, Result, EImageFormat::IF_RGBA_UBYTE); 
+				//check(Test->DataStorage == Base->DataStorage);
 
-            break;
-        }
+				ReleaseImage(TempBase);
 
+				break;
+			}
+			case EImageFormat::IF_RGB_UBYTE_RLE:
+			{
+				check(bFormatAllLODs);
 
-        case EImageFormat::IF_RGBA_UBYTE_RLE:
-        {
-            switch ( Base->GetFormat() )
-            {
+				// \todo: optimise
+				Ptr<Image> Temp1 = ImagePixelFormat(CompressionQuality, Base, EImageFormat::IF_RGB_UBYTE);
+				Ptr<Image> Temp2 = ImagePixelFormat(CompressionQuality, Temp1.get(), EImageFormat::IF_RGBA_UBYTE);
+				ReleaseImage(Temp1);
+				CompressRLE_RGBA(Temp2.get(), Result);
+				ReleaseImage(Temp2);
 
-            case EImageFormat::IF_RGBA_UBYTE:
-            {
-                check( OnlyLOD == -1 );
-                CompressRLE_RGBA( Base, pResult );
-                break;
-            }
+				break;
+			}
+			default:
+				// Case not implemented
+				check(false);
 
-            case EImageFormat::IF_RGB_UBYTE:
-            {
-                check( OnlyLOD == -1 );
+			}
 
-                // \todo: optimise
-                ImagePtr Temp = ImagePixelFormat( CompressionQuality, Base, EImageFormat::IF_RGBA_UBYTE);
-                CompressRLE_RGBA( Temp.get(), pResult );
+			break;
+		}
 
-                // Test
-                //                ImagePtr pTest = ImagePixelFormat( CompressionQuality,
-                //                pResult, EImageFormat::IF_RGBA_UBYTE ); check(
-                //                !memcmp(pTest->GetData(), Base->GetData(),
-                //                Base->GetDataSize() ) );
+		case EImageFormat::IF_RGB_UBYTE_RLE:
+		{
+			switch (Base->GetFormat())
+			{
 
-				ReleaseImage(Temp);
-                break;
-            }
+			case EImageFormat::IF_RGB_UBYTE:
+			{
+				check(bFormatAllLODs);
+				CompressRLE_RGB(Base, Result);
 
-            case EImageFormat::IF_RGB_UBYTE_RLE:
-            {
-                check( OnlyLOD == -1 );
+				// Test
+				// Ptr<Image> Test = ImagePixelFormat(CompressionQuality, Result, EImageFormat::IF_RGB_UBYTE); 
+				// check(Test->DataStorage == Base->DataStorage);
+				break;
+			}
 
-                // \todo: optimise
-                ImagePtr Temp1 = ImagePixelFormat( CompressionQuality, Base, EImageFormat::IF_RGB_UBYTE);
-                ImagePtr Temp2 = ImagePixelFormat( CompressionQuality, Temp1.get(), EImageFormat::IF_RGBA_UBYTE);
+			case EImageFormat::IF_RGBA_UBYTE_RLE:
+			{
+				check(bFormatAllLODs);
+
+				// \todo: optimise
+				Ptr<Image> Temp1 = ImagePixelFormat(CompressionQuality, Base, EImageFormat::IF_RGBA_UBYTE);
+				Ptr<Image> Temp2 = ImagePixelFormat(CompressionQuality, Temp1.get(), EImageFormat::IF_RGB_UBYTE);
 				ReleaseImage(Temp1);
-				CompressRLE_RGBA( Temp2.get(), pResult );
+				CompressRLE_RGB(Temp2.get(), Result);
 				ReleaseImage(Temp2);
+
 				break;
-            }
+			}
 
-            default:
-                // Case not implemented
-                check( false );
+			default:
+				// Case not implemented
+				check(false);
+			}
 
-            }
+			break;
+		}
 
-            break;
-        }
+		case EImageFormat::IF_BC1:
+		{
+			switch (Base->GetFormat())
+			{
+			case EImageFormat::IF_RGB_UBYTE:
+			{
+				ImagePixelFormatInternal::BlockCompressionLODRangeHelper(
+						Result, Base, ResultLODBegin, BaseLODBegin, NumLODs, CompressionQuality, miro::RGB_to_BC1);
+				break;
+			}
+			case EImageFormat::IF_RGBA_UBYTE:
+			{
+				ImagePixelFormatInternal::BlockCompressionLODRangeHelper(
+						Result, Base, ResultLODBegin, BaseLODBegin, NumLODs, CompressionQuality, miro::RGBA_to_BC1);
+				break;
+			}
+			case EImageFormat::IF_L_UBYTE:
+			{
+				Ptr<Image> TempBase = ImagePixelFormat(CompressionQuality, Base, EImageFormat::IF_RGB_UBYTE);
+				ImagePixelFormatInternal::BlockCompressionLODRangeHelper(
+						Result, TempBase.get(), ResultLODBegin, BaseLODBegin, NumLODs, CompressionQuality, miro::RGB_to_BC1);
+				ReleaseImage(TempBase);
 
+				break;
+			}
+			case EImageFormat::IF_BC3:
+			{
+				Ptr<Image> TempBase = ImagePixelFormat(CompressionQuality, Base, EImageFormat::IF_RGBA_UBYTE);
+				ImagePixelFormatInternal::BlockCompressionLODRangeHelper(
+						Result, TempBase.get(), ResultLODBegin, BaseLODBegin, NumLODs, CompressionQuality, miro::RGBA_to_BC1);
+				ReleaseImage(TempBase);
 
-        case EImageFormat::IF_RGB_UBYTE_RLE:
-        {
-            switch ( Base->GetFormat() )
-            {
+				break;
+			}
+			default:
+			{
+				UE_LOG(LogMutableCore, Log, TEXT("Image format conversion not implemented. Expensive generic one used. "));
 
-            case EImageFormat::IF_RGB_UBYTE:
-            {
-                check( OnlyLOD == -1 );
-                CompressRLE_RGB( Base, pResult );
+				// Case not implemented, try a generic more expensive operation.
+				Ptr<Image> TempBase = ImagePixelFormat(CompressionQuality, Base, EImageFormat::IF_RGBA_UBYTE);
+				ImagePixelFormatInternal::BlockCompressionLODRangeHelper(
+						Result, TempBase.get(), ResultLODBegin, BaseLODBegin, NumLODs, CompressionQuality, miro::RGBA_to_BC1);
+				ReleaseImage(TempBase);
+			}
 
-                // Test
-                // ImagePtr pTest = ImagePixelFormat( CompressionQuality, pDest.get(),
-                // EImageFormat::IF_RGB_UBYTE ); check( !memcmp(pTest->GetData(), Base->GetData(),
-                // Base->GetDataSize() ) );
-                break;
-            }
+			}
 
-            case EImageFormat::IF_RGBA_UBYTE_RLE:
-            {
-                check( OnlyLOD == -1 );
+			break;
+		}
 
-                // \todo: optimise
-                ImagePtr Temp1 = ImagePixelFormat( CompressionQuality, Base, EImageFormat::IF_RGBA_UBYTE);
-                ImagePtr Temp2 = ImagePixelFormat( CompressionQuality, Temp1.get(), EImageFormat::IF_RGB_UBYTE);
-				ReleaseImage(Temp1);
-				CompressRLE_RGB( Temp2.get(), pResult );
-				ReleaseImage(Temp2);
+		case EImageFormat::IF_BC2:
+		{
+			switch (Base->GetFormat())
+			{
+			case EImageFormat::IF_RGBA_UBYTE:
+			{
+				ImagePixelFormatInternal::BlockCompressionLODRangeHelper(
+						Result, Base, ResultLODBegin, BaseLODBegin, NumLODs, CompressionQuality, miro::RGBA_to_BC2);
 				break;
-            }
-
-            default:
-                // Case not implemented
-                check( false );
-
-            }
-
-            break;
-        }
-
-        case EImageFormat::IF_BC1:
-        {
-            switch ( Base->GetFormat() )
-            {
-            case EImageFormat::IF_RGB_UBYTE:
-                for ( int m=0; m<resultLODCount; ++m )
-                {
-                    FIntVector2 mipSize = pResult->CalculateMipSize( m );
-                    miro::RGB_to_BC1(
-                        mipSize[0], mipSize[1], Base->GetMipData( baseLOD + m ),
-                        pResult->GetMipData( m ), CompressionQuality );
-                }
-                break;
-            case EImageFormat::IF_RGBA_UBYTE:
-                for ( int m=0; m<resultLODCount; ++m )
-                {
-                    FIntVector2 mipSize = pResult->CalculateMipSize( m );
-                    miro::RGBA_to_BC1(
-                        mipSize[0], mipSize[1], Base->GetMipData( baseLOD + m ),
-                        pResult->GetMipData( m ), CompressionQuality );
-                }
-                break;
-
-            case EImageFormat::IF_L_UBYTE:
-            {
-                ImagePtr Temp = ImagePixelFormat( CompressionQuality, Base, EImageFormat::IF_RGB_UBYTE);
-                for ( int m=0; m<resultLODCount; ++m )
-                {
-                    FIntVector2 mipSize = pResult->CalculateMipSize( m );
-                    miro::RGB_to_BC1(
-                        mipSize[0], mipSize[1], Temp->GetMipData( baseLOD + m ),
-                        pResult->GetMipData( m ), CompressionQuality );
-                }
-				ReleaseImage(Temp);
-                break;
-            }
-
-            case EImageFormat::IF_BC3:
-            {
-                ImagePtr Temp = ImagePixelFormat( CompressionQuality, Base, EImageFormat::IF_RGBA_UBYTE);
-                for (int m = 0; m < resultLODCount; ++m)
-                {
-                    FIntVector2 mipSize = pResult->CalculateMipSize(m);
-                    miro::RGBA_to_BC1(
-                        mipSize[0], mipSize[1], Temp->GetMipData( baseLOD + m ),
-                        pResult->GetMipData( m ), CompressionQuality );
-                }
-				ReleaseImage(Temp);
-				break;
-            }
-
-            default:
+			}
+			default:
 			{
-				UE_LOG(LogMutableCore, Log, TEXT("Image format conversion not implemented. Expensive generic one used. "));
+				UE_LOG(LogMutableCore, Log, TEXT("Image format conversion not implemented. Expensive generic one used."));
 
 				// Case not implemented, try a generic more expensive operatino.
-				ImagePtr Temp = ImagePixelFormat(CompressionQuality, Base, EImageFormat::IF_RGBA_UBYTE);
-				for (int m = 0; m < resultLODCount; ++m)
-				{
-					FIntVector2 mipSize = pResult->CalculateMipSize(m);
-					miro::RGBA_to_BC1(
-						mipSize[0], mipSize[1], Temp->GetMipData(baseLOD + m),
-						pResult->GetMipData(m), CompressionQuality);
-				}
-				ReleaseImage(Temp);
+				Ptr<Image> TempBase = ImagePixelFormat(CompressionQuality, Base, EImageFormat::IF_RGBA_UBYTE);
+				ImagePixelFormatInternal::BlockCompressionLODRangeHelper(
+						Result, TempBase.get(), ResultLODBegin, BaseLODBegin, NumLODs, CompressionQuality, miro::RGBA_to_BC2);
+				ReleaseImage(TempBase);
 			}
 
-            }
-
-            break;
-        }
+			}
 
-        case EImageFormat::IF_BC2:
-        {
-            switch ( Base->GetFormat() )
-            {
-            case EImageFormat::IF_RGBA_UBYTE:
-                for ( int m=0; m<resultLODCount; ++m )
-                {
-                    FIntVector2 mipSize = pResult->CalculateMipSize( m );
-                    miro::RGBA_to_BC2(
-                        mipSize[0], mipSize[1], Base->GetMipData( baseLOD + m ),
-                        pResult->GetMipData( m ), CompressionQuality );
-                }
-                break;
+			break;
+		}
 
-            default:
+		case EImageFormat::IF_BC3:
+		{
+			switch (Base->GetFormat())
 			{
-				UE_LOG(LogMutableCore, Log, TEXT("Image format conversion not implemented. Expensive generic one used. "));
-
-				// Case not implemented, try a generic more expensive operatino.
-				ImagePtr Temp = ImagePixelFormat(CompressionQuality, Base, EImageFormat::IF_RGBA_UBYTE);
-				for (int m = 0; m < resultLODCount; ++m)
-				{
-					FIntVector2 mipSize = pResult->CalculateMipSize(m);
-					miro::RGBA_to_BC2(
-						mipSize[0], mipSize[1], Temp->GetMipData(baseLOD + m),
-						pResult->GetMipData(m), CompressionQuality);
-				}
-				ReleaseImage(Temp);
-			}
-
-            }
-
-            break;
-        }
-
-        case EImageFormat::IF_BC3:
-        {
-            switch ( Base->GetFormat() )
-            {
-            case EImageFormat::IF_RGBA_UBYTE:
-                for ( int m=0; m<resultLODCount; ++m )
-                {
-                    FIntVector2 mipSize = pResult->CalculateMipSize( m );
-                    miro::RGBA_to_BC3(
-                        mipSize[0], mipSize[1], Base->GetMipData( baseLOD + m ),
-                        pResult->GetMipData( m ), CompressionQuality );
-                }
-                break;
-
-            case EImageFormat::IF_RGB_UBYTE:          
-                for ( int m=0; m<resultLODCount; ++m )
-                {
-                    FIntVector2 mipSize = pResult->CalculateMipSize( m );
-                    miro::RGB_to_BC3(
-                        mipSize[0], mipSize[1], Base->GetMipData( baseLOD + m ),
-                        pResult->GetMipData( m ), CompressionQuality );
-                }
-                break;
-
-            case EImageFormat::IF_BC1:
-            {
-                MUTABLE_CPUPROFILER_SCOPE(BC1toBC3);
-
-                for ( int m=0; m<resultLODCount; ++m )
-                {
-                    FIntVector2 mipSize = pResult->CalculateMipSize( m );
-                    miro::BC1_to_BC3(
-                        mipSize[0], mipSize[1], Base->GetMipData( baseLOD + m ),
-                        pResult->GetMipData( m ), CompressionQuality );
-                }
-                break;
-            }
-
-            default:
+			case EImageFormat::IF_RGBA_UBYTE:
 			{
-				UE_LOG(LogMutableCore, Log, TEXT("Image format conversion not implemented. Expensive generic one used. "));
+				ImagePixelFormatInternal::BlockCompressionLODRangeHelper(
+						Result, Base, ResultLODBegin, BaseLODBegin, NumLODs, CompressionQuality, miro::RGBA_to_BC3);
+				break;
+			}
+			case EImageFormat::IF_RGB_UBYTE:
+			{
+				ImagePixelFormatInternal::BlockCompressionLODRangeHelper(
+						Result, Base, ResultLODBegin, BaseLODBegin, NumLODs, CompressionQuality, miro::RGB_to_BC3);
+				break;
+			}
+			case EImageFormat::IF_BC1:
+			{
+				MUTABLE_CPUPROFILER_SCOPE(BC1toBC3);
+				ImagePixelFormatInternal::BlockCompressionLODRangeHelper(
+						Result, Base, ResultLODBegin, BaseLODBegin, NumLODs, CompressionQuality, miro::BC1_to_BC3);
+				break;
+			}
+			default:
+			{
+				UE_LOG(LogMutableCore, Log, TEXT("Image format conversion not implemented. Expensive generic one used."));
 
-				// Case not implemented, try a generic more expensive operatino.
-				ImagePtr Temp = ImagePixelFormat(CompressionQuality, Base, EImageFormat::IF_RGBA_UBYTE);
-				for (int m = 0; m < resultLODCount; ++m)
-				{
-					FIntVector2 mipSize = pResult->CalculateMipSize(m);
-					miro::RGBA_to_BC3(
-						mipSize[0], mipSize[1], Temp->GetMipData(baseLOD + m),
-						pResult->GetMipData(m), CompressionQuality);
-				}
-				ReleaseImage(Temp);
-			}               
+				// Case not implemented, try a generic more expensive operation.
+				Ptr<Image> TempBase = ImagePixelFormat(CompressionQuality, Base, EImageFormat::IF_RGBA_UBYTE);
+				ImagePixelFormatInternal::BlockCompressionLODRangeHelper(
+						Result, TempBase.get(), ResultLODBegin, BaseLODBegin, NumLODs, CompressionQuality, miro::RGBA_to_BC3);
+				ReleaseImage(TempBase);
+			}
 
-            }
+			}
 
-            break;
-        }
+			break;
+		}
 
-        case EImageFormat::IF_BC4:
-        {
-            switch ( Base->GetFormat() )
-            {
-            case EImageFormat::IF_L_UBYTE:
+		case EImageFormat::IF_BC4:
+		{
+			switch (Base->GetFormat())
 			{
-				for (int m = 0; m < resultLODCount; ++m)
-				{
-					FIntVector2 mipSize = pResult->CalculateMipSize(m);
-					miro::L_to_BC4(
-						mipSize[0], mipSize[1], Base->GetMipData(baseLOD + m),
-						pResult->GetMipData(m), CompressionQuality);
-				}
+			case EImageFormat::IF_L_UBYTE:
+			{
+				ImagePixelFormatInternal::BlockCompressionLODRangeHelper(
+						Result, Base, ResultLODBegin, BaseLODBegin, NumLODs, CompressionQuality, miro::L_to_BC4);
 				break;
 			}
 			case EImageFormat::IF_RGB_UBYTE:
 			case EImageFormat::IF_RGBA_UBYTE:
 			case EImageFormat::IF_BGRA_UBYTE:
 			default:
-			{
-				UE_LOG(LogMutableCore, Log, TEXT("Image format conversion not implemented. Expensive generic one used. "));
-				
-				Ptr<Image> TempBase = ImagePixelFormat( CompressionQuality, Base, EImageFormat::IF_L_UBYTE, -1);
-                for (int m = 0; m < resultLODCount; ++m)
-                {
-                    FIntVector2 MipSize = TempBase->CalculateMipSize( m );
-                    miro::L_to_BC4( MipSize[0], MipSize[1], TempBase->GetMipData( baseLOD + m ), pResult->GetMipData( m ), CompressionQuality );
-                }
+			{	
+				Ptr<Image> TempBase = ImagePixelFormat(CompressionQuality, Base, EImageFormat::IF_L_UBYTE, -1);
+				ImagePixelFormatInternal::BlockCompressionLODRangeHelper(
+						Result, TempBase.get(), ResultLODBegin, BaseLODBegin, NumLODs, CompressionQuality, miro::L_to_BC4);
 				ReleaseImage(TempBase);
 				break;
 			}
+			}
 
-            }
-
-            break;
-        }
-
-        case EImageFormat::IF_BC5:
-        {
-            switch ( Base->GetFormat() )
-            {
-            case EImageFormat::IF_RGBA_UBYTE:
-                for ( int m=0; m<resultLODCount; ++m )
-                {
-                    FIntVector2 mipSize = pResult->CalculateMipSize( m );
-                    miro::RGBA_to_BC5(
-                        mipSize[0], mipSize[1], Base->GetMipData( baseLOD + m ),
-                        pResult->GetMipData( m ), CompressionQuality );
-                }
-                break;
-
-            case EImageFormat::IF_RGB_UBYTE:
-                for ( int m=0; m<resultLODCount; ++m )
-                {
-                    FIntVector2 mipSize = pResult->CalculateMipSize( m );
-                    miro::RGB_to_BC5(
-                        mipSize[0], mipSize[1], Base->GetMipData( baseLOD + m ),
-                        pResult->GetMipData( m ), CompressionQuality );
-                }
-                break;
+			break;
+		}
 
+		case EImageFormat::IF_BC5:
+		{
+			switch (Base->GetFormat())
+			{
+			case EImageFormat::IF_RGBA_UBYTE:
+			{
+				ImagePixelFormatInternal::BlockCompressionLODRangeHelper(
+						Result, Base, ResultLODBegin, BaseLODBegin, NumLODs, CompressionQuality, miro::RGBA_to_BC5);
+				break;
+			}
+			case EImageFormat::IF_RGB_UBYTE:
+			{
+				ImagePixelFormatInternal::BlockCompressionLODRangeHelper(
+						Result, Base, ResultLODBegin, BaseLODBegin, NumLODs, CompressionQuality, miro::RGB_to_BC5);
+				break;
+			}
 			default:
 			{
 				UE_LOG(LogMutableCore, Log, TEXT("Image format conversion not implemented. Expensive generic one used. "));
 
 				Ptr<Image> TempBase = ImagePixelFormat(CompressionQuality, Base, EImageFormat::IF_RGBA_UBYTE, -1);
-				for (int m = 0; m < resultLODCount; ++m)
-				{
-					FIntVector2 MipSize = TempBase->CalculateMipSize(m);
-					miro::RGBA_to_BC5(MipSize[0], MipSize[1], TempBase->GetMipData(baseLOD + m), pResult->GetMipData(m), CompressionQuality);
-				}
+				ImagePixelFormatInternal::BlockCompressionLODRangeHelper(
+						Result, TempBase.get(), ResultLODBegin, BaseLODBegin, NumLODs, CompressionQuality, miro::RGBA_to_BC5);
 				ReleaseImage(TempBase);
 				break;
 			}
+			}
 
-            }
-
-            break;
-        }
-
-        case EImageFormat::IF_ASTC_4x4_RGB_LDR:
-        {
-            switch ( Base->GetFormat() )
-            {
-            case EImageFormat::IF_RGBA_UBYTE:
-                for ( int m=0; m<resultLODCount; ++m )
-                {
-                    FIntVector2 mipSize = pResult->CalculateMipSize( m );
-                    miro::RGBA_to_ASTC4x4RGBL(
-                        mipSize[0], mipSize[1], Base->GetMipData( baseLOD + m ),
-                        pResult->GetMipData( m ), CompressionQuality );
-                }
-                break;
-
-            case EImageFormat::IF_RGB_UBYTE:
-                for ( int m=0; m<resultLODCount; ++m )
-                {
-                    FIntVector2 mipSize = pResult->CalculateMipSize( m );
-                    miro::RGB_to_ASTC4x4RGBL(
-                        mipSize[0], mipSize[1], Base->GetMipData( baseLOD + m ),
-                        pResult->GetMipData( m ), CompressionQuality );
-                }
-                break;
-
-            case EImageFormat::IF_ASTC_4x4_RGBA_LDR:
-                for ( int m=0; m<resultLODCount; ++m )
-                {
-                    FIntVector2 mipSize = pResult->CalculateMipSize( m );
-                    miro::ASTC4x4RGBAL_to_ASTC4x4RGBL( mipSize[0], mipSize[1],
-                                       Base->GetMipData(baseLOD+m),
-                                       pResult->GetMipData(m) );
-                }
-                break;
+			break;
+		}
 
+		case EImageFormat::IF_ASTC_4x4_RGB_LDR:
+		{
+			switch (Base->GetFormat())
+			{
+			case EImageFormat::IF_RGBA_UBYTE:
+			{
+				ImagePixelFormatInternal::BlockCompressionLODRangeHelper(
+						Result, Base, ResultLODBegin, BaseLODBegin, NumLODs, CompressionQuality, miro::RGBA_to_ASTC4x4RGBL);
+				break;
+			}
+			case EImageFormat::IF_RGB_UBYTE:
+			{
+				ImagePixelFormatInternal::BlockCompressionLODRangeHelper(
+						Result, Base, ResultLODBegin, BaseLODBegin, NumLODs, CompressionQuality, miro::RGB_to_ASTC4x4RGBL);
+				break;
+			}
+			case EImageFormat::IF_ASTC_4x4_RGBA_LDR:
+			{
+				ImagePixelFormatInternal::BlockDecompressionLODRangeHelper(
+						Result, Base, ResultLODBegin, BaseLODBegin, NumLODs, miro::ASTC4x4RGBAL_to_ASTC4x4RGBL);
+				break;
+			}
 			case EImageFormat::IF_L_UBYTE:
-				for (int m = 0; m < resultLODCount; ++m)
-				{
-					FIntVector2 mipSize = pResult->CalculateMipSize(m);
-					miro::L_to_ASTC4x4RGBL(mipSize[0], mipSize[1],
-						Base->GetMipData(baseLOD + m),
-						pResult->GetMipData(m), CompressionQuality);
-				}
+			{
+				ImagePixelFormatInternal::BlockCompressionLODRangeHelper(
+						Result, Base, ResultLODBegin, BaseLODBegin, NumLODs, CompressionQuality, miro::L_to_ASTC4x4RGBL);
 				break;
-
+			}
 			default:
 			{
 				UE_LOG(LogMutableCore, Log, TEXT("Image format conversion not implemented. Expensive generic one used. "));
 
 				Ptr<Image> TempBase = ImagePixelFormat(CompressionQuality, Base, EImageFormat::IF_RGBA_UBYTE, -1);
-				for (int m = 0; m < resultLODCount; ++m)
-				{
-					FIntVector2 MipSize = TempBase->CalculateMipSize(m);
-					miro::RGBA_to_ASTC4x4RGBL(MipSize[0], MipSize[1], TempBase->GetMipData(baseLOD + m), pResult->GetMipData(m), CompressionQuality);
-				}
+				ImagePixelFormatInternal::BlockCompressionLODRangeHelper(
+						Result, TempBase.get(), ResultLODBegin, BaseLODBegin, NumLODs, CompressionQuality, miro::RGBA_to_ASTC4x4RGBL);
 				ReleaseImage(TempBase);
 				break;
 			}
+			}
 
-            }
-
-            break;
-        }
-
-        case EImageFormat::IF_ASTC_4x4_RGBA_LDR:
-        {
-            switch ( Base->GetFormat() )
-            {
-            case EImageFormat::IF_RGBA_UBYTE:
-                for ( int m=0; m<resultLODCount; ++m )
-                {
-                    FIntVector2 mipSize = pResult->CalculateMipSize( m );
-                    miro::RGBA_to_ASTC4x4RGBAL(
-                        mipSize[0], mipSize[1], Base->GetMipData( baseLOD + m ),
-                        pResult->GetMipData( m ), CompressionQuality );
-                }
-                break;
-
-            case EImageFormat::IF_RGB_UBYTE:
-                for ( int m=0; m<resultLODCount; ++m )
-                {
-                    FIntVector2 mipSize = pResult->CalculateMipSize( m );
-                    miro::RGB_to_ASTC4x4RGBAL(
-                        mipSize[0], mipSize[1], Base->GetMipData( baseLOD + m ),
-                        pResult->GetMipData( m ), CompressionQuality );
-                }
-                break;
-
-            case EImageFormat::IF_ASTC_4x4_RGB_LDR:
-                for ( int m=0; m<resultLODCount; ++m )
-                {
-                    FIntVector2 mipSize = pResult->CalculateMipSize( m );
-                    miro::ASTC4x4RGBL_to_ASTC4x4RGBAL( mipSize[0], mipSize[1],
-                                       Base->GetMipData(baseLOD+m),
-                                       pResult->GetMipData(m) );
-                }
-                break;
+			break;
+		}
 
+		case EImageFormat::IF_ASTC_4x4_RGBA_LDR:
+		{
+			switch (Base->GetFormat())
+			{
+			case EImageFormat::IF_RGBA_UBYTE:
+			{
+				ImagePixelFormatInternal::BlockCompressionLODRangeHelper(
+						Result, Base, ResultLODBegin, BaseLODBegin, NumLODs, CompressionQuality, miro::RGBA_to_ASTC4x4RGBAL);
+				break;
+			}
+			case EImageFormat::IF_RGB_UBYTE:
+			{
+				ImagePixelFormatInternal::BlockCompressionLODRangeHelper(
+						Result, Base, ResultLODBegin, BaseLODBegin, NumLODs, CompressionQuality, miro::RGB_to_ASTC4x4RGBAL);
+				break;
+			}
+			case EImageFormat::IF_ASTC_4x4_RGB_LDR:
+			{
+				ImagePixelFormatInternal::BlockDecompressionLODRangeHelper(
+						Result, Base, ResultLODBegin, BaseLODBegin, NumLODs, miro::ASTC4x4RGBL_to_ASTC4x4RGBAL);
+				break;
+			}
 			default:
 			{
 				UE_LOG(LogMutableCore, Log, TEXT("Image format conversion not implemented. Expensive generic one used. "));
 
 				Ptr<Image> TempBase = ImagePixelFormat(CompressionQuality, Base, EImageFormat::IF_RGBA_UBYTE, -1);
-				for (int m = 0; m < resultLODCount; ++m)
-				{
-					FIntVector2 MipSize = TempBase->CalculateMipSize(m);
-					miro::RGBA_to_ASTC4x4RGBAL(MipSize[0], MipSize[1], TempBase->GetMipData(baseLOD + m), pResult->GetMipData(m), CompressionQuality);
-				}
+				ImagePixelFormatInternal::BlockCompressionLODRangeHelper(
+						Result, TempBase.get(), ResultLODBegin, BaseLODBegin, NumLODs, CompressionQuality, miro::RGBA_to_ASTC4x4RGBAL);
 				ReleaseImage(TempBase);
 				break;
 			}
+			}
 
-            }
-
-            break;
-        }
-
-        case EImageFormat::IF_ASTC_4x4_RG_LDR:
-        {
-            switch ( Base->GetFormat() )
-            {
-            case EImageFormat::IF_RGBA_UBYTE:
-                for ( int m=0; m<resultLODCount; ++m )
-                {
-                    FIntVector2 mipSize = pResult->CalculateMipSize( m );
-                    miro::RGBA_to_ASTC4x4RGL(
-                        mipSize[0], mipSize[1], Base->GetMipData( baseLOD + m ),
-                        pResult->GetMipData( m ), CompressionQuality );
-                }
-                break;
-
-            case EImageFormat::IF_RGB_UBYTE:
-                for ( int m=0; m<resultLODCount; ++m )
-                {
-                    FIntVector2 mipSize = pResult->CalculateMipSize( m );
-                    miro::RGB_to_ASTC4x4RGL(
-                        mipSize[0], mipSize[1], Base->GetMipData( baseLOD + m ),
-                        pResult->GetMipData( m ), CompressionQuality );
-                }
-                break;
-
-            case EImageFormat::IF_ASTC_4x4_RG_LDR:
-                for ( int m=0; m<resultLODCount; ++m )
-                {
-                    FIntVector2 mipSize = pResult->CalculateMipSize( m );
-                    // Hack that actually works because of block size.
-                    miro::ASTC4x4RGBAL_to_ASTC4x4RGBL( mipSize[0], mipSize[1],
-                                       Base->GetMipData(baseLOD+m),
-                                       pResult->GetMipData(m) );
-                }
-                break;
+			break;
+		}
 
+		case EImageFormat::IF_ASTC_4x4_RG_LDR:
+		{
+			switch (Base->GetFormat())
+			{
+			case EImageFormat::IF_RGBA_UBYTE:
+			{
+				ImagePixelFormatInternal::BlockCompressionLODRangeHelper(
+						Result, Base, ResultLODBegin, BaseLODBegin, NumLODs, CompressionQuality, miro::RGBA_to_ASTC4x4RGL);
+				break;
+			}
+			case EImageFormat::IF_RGB_UBYTE:
+			{
+				ImagePixelFormatInternal::BlockCompressionLODRangeHelper(
+						Result, Base, ResultLODBegin, BaseLODBegin, NumLODs, CompressionQuality, miro::RGB_to_ASTC4x4RGL);
+				break;
+			}
+			case EImageFormat::IF_ASTC_4x4_RG_LDR:
+			{
+				// Hack that actually works because of block size.
+				ImagePixelFormatInternal::BlockDecompressionLODRangeHelper(
+						Result, Base, ResultLODBegin, BaseLODBegin, NumLODs, miro::ASTC4x4RGBAL_to_ASTC4x4RGBL);
+				break;
+			}
 			default:
 			{
 				UE_LOG(LogMutableCore, Log, TEXT("Image format conversion not implemented. Expensive generic one used. "));
 
 				Ptr<Image> TempBase = ImagePixelFormat(CompressionQuality, Base, EImageFormat::IF_RGBA_UBYTE, -1);
-				for (int m = 0; m < resultLODCount; ++m)
-				{
-					FIntVector2 MipSize = TempBase->CalculateMipSize(m);
-					miro::RGBA_to_ASTC4x4RGL(MipSize[0], MipSize[1], TempBase->GetMipData(baseLOD + m), pResult->GetMipData(m), CompressionQuality);
-				}
+				ImagePixelFormatInternal::BlockCompressionLODRangeHelper(
+						Result, TempBase.get(), ResultLODBegin, BaseLODBegin, NumLODs, CompressionQuality, miro::RGBA_to_ASTC4x4RGL);
 				ReleaseImage(TempBase);
 				break;
 			}
+			}
 
-            }
-
-            break;
-        }
+			break;
+		}
 
 		case EImageFormat::IF_ASTC_6x6_RGB_LDR:
 		{
 			switch (Base->GetFormat())
 			{
 			case EImageFormat::IF_RGBA_UBYTE:
-				for (int m = 0; m < resultLODCount; ++m)
-				{
-					FIntVector2 mipSize = pResult->CalculateMipSize(m);
-					miro::RGBA_to_ASTC6x6RGBL(
-						mipSize[0], mipSize[1], Base->GetMipData(baseLOD + m),
-						pResult->GetMipData(m), CompressionQuality);
-				}
+			{
+				ImagePixelFormatInternal::BlockCompressionLODRangeHelper(
+						Result, Base, ResultLODBegin, BaseLODBegin, NumLODs, CompressionQuality, miro::RGBA_to_ASTC6x6RGBL);
 				break;
-
+			}
 			case EImageFormat::IF_RGB_UBYTE:
-				for (int m = 0; m < resultLODCount; ++m)
-				{
-					FIntVector2 mipSize = pResult->CalculateMipSize(m);
-					miro::RGB_to_ASTC6x6RGBL(
-						mipSize[0], mipSize[1], Base->GetMipData(baseLOD + m),
-						pResult->GetMipData(m), CompressionQuality);
-				}
+			{
+				ImagePixelFormatInternal::BlockCompressionLODRangeHelper(
+						Result, Base, ResultLODBegin, BaseLODBegin, NumLODs, CompressionQuality, miro::RGB_to_ASTC6x6RGBL);
 				break;
-
+			}
 			//case EImageFormat::IF_ASTC_6x6_RGBA_LDR:
-			//	for (int m = 0; m < resultLODCount; ++m)
-			//	{
-			//		FIntVector2 mipSize = pResult->CalculateMipSize(m);
-			//		miro::ASTC6x6RGBAL_to_ASTC6x6RGBL(mipSize[0], mipSize[1],
-			//			Base->GetMipData(baseLOD + m),
-			//			pResult->GetMipData(m));
-			//	}
+			//{
+			//  ImagePixelFormatInternal::BlockCompressionLODRangeHelper(
+			//  		Result, Base, ResultLODBegin, BaseLODBegin, NumLODs, CompressionQuality, miro::ASTC6x6RGBAL_to_ASTC6x6RGBL);
 			//	break;
-
+			//}
 			//case EImageFormat::IF_L_UBYTE:
-			//	for (int m = 0; m < resultLODCount; ++m)
-			//	{
-			//		FIntVector2 mipSize = pResult->CalculateMipSize(m);
-			//		miro::L_to_ASTC6x6RGBL(mipSize[0], mipSize[1],
-			//			Base->GetMipData(baseLOD + m),
-			//			pResult->GetMipData(m), CompressionQuality);
-			//	}
+			//{
+			//  ImagePixelFormatInternal::BlockCompressionLODRangeHelper(
+			//  		Result, Base, ResultLODBegin, BaseLODBegin, NumLODs, CompressionQuality, miro::L_to_ASTC6x6RGBL);
 			//	break;
+			//}
 
 			default:
 			{
 				UE_LOG(LogMutableCore, Log, TEXT("Image format conversion not implemented. Expensive generic one used. "));
 
 				Ptr<Image> TempBase = ImagePixelFormat(CompressionQuality, Base, EImageFormat::IF_RGBA_UBYTE, -1);
-				for (int m = 0; m < resultLODCount; ++m)
-				{
-					FIntVector2 MipSize = TempBase->CalculateMipSize(m);
-					miro::RGBA_to_ASTC6x6RGBL(MipSize[0], MipSize[1], TempBase->GetMipData(baseLOD + m), pResult->GetMipData(m), CompressionQuality);
-				}
+				ImagePixelFormatInternal::BlockCompressionLODRangeHelper(
+						Result, TempBase.get(), ResultLODBegin, BaseLODBegin, NumLODs, CompressionQuality, miro::RGBA_to_ASTC6x6RGBL);
 				ReleaseImage(TempBase);
 				break;
 			}
-
 			}
 
 			break;
@@ -1656,49 +1403,33 @@ void FImageOperator::ImagePixelFormat( bool& bOutSuccess, int32 CompressionQuali
 			switch (Base->GetFormat())
 			{
 			case EImageFormat::IF_RGBA_UBYTE:
-				for (int m = 0; m < resultLODCount; ++m)
-				{
-					FIntVector2 mipSize = pResult->CalculateMipSize(m);
-					miro::RGBA_to_ASTC6x6RGBAL(
-						mipSize[0], mipSize[1], Base->GetMipData(baseLOD + m),
-						pResult->GetMipData(m), CompressionQuality);
-				}
+			{
+				ImagePixelFormatInternal::BlockCompressionLODRangeHelper(
+						Result, Base, ResultLODBegin, BaseLODBegin, NumLODs, CompressionQuality, miro::RGBA_to_ASTC6x6RGBAL);
 				break;
-
+			}
 			case EImageFormat::IF_RGB_UBYTE:
-				for (int m = 0; m < resultLODCount; ++m)
-				{
-					FIntVector2 mipSize = pResult->CalculateMipSize(m);
-					miro::RGB_to_ASTC6x6RGBAL(
-						mipSize[0], mipSize[1], Base->GetMipData(baseLOD + m),
-						pResult->GetMipData(m), CompressionQuality);
-				}
+			{
+				ImagePixelFormatInternal::BlockCompressionLODRangeHelper(
+						Result, Base, ResultLODBegin, BaseLODBegin, NumLODs, CompressionQuality, miro::RGB_to_ASTC6x6RGBAL);
 				break;
-
+			}
 			//case EImageFormat::IF_ASTC_6x6_RGB_LDR:
-			//	for (int m = 0; m < resultLODCount; ++m)
-			//	{
-			//		FIntVector2 mipSize = pResult->CalculateMipSize(m);
-			//		miro::ASTC6x6RGBL_to_ASTC6x6RGBAL(mipSize[0], mipSize[1],
-			//			Base->GetMipData(baseLOD + m),
-			//			pResult->GetMipData(m));
-			//	}
+			//{
+			//	ImagePixelFormatInternal::BlockCompressionLODRangeHelper(
+			//			Result, Base, ResultLODBegin, BaseLODBegin, NumLODs, CompressionQuality, miro::ASTC6x6RGBL_to_ASTC6x6RGBAL);
 			//	break;
-
+			//}
 			default:
 			{
 				UE_LOG(LogMutableCore, Log, TEXT("Image format conversion not implemented. Expensive generic one used. "));
 
 				Ptr<Image> TempBase = ImagePixelFormat(CompressionQuality, Base, EImageFormat::IF_RGBA_UBYTE, -1);
-				for (int m = 0; m < resultLODCount; ++m)
-				{
-					FIntVector2 MipSize = TempBase->CalculateMipSize(m);
-					miro::RGBA_to_ASTC6x6RGBAL(MipSize[0], MipSize[1], TempBase->GetMipData(baseLOD + m), pResult->GetMipData(m), CompressionQuality);
-				}
+				ImagePixelFormatInternal::BlockCompressionLODRangeHelper(
+						Result, TempBase.get(), ResultLODBegin, BaseLODBegin, NumLODs, CompressionQuality, miro::RGBA_to_ASTC6x6RGBAL);
 				ReleaseImage(TempBase);
 				break;
 			}
-
 			}
 
 			break;
@@ -1709,50 +1440,35 @@ void FImageOperator::ImagePixelFormat( bool& bOutSuccess, int32 CompressionQuali
 			switch (Base->GetFormat())
 			{
 			case EImageFormat::IF_RGBA_UBYTE:
-				for (int m = 0; m < resultLODCount; ++m)
-				{
-					FIntVector2 mipSize = pResult->CalculateMipSize(m);
-					miro::RGBA_to_ASTC6x6RGL(
-						mipSize[0], mipSize[1], Base->GetMipData(baseLOD + m),
-						pResult->GetMipData(m), CompressionQuality);
-				}
+			{
+				ImagePixelFormatInternal::BlockCompressionLODRangeHelper(
+						Result, Base, ResultLODBegin, BaseLODBegin, NumLODs, CompressionQuality, miro::RGBA_to_ASTC6x6RGL);
 				break;
+			}
 
 			case EImageFormat::IF_RGB_UBYTE:
-				for (int m = 0; m < resultLODCount; ++m)
-				{
-					FIntVector2 mipSize = pResult->CalculateMipSize(m);
-					miro::RGB_to_ASTC6x6RGL(
-						mipSize[0], mipSize[1], Base->GetMipData(baseLOD + m),
-						pResult->GetMipData(m), CompressionQuality);
-				}
+			{
+				ImagePixelFormatInternal::BlockCompressionLODRangeHelper(
+						Result, Base, ResultLODBegin, BaseLODBegin, NumLODs, CompressionQuality, miro::RGB_to_ASTC6x6RGL);
 				break;
-
+			}
 			//case EImageFormat::IF_ASTC_6x6_RG_LDR:
-			//	for (int m = 0; m < resultLODCount; ++m)
-			//	{
-			//		FIntVector2 mipSize = pResult->CalculateMipSize(m);
-			//		// Hack that actually works because of block size.
-			//		miro::ASTC6x6RGBAL_to_ASTC6x6RGBL(mipSize[0], mipSize[1],
-			//			Base->GetMipData(baseLOD + m),
-			//			pResult->GetMipData(m));
-			//	}
+			//{
+			//	// Hack that actually works because of block size.
+			//	ImagePixelFormatInternal::BlockCompressionLODRangeHelper(
+			//			Result, Base, ResultLODBegin, BaseLODBegin, NumLODs, CompressionQuality, miro::ASTC6x6RGBAL_to_ASTC6x6RGBL);
 			//	break;
-
+			//}
 			default:
 			{
 				UE_LOG(LogMutableCore, Log, TEXT("Image format conversion not implemented. Expensive generic one used. "));
 
 				Ptr<Image> TempBase = ImagePixelFormat(CompressionQuality, Base, EImageFormat::IF_RGBA_UBYTE, -1);
-				for (int m = 0; m < resultLODCount; ++m)
-				{
-					FIntVector2 MipSize = TempBase->CalculateMipSize(m);
-					miro::RGBA_to_ASTC6x6RGL(MipSize[0], MipSize[1], TempBase->GetMipData(baseLOD + m), pResult->GetMipData(m), CompressionQuality);
-				}
+				ImagePixelFormatInternal::BlockCompressionLODRangeHelper(
+						Result, TempBase.get(), ResultLODBegin, BaseLODBegin, NumLODs, CompressionQuality, miro::RGBA_to_ASTC6x6RGL);
 				ReleaseImage(TempBase);
 				break;
 			}
-
 			}
 
 			break;
@@ -1763,59 +1479,39 @@ void FImageOperator::ImagePixelFormat( bool& bOutSuccess, int32 CompressionQuali
 			switch (Base->GetFormat())
 			{
 			case EImageFormat::IF_RGBA_UBYTE:
-				for (int m = 0; m < resultLODCount; ++m)
-				{
-					FIntVector2 mipSize = pResult->CalculateMipSize(m);
-					miro::RGBA_to_ASTC8x8RGBL(
-						mipSize[0], mipSize[1], Base->GetMipData(baseLOD + m),
-						pResult->GetMipData(m), CompressionQuality);
-				}
+			{
+				ImagePixelFormatInternal::BlockCompressionLODRangeHelper(
+						Result, Base, ResultLODBegin, BaseLODBegin, NumLODs, CompressionQuality, miro::RGBA_to_ASTC8x8RGBL);
 				break;
-
+			}
 			case EImageFormat::IF_RGB_UBYTE:
-				for (int m = 0; m < resultLODCount; ++m)
-				{
-					FIntVector2 mipSize = pResult->CalculateMipSize(m);
-					miro::RGB_to_ASTC8x8RGBL(
-						mipSize[0], mipSize[1], Base->GetMipData(baseLOD + m),
-						pResult->GetMipData(m), CompressionQuality);
-				}
+			{
+				ImagePixelFormatInternal::BlockCompressionLODRangeHelper(
+						Result, Base, ResultLODBegin, BaseLODBegin, NumLODs, CompressionQuality, miro::RGB_to_ASTC8x8RGBL);
 				break;
-
-				//case EImageFormat::IF_ASTC_8x8_RGBA_LDR:
-				//	for (int m = 0; m < resultLODCount; ++m)
-				//	{
-				//		FIntVector2 mipSize = pResult->CalculateMipSize(m);
-				//		miro::ASTC8x8RGBAL_to_ASTC8x8RGBL(mipSize[0], mipSize[1],
-				//			Base->GetMipData(baseLOD + m),
-				//			pResult->GetMipData(m));
-				//	}
-				//	break;
-
-				//case EImageFormat::IF_L_UBYTE:
-				//	for (int m = 0; m < resultLODCount; ++m)
-				//	{
-				//		FIntVector2 mipSize = pResult->CalculateMipSize(m);
-				//		miro::L_to_ASTC8x8RGBL(mipSize[0], mipSize[1],
-				//			Base->GetMipData(baseLOD + m),
-				//			pResult->GetMipData(m), CompressionQuality);
-				//	}
-				//	break;
-
+			}
+			//case EImageFormat::IF_ASTC_8x8_RGBA_LDR:
+			//{
+			//	ImagePixelFormatInternal::BlockCompressionLODRangeHelper(
+			//			Result, Base, ResultLODBegin, BaseLODBegin, NumLODs, CompressionQuality, miro::ASTC8x8RGBAL_to_ASTC8x8RGBL);
+			//	break;
+			//}
+			//case EImageFormat::IF_L_UBYTE:
+			//{
+			//	ImagePixelFormatInternal::BlockCompressionLODRangeHelper(
+			//			Result, Base, ResultLODBegin, BaseLODBegin, NumLODs, CompressionQuality, miro::L_to_ASTC8x8RGBL);
+			//	break;
+			//}
 			default:
 			{
 				UE_LOG(LogMutableCore, Log, TEXT("Image format conversion not implemented. Expensive generic one used. "));
 
 				Ptr<Image> TempBase = ImagePixelFormat(CompressionQuality, Base, EImageFormat::IF_RGBA_UBYTE, -1);
-				for (int m = 0; m < resultLODCount; ++m)
-				{
-					FIntVector2 MipSize = TempBase->CalculateMipSize(m);
-					miro::RGBA_to_ASTC8x8RGBL(MipSize[0], MipSize[1], TempBase->GetMipData(baseLOD + m), pResult->GetMipData(m), CompressionQuality);
-				}
+				ImagePixelFormatInternal::BlockCompressionLODRangeHelper(
+						Result, TempBase.get(), ResultLODBegin, BaseLODBegin, NumLODs, CompressionQuality, miro::RGBA_to_ASTC8x8RGBL);
 				ReleaseImage(TempBase);
 				break;
 			}
-
 			}
 
 			break;
@@ -1826,49 +1522,33 @@ void FImageOperator::ImagePixelFormat( bool& bOutSuccess, int32 CompressionQuali
 			switch (Base->GetFormat())
 			{
 			case EImageFormat::IF_RGBA_UBYTE:
-				for (int m = 0; m < resultLODCount; ++m)
-				{
-					FIntVector2 mipSize = pResult->CalculateMipSize(m);
-					miro::RGBA_to_ASTC8x8RGBAL(
-						mipSize[0], mipSize[1], Base->GetMipData(baseLOD + m),
-						pResult->GetMipData(m), CompressionQuality);
-				}
+			{
+				ImagePixelFormatInternal::BlockCompressionLODRangeHelper(
+						Result, Base, ResultLODBegin, BaseLODBegin, NumLODs, CompressionQuality, miro::RGBA_to_ASTC8x8RGBAL);
 				break;
-
+			}
 			case EImageFormat::IF_RGB_UBYTE:
-				for (int m = 0; m < resultLODCount; ++m)
-				{
-					FIntVector2 mipSize = pResult->CalculateMipSize(m);
-					miro::RGB_to_ASTC8x8RGBAL(
-						mipSize[0], mipSize[1], Base->GetMipData(baseLOD + m),
-						pResult->GetMipData(m), CompressionQuality);
-				}
+			{
+				ImagePixelFormatInternal::BlockCompressionLODRangeHelper(
+						Result, Base, ResultLODBegin, BaseLODBegin, NumLODs, CompressionQuality, miro::RGB_to_ASTC8x8RGBAL);
 				break;
-
+			}
 			//case EImageFormat::IF_ASTC_8x8_RGB_LDR:
-			//	for (int m = 0; m < resultLODCount; ++m)
-			//	{
-			//		FIntVector2 mipSize = pResult->CalculateMipSize(m);
-			//		miro::ASTC8x8RGBL_to_ASTC8x8RGBAL(mipSize[0], mipSize[1],
-			//			Base->GetMipData(baseLOD + m),
-			//			pResult->GetMipData(m));
-			//	}
+			//{
+			//	ImagePixelFormatInternal::BlockCompressionLODRangeHelper(
+			//			Result, Base, ResultLODBegin, BaseLODBegin, NumLODs, CompressionQuality, miro::ASTC8x8RGBL_to_ASTC8x8RGBAL);
 			//	break;
-
+			//}
 			default:
 			{
 				UE_LOG(LogMutableCore, Log, TEXT("Image format conversion not implemented. Expensive generic one used. "));
 
 				Ptr<Image> TempBase = ImagePixelFormat(CompressionQuality, Base, EImageFormat::IF_RGBA_UBYTE, -1);
-				for (int m = 0; m < resultLODCount; ++m)
-				{
-					FIntVector2 MipSize = TempBase->CalculateMipSize(m);
-					miro::RGBA_to_ASTC8x8RGBAL(MipSize[0], MipSize[1], TempBase->GetMipData(baseLOD + m), pResult->GetMipData(m), CompressionQuality);
-				}
+				ImagePixelFormatInternal::BlockCompressionLODRangeHelper(
+						Result, TempBase.get(), ResultLODBegin, BaseLODBegin, NumLODs, CompressionQuality, miro::RGBA_to_ASTC8x8RGBAL);
 				ReleaseImage(TempBase);
 				break;
 			}
-
 			}
 
 			break;
@@ -1879,51 +1559,37 @@ void FImageOperator::ImagePixelFormat( bool& bOutSuccess, int32 CompressionQuali
 			switch (Base->GetFormat())
 			{
 			case EImageFormat::IF_RGBA_UBYTE:
-				for (int m = 0; m < resultLODCount; ++m)
-				{
-					FIntVector2 mipSize = pResult->CalculateMipSize(m);
-					miro::RGBA_to_ASTC8x8RGL(
-						mipSize[0], mipSize[1], Base->GetMipData(baseLOD + m),
-						pResult->GetMipData(m), CompressionQuality);
-				}
+			{
+				ImagePixelFormatInternal::BlockCompressionLODRangeHelper(
+						Result, Base, ResultLODBegin, BaseLODBegin, NumLODs, CompressionQuality, miro::RGBA_to_ASTC8x8RGL);
 				break;
-
+			}
 			case EImageFormat::IF_RGB_UBYTE:
-				for (int m = 0; m < resultLODCount; ++m)
-				{
-					FIntVector2 mipSize = pResult->CalculateMipSize(m);
-					miro::RGB_to_ASTC8x8RGL(
-						mipSize[0], mipSize[1], Base->GetMipData(baseLOD + m),
-						pResult->GetMipData(m), CompressionQuality);
-				}
+			{
+				ImagePixelFormatInternal::BlockCompressionLODRangeHelper(
+						Result, Base, ResultLODBegin, BaseLODBegin, NumLODs, CompressionQuality, miro::RGB_to_ASTC8x8RGL);
 				break;
-
+			}
 			default:
 			{
 				UE_LOG(LogMutableCore, Log, TEXT("Image format conversion not implemented. Expensive generic one used. "));
 
 				Ptr<Image> TempBase = ImagePixelFormat(CompressionQuality, Base, EImageFormat::IF_RGBA_UBYTE, -1);
-				for (int m = 0; m < resultLODCount; ++m)
-				{
-					FIntVector2 MipSize = TempBase->CalculateMipSize(m);
-					miro::RGBA_to_ASTC8x8RGL(MipSize[0], MipSize[1], TempBase->GetMipData(baseLOD + m), pResult->GetMipData(m), CompressionQuality);
-				}
+				ImagePixelFormatInternal::BlockCompressionLODRangeHelper(
+						Result, TempBase.get(), ResultLODBegin, BaseLODBegin, NumLODs, CompressionQuality, miro::RGBA_to_ASTC8x8RGL);
 				ReleaseImage(TempBase);
 				break;
 			}
-
 			}
-
+			
 			break;
 		}
 
-        default:
-            // Case not implemented
-            check( false );
-        }
-    }
-
+		default:
+			// Case not implemented
+			check(false);
+		}
+	}
 }
 
-
 }
