[Mutable] Store Image LODs in different buffers.

diff --git a/Engine/Plugins/Experimental/Mutable/Source/MutableRuntime/Private/MuR/OpImageProject.cpp b/Engine/Plugins/Experimental/Mutable/Source/MutableRuntime/Private/MuR/OpImageProject.cpp
--- a/Engine/Plugins/Experimental/Mutable/Source/MutableRuntime/Private/MuR/OpImageProject.cpp
+++ b/Engine/Plugins/Experimental/Mutable/Source/MutableRuntime/Private/MuR/OpImageProject.cpp
@@ -194,7 +194,7 @@ namespace Private
 
             Context.Source0SizeX = static_cast<uint16>(Source->GetSizeX());
             Context.Source0SizeY = static_cast<uint16>(Source->GetSizeY());
-			Context.Source0Data = (Context.Source0SizeX > 0 && Context.Source0SizeY > 0) ? Source->GetData() : nullptr;
+			Context.Source0Data = (Context.Source0SizeX > 0 && Context.Source0SizeY > 0) ? Source->GetLODData(0) : nullptr;
 
 			if constexpr (EnumHasAnyFlags(Features, EPixelProcessorFeatures::SamplingLinear))
 			{
@@ -696,7 +696,7 @@ namespace Private
 				!Scratch->CulledVertex[Index2])
 			{
 				constexpr int32 NumInterpolators = 4;
-				Triangle<NumInterpolators>(pImage->GetData(), pImage->GetDataSize(),
+				Triangle<NumInterpolators>(pImage->GetLODData(0), pImage->GetDataSize(),
 					SizeX, SizeY,
 					PixelSize,
 					Scratch->Vertices[Index0],
@@ -824,7 +824,7 @@ namespace Private
 				!Scratch->CulledVertex[Index2])
 			{
 				constexpr int32 NumInterpolators = 4;
-				Triangle<NumInterpolators>(pImage->GetData(), pImage->GetDataSize(),
+				Triangle<NumInterpolators>(pImage->GetLODData(0), pImage->GetDataSize(),
 					SizeX, SizeY,
 					PixelSize,
 					Scratch->Vertices[Index0], Scratch->Vertices[Index1], Scratch->Vertices[Index2],
@@ -841,7 +841,7 @@ namespace Private
 
 		const Private::FProjectedPixelProcessorContext Context = 
 			Private::TProjectedPixelProcessor<Features>::MakeContext(
-				Args.SourcePtr, Args.ImagePtr->GetData(), Args.MaskPtr ? Args.MaskPtr->GetData() : nullptr,
+				Args.SourcePtr, Args.ImagePtr->GetLODData(0), Args.MaskPtr ? Args.MaskPtr->GetLODData(0) : nullptr,
 				Args.bIsRGBFadingEnabled, Args.bIsAlphaFadingEnabled, Args.FadeStart, Args.FadeEnd, 
 				Args.ProjectionAngle, Args.MipInterpolationFactor);
 
@@ -860,7 +860,7 @@ namespace Private
 
 		const Private::FProjectedPixelProcessorContext Context = 
 			Private::TProjectedPixelProcessor<Features>::MakeContext(
-				Args.SourcePtr, Args.ImagePtr->GetData(), Args.MaskPtr ? Args.MaskPtr->GetData() : nullptr,
+				Args.SourcePtr, Args.ImagePtr->GetLODData(0), Args.MaskPtr ? Args.MaskPtr->GetLODData(0) : nullptr,
 				Args.bIsRGBFadingEnabled, Args.bIsAlphaFadingEnabled, Args.FadeStart, Args.FadeEnd, 
 				Args.ProjectionAngle, Args.MipInterpolationFactor);
 
@@ -879,7 +879,7 @@ namespace Private
 
 		const Private::FProjectedPixelProcessorContext Context = 
 			Private::TProjectedPixelProcessor<Features>::MakeContext(
-				Args.SourcePtr, Args.ImagePtr->GetData(), Args.MaskPtr ? Args.MaskPtr->GetData() : nullptr,
+				Args.SourcePtr, Args.ImagePtr->GetLODData(0), Args.MaskPtr ? Args.MaskPtr->GetLODData(0) : nullptr,
 				Args.bIsRGBFadingEnabled, Args.bIsAlphaFadingEnabled, Args.FadeStart, Args.FadeEnd, 
 				Args.ProjectionAngle, Args.MipInterpolationFactor);
 
diff --git a/Engine/Plugins/Experimental/Mutable/Source/MutableRuntime/Private/MuR/OpImageRasterMesh.h b/Engine/Plugins/Experimental/Mutable/Source/MutableRuntime/Private/MuR/OpImageRasterMesh.h
--- a/Engine/Plugins/Experimental/Mutable/Source/MutableRuntime/Private/MuR/OpImageRasterMesh.h
+++ b/Engine/Plugins/Experimental/Mutable/Source/MutableRuntime/Private/MuR/OpImageRasterMesh.h
@@ -84,13 +84,15 @@ namespace mu
 		{
 			// Raster all the faces
             WhitePixelProcessor pixelProc;
+			const TArrayView<uint8> ImageData = pImage->DataStorage.GetLOD(0);
+
 			//for ( int f=0; f<faceCount; ++f )
 			const auto& ProcessFace = [
-				vertices, indices, pImage, sizeX, sizeY, pixelProc
+				vertices, indices, ImageData, sizeX, sizeY, pixelProc
 			] (int32 f)
 			{
 				constexpr int32 NumInterpolators = 1;
-				Triangle<NumInterpolators>(pImage->GetData(), pImage->GetDataSize(),
+				Triangle<NumInterpolators>(ImageData.GetData(), ImageData.Num(),
 					sizeX, sizeY,
 					1,
 					vertices[indices[f * 3 + 0]],
@@ -120,16 +122,19 @@ namespace mu
 			}
 
             WhitePixelProcessor pixelProc;
+
+			const TArrayView<uint8> ImageData = pImage->DataStorage.GetLOD(0); 
+
 			//for (int f = 0; f < faceCount; ++f)
 			const auto& ProcessFace = [
-				vertices, indices, blocks, BlockId, pImage, sizeX, sizeY, pixelProc
+				vertices, indices, blocks, BlockId, ImageData, sizeX, sizeY, pixelProc
 			] (int32 f)
 			{
 				// TODO: Select faces outside for loop?
 				if (blocks[indices[f * 3 + 0]] == BlockId)
 				{
 					constexpr int32 NumInterpolators = 1;
-					Triangle<NumInterpolators>(pImage->GetData(), pImage->GetDataSize(),
+					Triangle<NumInterpolators>(ImageData.GetData(), ImageData.Num(),
 						sizeX, sizeY,
 						1,
 						vertices[indices[f * 3 + 0]],
diff --git a/Engine/Plugins/Experimental/Mutable/Source/MutableRuntime/Private/MuR/OpImageResize.cpp b/Engine/Plugins/Experimental/Mutable/Source/MutableRuntime/Private/MuR/OpImageResize.cpp
--- a/Engine/Plugins/Experimental/Mutable/Source/MutableRuntime/Private/MuR/OpImageResize.cpp
+++ b/Engine/Plugins/Experimental/Mutable/Source/MutableRuntime/Private/MuR/OpImageResize.cpp
@@ -43,8 +43,8 @@ namespace mu
 		] (uint32 y)
 		{
 			uint32 px_16 = 0;
-			const uint8* BaseBuf = Base->GetData() + y * BaseSizeX * NC;
-			uint8* DestBuf = Dest->GetData() + y * DestSizeX * NC;
+			const uint8* BaseBuf = Base->GetLODData(0) + y * BaseSizeX * NC;
+			uint8* DestBuf = Dest->GetLODData(0) + y * DestSizeX * NC;
 
 			uint32 LastPX = BaseSizeX - 1;
 			for (int32 x = 0; x < DestSizeX; ++x)
@@ -126,8 +126,8 @@ namespace mu
 			Dest, Base, dx_16, BaseSizeX, DestSizeX
 		] (uint32 y)
 		{
-			const uint8* BaseBuf = Base->GetData() + y * BaseSizeX * NC;
-			uint8* DestBuf = Dest->GetData() + y * DestSizeX * NC;
+			const uint8* BaseBuf = Base->GetLODData(0) + y * BaseSizeX * NC;
+			uint8* DestBuf = Dest->GetLODData(0) + y * DestSizeX * NC;
 
 			uint32 px_16 = 0;
 			for (int32 x = 0; x < DestSizeX; ++x)
@@ -199,16 +199,16 @@ namespace mu
 		int32 DestSizeX = Dest->GetSizeX();
 		int32 SizeY = Base->GetSizeY();
 
-		const uint8* BaseBuf = Base->GetData();
-		uint8* DestBuf = Dest->GetData();
+		const uint8* BaseBuf = Base->GetLODData(0);
+		uint8* DestBuf = Dest->GetLODData(0);
 
 		// Linear filtering
 		const auto& ProcessLine = [
 			Dest, Base, BaseSizeX, DestSizeX
 		] (uint32 y)
 		{
-			const uint8* BaseBuf = Base->GetData() + y * BaseSizeX * NC;
-			uint8* DestBuf = Dest->GetData() + y * DestSizeX * NC;
+			const uint8* BaseBuf = Base->GetLODData(0) + y * BaseSizeX * NC;
+			uint8* DestBuf = Dest->GetLODData(0) + y * DestSizeX * NC;
 
 			uint32 r[NC];
 			for (int32 x = 0; x < DestSizeX; ++x)
@@ -235,10 +235,10 @@ namespace mu
 		ParallelFor(SizeY, ProcessLine);
 	}
 
-	inline uint32 AverageChannel(uint32 a, uint32 b)
+	inline uint32 AverageChannel(uint32 A, uint32 B)
 	{
-		uint32 r = (a + b) >> 1;
-		return r;
+		uint32 Result = (A + B) >> 1;
+		return Result;
 	}
 
 	template<>
@@ -248,8 +248,8 @@ namespace mu
 		int32 DestSizeX = Dest->GetSizeX();
 		int32 SizeY = Base->GetSizeY();
 
-		const uint8* BaseBuf = Base->GetData();
-		uint8* DestBuf = Dest->GetData();
+		const uint8* BaseBuf = Base->GetLODData(0);
+		uint8* DestBuf = Dest->GetLODData(0);
 
 		int32 TotalBasePixels = BaseSizeX * SizeY;
 		constexpr int32 BasePixelsPerBatch = 4096 * 2;
@@ -411,8 +411,8 @@ namespace mu
 				Dest, Base, rowSize
 			] (uint32 y)
 			{
-				uint8* DestBuf = Dest->GetData() + 2 * y * rowSize;
-				const uint8* BaseBuf = Base->GetData() + y * rowSize;
+				uint8* DestBuf = Dest->GetLODData(0) + 2 * y * rowSize;
+				const uint8* BaseBuf = Base->GetLODData(0) + y * rowSize;
 
 				FMemory::Memcpy(DestBuf, BaseBuf, rowSize);
 				DestBuf += rowSize;
@@ -436,8 +436,8 @@ namespace mu
 			] (uint32 x)
 			{
 				uint32 py_16 = 0;
-				uint8* DestBuf = Dest->GetData() + x * NC;
-				const uint8* BaseBuf = Base->GetData();
+				uint8* DestBuf = Dest->GetLODData(0) + x * NC;
+				const uint8* BaseBuf = Base->GetLODData(0);
 
 				for (int32 y = 0; y < DestSizeY; ++y)
 				{
@@ -515,7 +515,7 @@ namespace mu
 
 		uint32 dy_16 = (uint32(BaseSizeY) << 16) / DestSizeY;
 
-		const uint8* BaseBuf = Base->GetData();
+		const uint8* BaseBuf = Base->GetLODData(0);
 
 		// Linear filtering
 		//for (int32 x = 0; x < SizeX; ++x)
@@ -523,7 +523,7 @@ namespace mu
 			Dest, BaseBuf, SizeX, DestSizeY, dy_16
 		] (uint32 x)
 		{
-			uint8* DestBuf = Dest->GetData() + x * NC;
+			uint8* DestBuf = Dest->GetLODData(0) + x * NC;
 			uint32 py_16 = 0;
 			for (int32 y = 0; y < DestSizeY; ++y)
 			{
@@ -600,8 +600,8 @@ namespace mu
 			Dest, Base, SizeX, DestSizeY
 		] (uint32 y)
 		{
-			uint8* DestBuf = Dest->GetData() + y * NC * SizeX;
-			const uint8* BaseBuf = Base->GetData() + y * FACTOR * SizeX * NC;
+			uint8* DestBuf = Dest->GetLODData(0) + y * NC * SizeX;
+			const uint8* BaseBuf = Base->GetLODData(0) + y * FACTOR * SizeX * NC;
 
 			for (int32 x = 0; x < SizeX; ++x)
 			{
@@ -917,8 +917,8 @@ namespace
 		const EImageFormat SrcFormat = InBase->GetFormat();
 		const EImageFormat DestFormat = Dest->GetFormat();
 	
-		uint8* DestData = Dest->GetData();
-		const uint8* SrcData = InBase->GetData();
+		uint8* DestData = Dest->GetLODData(0);
+		const uint8* SrcData = InBase->GetLODData(0);
 		
 		const FImageFormatData& DestFormatData = GetImageFormatData(Dest->GetFormat());
 		const FImageFormatData& SrcFormatData = GetImageFormatData(InBase->GetFormat());
diff --git a/Engine/Plugins/Experimental/Mutable/Source/MutableRuntime/Private/MuR/OpImageSaturate.h b/Engine/Plugins/Experimental/Mutable/Source/MutableRuntime/Private/MuR/OpImageSaturate.h
--- a/Engine/Plugins/Experimental/Mutable/Source/MutableRuntime/Private/MuR/OpImageSaturate.h
+++ b/Engine/Plugins/Experimental/Mutable/Source/MutableRuntime/Private/MuR/OpImageSaturate.h
@@ -8,13 +8,12 @@
 
 namespace mu
 {
-
 	template<bool bUseVectorImpl = false>
-	inline void ImageSaturateInPlace(Image* pA, float Factor)
+	inline void ImageSaturate(Image* AImage, float Factor)
 	{
 		MUTABLE_CPUPROFILER_SCOPE(ImageSaturateInplace);
 
-		if (!pA)
+		if (!AImage)
 		{
 			return;
 		}
@@ -22,189 +21,204 @@ namespace mu
 		// Clamp the factor
 		// TODO: See what happens if we don't
 		Factor = FMath::Max(0.0f, Factor);
-        int32_t f_8 = (int32_t)(Factor*255);
-
-        uint8_t* pABuf = pA->GetData();
+        uint16 FactorUNorm = (uint16)(Factor*255.0f);
 
 		// Generic implementation
-		const int32 PixelCount = (int32)pA->CalculatePixelCount();
 
-		switch ( pA->GetFormat() )
-		{
-		case EImageFormat::IF_RGB_UBYTE:
+		const int32 BytesPerElem = GetImageFormatData(AImage->GetFormat()).BytesPerBlock;
+		
+		constexpr int32 NumBatchElems = 1 << 14;
+		const int32 NumBatches = AImage->DataStorage.GetNumBatches(NumBatchElems, BytesPerElem); 
+
+		auto ProcessBatch = [AImage, NumBatchElems, BytesPerElem, FactorUNorm](int32 BatchId)
 		{
-			if constexpr (!bUseVectorImpl)
+			TArrayView<uint8> AView = AImage->DataStorage.GetBatch(BatchId, NumBatchElems, BytesPerElem);
+			uint8* ABuf = AView.GetData(); 
+			const int32 NumElems = AView.Num() / BytesPerElem;
+
+			switch (AImage->GetFormat())
+			{
+			case EImageFormat::IF_RGB_UBYTE:
 			{
-				for (int i = 0; i < PixelCount; ++i)
+				if constexpr (!bUseVectorImpl)
 				{
-					int32_t l_16 = 76 * pABuf[3 * i + 0] + 150 * pABuf[3 * i + 1] + 29 * pABuf[3 * i + 2];
-
-					for (int c = 0; c < 3; ++c)
+					for (int32 I = 0; I < NumElems; ++I)
 					{
-						int32_t d_16 = (pABuf[3 * i + c] << 8) - l_16;
-						int32_t r_16 = l_16 + ((d_16 * f_8) >> 8);
-						pABuf[3 * i + c] = (uint8_t)FMath::Min((uint32_t)r_16 >> 8, 255u);
+						uint16 PixelData[4] = {ABuf[3*I + 0], ABuf[3*I + 1], ABuf[3*I + 2], 0};
+
+						const uint16 LumTimesFactor = ((76 * PixelData[0] + 150 * PixelData[1] + 29 * PixelData[2]) / 255) * (255 - FactorUNorm);
+
+						ABuf[3*I + 0] = static_cast<uint8>((PixelData[0] * FactorUNorm + LumTimesFactor) / 255);
+						ABuf[3*I + 1] = static_cast<uint8>((PixelData[1] * FactorUNorm + LumTimesFactor) / 255);
+						ABuf[3*I + 2] = static_cast<uint8>((PixelData[2] * FactorUNorm + LumTimesFactor) / 255);
 					}
 				}
-			}
-			else // if constexpr bUseVectorImpl
-			{
-				constexpr VectorRegister4Int StaurationFactorRGB = MakeVectorRegisterIntConstant(76, 150, 29, 0);
-				const VectorRegister4Int F = VectorIntSet1(static_cast<int32>(Factor * 255.0f));
-				const VectorRegister4Int OneMinusF = VectorIntSet1(255 - static_cast<int32>(Factor * 255.0f));
-				for (int32 PixelIndex = 0; PixelIndex < PixelCount; ++PixelIndex)
+				else // if constexpr bUseVectorImpl
 				{
-					const VectorRegister4Int Value = MakeVectorRegisterInt(
-						static_cast<int32>(pABuf[3 * PixelIndex + 0]), 
-						static_cast<int32>(pABuf[3 * PixelIndex + 1]), 
-						static_cast<int32>(pABuf[3 * PixelIndex + 2]), 
-						0);
-
-					const VectorRegister4Int ScaledValue = VectorIntMultiply(Value, StaurationFactorRGB);
+					constexpr VectorRegister4Int StaurationFactorRGB = MakeVectorRegisterIntConstant(76, 150, 29, 0);
+					constexpr VectorRegister4Int Value255 = MakeVectorRegisterIntConstant(255, 255, 255, 255);
+					const VectorRegister4Int F = VectorIntSet1(FactorUNorm);
+					const VectorRegister4Int OneMinusF = VectorIntSubtract(Value255, F);
 					
-					alignas(VectorRegister4Int) int32 IndexableScaledValue[4];
-					VectorIntStoreAligned(ScaledValue, IndexableScaledValue);
-
-					// Add scaled values and divide by 255
-					const VectorRegister4Int L = VectorShiftRightImmLogical(
-						VectorIntMultiply(
-							VectorIntSet1(IndexableScaledValue[0] + IndexableScaledValue[1] + IndexableScaledValue[2]), 
-							VectorIntSet1(32897)), 23);
+					for (int32 I = 0; I < NumElems; ++I)
+					{
+						const VectorRegister4Int Value = MakeVectorRegisterInt(
+							static_cast<int32>(ABuf[3*I + 0]), 
+							static_cast<int32>(ABuf[3*I + 1]), 
+							static_cast<int32>(ABuf[3*I + 2]), 
+							0);
+
+						const VectorRegister4Int ScaledValue = VectorIntMultiply(Value, StaurationFactorRGB);
+						
+						alignas(VectorRegister4Int) int32 IndexableScaledValue[4];
+						VectorIntStoreAligned(ScaledValue, IndexableScaledValue);
+
+						// Add scaled values and divide by 255
+						const VectorRegister4Int L = VectorShiftRightImmLogical(
+							VectorIntMultiply(
+								VectorIntSet1(IndexableScaledValue[0] + IndexableScaledValue[1] + IndexableScaledValue[2]), 
+								VectorIntSet1(32897)), 23);
+						
+						const VectorRegister4Int Lerp = VectorIntAdd(VectorIntMultiply(L, OneMinusF), VectorIntMultiply(Value, F));
 					
-					const VectorRegister4Int Lerp = VectorIntAdd(VectorIntMultiply(L, OneMinusF), VectorIntMultiply(Value, F));
-				
-					alignas(VectorRegister4Int) int32 IndexableResult[4];
+						alignas(VectorRegister4Int) int32 IndexableResult[4];
 
-					// divide by 255 and store. do we need to clamp?
-					VectorIntStoreAligned(VectorShiftRightImmLogical(VectorIntMultiply(Lerp, VectorIntSet1(32897)), 23),  IndexableResult);
+						// divide by 255 and store. do we need to clamp?
+						VectorIntStoreAligned(VectorShiftRightImmLogical(VectorIntMultiply(Lerp, VectorIntSet1(32897)), 23),  IndexableResult);
 
-					pABuf[3 * PixelIndex + 0] = static_cast<uint8>(IndexableResult[0]);
-					pABuf[3 * PixelIndex + 1] = static_cast<uint8>(IndexableResult[1]);
-					pABuf[3 * PixelIndex + 2] = static_cast<uint8>(IndexableResult[2]);
+						ABuf[3*I + 0] = static_cast<uint8>(IndexableResult[0]);
+						ABuf[3*I + 1] = static_cast<uint8>(IndexableResult[1]);
+						ABuf[3*I + 2] = static_cast<uint8>(IndexableResult[2]);
+					}
 				}
+				break;
 			}
-			break;
-		}
 
-        case EImageFormat::IF_RGBA_UBYTE:
-		{
-			if constexpr (!bUseVectorImpl)
+			case EImageFormat::IF_RGBA_UBYTE:
 			{
-				for (int i = 0; i < PixelCount; ++i)
+				if constexpr (!bUseVectorImpl)
 				{
-					int32_t l_16 = 76 * pABuf[4 * i + 0] + 150 * pABuf[4 * i + 1] + 29 * pABuf[4 * i + 2];
-
-					for (int c = 0; c < 3; ++c)
+					for (int32 I = 0; I < NumElems; ++I)
 					{
-						int32_t d_16 = (pABuf[4 * i + c] << 8) - l_16;
-						int32_t r_16 = l_16 + ((d_16 * f_8) >> 8);
-						pABuf[4 * i + c] = (uint8_t)FMath::Min((uint32_t)r_16 >> 8, 255u);
-					}
+						uint16 PixelData[4] = {ABuf[4*I + 0], ABuf[4*I + 1], ABuf[4*I + 2], ABuf[4*I + 3]};
 
-					//pABuf[4 * i + 3] = pABuf[4 * i + 3];
+						const uint16 LumTimesFactor = ((76 * PixelData[0] + 150 * PixelData[1] + 29 * PixelData[2]) / 255) * (255 - FactorUNorm);
+
+						ABuf[4*I + 0] = static_cast<uint8>((PixelData[0] * FactorUNorm + LumTimesFactor) / 255);
+						ABuf[4*I + 1] = static_cast<uint8>((PixelData[1] * FactorUNorm + LumTimesFactor) / 255);
+						ABuf[4*I + 2] = static_cast<uint8>((PixelData[2] * FactorUNorm + LumTimesFactor) / 255);
+					}
 				}
-			}
-			else // if constexpr bUseVectorImpl
-			{
-				constexpr VectorRegister4Int StaurationFactorRGB = MakeVectorRegisterIntConstant(76, 150, 29, 0);
-				const VectorRegister4Int F = VectorIntSet1(static_cast<int32>(Factor * 255.0f));
-				const VectorRegister4Int OneMinusF = VectorIntSet1(255 - static_cast<int32>(Factor * 255.0f));
-				for (int32 PixelIndex = 0; PixelIndex < PixelCount; ++PixelIndex)
+				else // if constexpr bUseVectorImpl
 				{
-					const VectorRegister4Int Value = MakeVectorRegisterInt(
-						static_cast<int32>(pABuf[4 * PixelIndex + 0]), 
-						static_cast<int32>(pABuf[4 * PixelIndex + 1]), 
-						static_cast<int32>(pABuf[4 * PixelIndex + 2]), 
-						0);
-
-					const VectorRegister4Int ScaledValue = VectorIntMultiply(Value, StaurationFactorRGB);
+					constexpr VectorRegister4Int StaurationFactorRGB = MakeVectorRegisterIntConstant(76, 150, 29, 0);
+					constexpr VectorRegister4Int Value255 = MakeVectorRegisterIntConstant(255, 255, 255, 255);
+					const VectorRegister4Int F = VectorIntSet1(FactorUNorm);
+					const VectorRegister4Int OneMinusF = VectorIntSubtract(Value255, F);
 					
-					alignas(VectorRegister4Int) int32 IndexableScaledValue[4];
-					VectorIntStoreAligned(ScaledValue, IndexableScaledValue);
-
-					// Add scaled values and divide by 255
-					const VectorRegister4Int L = VectorShiftRightImmLogical(
-						VectorIntMultiply(
-							VectorIntSet1(IndexableScaledValue[0] + IndexableScaledValue[1] + IndexableScaledValue[2]), 
-							VectorIntSet1(32897)), 23);
+					for (int32 I = 0; I < NumElems; ++I)
+					{
+						const VectorRegister4Int Value = MakeVectorRegisterInt(
+							static_cast<int32>(ABuf[4*I + 0]), 
+							static_cast<int32>(ABuf[4*I + 1]), 
+							static_cast<int32>(ABuf[4*I + 2]), 
+							0);
+
+						const VectorRegister4Int ScaledValue = VectorIntMultiply(Value, StaurationFactorRGB);
+						
+						alignas(VectorRegister4Int) int32 IndexableScaledValue[4];
+						VectorIntStoreAligned(ScaledValue, IndexableScaledValue);
+
+						// Add scaled values and divide by 255
+						const VectorRegister4Int L = VectorShiftRightImmLogical(
+							VectorIntMultiply(
+								VectorIntSet1(IndexableScaledValue[0] + IndexableScaledValue[1] + IndexableScaledValue[2]), 
+								VectorIntSet1(32897)), 23);
+						
+						const VectorRegister4Int Lerp = VectorIntAdd(VectorIntMultiply(L, OneMinusF), VectorIntMultiply(Value, F));
 					
-					const VectorRegister4Int Lerp = VectorIntAdd(VectorIntMultiply(L, OneMinusF), VectorIntMultiply(Value, F));
-				
-					alignas(VectorRegister4Int) int32 IndexableResult[4];
-					// divide by 255 and store. do we need to clamp?
-					VectorIntStoreAligned(VectorShiftRightImmLogical(VectorIntMultiply(Lerp, VectorIntSet1(32897)), 23),  IndexableResult);
-
-					pABuf[4 * PixelIndex + 0] = static_cast<uint8>(IndexableResult[0]);
-					pABuf[4 * PixelIndex + 1] = static_cast<uint8>(IndexableResult[1]);
-					pABuf[4 * PixelIndex + 2] = static_cast<uint8>(IndexableResult[2]);
-					//pABuf[4 * PixelIndex + 3] = pABuf[4 * PixelIndex + 3];
-				}
-			}
+						alignas(VectorRegister4Int) int32 IndexableResult[4];
 
-            break;
-        }
+						// divide by 255 and store. do we need to clamp?
+						VectorIntStoreAligned(VectorShiftRightImmLogical(VectorIntMultiply(Lerp, VectorIntSet1(32897)), 23),  IndexableResult);
 
-        case EImageFormat::IF_BGRA_UBYTE:
-        {
-			if constexpr (!bUseVectorImpl)
-			{
-				for (int i = 0; i < PixelCount; ++i)
-				{
-					int32_t l_16 = 76 * pABuf[4 * i + 2] + 150 * pABuf[4 * i + 1] + 29 * pABuf[4 * i + 0];
-
-					for (int c = 0; c < 3; ++c)
-					{
-						int32_t d_16 = (pABuf[4 * i + c] << 8) - l_16;
-						int32_t r_16 = l_16 + ((d_16 * f_8) >> 8);
-						pABuf[4 * i + c] = (uint8_t)FMath::Min((uint32_t)r_16 >> 8, 255u);
+						ABuf[4*I + 0] = static_cast<uint8>(IndexableResult[0]);
+						ABuf[4*I + 1] = static_cast<uint8>(IndexableResult[1]);
+						ABuf[4*I + 2] = static_cast<uint8>(IndexableResult[2]);
 					}
+				}	
 
-					//pABuf[4 * i + 3] = pABuf[4 * i + 3];
-				}
+				break;
 			}
-			else // if constexpr bUseVectorImpl
+
+			case EImageFormat::IF_BGRA_UBYTE:
 			{
-				constexpr VectorRegister4Int StaurationFactorBGR = MakeVectorRegisterIntConstant(29, 150, 76, 0);
-				const VectorRegister4Int F = VectorIntSet1(static_cast<int32>(Factor * 255.0f));
-				const VectorRegister4Int OneMinusF = VectorIntSet1(255 - static_cast<int32>(Factor * 255.0f));
-				for (int32 PixelIndex = 0; PixelIndex < PixelCount; ++PixelIndex)
+				if constexpr (!bUseVectorImpl)
 				{
-					const VectorRegister4Int Value = MakeVectorRegisterInt(
-						static_cast<int32>(pABuf[4 * PixelIndex + 0]), 
-						static_cast<int32>(pABuf[4 * PixelIndex + 1]), 
-						static_cast<int32>(pABuf[4 * PixelIndex + 2]), 
-						0);
+					for (int32 I = 0; I < NumElems; ++I)
+					{
+						uint16 PixelData[4] = {ABuf[4*I + 0], ABuf[4*I + 1], ABuf[4*I + 2], ABuf[4*I + 3]};
 
-					const VectorRegister4Int ScaledValue = VectorIntMultiply(Value, StaurationFactorBGR);
-					
-					alignas(VectorRegister4Int) int32 IndexableScaledValue[4];
-					VectorIntStoreAligned(ScaledValue, IndexableScaledValue);
-
-					// Add scaled values and divide by 255
-					const VectorRegister4Int L = VectorShiftRightImmLogical(
-						VectorIntMultiply(
-							VectorIntSet1(IndexableScaledValue[0] + IndexableScaledValue[1] + IndexableScaledValue[2]), 
-							VectorIntSet1(32897)), 23);
+						const uint16 LumTimesFactor = ((29 * PixelData[0] + 150 * PixelData[1] + 76 * PixelData[2]) / 255) * (255 - FactorUNorm);
+
+						ABuf[4*I + 0] = static_cast<uint8>((PixelData[0] * FactorUNorm + LumTimesFactor) / 255);
+						ABuf[4*I + 1] = static_cast<uint8>((PixelData[1] * FactorUNorm + LumTimesFactor) / 255);
+						ABuf[4*I + 2] = static_cast<uint8>((PixelData[2] * FactorUNorm + LumTimesFactor) / 255);
+					}
+				}
+				else // if constexpr bUseVectorImpl
+				{
+					constexpr VectorRegister4Int StaurationFactorRGB = MakeVectorRegisterIntConstant(29, 150, 76, 0);
+					constexpr VectorRegister4Int Value255 = MakeVectorRegisterIntConstant(255, 255, 255, 255);
+					const VectorRegister4Int F = VectorIntSet1(FactorUNorm);
+					const VectorRegister4Int OneMinusF = VectorIntSubtract(Value255, F);
+					for (int32 I = 0; I < NumElems; ++I)
+					{
+						const VectorRegister4Int Value = MakeVectorRegisterInt(
+							static_cast<int32>(ABuf[4*I + 0]), 
+							static_cast<int32>(ABuf[4*I + 1]), 
+							static_cast<int32>(ABuf[4*I + 2]), 
+							0);
+
+						const VectorRegister4Int ScaledValue = VectorIntMultiply(Value, StaurationFactorRGB);
+						
+						alignas(VectorRegister4Int) int32 IndexableScaledValue[4];
+						VectorIntStoreAligned(ScaledValue, IndexableScaledValue);
+
+						// Add scaled values and divide by 255
+						const VectorRegister4Int L = VectorShiftRightImmLogical(
+							VectorIntMultiply(
+								VectorIntSet1(IndexableScaledValue[0] + IndexableScaledValue[1] + IndexableScaledValue[2]), 
+								VectorIntSet1(32897)), 23);
+						
+						const VectorRegister4Int Lerp = VectorIntAdd(VectorIntMultiply(L, OneMinusF), VectorIntMultiply(Value, F));
 					
-					const VectorRegister4Int Lerp = VectorIntAdd(VectorIntMultiply(L, OneMinusF), VectorIntMultiply(Value, F));
-				
-					alignas(VectorRegister4Int) int32 IndexableResult[4];
+						alignas(VectorRegister4Int) int32 IndexableResult[4];
 
-					// divide by 255 and store. do we need to clamp?
-					VectorIntStoreAligned(VectorShiftRightImmLogical(VectorIntMultiply(Lerp, VectorIntSet1(32897)), 23),  IndexableResult);
+						// divide by 255 and store. do we need to clamp?
+						VectorIntStoreAligned(VectorShiftRightImmLogical(VectorIntMultiply(Lerp, VectorIntSet1(32897)), 23),  IndexableResult);
 
-					pABuf[4 * PixelIndex + 0] = static_cast<uint8>(IndexableResult[0]);
-					pABuf[4 * PixelIndex + 1] = static_cast<uint8>(IndexableResult[1]);
-					pABuf[4 * PixelIndex + 2] = static_cast<uint8>(IndexableResult[2]);
-					//pABuf[4 * PixelIndex + 3] = pABuf[4 * PixelIndex + 3];
-				}
+						ABuf[4*I + 0] = static_cast<uint8>(IndexableResult[0]);
+						ABuf[4*I + 1] = static_cast<uint8>(IndexableResult[1]);
+						ABuf[4*I + 2] = static_cast<uint8>(IndexableResult[2]);
+					}
+				}	
+				
+				break;
+			}
+
+			default:
+				check(false);
 			}
-            
-			break;
-        }
+		};
 
-		default:
-			check(false);
+		if (NumBatches == 1)
+		{
+			ProcessBatch(0);
+		}
+		else
+		{
+			ParallelFor(NumBatches, ProcessBatch);
 		}
 	}
 }
diff --git a/Engine/Plugins/Experimental/Mutable/Source/MutableRuntime/Private/MuR/OpImageSwizzle.cpp b/Engine/Plugins/Experimental/Mutable/Source/MutableRuntime/Private/MuR/OpImageSwizzle.cpp
--- a/Engine/Plugins/Experimental/Mutable/Source/MutableRuntime/Private/MuR/OpImageSwizzle.cpp
+++ b/Engine/Plugins/Experimental/Mutable/Source/MutableRuntime/Private/MuR/OpImageSwizzle.cpp
@@ -5,8 +5,7 @@
 
 namespace mu
 {
-
-	void ImageSwizzle( Image* Result, const Ptr<const Image> Sources[], const uint8 Channels[] )
+	void ImageSwizzle(Image* Result, const Ptr<const Image> Sources[], const uint8 Channels[])
 	{
 		MUTABLE_CPUPROFILER_SCOPE(ImageSwizzle);
 
@@ -15,28 +14,18 @@ namespace mu
 			return;
 		}
 
-		// Very slow generic implementations
-        int32 PixelCount = Result->CalculatePixelCount();
-
 		EImageFormat Format = Result->GetFormat();
 
-        // Pixelcount should already match, but due to bugs it may not be the case. Try to detect it,
-        // but avoid crashing below:
-        uint16 numChannels = GetImageFormatData(Format).Channels;
-        for (uint16 c=0;c<numChannels; ++c)
-        {
-            if (Sources[c])
-            {
-                int32 SourcePixelCount = Sources[c]->CalculatePixelCount();
-                if (PixelCount> SourcePixelCount)
-                {
-                    // Something went wrong
-					PixelCount = SourcePixelCount;
-                }
-            }
-        }
-
-		int NumDestChannels = 0;
+		// LODs may not match due to bugs, only precess the common avaliable lODs.
+		int32 NumLODs = Result->GetLODCount();
+
+        uint16 NumChannels = GetImageFormatData(Format).Channels;
+		for (uint16 C = 0; C < NumChannels; ++C)
+		{
+			NumLODs = Sources[C] ? FMath::Min(NumLODs, Sources[C]->GetLODCount()) : NumLODs;
+		}
+
+		int32 NumDestChannels = 0;
 
 		switch (Format)
 		{
@@ -57,71 +46,46 @@ namespace mu
 			check(false);
 		}
 
-		for (int i = 0; i < NumDestChannels; ++i)
+		for (int32 Channel = 0; Channel < NumDestChannels; ++Channel)
 		{
-			uint8* pDestBuf = Result->GetData() + i;
-
-			if (Format == EImageFormat::IF_BGRA_UBYTE)
-			{
-				if (i == 0)
-				{
-					pDestBuf = Result->GetData() + 2;
-				}
-				else if (i == 2)
-				{
-					pDestBuf = Result->GetData() + 0;
-				}
-			}
+			const Image* Src = Sources[Channel].get();
+			
+			const int32 DestChannel = Format == EImageFormat::IF_BGRA_UBYTE && Channel < 3
+								    ? FMath::Abs(Channel - 2)
+									: Channel;
 
 			bool bFilled = false;
 
 			constexpr int32 NumBatchElems = 4096*2;
-			const int32 NumBatches = FMath::DivideAndRoundUp(PixelCount, NumBatchElems);
 
-			if (Sources[i])
+			if (Sources[Channel])
 			{
-				const uint8* pSourceBuf = Sources[i]->GetData() + Channels[i];
-
-				switch (Sources[i]->GetFormat())
+				EImageFormat SrcFormat = Sources[Channel]->GetFormat();
+				switch (SrcFormat)
 				{
 				case EImageFormat::IF_L_UBYTE:
-					if (Channels[i] < 1)
+				{
+					if (Channels[Channel] < 1)
 					{
-						auto ProcessBatch = [pDestBuf, pSourceBuf, NumDestChannels, PixelCount, NumBatchElems ](int32 BatchId)
-						{
-							const int32 BatchBegin = BatchId * NumBatchElems;
-							const int32 BatchEnd = FMath::Min(BatchBegin + NumBatchElems, PixelCount);
+						const int32 NumBatches = Src->DataStorage.GetNumBatchesLODRange(NumBatchElems, 1, 0, NumLODs);
+						check(NumBatches == Result->DataStorage.GetNumBatchesLODRange(NumBatchElems, NumDestChannels, 0, NumLODs));
+		
+						int32 SrcChannel = Channels[Channel];
 
-							for (int32 p = BatchBegin; p < BatchEnd; ++p)
-							{
-								pDestBuf[p * NumDestChannels] = pSourceBuf[p];
-							}
-						};
-
-						if (NumBatches == 1)
-						{
-							ProcessBatch(0);
-						}
-						else if (NumBatches > 1)
+						auto ProcessBatch = [Result, Src, Sources, NumDestChannels, DestChannel, SrcChannel, NumBatchElems, NumLODs](int32 BatchId)
 						{
-							ParallelFor(NumBatches, ProcessBatch);
-						}
+							TArrayView<const uint8> SrcView = Src->DataStorage.GetBatchLODRange(BatchId, NumBatchElems, 1, 0, NumLODs);
+							TArrayView<uint8> ResultView = Result->DataStorage.GetBatchLODRange(BatchId, NumBatchElems, NumDestChannels, 0, NumLODs);
 
-						bFilled = true;
-					}
-					break;
+							const int32 NumElems = SrcView.Num() / 1;
+							check(NumElems == ResultView.Num() / NumDestChannels);
 
-				case EImageFormat::IF_RGB_UBYTE:
-					if (Channels[i] < 3)
-					{
-						auto ProcessBatch = [pDestBuf, pSourceBuf, NumDestChannels, PixelCount, NumBatchElems](int32 BatchId)
-						{
-							const int32 BatchBegin = BatchId * NumBatchElems;
-							const int32 BatchEnd = FMath::Min(BatchBegin + NumBatchElems, PixelCount);
+							uint8* DestBuf = ResultView.GetData() + DestChannel;
+							const uint8* SrcBuf = SrcView.GetData() + SrcChannel;
 
-							for (int32 p = BatchBegin; p < BatchEnd; ++p)
+							for (int32 I = 0; I < NumElems; ++I)
 							{
-								pDestBuf[p * NumDestChannels] = pSourceBuf[p * 3];
+								DestBuf[I*NumDestChannels] = SrcBuf[I];
 							}
 						};
 
@@ -137,18 +101,30 @@ namespace mu
 						bFilled = true;
 					}
 					break;
+				}
+				case EImageFormat::IF_RGB_UBYTE:
+				{
+					if (Channels[Channel] < 3)
+					{
+						const int32 NumBatches = Src->DataStorage.GetNumBatchesLODRange(NumBatchElems, 3, 0, NumLODs);
+						check(NumBatches == Result->DataStorage.GetNumBatchesLODRange(NumBatchElems, NumDestChannels, 0, NumLODs));
+		
+						const int32 SrcChannel = Channels[Channel];
 
-				case EImageFormat::IF_RGBA_UBYTE:
-					if (Channels[i] < 4)
-					{	
-						auto ProcessBatch = [pDestBuf, pSourceBuf, NumDestChannels, PixelCount, NumBatchElems](int32 BatchId)
+						auto ProcessBatch = [Result, Src, Sources, NumDestChannels, DestChannel, SrcChannel, NumBatchElems, NumLODs](int32 BatchId)
 						{
-							const int32 BatchBegin = BatchId * NumBatchElems;
-							const int32 BatchEnd = FMath::Min(BatchBegin + NumBatchElems, PixelCount);
+							TArrayView<const uint8> SrcView = Src->DataStorage.GetBatchLODRange(BatchId, NumBatchElems, 3, 0, NumLODs);
+							TArrayView<uint8> ResultView = Result->DataStorage.GetBatchLODRange(BatchId, NumBatchElems, NumDestChannels, 0, NumLODs);
+
+							const int32 NumElems = SrcView.Num() / 3;
+							check(NumElems == ResultView.Num()/NumDestChannels);
 
-							for (int32 p = BatchBegin; p < BatchEnd; ++p)
+							uint8* DestBuf = ResultView.GetData() + DestChannel;
+							const uint8* SrcBuf = SrcView.GetData() + SrcChannel;
+							
+							for (int32 I = 0; I < NumElems; ++I)
 							{
-								pDestBuf[p * NumDestChannels] = pSourceBuf[p * 4];
+								DestBuf[I*NumDestChannels] = SrcBuf[I*3];
 							}
 						};
 
@@ -164,25 +140,33 @@ namespace mu
 						bFilled = true;
 					}
 					break;
-
+				}
+				case EImageFormat::IF_RGBA_UBYTE:
 				case EImageFormat::IF_BGRA_UBYTE:
-					if (Channels[i] == 0)
-					{
-						pSourceBuf = Sources[i]->GetData() + 2;
-					}
-					else if (Channels[i] == 2)
-					{
-						pSourceBuf = Sources[i]->GetData() + 0;
-					}
-					if (Channels[i] < 4)
+				{
+					if (Channels[Channel] < 4)
 					{
-						auto ProcessBatch = [pDestBuf, pSourceBuf, NumDestChannels, PixelCount, NumBatchElems](int32 BatchId)
+						const int32 SrcChannel = SrcFormat == EImageFormat::IF_BGRA_UBYTE && Channels[Channel] < 3
+											   ? FMath::Abs(int32(Channels[Channel]) - 2)
+											   : Channels[Channel];
+
+						const int32 NumBatches = Src->DataStorage.GetNumBatchesLODRange(NumBatchElems, 4, 0, NumLODs);
+						check(NumBatches == Result->DataStorage.GetNumBatchesLODRange(NumBatchElems, NumDestChannels, 0, NumLODs));
+			
+						auto ProcessBatch = [Result, Src, Sources, NumDestChannels, DestChannel, SrcChannel, NumBatchElems, NumLODs](int32 BatchId)
 						{
-							const int32 BatchBegin = BatchId * NumBatchElems;
-							const int32 BatchEnd = FMath::Min(BatchBegin + NumBatchElems, PixelCount);
-							for (int32 p = BatchBegin; p < BatchEnd; ++p)
+							TArrayView<const uint8> SrcView = Src->DataStorage.GetBatchLODRange(BatchId, NumBatchElems, 4, 0, NumLODs);  
+							TArrayView<uint8> ResultView = Result->DataStorage.GetBatchLODRange(BatchId, NumBatchElems, NumDestChannels, 0, NumLODs);  
+
+							const int32 NumElems = SrcView.Num() / 4;
+							check(NumElems == ResultView.Num()/NumDestChannels);
+
+							uint8* DestBuf = ResultView.GetData() + DestChannel;
+							const uint8* SrcBuf = SrcView.GetData() + SrcChannel;
+							
+							for (int32 I = 0; I < NumElems; ++I)
 							{
-								pDestBuf[p * NumDestChannels] = pSourceBuf[p * 4];
+								DestBuf[I*NumDestChannels] = SrcBuf[I*4];
 							}
 						};
 
@@ -198,24 +182,31 @@ namespace mu
 						bFilled = true;
 					}
 					break;
-
+				}
 				default:
+				{
 					check(false);
 				}
+				}
 			}
 
 			if (!bFilled)
 			{
-				// Source not set. Clear to 0 if it is a color channel, or to 255 if it is alpha
-				uint8 Value = (Channels[i] == 3) ? 255 : 0;
+				const int32 NumBatches = Result->DataStorage.GetNumBatchesLODRange(NumBatchElems, NumDestChannels, 0, NumLODs);
 
-				auto ProcessBatch = [pDestBuf, NumDestChannels, PixelCount, NumBatchElems, Value](int32 BatchId)
+				// Alpha is expected to be filled with 1.
+				const uint8 FillValue = DestChannel < 3 ? 0 : 255;
+
+				auto ProcessBatch = [Result, NumDestChannels, NumBatchElems, NumLODs, DestChannel, FillValue](int32 BatchId)
 				{
-					const int32 BatchBegin = BatchId * NumBatchElems;
-					const int32 BatchEnd = FMath::Min(BatchBegin + NumBatchElems, PixelCount);
-					for (int32 p = BatchBegin; p < BatchEnd; ++p)
+					TArrayView<uint8> ResultView = Result->DataStorage.GetBatchLODRange(BatchId, NumBatchElems, NumDestChannels, 0, NumLODs);
+					int32 NumElems = ResultView.Num() / NumDestChannels;
+
+					uint8* DestBuf = ResultView.GetData() + DestChannel;
+
+					for (int32 I = 0; I < NumElems; ++I)
 					{
-						pDestBuf[p * NumDestChannels] = Value;
+						DestBuf[I * NumDestChannels] = FillValue;
 					}
 				};
 
@@ -231,7 +222,6 @@ namespace mu
 		}
 	}
 
-
 	Ptr<Image> FImageOperator::ImageSwizzle( EImageFormat Format, const Ptr<const Image> Sources[], const uint8 Channels[] )
 	{
 		MUTABLE_CPUPROFILER_SCOPE(ImageSwizzle);
@@ -248,5 +238,4 @@ namespace mu
 		return Dest;
 	}
 
-
 }
diff --git a/Engine/Plugins/Experimental/Mutable/Source/MutableRuntime/Private/MuR/OpImageTransform.h b/Engine/Plugins/Experimental/Mutable/Source/MutableRuntime/Private/MuR/OpImageTransform.h
--- a/Engine/Plugins/Experimental/Mutable/Source/MutableRuntime/Private/MuR/OpImageTransform.h
+++ b/Engine/Plugins/Experimental/Mutable/Source/MutableRuntime/Private/MuR/OpImageTransform.h
@@ -57,7 +57,7 @@ namespace mu
 			pDestImage->RelevancyMaxY = FMath::Max(DestCropRect.Max.Y - 1, DestCropRect.Min.Y);
 		}
 
-		uint8* DestData       = pDestImage->GetData();
+		uint8* DestData       = pDestImage->GetMipData(0);
 		const uint8* Src0Data = pImage->GetMipData(0);
 		const uint8* Src1Data = pImage->GetLODCount() > 1 ? pImage->GetMipData(1) : Src0Data;
 
diff --git a/Engine/Plugins/Experimental/Mutable/Source/MutableRuntime/Private/MuR/Serialisation.cpp b/Engine/Plugins/Experimental/Mutable/Source/MutableRuntime/Private/MuR/Serialisation.cpp
--- a/Engine/Plugins/Experimental/Mutable/Source/MutableRuntime/Private/MuR/Serialisation.cpp
+++ b/Engine/Plugins/Experimental/Mutable/Source/MutableRuntime/Private/MuR/Serialisation.cpp
@@ -16,15 +16,15 @@ namespace mu
 	MUTABLE_IMPLEMENT_POD_SERIALISABLE(float);    
 	MUTABLE_IMPLEMENT_POD_SERIALISABLE(double);   
                                                   
-    MUTABLE_IMPLEMENT_POD_SERIALISABLE( int8 );   
-    MUTABLE_IMPLEMENT_POD_SERIALISABLE( int16 );  
-    MUTABLE_IMPLEMENT_POD_SERIALISABLE( int32 );  
-    MUTABLE_IMPLEMENT_POD_SERIALISABLE( int64 );  
-                                                  
-    MUTABLE_IMPLEMENT_POD_SERIALISABLE( uint8 );  
-    MUTABLE_IMPLEMENT_POD_SERIALISABLE( uint16 ); 
-    MUTABLE_IMPLEMENT_POD_SERIALISABLE( uint32 ); 
-    MUTABLE_IMPLEMENT_POD_SERIALISABLE( uint64 )
+    MUTABLE_IMPLEMENT_POD_SERIALISABLE(int8);   
+    MUTABLE_IMPLEMENT_POD_SERIALISABLE(int16);  
+    MUTABLE_IMPLEMENT_POD_SERIALISABLE(int32);  
+    MUTABLE_IMPLEMENT_POD_SERIALISABLE(int64);  
+                                                 
+    MUTABLE_IMPLEMENT_POD_SERIALISABLE(uint8);  
+    MUTABLE_IMPLEMENT_POD_SERIALISABLE(uint16); 
+    MUTABLE_IMPLEMENT_POD_SERIALISABLE(uint32); 
+    MUTABLE_IMPLEMENT_POD_SERIALISABLE(uint64)
 
 	MUTABLE_IMPLEMENT_POD_SERIALISABLE(FGuid);
 
@@ -39,6 +39,7 @@ namespace mu
 	MUTABLE_IMPLEMENT_POD_SERIALISABLE(FUintVector2);                                 
 	MUTABLE_IMPLEMENT_POD_SERIALISABLE(UE::Math::TIntVector2<uint16>);                
 	MUTABLE_IMPLEMENT_POD_SERIALISABLE(UE::Math::TIntVector2<int16>); 
+	MUTABLE_IMPLEMENT_POD_SERIALISABLE(UE::Math::TVector4<float>); 
 
 	
     void operator<<(OutputArchive& arch, const FString& t)
@@ -50,7 +51,7 @@ namespace mu
 
     void operator>>(InputArchive& arch, FString& t)
     {
-	    TArray<TCHAR> Data;
+		TArray<TCHAR> Data;
     	arch >> Data;
 
     	t = FString(Data.GetData()); // Construct from raw pointer to avoid double zero terminating character
@@ -83,18 +84,18 @@ namespace mu
     }
 
 	
-	void operator<< ( OutputArchive& arch, const bool& t )
+	void operator<<(OutputArchive& Arch, const bool& T)
     {
-    	uint8 s = t ? 1 : 0;
-    	arch.GetPrivate()->m_pStream->Write( &s, sizeof(uint8) );
+    	uint8 S = T ? 1 : 0;
+    	Arch.GetPrivate()->m_pStream->Write(&S, sizeof(uint8));
     }
 
 
-	void operator>> ( InputArchive& arch, bool& t )
+	void operator>>(InputArchive& Arch, bool& T)
     {
-    	uint8 s;
-    	arch.GetPrivate()->m_pStream->Read( &s, sizeof(uint8) );
-    	t = s!=0;
+    	uint8 S;
+    	Arch.GetPrivate()->m_pStream->Read(&S, sizeof(uint8));
+    	T = S != 0;
     }
 
 	
diff --git a/Engine/Plugins/Experimental/Mutable/Source/MutableRuntime/Public/MuR/Image.h b/Engine/Plugins/Experimental/Mutable/Source/MutableRuntime/Public/MuR/Image.h
--- a/Engine/Plugins/Experimental/Mutable/Source/MutableRuntime/Public/MuR/Image.h
+++ b/Engine/Plugins/Experimental/Mutable/Source/MutableRuntime/Public/MuR/Image.h
@@ -14,13 +14,7 @@
 #include "HAL/Platform.h"
 #include "Misc/AssertionMacros.h"
 
-#include "MuR/MemoryTrackingAllocationPolicy.h"
-
-namespace mu::MemoryCounters
-{
-	struct FImageMemoryCounterTag {};
-	using FImageMemoryCounter = TMemoryCounter<FImageMemoryCounterTag>;
-}
+#include "MuR/ImageDataStorage.h"
 
 namespace mu
 {
@@ -30,11 +24,6 @@ namespace mu
 	typedef Ptr<Image> ImagePtr;
 	typedef Ptr<const Image> ImagePtrConst;
 
-	MUTABLE_DEFINE_ENUM_SERIALISABLE( EBlendType );
-	MUTABLE_DEFINE_ENUM_SERIALISABLE( ECompositeImageMode );
-	MUTABLE_DEFINE_ENUM_SERIALISABLE( ESamplingMethod );
-	MUTABLE_DEFINE_ENUM_SERIALISABLE( EMinFilterMethod );
-
     //! \brief 2D image resource with mipmaps.
 	//! \ingroup runtime
     class MUTABLERUNTIME_API Image : public Resource
@@ -55,10 +44,10 @@ namespace mu
         //!         base level. It must be a number between 1 and the maximum possible levels, which
         //!         depends on the image size.
         //! \param format Pixel format.
-        Image( uint32 sizeX, uint32 sizeY, uint32 lods, EImageFormat format, EInitializationType Init );
+        Image(uint32 SizeX, uint32 SizeY, uint32 lods, EImageFormat Format, EInitializationType InitType);
 
 		/** */
-		static Ptr<Image> CreateAsReference( uint32 ID, const FImageDesc& Desc, bool bForceLoad);
+		static Ptr<Image> CreateAsReference(uint32 ID, const FImageDesc& Desc, bool bForceLoad);
 
 		//! Serialisation
 		static void Serialise( const Image* p, OutputArchive& arch );
@@ -72,7 +61,7 @@ namespace mu
 		//-----------------------------------------------------------------------------------------
 
 		/** */
-		void Init(uint32 sizeX, uint32 sizeY, uint32 lods, EImageFormat format, EInitializationType Init);
+		void Init(uint32 SizeX, uint32 SizeY, uint32 Lods, EImageFormat Format, EInitializationType InitType);
 
 		//! Return the width of the image.
         uint16 GetSizeX() const;
@@ -87,15 +76,15 @@ namespace mu
 
 		//! Return the number of levels of detail (mipmaps) in the texture. The base lavel is also
 		//! counted, so the minimum is 1.
-		int GetLODCount() const;
+		int32 GetLODCount() const;
+		void SetLODCount(int32 LODCount);
 
-		//! Return a pointer to a instance-owned buffer where the image pixels are. The pixels are
-		//! stored without any padding. All the LODs are stored in the same buffer, contiguously.
-        const uint8* GetData() const;
-        uint8* GetData();
+		//! Return a pointer to a instance-owned buffer where the image pixels are.
+        const uint8* GetLODData(int32 LODIndex) const;
+        uint8* GetLODData(int32 LODIndex);
 
         //! Return the size in bytes of a specific LOD of the image.
-        int32 GetLODDataSize( int lod ) const;
+        int32 GetLODDataSize(int32 LODIndex) const;
 
 		/** Return true if this is a reference to an engine image. */
 		bool IsReference() const;
@@ -118,16 +107,6 @@ namespace mu
 
 		// This used to be the data in the private implementation of the image interface
 		//-----------------------------------------------------------------------------------------
-
-		/** */
-		FImageSize m_size = FImageSize(0, 0);
-
-		/** */
-		EImageFormat m_format = EImageFormat::IF_NONE;
-
-		/** Levels of detail (mipmaps) */
-		uint8 m_lods = 0;
-
 		/** These set of flags are used to cache information for images at runtime. */
 		typedef enum
 		{
@@ -162,9 +141,7 @@ namespace mu
 
 		/** Pixel data for all lods. */
 
-		using ImageDataContainerType = TArray<uint8, FDefaultMemoryTrackingAllocator<MemoryCounters::FImageMemoryCounter>>;
-		ImageDataContainerType m_data;
-
+		FImageDataStorage DataStorage;
 
 		// This used to be the methods in the private implementation of the image interface
 		//-----------------------------------------------------------------------------------------
@@ -177,13 +154,10 @@ namespace mu
 
 			Ptr<Image> pResult = new Image();
 
-			pResult->m_size = m_size;
-			pResult->m_format = m_format;
-			pResult->m_lods = m_lods;
 			pResult->m_flags = m_flags;
 			pResult->RelevancyMinY = RelevancyMinY;
 			pResult->RelevancyMaxY = RelevancyMaxY;
-			pResult->m_data = m_data;
+			pResult->DataStorage = DataStorage;
 			pResult->ReferenceID = ReferenceID;
 
 			return pResult;
@@ -200,13 +174,10 @@ namespace mu
 				return;
 			}
 
-			m_size = Other->m_size;
-			m_format = Other->m_format;
-			m_lods = Other->m_lods;
 			m_flags = Other->m_flags;
 			RelevancyMinY = Other->RelevancyMinY;
 			RelevancyMaxY = Other->RelevancyMaxY;
-			m_data =Other->m_data;
+			DataStorage = Other->DataStorage;
 			ReferenceID = Other->ReferenceID;
 		}
 
@@ -220,23 +191,12 @@ namespace mu
 				return;
 			}
 
-			m_size = Other->m_size;
-			m_format = Other->m_format;
-			m_lods = Other->m_lods;
 			m_flags = Other->m_flags;
 			RelevancyMinY = Other->RelevancyMinY;
 			RelevancyMaxY = Other->RelevancyMaxY;
-			m_data = MoveTemp(Other->m_data);
 			ReferenceID = Other->ReferenceID;
-
-			Other->m_size = FImageSize(0, 0);
-			Other->m_format = EImageFormat::IF_NONE;
-			Other->m_lods = 0;
-			Other->m_flags = 0;
-			Other->RelevancyMinY = 0;
-			Other->RelevancyMaxY = 0;
-			Other->m_data.SetNum(0);
-			Other->ReferenceID = 0;
+			
+			DataStorage = MoveTemp(Other->DataStorage);
 		}
 
 
@@ -247,13 +207,11 @@ namespace mu
 		void Unserialise(InputArchive& arch);
 
 		//!
-		inline bool operator==(const Image& o) const
+		inline bool operator==(const Image& Other) const
 		{
-			return (m_lods == o.m_lods) &&
-				(m_size == o.m_size) &&
-				(m_format == o.m_format) &&
-				(m_data == o.m_data) &&
-				(ReferenceID == o.ReferenceID);
+			return 
+				(DataStorage == Other.DataStorage) &&
+				(ReferenceID == Other.ReferenceID);
 		}
 
 		//-----------------------------------------------------------------------------------------
@@ -271,25 +229,25 @@ namespace mu
 		//! Calculate the size of a lod of the image data in bytes, regardless of what is allocated
 		//! in m_data, only using the image descriptions. For non-block-compressed images, it
 		//! returns 0.
-		int32 CalculateDataSize(int lod) const;
+		//int32 CalculateDataSize(int lod) const;
 
-		//! Calculate the number of pixels of the image, regardless of what is allocated in
-		//! m_data, only using the image descriptions.
-		//! For block-compressed images, it includes the wasted pixels in the blocks, in case
-		//! the size is not a multiple of it in all lods.
-		//! It includes the pixels in all lods
-		int32 CalculatePixelCount() const;
+		////! Calculate the number of pixels of the image, regardless of what is allocated in
+		////! m_data, only using the image descriptions.
+		////! For block-compressed images, it includes the wasted pixels in the blocks, in case
+		////! the size is not a multiple of it in all lods.
+		////! It includes the pixels in all lods
+		//int32 CalculatePixelCount() const;
 
-		//! Same as above, but only calculates the pixels of an lod
-		int32 CalculatePixelCount(int lod) const;
+		////! Same as above, but only calculates the pixels of an lod
+		//int32 CalculatePixelCount(int32 LOD) const;
 
 		//! Calculate the size in pixels of a particular mipmap of this image. The size doesn't
 		//! include pixels necessary for completing blocks in block-compressed formats.
-		FIntVector2 CalculateMipSize(int lod) const;
+		FIntVector2 CalculateMipSize(int32 LOD) const;
 
 		//! Return a pointer to the beginning of the data for a particular mip.
-		uint8* GetMipData(int mip);
-		const uint8* GetMipData(int mip) const;
+		uint8* GetMipData(int32 Mip);
+		const uint8* GetMipData(int32 Mip) const;
 		
 		//! Return the size of the resident mips. for block-compressed images the result is the same as in CalculateDataSize()
 		//! for non-block-compressed images, the sizes encoded in the image data is used to compute the final size.
@@ -362,6 +320,7 @@ namespace mu
 		* \return false if the conversion failed, usually because not enough memory was allocated in the result. This is only checked for RLE compression.
 		*/
 		MUTABLERUNTIME_API void ImagePixelFormat(bool& bOutSuccess, int32 Quality, Image* Result, const Image* Base, int32 OnlyLOD = -1);
+		MUTABLERUNTIME_API void ImagePixelFormat(bool& bOutSuccess, int32 Quality, Image* Result, const Image* Base, int32 BeginResultLOD, int32 BeginBaseLOD, int32 NumLODs);
 
 		MUTABLERUNTIME_API Ptr<Image> ImageSwizzle(EImageFormat Format, const Ptr<const Image> Sources[], const uint8 Channels[]);
 
@@ -388,13 +347,13 @@ namespace mu
 		* from Base up to LevelCount and append them in Dest to the already generated Base's mips.
 		*/
 		void ImageMipmap(int32 CompressionQuality, Image* Dest, const Image* Base,
-			int32 LevelCount,
+			int32 StartLevel, int32 LevelCount,
 			const FMipmapGenerationSettings&, bool bGenerateOnlyTail = false);
 
 		/** Mipmap separating the worst case treatment in 3 steps to manage allocations of temp data. */
-		void ImageMipmap_PrepareScratch(Image* Dest, const Image* Base, int32 LevelCount, FScratchImageMipmap&);
+		void ImageMipmap_PrepareScratch(const Image* DestImage, int32 StartLevel, int32 LevelCount, FScratchImageMipmap&);
 		void ImageMipmap(FImageOperator::FScratchImageMipmap&, int32 CompressionQuality, Image* Dest, const Image* Base,
-			int32 LevelCount,
+			int32 LevelStart, int32 LevelCount,
 			const FMipmapGenerationSettings&, bool bGenerateOnlyTail = false);
 		void ImageMipmap_ReleaseScratch(FScratchImageMipmap&);
 
diff --git a/Engine/Plugins/Experimental/Mutable/Source/MutableRuntime/Public/MuR/ImageDataStorage.h b/Engine/Plugins/Experimental/Mutable/Source/MutableRuntime/Public/MuR/ImageDataStorage.h
new file mode 100644
--- /dev/null
+++ b/Engine/Plugins/Experimental/Mutable/Source/MutableRuntime/Public/MuR/ImageDataStorage.h
@@ -0,0 +1,403 @@
+// Copyright Epic Games, Inc. All Rights Reserved.
+
+#pragma once
+
+#include "MuR/ImageTypes.h"
+#include "MuR/MemoryTrackingAllocationPolicy.h"
+
+#include "Containers/Array.h"
+#include "Containers/ArrayView.h"
+#include "Containers/StaticArray.h"
+
+namespace mu::MemoryCounters
+{
+	struct FImageMemoryCounterTag {};
+	using FImageMemoryCounter = TMemoryCounter<FImageMemoryCounterTag>;
+}
+
+namespace mu
+{
+	using FImageArray = TArray<uint8, FDefaultMemoryTrackingAllocator<MemoryCounters::FImageMemoryCounter>>;
+}
+
+namespace mu::FImageDataStorageInternal
+{
+	FORCEINLINE int32 ComputeNumLODsForSize(FImageSize Size)
+	{
+		return static_cast<int32>(FMath::CeilLogTwo(FMath::Max(Size.X, Size.Y))) + 1;
+	}
+}
+
+namespace mu
+{
+class MUTABLERUNTIME_API FImageDataStorage
+{
+public:
+	TArray<FImageArray, TInlineAllocator<1>> Buffers;
+	
+	FImageSize ImageSize = FImageSize(0, 0);
+	EImageFormat ImageFormat = EImageFormat::IF_NONE;
+	uint8 NumLODs = 0;
+
+	static constexpr int32 NumLODsInCompactedTail = 7;
+	// This is needed for images its size cannot be known from dimensions and format, e.g., RLE compressed formats.
+	// It stores the offset to the end of the LOD so that CompactedTailOffsets[LOD] - CompactedTailOffsets[LOD - 1]
+	// is the size of LOD.
+	TStaticArray<uint16, NumLODsInCompactedTail> CompactedTailOffsets = MakeUniformStaticArray<uint16, NumLODsInCompactedTail>(0);
+
+public:
+
+	FImageDataStorage();
+	FImageDataStorage(const FImageDesc& Desc);
+
+	FImageDataStorage(const FImageDataStorage& Other);
+	FImageDataStorage& operator=(const FImageDataStorage& Other);
+	
+	FImageDataStorage(FImageDataStorage&& Other) = default;
+	FImageDataStorage& operator=(FImageDataStorage&& Other) = default;
+
+	bool operator==(const FImageDataStorage& Other) const;
+
+	FORCEINLINE FImageDesc MakeImageDesc() const
+	{
+		return FImageDesc(ImageSize, ImageFormat, NumLODs);
+	}
+
+	FORCEINLINE void InitVoid(const FImageDesc& Desc) 
+	{
+		Buffers.Empty();	
+		CompactedTailOffsets[0] = TNumericLimits<uint16>::Max();
+
+		ImageSize = Desc.m_size;
+		ImageFormat = Desc.m_format;
+		NumLODs = Desc.m_lods;
+	}
+
+	FORCEINLINE bool IsVoid() const
+	{
+		return Buffers.IsEmpty() && CompactedTailOffsets[0] == TNumericLimits<uint16>::Max();
+	}
+
+	void InitInternalArray(int32 Index, int32 Size, EInitializationType InitType);
+	void Init(const FImageDesc& ImageDesc, EInitializationType InitType);
+
+	void CopyFrom(const FImageDataStorage& Other);
+
+	FORCEINLINE int32 GetNumLODs() const
+	{
+		return NumLODs;
+	}
+
+	/**
+	* Get a const view to the data containing the LODIndex.
+	*/
+	FORCEINLINE TArrayView<const uint8> GetLOD(int32 LODIndex) const
+	{
+		check(!IsVoid());
+		if (LODIndex >= NumLODs)
+		{
+			//check(false);
+			return TArrayView<const uint8>();
+		}
+
+		const int32 FirstTailLOD = ComputeFirstCompactedTailLOD();
+		if (LODIndex >= FirstTailLOD)
+		{
+			const int32 LODIndexInTail = FMath::Max(0, LODIndex - FirstTailLOD);
+			
+			const int32 LODInTailBegin = LODIndexInTail == 0 ? 0 : CompactedTailOffsets[LODIndexInTail - 1];
+			const int32 LODInTailEnd = CompactedTailOffsets[LODIndexInTail];
+
+			const FImageArray& TailBuffer = Buffers.Last();
+			return MakeArrayView(TailBuffer.GetData() + LODInTailBegin, LODInTailEnd - LODInTailBegin);
+		}
+
+		return MakeArrayView(Buffers[LODIndex].GetData(), Buffers[LODIndex].Num());
+	}
+
+	/**
+	* Get a view to the data containing the LODIndex.
+	*/
+	FORCEINLINE TArrayView<uint8> GetLOD(int32 LODIndex)
+	{	
+		const TArrayView<const uint8> ConstLODView = const_cast<const FImageDataStorage*>(this)->GetLOD(LODIndex);
+		return TArrayView<uint8>(const_cast<uint8*>(ConstLODView.GetData()), ConstLODView.Num());
+	}
+
+	/**
+	 * Change the number of lods in the image. If InNumLODs is smaller than the current lod count, 
+	 * LODs are dropped from the tail up.
+	 */
+	void SetNumLODs(int32 InNumLODs, EInitializationType InitType = EInitializationType::NotInitialized);
+
+	/**
+	 * Remove NumLODsToDrop starting form LOD 0. The resulting image will be smaller. 
+	 */
+	void DropLODs(int32 NumLODsToDrop);
+
+	/**
+	 * Resizes LODIndex Data to NewSizeBytes
+	 */
+	void ResizeLOD(int32 LODIndex, int32 NewSizeBytes);
+
+	FORCEINLINE const FImageArray& GetInternalArray(int32 LODIndex) const
+	{	
+		check(!IsVoid());
+		check(LODIndex < NumLODs);
+		return Buffers[FMath::Min(LODIndex, ComputeFirstCompactedTailLOD())];
+	}
+
+	FORCEINLINE FImageArray& GetInternalArray(int32 LODIndex)
+	{	
+		check(!IsVoid());
+		check(LODIndex < NumLODs);
+		return Buffers[FMath::Min(LODIndex, ComputeFirstCompactedTailLOD())];
+	}
+
+	FORCEINLINE int32 ComputeFirstCompactedTailLOD() const
+	{
+		using namespace FImageDataStorageInternal;
+		return FMath::Max(0, ComputeNumLODsForSize(ImageSize) - NumLODsInCompactedTail); 
+	}
+
+	/**
+	 * Stateless iteration in batches.
+	 */
+	FORCEINLINE int32 GetNumBatches(int32 BatchSizeInElems, int32 BatchElemSizeInBytes) const
+	{	
+		return GetNumBatchesLODRange(BatchSizeInElems, BatchElemSizeInBytes, 0, NumLODs);
+	}
+
+	FORCEINLINE TArrayView<const uint8> GetBatch(int32 BatchId, int32 BatchSizeInElems, int32 BatchElemSizeInBytes) const
+	{
+		return GetBatchLODRange(BatchId, BatchSizeInElems, BatchElemSizeInBytes, 0, NumLODs);
+	}
+
+	FORCEINLINE TArrayView<uint8> GetBatch(int32 BatchId, int32 BatchSizeInElems, int32 BatchElemSizeInBytes)
+	{
+		// Use the const_cast idiom to generate the non-const operation. 
+		const TArrayView<const uint8> ConstBatchView = 
+				const_cast<const FImageDataStorage*>(this)->GetBatch(BatchId, BatchSizeInElems, BatchElemSizeInBytes);
+		return TArrayView<uint8>(const_cast<uint8*>(ConstBatchView.GetData()), ConstBatchView.Num());
+	}
+
+	/**
+	 * Returns the number of batches GetBatchFirtsLODOffet will admit and that will cover
+	 * the first LOD buffer with OffsetInBytes removed from the front. 
+	 * 
+	 * A batch cannot be larger than BatchElemsSize. The last batch of a buffer will be smaller if 
+	 * BatchSizeInElems*BatchElemSizeInBytes is not multiple of the buffer size.
+	 *
+	 **/
+	FORCEINLINE int32 GetNumBatchesFirstLODOffset(int32 BatchSizeInElems, int32 BatchElemSizeInBytes, int32 OffsetInBytes) const
+	{
+		check(!IsVoid());
+		TArrayView<const uint8> FirstLODView = GetLOD(0);
+		
+		check(OffsetInBytes % BatchElemSizeInBytes == 0);
+		check(FirstLODView.Num() % BatchElemSizeInBytes == 0);
+		check(FirstLODView.Num() > OffsetInBytes); 
+
+		return FMath::DivideAndRoundUp(FirstLODView.Num() - OffsetInBytes, BatchSizeInElems*BatchElemSizeInBytes);
+	}
+
+	/**
+	 * Returns a non modifiable view to the portion of the first LOD buffer with OffsetInBytes for the BatchId 
+	 *
+	 * A batch cannot be larger than BatchElemsSize. The last batch of a buffer will be smaller if 
+	 * BatchSizeInElems*BatchElemSizeInBytes is not multiple of the buffer size.
+	 *
+	 */
+	FORCEINLINE TArrayView<const uint8> GetBatchFirstLODOffset(int32 BatchId, int32 BatchSizeInElems, int32 BatchElemSizeInBytes, int32 OffsetInBytes = 0) const
+	{
+		check(!IsVoid());
+		TArrayView<const uint8> FirstLODView = GetLOD(0);
+		
+		check(OffsetInBytes < FirstLODView.Num());
+		check(FirstLODView.Num() % BatchElemSizeInBytes == 0);
+		check(OffsetInBytes % BatchElemSizeInBytes == 0);
+	
+		TArrayView<const uint8> OffsetedLODView = MakeArrayView(FirstLODView.GetData() + OffsetInBytes, FirstLODView.Num() - OffsetInBytes);
+
+		const int32 BatchOffset = BatchId * BatchSizeInElems * BatchElemSizeInBytes;
+
+		if (OffsetedLODView.Num() > BatchOffset)
+		{
+			return TArrayView<const uint8>();
+		}
+
+		return MakeArrayView(OffsetedLODView.GetData() + BatchOffset, FMath::Min(BatchSizeInElems, OffsetedLODView.Num() - BatchOffset));
+	}
+
+	/**
+	 * Non const version of GetBatchFirstLODOffset(). 
+	 */
+	FORCEINLINE TArrayView<uint8> GetBatchFirstLODOffset(int32 BatchId, int32 BatchSizeInElems, int32 BatchElemSizeInBytes, int32 OffsetInBytes = 0)
+	{
+		const TArrayView<const uint8> ConstBatchView = 
+				const_cast<const FImageDataStorage*>(this)->GetBatchFirstLODOffset(BatchId, BatchSizeInElems, BatchElemSizeInBytes, OffsetInBytes);
+
+		return TArrayView<uint8>(const_cast<uint8*>(ConstBatchView.GetData()), ConstBatchView.Num());
+	}
+
+	/**
+	 * Returns the number of batches GetBatchLODRange will admit and that will cover
+	 * the [LODBegin, LODEnd) range.
+	 *
+	 **/
+	FORCEINLINE int32 GetNumBatchesLODRange(int32 BatchSizeInElems, int32 BatchElemSizeInBytes, int32 LODBegin, int32 LODEnd) const
+	{
+		check(!IsVoid());
+		check(LODBegin < LODEnd);
+		check(LODBegin >= 0 && LODEnd <= NumLODs);
+
+		const int32 BatchNumBytes = BatchSizeInElems*BatchElemSizeInBytes;
+		
+		const int32 FirstCompactedTailLOD = ComputeFirstCompactedTailLOD();
+		const int32 NumLODsInTail = FMath::Max(0, NumLODs - FirstCompactedTailLOD);
+		const int32 NumLODsNotInTail = FMath::Max(0, NumLODs - NumLODsInTail);
+
+		const int32 LODsNotInTailLODRangeEnd = FMath::Min(LODEnd, NumLODsNotInTail);
+
+		int32 NumBatches = 0;
+		for (int32 BufferIndex = LODBegin; BufferIndex < LODsNotInTailLODRangeEnd; ++BufferIndex)
+		{	
+			NumBatches += FMath::DivideAndRoundUp(Buffers[BufferIndex].Num(), BatchNumBytes);
+		}
+
+		if (FirstCompactedTailLOD < LODEnd)
+		{
+			TArrayView<const uint8> FirstLODInRangeView = GetLOD(FirstCompactedTailLOD);
+			TArrayView<const uint8> LastLODInRangeView = GetLOD(LODEnd - 1);
+
+			const int32 TailInLODRangeNumBytes = LastLODInRangeView.GetData() - FirstLODInRangeView.GetData() + LastLODInRangeView.Num();
+			NumBatches += FMath::DivideAndRoundUp(TailInLODRangeNumBytes, BatchNumBytes);
+		}
+		return NumBatches;
+	}
+
+	
+	/**
+	 * Returns a non modifiable view to the portion of the [LODBegin, LODEnd) for the BatchId. 
+	 *
+	 * A batch cannot be larger than BatchElemsSize. The last batch of a buffer will be smaller if 
+	 * BatchSizeInElems*BatchElemSizeInBytes is not multiple of the buffer size.
+	 *
+	 */
+	FORCEINLINE TArrayView<const uint8> GetBatchLODRange(int32 BatchId, int32 BatchSizeInElems, int32 BatchElemSizeInBytes, int32 LODBegin, int32 LODEnd) const
+	{
+		check(!IsVoid());
+		check(LODBegin < LODEnd);
+		check(LODBegin >= 0 && LODEnd <= NumLODs);
+
+		const int32 BatchNumBytes = BatchSizeInElems*BatchElemSizeInBytes;
+
+		const int32 FirstCompactedTailLOD = ComputeFirstCompactedTailLOD();
+		const int32 NonTailBuffersEnd = FMath::Min(LODEnd, FirstCompactedTailLOD);
+
+		int32 NumBatches = 0;	
+		for (int32 BufferIndex = LODBegin; BufferIndex < NonTailBuffersEnd; ++BufferIndex)
+		{
+			const TArrayView<const uint8> BufferView = MakeArrayView(Buffers[BufferIndex].GetData(), Buffers[BufferIndex].Num());
+			const int32 BufferNumBatches = FMath::DivideAndRoundUp(BufferView.Num(), BatchNumBytes); 
+			
+			check(BatchId >= NumBatches);
+			int32 BufferBatchId = BatchId - NumBatches;
+
+			if (BufferBatchId < BufferNumBatches)
+			{ 
+				const int32 BufferBatchBeginOffset = BufferBatchId * BatchNumBytes;
+
+				return TArrayView<const uint8>(
+						BufferView.GetData() + BufferBatchBeginOffset,
+						FMath::Min(BufferView.Num() - BufferBatchBeginOffset, BatchNumBytes));
+			}
+
+			NumBatches += BufferNumBatches;
+		}
+
+		if (FirstCompactedTailLOD < LODEnd)
+		{
+			TArrayView<const uint8> FirstLODInRangeView = GetLOD(FirstCompactedTailLOD);
+			TArrayView<const uint8> LastLODInRangeView = GetLOD(LODEnd - 1);
+			
+			const int32 TailInLODRangeNumBytes = LastLODInRangeView.GetData() - FirstLODInRangeView.GetData() + LastLODInRangeView.Num();
+			TArrayView<const uint8> TailInLODRangeView = MakeArrayView(FirstLODInRangeView.GetData(), TailInLODRangeNumBytes);
+
+			check(BatchId >= NumBatches);
+			const int32 TailInLODRangeBatchId = BatchId - NumBatches;
+			const int32 TailInLODRangeNumBatches = FMath::DivideAndRoundUp(TailInLODRangeView.Num(), BatchNumBytes);
+
+			if (TailInLODRangeBatchId < TailInLODRangeNumBatches)
+			{
+				const int32 BufferBatchBeginOffset = TailInLODRangeBatchId * BatchNumBytes;
+				
+				return TArrayView<const uint8>(
+							TailInLODRangeView.GetData() + BufferBatchBeginOffset,
+							FMath::Min(TailInLODRangeView.Num() - BufferBatchBeginOffset, BatchNumBytes));
+			}
+		}
+		// If the BatchId is not found, return an empty view.
+		return TArrayView<const uint8>();
+	}
+
+	/**
+	 * Non const version of GetBatchLODRange().
+	 */
+	FORCEINLINE TArrayView<uint8> GetBatchLODRange(int32 BatchId, int32 BatchSizeInElems, int32 BatchElemSizeInBytes, int32 LODBegin, int32 LODEnd)
+	{
+		const TArrayView<const uint8> ConstBatchView = 
+			const_cast<const FImageDataStorage*>(this)->GetBatchLODRange(
+					BatchId, BatchSizeInElems, BatchElemSizeInBytes, LODBegin, LODEnd);
+
+		return TArrayView<uint8>(const_cast<uint8*>(ConstBatchView.GetData()), ConstBatchView.Num());
+	}
+
+	FORCEINLINE int32 GetAllocatedSize() const
+	{ 
+		SIZE_T Result = 0;
+		for (const FImageArray& Buffer : Buffers)
+		{
+			Result += Buffer.GetAllocatedSize();
+		}
+
+		check(Result < TNumericLimits<int32>::Max())
+		return (int32)Result;
+	}
+
+	FORCEINLINE int32 GetDataSize() const
+	{
+		SIZE_T Result = 0;
+		for (const FImageArray& Buffer : Buffers)
+		{
+			Result += Buffer.Num();
+		}
+
+		check(Result < TNumericLimits<int32>::Max())
+		return (int32)Result;
+	}
+
+	/**
+	 * Returns true if all buffers are empty. This can be true after initialization for image formats that report 
+	 * BytesPerBlock 0, e.g., RLE formats. If not initialized will also return true.
+	 */
+	FORCEINLINE bool IsEmpty() const
+	{
+		check(!IsVoid());
+
+		for (const FImageArray& Buffer : Buffers)
+		{
+			if (!Buffer.IsEmpty())
+			{
+				return false;
+			}
+		}
+
+		return true;
+	}
+
+	void Serialise(OutputArchive& Arch) const;
+	void Unserialise(InputArchive& Arch);
+};
+
+}
diff --git a/Engine/Plugins/Experimental/Mutable/Source/MutableRuntime/Public/MuR/ImageTypes.h b/Engine/Plugins/Experimental/Mutable/Source/MutableRuntime/Public/MuR/ImageTypes.h
--- a/Engine/Plugins/Experimental/Mutable/Source/MutableRuntime/Public/MuR/ImageTypes.h
+++ b/Engine/Plugins/Experimental/Mutable/Source/MutableRuntime/Public/MuR/ImageTypes.h
@@ -67,7 +67,9 @@ namespace mu
 
         IF_COUNT
 	};
-	
+	MUTABLE_DEFINE_ENUM_SERIALISABLE(EImageFormat);
+
+
 	struct MUTABLERUNTIME_API FImageDesc
 	{
 		FImageDesc()
@@ -113,8 +115,9 @@ namespace mu
 		BT_LIGHTEN,				// Increase the channel value by a given proportion of what is missing from white 
 		BT_NORMAL_COMBINE,
 		_BT_COUNT
-	};
-	
+	};	
+	MUTABLE_DEFINE_ENUM_SERIALISABLE(EBlendType);
+
 	enum class EAddressMode
 	{
 		None,
@@ -129,8 +132,9 @@ namespace mu
 		MFT_SimpleAverage,
 		MFT_Sharpen,
 		_MFT_COUNT
-	};
-	
+	};	
+	MUTABLE_DEFINE_ENUM_SERIALISABLE(EMipmapFilterType);
+
 	enum class ECompositeImageMode
 	{
 		CIM_Disabled,
@@ -139,8 +143,9 @@ namespace mu
 		CIM_NormalRoughnessToBlue,
 		CIM_NormalRoughnessToAlpha,
 		_CIM_COUNT
-	};
-	
+	};	
+	MUTABLE_DEFINE_ENUM_SERIALISABLE(ECompositeImageMode);
+
 	enum class ESamplingMethod : uint8
 	{
 		Point = 0,
@@ -148,6 +153,7 @@ namespace mu
 		MaxValue
 	};
 	static_assert(uint32(ESamplingMethod::MaxValue) <= (1 << 3), "ESampligMethod enum cannot hold more than 8 values");
+	MUTABLE_DEFINE_ENUM_SERIALISABLE(ESamplingMethod);
 
 	enum class EInitializationType
 	{
@@ -162,6 +168,7 @@ namespace mu
 		MaxValue
 	};
 	static_assert(uint32(EMinFilterMethod::MaxValue) <= (1 << 3), "EMinFilterMethod enum cannot hold more than 8 values");
+	MUTABLE_DEFINE_ENUM_SERIALISABLE(EMinFilterMethod);
 
 	/** */
 	struct FImageFormatData
diff --git a/Engine/Plugins/Experimental/Mutable/Source/MutableRuntime/Public/MuR/Serialisation.h b/Engine/Plugins/Experimental/Mutable/Source/MutableRuntime/Public/MuR/Serialisation.h
--- a/Engine/Plugins/Experimental/Mutable/Source/MutableRuntime/Public/MuR/Serialisation.h
+++ b/Engine/Plugins/Experimental/Mutable/Source/MutableRuntime/Public/MuR/Serialisation.h
@@ -20,26 +20,20 @@ namespace mu
     class Image;
 	class Model;
 
-#define MUTABLE_DEFINE_POD_SERIALISABLE(T)							\
-	template<>														\
-	void DLLEXPORT operator<< <T>(OutputArchive& arch, const T& t);	\
-																	\
-	template<>														\
-	void DLLEXPORT operator>> <T>(InputArchive& arch, T& t);		\
-
-#define MUTABLE_DEFINE_POD_VECTOR_SERIALISABLE(T)									     \
-	template<typename Alloc>                                                             \
-	void operator <<(OutputArchive& arch, const TArray<T, Alloc>& v);					 \
-	                                                                                     \
-	template<typename Alloc>                                                             \
-	void operator >>(InputArchive& arch, TArray<T, Alloc>& v);                           \
-
-#define MUTABLE_DEFINE_ENUM_SERIALISABLE(T)							\
-	template<>														\
-    void DLLEXPORT operator<< <T>(OutputArchive& arch, const T& t);	\
-																	\
-	template<>														\
-    void DLLEXPORT operator>> <T>(InputArchive& arch, T& t);		\
+#define MUTABLE_DEFINE_POD_SERIALISABLE(Type)					   \
+	void DLLEXPORT operator<<(OutputArchive& Arch, const Type& T); \
+	void DLLEXPORT operator>>(InputArchive& Arch, Type& T);		   \
+
+#define MUTABLE_DEFINE_POD_VECTOR_SERIALISABLE(Type)					\
+	template<typename Alloc>                                            \
+	void operator<<(OutputArchive& Arch, const TArray<Type, Alloc>& V);	\
+	                                                                    \
+	template<typename Alloc>                                            \
+	void operator>>(InputArchive& Arch, TArray<Type, Alloc>& V);        \
+
+#define MUTABLE_DEFINE_ENUM_SERIALISABLE(Type)					   \
+    void DLLEXPORT operator<<(OutputArchive& Arch, const Type& T); \
+    void DLLEXPORT operator>>(InputArchive& Arch, Type& T);		   \
 	
 	
     /** This repesents a handle to Mutable resource. */
@@ -327,16 +321,16 @@ namespace mu
     };
 
 
-	template< typename T >
-	void operator<< ( OutputArchive& arch, const T& t )
+	template<typename Type>
+	void operator<<(OutputArchive& Arch, const Type& Value)
 	{
-        t.Serialise( arch );
+        Value.Serialise(Arch);
 	}
 
-	template< typename T >
-	void operator>> ( InputArchive& arch, T& t )
+	template<typename Type>
+	void operator>>(InputArchive& Arch, Type& Value)
 	{
-        t.Unserialise( arch );
+        Value.Unserialise(Arch);
 	}
 
 	MUTABLE_DEFINE_POD_SERIALISABLE(float);
@@ -419,8 +413,8 @@ namespace mu
 	MUTABLERUNTIME_API void operator>> ( InputArchive& arch, std::string& v );
 
 	// Bool size is not a standard
-	MUTABLERUNTIME_API void operator<< ( OutputArchive& arch, const bool& t );
+	MUTABLERUNTIME_API void operator<<(OutputArchive& Arch, const bool& T);
 
-	MUTABLERUNTIME_API void operator>> ( InputArchive& arch, bool& t );
+	MUTABLERUNTIME_API void operator>>(InputArchive& Arch, bool& T);
 }
 
diff --git a/Engine/Plugins/Experimental/Mutable/Source/MutableTools/Internal/MuT/AST.h b/Engine/Plugins/Experimental/Mutable/Source/MutableTools/Internal/MuT/AST.h
--- a/Engine/Plugins/Experimental/Mutable/Source/MutableTools/Internal/MuT/AST.h
+++ b/Engine/Plugins/Experimental/Mutable/Source/MutableTools/Internal/MuT/AST.h
@@ -266,14 +266,11 @@ namespace mu
 			static FORCEINLINE uint32 GetKeyHash(KeyInitType Key)
 			{
 				const mu::Image* Data = Key.get();
-				uint32 Hash = HashCombineFast(
-					::GetTypeHash(Data->m_format),
-					HashCombineFast(::GetTypeHash(Data->m_size[0]),
-								::GetTypeHash(Data->m_size[1]))
-				);
+				uint32 Hash = HashCombineFast(::GetTypeHash(Data->GetFormat()), GetTypeHash(Data->GetSize()));
 
-				uint64 DataHash = CityHash64(reinterpret_cast<const char*>(Data->m_data.GetData()), Data->m_data.Num());
-				Hash = HashCombineFast( Hash, ::GetTypeHash(DataHash) );
+				TArrayView<const uint8> DataView = Data->DataStorage.GetLOD(0);
+				uint64 DataHash = CityHash64(reinterpret_cast<const char*>(DataView.GetData()), DataView.Num());
+				Hash = HashCombineFast(Hash, ::GetTypeHash(DataHash));
 
 				return Hash;
 			}
diff --git a/Engine/Plugins/Experimental/Mutable/Source/MutableTools/Private/MuT/ASTOpConstantResource.cpp b/Engine/Plugins/Experimental/Mutable/Source/MutableTools/Private/MuT/ASTOpConstantResource.cpp
--- a/Engine/Plugins/Experimental/Mutable/Source/MutableTools/Private/MuT/ASTOpConstantResource.cpp
+++ b/Engine/Plugins/Experimental/Mutable/Source/MutableTools/Private/MuT/ASTOpConstantResource.cpp
@@ -151,13 +151,6 @@ namespace mu
 					}
 					check(NewMip);
 
-					// Shrink the buffer to the minimum necessary size.
-					uint32 CalculatedSize = NewMip->CalculateDataSize();
-					if (CalculatedSize)
-					{
-						NewMip->m_data.SetNum(CalculatedSize);
-					}
-
 					pMip = NewMip;
 				}
 			}
@@ -173,7 +166,6 @@ namespace mu
 		}
 	}
 
-
 	//-------------------------------------------------------------------------------------------------
 	void ASTOpConstantResource::Link(FProgram& program, FLinkerOptions* Options)
 	{
@@ -283,25 +275,24 @@ namespace mu
 	//-------------------------------------------------------------------------------------------------
 	FImageDesc ASTOpConstantResource::GetImageDesc(bool, class FGetImageDescContext*) const
 	{
-		FImageDesc res;
+		FImageDesc Result;
 
 		if (Type == OP_TYPE::IM_CONSTANT)
 		{
 			// TODO: cache to avoid disk loading
-			Ptr<const Image> pConst = static_cast<const Image*>(GetValue().get());
-			res.m_format = pConst->m_format;
-			res.m_lods = pConst->m_lods;
-			res.m_size = pConst->m_size;
+			Ptr<const Image> ConstImage = static_cast<const Image*>(GetValue().get());
+			Result.m_format = ConstImage->GetFormat();
+			Result.m_lods = ConstImage->GetLODCount();
+			Result.m_size = ConstImage->GetSize();
 		}
 		else
 		{
 			check(false);
 		}
 
-		return res;
+		return Result;
 	}
 
-
 	//-------------------------------------------------------------------------------------------------
 	void ASTOpConstantResource::GetBlockLayoutSize(int blockIndex, int* pBlockX, int* pBlockY, FBlockLayoutSizeCache*)
 	{
@@ -380,7 +371,7 @@ namespace mu
 		case OP_TYPE::IM_CONSTANT:
 		{
 			Ptr<const Image> pImage = static_cast<const Image*>(GetValue().get());
-			if ( (pImage->m_size[0] <= 0) || (pImage->m_size[1] <= 0) )
+			if (pImage->GetSizeX() <= 0 || pImage->GetSizeY() <= 0)
 			{
 				res = true;
 				colour = FVector4f(0.0f,0.0f,0.0f,1.0f);
@@ -529,7 +520,7 @@ namespace mu
 			Ptr<ImageSizeExpression> pRes = new ImageSizeExpression;
 			pRes->type = ImageSizeExpression::ISET_CONSTANT;
 			Ptr<const Image> pConst = static_cast<const Image*>(GetValue().get());
-			pRes->size = pConst->m_size;
+			pRes->size = pConst->GetSize();
 			return pRes;
 		}
 
diff --git a/Engine/Plugins/Experimental/Mutable/Source/MutableTools/Private/MuT/CodeGenerator_Image.cpp b/Engine/Plugins/Experimental/Mutable/Source/MutableTools/Private/MuT/CodeGenerator_Image.cpp
--- a/Engine/Plugins/Experimental/Mutable/Source/MutableTools/Private/MuT/CodeGenerator_Image.cpp
+++ b/Engine/Plugins/Experimental/Mutable/Source/MutableTools/Private/MuT/CodeGenerator_Image.cpp
@@ -1913,118 +1913,114 @@ namespace mu
 	}
 
 
-    //---------------------------------------------------------------------------------------------
-    Ptr<Image> CodeGenerator::GenerateMissingImage( EImageFormat format )
-    {
-        // Create the image node if it hasn't been created yet.
-        if (!m_missingImage[size_t(format)])
-        {
-            // Make a checkered debug image
-            const FImageSize size(16,16);
+	Ptr<Image> CodeGenerator::GenerateMissingImage(EImageFormat Format)
+	{
+		// Create the image node if it hasn't been created yet.
+		if (!m_missingImage[SIZE_T(Format)])
+		{
+			// Make a checkered debug image
+			const FImageSize Size(16, 16);
 
-            ImagePtr pImage = new Image( size[0], size[1], 1, format, EInitializationType::NotInitialized);
+			Ptr<Image> GeneratedImage = new Image(Size[0], Size[1], 1, Format, EInitializationType::NotInitialized);
 
-            switch (format)
-            {
-            case EImageFormat::IF_L_UBYTE:
-                {
-                    uint8_t* pData = pImage->GetData();
-                    for ( int p=0; p<size[0]*size[1]; ++p )
-                    {
-                        if ((p+p/size[0])%2)
-                        {
-                            pData[0] = 255;
-                        }
-                        else
-                        {
-                            pData[0] = 64;
-                        }
-
-                        pData++;
-                    }
-                    break;
-                }
+			switch (Format)
+			{
+			case EImageFormat::IF_L_UBYTE:
+			{
+				uint8* DataPtr = GeneratedImage->GetLODData(0);
+				for (int32 P = 0; P < Size[0]*Size[1]; ++P)
+				{
+					if ((P + P/Size[0]) % 2)
+					{
+						DataPtr[0] = 255;
+					}
+					else
+					{
+						DataPtr[0] = 64;
+					}
 
-            case EImageFormat::IF_RGB_UBYTE:
-                {
-                    uint8_t* pData = pImage->GetData();
-                    for ( int p=0; p<size[0]*size[1]; ++p )
-                    {
-                        if ((p+p/size[0])%2)
-                        {
-                            pData[0] = 255;
-                            pData[1] = 255;
-                            pData[2] = 64;
-                        }
-                        else
-                        {
-                            pData[0] = 64;
-                            pData[1] = 64;
-                            pData[2] = 255;
-                        }
-
-                        pData += 3;
-                    }
-                    break;
-                }
+					DataPtr++;
+				}
+				break;
+			}
+			case EImageFormat::IF_RGB_UBYTE:
+			{
+				uint8* DataPtr = GeneratedImage->GetLODData(0);
+				for (int32 P = 0; P < Size[0]*Size[1]; ++P)
+				{
+					if ((P + P/Size[0]) % 2)
+					{
+						DataPtr[0] = 255;
+						DataPtr[1] = 255;
+						DataPtr[2] = 64;
+					}
+					else
+					{
+						DataPtr[0] = 64;
+						DataPtr[1] = 64;
+						DataPtr[2] = 255;
+					}
 
-            case EImageFormat::IF_BGRA_UBYTE:
-            case EImageFormat::IF_RGBA_UBYTE:
-                {
-                    uint8_t* pData = pImage->GetData();
-                    for ( int p=0; p<size[0]*size[1]; ++p )
-                    {
-                        if ((p+p/size[0])%2)
-                        {
-                            pData[0] = 255;
-                            pData[1] = 255;
-                            pData[2] = 64;
-                            pData[3] = 255;
-                        }
-                        else
-                        {
-                            pData[0] = 64;
-                            pData[1] = 64;
-                            pData[2] = 255;
-                            pData[3] = 128;
-                        }
-
-                        pData += 4;
-                    }
-                    break;
-                }
+					DataPtr += 3;
+				}
+				break;
+			}
+			case EImageFormat::IF_BGRA_UBYTE:
+			case EImageFormat::IF_RGBA_UBYTE:
+			{
+				uint8* DataPtr = GeneratedImage->GetLODData(0);
+				for (int32 P = 0; P < Size[0]*Size[1]; ++P)
+				{
+					if ((P + P/Size[0]) % 2)
+					{
+						DataPtr[0] = 255;
+						DataPtr[1] = 255;
+						DataPtr[2] = 64;
+						DataPtr[3] = 255;
+					}
+					else
+					{
+						DataPtr[0] = 64;
+						DataPtr[1] = 64;
+						DataPtr[2] = 255;
+						DataPtr[3] = 128;
+					}
+
+					DataPtr += 4;
+				}
+				break;
+			}
 
-            default:
-                check( false );
-                break;
+			default:
+				check( false );
+				break;
 
-            }
+			}
 
-            m_missingImage[(size_t)format] = pImage;
-        }
+			m_missingImage[(SIZE_T)Format] = GeneratedImage;
+		}
 
-        return m_missingImage[(size_t)format].get();
-    }
+		return m_missingImage[(SIZE_T)Format].get();
+	}
 
 
-    //---------------------------------------------------------------------------------------------
-    Ptr<ASTOp> CodeGenerator::GenerateMissingImageCode(const TCHAR* strWhere, EImageFormat format, const void* errorContext, const FImageGenerationOptions& Options )
-    {
-        // Log an error message
+	Ptr<ASTOp> CodeGenerator::GenerateMissingImageCode(const TCHAR* strWhere, EImageFormat format, const void* errorContext, const FImageGenerationOptions& Options )
+	{
+		// Log an error message
 		FString Msg = FString::Printf(TEXT("Required connection not found: %s"), strWhere );
-        m_pErrorLog->GetPrivate()->Add( Msg, ELMT_ERROR, errorContext );
+		m_pErrorLog->GetPrivate()->Add( Msg, ELMT_ERROR, errorContext );
 
-        // Make a checkered debug image
-        ImagePtr pImage = GenerateMissingImage( format );
+		// Make a checkered debug image
+		Ptr<Image> GeneratedImage = GenerateMissingImage( format );
 
-        NodeImageConstantPtr pNode = new NodeImageConstant();
-        pNode->SetValue( pImage.get() );
+		NodeImageConstantPtr pNode = new NodeImageConstant();
+		pNode->SetValue(GeneratedImage.get());
 
 		FImageGenerationResult Result;
 		GenerateImage(Options, Result, pNode);
 
-        return Result.op;
-    }
+		return Result.op;
+	}
 
 
     //---------------------------------------------------------------------------------------------
diff --git a/Engine/Plugins/Experimental/Mutable/Source/MutableTools/Private/MuT/UnrealPixelFormatOverride.cpp b/Engine/Plugins/Experimental/Mutable/Source/MutableTools/Private/MuT/UnrealPixelFormatOverride.cpp
--- a/Engine/Plugins/Experimental/Mutable/Source/MutableTools/Private/MuT/UnrealPixelFormatOverride.cpp
+++ b/Engine/Plugins/Experimental/Mutable/Source/MutableTools/Private/MuT/UnrealPixelFormatOverride.cpp
@@ -170,8 +170,10 @@ void MutableToImageCore(const mu::Image* InMutable, FImage& CoreImage, int32 LOD
 
 void ImageCoreToMutable(const FCompressedImage2D& Compressed, mu::Image* Mutable, int32 LOD)
 {
-	check(Compressed.RawData.Num()==Mutable->CalculateDataSize(LOD));
-	FMemory::Memcpy(Mutable->GetMipData(LOD), Compressed.RawData.GetData(), Compressed.RawData.Num());
+	TArrayView<uint8> MutableView = Mutable->DataStorage.GetLOD(LOD);
+
+	check(Compressed.RawData.Num() == MutableView.Num());
+	FMemory::Memcpy(MutableView.GetData(), Compressed.RawData.GetData(), Compressed.RawData.Num());
 }
 
 
