[Mutable] Add a helper function to calculate the FSkeletalMeshLODRenderData's BuffersSize.
- Add a helper function to initialize streamed LODs with dummy data.
- Simplify the flow when creating meshes by moving the code that copies data between mesh LODs to a separate for loop.
- Remove double initialization of weight buffers.

#rnx

diff --git a/Engine/Plugins/Experimental/Mutable/Source/CustomizableObject/Internal/MuCO/UnrealConversionUtils.h b/Engine/Plugins/Experimental/Mutable/Source/CustomizableObject/Internal/MuCO/UnrealConversionUtils.h
--- a/Engine/Plugins/Experimental/Mutable/Source/CustomizableObject/Internal/MuCO/UnrealConversionUtils.h
+++ b/Engine/Plugins/Experimental/Mutable/Source/CustomizableObject/Internal/MuCO/UnrealConversionUtils.h
@@ -32,23 +32,22 @@ namespace UnrealConversionUtils
 	
 	/**
 	 * Prepares the render sections found on the InSkeletalMesh and sets them up accordingly what the InMutableMesh requires
-	 * @param OutSkeletalMesh - The Skeletal mesh whose sections are ought to be updated
+	 * @param LODResource - LODRenderData whose sections are ought to be updated
 	 * @param InMutableMesh - Mutable mesh to be used as reference for the section data update on the Skeletal Mesh
-	 * @param MeshLODIndex - The Mesh lod to be working with
-	 * @param InBoneMap - The bones to be set as part of the sections.
+	 * @param InBoneMap - Bones to be set as part of the sections.
+	 * @param InFirstBoneMapIndex - Index to the first BoneMap bone that belongs to this LODResource.
 	 */
 	CUSTOMIZABLEOBJECT_API void SetupRenderSections(
 		FSkeletalMeshLODRenderData& LODResource,
 		const mu::Ptr<const mu::Mesh> InMutableMesh,
-		const int32 InMeshLODIndex,
 		const TArray<uint16>& InBoneMap,
 		const int32 InFirstBoneMapIndex);
 
 
 	/** Performs a copy of the data found on the vertex buffers on the mutable mesh to the buffers of the skeletal mesh
-	 * @param OutSkeletalMesh - Skeletal Mesh to be updated with new buffer data
+	 * @param LODResource - LODRenderData to update
 	 * @param InMutableMesh - Mutable mesh to be used as reference for the section data update on the Skeletal Mesh
-	 * @param MeshLODIndex - The LOD index we are working with.
+	 * @param bAllowCPUAccess - Keeps this LODs data on the CPU so it can be used for things such as sampling in FX.
 	 */
 	CUSTOMIZABLEOBJECT_API void CopyMutableVertexBuffers(
 		FSkeletalMeshLODRenderData& LODResource,
@@ -58,8 +57,8 @@ namespace UnrealConversionUtils
 	
 	/**
 	 *Performs a copy of the data found on the index buffers on the mutable mesh to the buffers of the skeletal mesh
+	 * @param LODResource - LODRenderData to be updated.
 	 * @param InMutableMesh - The mutable mesh whose index buffers you want to work with
-	 * @param OutLODModel - The LOD model to be updated.
 	 * @return True if the operation could be performed successfully, false if not.
 	 */
 	CUSTOMIZABLEOBJECT_API bool CopyMutableIndexBuffers(
@@ -69,7 +68,7 @@ namespace UnrealConversionUtils
 
 	/**
 	 *Performs a copy of the data found on the index buffers on the mutable mesh to the buffers of the skeletal mesh
-	 * @param OutLODModel - The LOD model to be updated.
+	 * @param LODResource - LODRenderData to be updated.
 	 * @param InProfileName - Name of the profile to generate.
 	 * @param InMutableMeshVertexBuffers - The mutable buffers to be reading data from
 	 * @param InBoneIndexBuffer - The buffer containing the indices for the bones.
@@ -84,17 +83,25 @@ namespace UnrealConversionUtils
 
 	/**
 	 *Performs a copy of the render data of a specific Skeletal Mesh LOD to another Skeletal Mesh
-	 * @param SrcSkeletalMesh - Skeletal Mesh with the data to copy.
-	 * @param DestSkeletalMesh - Skeletal Mesh where the data will be copied.
-	 * @param SrcLODIndex - Index of the LOD to copy.
-	 * @param DestSkeletalMesh - Index of the LOD that will be copyed to.
+	 * @param LODResource - LODRenderData to copy to.
+	 * @param SourceLODResource - LODRenderData to copy from.
+	 * @param SkeletalMesh - Owner of both LODResource and SourceLODResource. 
+	 * @param bAllowCPUAccess - Keeps this LODs data on the CPU so it can be used for things such as sampling in FX.
 	 */
 	CUSTOMIZABLEOBJECT_API void CopySkeletalMeshLODRenderData(
-		const USkeletalMesh* SrcSkeletalMesh,
-		USkeletalMesh* DestSkeletalMesh,
-		int32 SrcLODIndex,
-		int32 DestLODIndex);
+		FSkeletalMeshLODRenderData& LODResource,
+		FSkeletalMeshLODRenderData& SourceLODResource,
+		const USkeletalMesh& SkeletalMesh,
+		const bool bAllowCPUAccess
+	);
 
+	/**
+	 * Update SkeletalMeshLODRenderData buffers size.
+	 * @param LODResource - LODRenderData to be updated.
+	 */
+	CUSTOMIZABLEOBJECT_API void UpdateSkeletalMeshLODRenderDataBuffersSize(
+		FSkeletalMeshLODRenderData& LODResource
+	);
 }
 
 #if UE_ENABLE_INCLUDE_ORDER_DEPRECATED_IN_5_2
diff --git a/Engine/Plugins/Experimental/Mutable/Source/CustomizableObject/Private/MuCO/CustomizableObjectInstance.cpp b/Engine/Plugins/Experimental/Mutable/Source/CustomizableObject/Private/MuCO/CustomizableObjectInstance.cpp
--- a/Engine/Plugins/Experimental/Mutable/Source/CustomizableObject/Private/MuCO/CustomizableObjectInstance.cpp
+++ b/Engine/Plugins/Experimental/Mutable/Source/CustomizableObject/Private/MuCO/CustomizableObjectInstance.cpp
@@ -3527,8 +3527,8 @@ void UCustomizableInstancePrivate::BuildOrCopyMorphTargetsData(const TSharedRef<
 	const UCustomizableObject* CustomizableObject = CustomizableObjectInstance->GetCustomizableObject();
 	const TArray<FName>& MorphTargetNames = CustomizableObject->GetPrivate()->GetModelResources().RealTimeMorphTargetNames;
 
-	const TMap<uint32, TArray<FMorphTargetVertexData>>& ResourceIdToVertexDataMap = 
-			OperationData->InstanceUpdateData.MorphTargetsVertexData;
+	const TMap<uint32, TArray<FMorphTargetVertexData>>& ResourceIdToVertexDataMap =
+		OperationData->InstanceUpdateData.MorphTargetsVertexData;
 
 	if (MorphTargetNames.IsEmpty() || ResourceIdToVertexDataMap.IsEmpty())
 	{
@@ -3545,211 +3545,204 @@ void UCustomizableInstancePrivate::BuildOrCopyMorphTargetsData(const TSharedRef<
 	TArray<TArray<FMorphTargetLODModel>> MorphsData;
 	MorphsData.SetNum(MorphTargetNames.Num());
 
-	const int32 NumLODs = OperationData->NumLODsAvailable;
+	TArray<TObjectPtr<UMorphTarget>>& MorphTargets = SkeletalMesh->GetMorphTargets();
 
-	int32 LastValidLODIndex = OperationData->NumLODsAvailable - 1;
-	for (int32 LODIndex = LastValidLODIndex; LODIndex >= FirstLODAvailable; --LODIndex)
+	const int32 FirstGeneratedLOD = FMath::Max((int32)OperationData->GetRequestedLODs()[ComponentIndex], OperationData->GetMinLOD());
+	for (int32 LODIndex = FirstGeneratedLOD; LODIndex < OperationData->NumLODsAvailable; ++LODIndex)
 	{
+		MUTABLE_CPUPROFILER_SCOPE(BuildMorphTargetsData);
+
 		int32 NumNotFoundLoadedMorphsResources = 0;
 
 		const FInstanceUpdateData::FLOD& LOD = OperationData->InstanceUpdateData.LODs[LODIndex];
+		const FInstanceUpdateData::FComponent& Component = OperationData->InstanceUpdateData.Components[LOD.FirstComponent + ComponentIndex];
+		check(Component.bGenerated);
+		check(Component.Mesh);
 
-		if (LODIndex >= OperationData->GetMinLOD() && OperationData->InstanceUpdateData.Components[LOD.FirstComponent + ComponentIndex].bGenerated)
-		{
-			LastValidLODIndex = LODIndex;
-		}
+		const mu::FMeshBufferSet& MeshSet = Component.Mesh->GetVertexBuffers();
 
-		const FInstanceUpdateData::FComponent& Component = OperationData->InstanceUpdateData.Components[OperationData->InstanceUpdateData.LODs[LastValidLODIndex].FirstComponent + ComponentIndex];
+		int32 VertexMorphsInfoIndexBufferIndex, VertexMorphsInfoIndexBufferChannel;
+		MeshSet.FindChannel(mu::MBS_OTHER, 0, &VertexMorphsInfoIndexBufferIndex, &VertexMorphsInfoIndexBufferChannel);
 
-		if (Component.Mesh)
-		{
-			MUTABLE_CPUPROFILER_SCOPE(BuildOrCopyMorphTargetsData_Build);
+		int32 VertexMorphsCountBufferIndex, VertexMorphsCountBufferChannel;
+		MeshSet.FindChannel(mu::MBS_OTHER, 1, &VertexMorphsCountBufferIndex, &VertexMorphsCountBufferChannel);
 
-			const mu::FMeshBufferSet& MeshSet = Component.Mesh->GetVertexBuffers();
+		int32 VertexMorphsResourceIdBufferIndex, VertexMorphsResourceIdBufferChannel;
+		MeshSet.FindChannel(mu::MBS_OTHER, 2, &VertexMorphsResourceIdBufferIndex, &VertexMorphsResourceIdBufferChannel);
 
-			int32 VertexMorphsInfoIndexBufferIndex, VertexMorphsInfoIndexBufferChannel;
-			MeshSet.FindChannel(mu::MBS_OTHER, 0, &VertexMorphsInfoIndexBufferIndex, &VertexMorphsInfoIndexBufferChannel);
+		if (VertexMorphsInfoIndexBufferIndex < 0 || VertexMorphsCountBufferIndex < 0 || VertexMorphsResourceIdBufferIndex < 0)
+		{
+			continue;
+		}
+
+		const int32* const VertexMorphsInfoIndexBuffer = reinterpret_cast<const int32*>(MeshSet.GetBufferData(VertexMorphsInfoIndexBufferIndex));
+		TArrayView<const int32> VertexMorphsInfoIndexView(VertexMorphsInfoIndexBuffer, MeshSet.GetElementCount());
 
-			int32 VertexMorphsCountBufferIndex, VertexMorphsCountBufferChannel;
-			MeshSet.FindChannel(mu::MBS_OTHER, 1, &VertexMorphsCountBufferIndex, &VertexMorphsCountBufferChannel);
+		const uint16* const VertexMorphsCountBuffer = reinterpret_cast<const uint16*>(MeshSet.GetBufferData(VertexMorphsCountBufferIndex));
+		TArrayView<const uint16> VertexMorphsCountView(VertexMorphsCountBuffer, MeshSet.GetElementCount());
 
-			int32 VertexMorphsResourceIdBufferIndex, VertexMorphsResourceIdBufferChannel;
-			MeshSet.FindChannel(mu::MBS_OTHER, 2, &VertexMorphsResourceIdBufferIndex, &VertexMorphsResourceIdBufferChannel);
+		const uint16* const VertexMorphsResourceIdBuffer = reinterpret_cast<const uint16*>(MeshSet.GetBufferData(VertexMorphsResourceIdBufferIndex));
+		TArrayView<const uint16> VertexMorphsResourceIdView(VertexMorphsResourceIdBuffer, MeshSet.GetElementCount());
 
-			if (VertexMorphsInfoIndexBufferIndex < 0 || VertexMorphsCountBufferIndex < 0 || VertexMorphsResourceIdBufferIndex < 0)
+		const int32 SurfaceCount = Component.Mesh->GetSurfaceCount();
+		for (int32 Section = 0; Section < SurfaceCount; ++Section)
+		{
+			// Reset SectionMorphTargets.
+			for (int32& Elem : SectionMorphTargetVertices)
 			{
-				continue;
+				Elem = 0;
 			}
 
-			const int32* const VertexMorphsInfoIndexBuffer = reinterpret_cast<const int32*>(MeshSet.GetBufferData(VertexMorphsInfoIndexBufferIndex));
-			TArrayView<const int32> VertexMorphsInfoIndexView(VertexMorphsInfoIndexBuffer, MeshSet.GetElementCount());
+			int32 FirstVertex, VerticesCount, FirstIndex, IndiciesCount;
+			Component.Mesh->GetSurface(Section, &FirstVertex, &VerticesCount, &FirstIndex, &IndiciesCount, nullptr, nullptr, nullptr);
 
-			const uint16* const VertexMorphsCountBuffer = reinterpret_cast<const uint16*>(MeshSet.GetBufferData(VertexMorphsCountBufferIndex));
-			TArrayView<const uint16> VertexMorphsCountView(VertexMorphsCountBuffer, MeshSet.GetElementCount());
-			
-			const uint16* const VertexMorphsResourceIdBuffer = reinterpret_cast<const uint16*>(MeshSet.GetBufferData(VertexMorphsResourceIdBufferIndex));
-			TArrayView<const uint16> VertexMorphsResourceIdView(VertexMorphsResourceIdBuffer, MeshSet.GetElementCount());
-
-			const int32 SurfaceCount = Component.Mesh->GetSurfaceCount();
-			for (int32 Section = 0; Section < SurfaceCount; ++Section)
+			for (int32 VertexIdx = FirstVertex; VertexIdx < FirstVertex + VerticesCount;)
 			{
-				// Reset SectionMorphTargets.
-				for (int32& Elem : SectionMorphTargetVertices)
+				// Find a span with the same VertexMorphResourceId to amortise the cost of finding 
+				// in the loaded resources map. It is expected to find large consecutive mesh sections pointing to
+				// the same loaded resource.
+
+				const int32 SpanStart = VertexIdx++;
+				const uint16 CurrentResourceId = VertexMorphsResourceIdView[SpanStart];
+
+				// Vertex with no morphs are marked with TNumericLimits<uint16>::Max(), skip vertex if the case.
+				if (CurrentResourceId == TNumericLimits<uint16>::Max())
 				{
-					Elem = 0;
+					continue;
 				}
 
-				int32 FirstVertex, VerticesCount, FirstIndex, IndiciesCount;
-				Component.Mesh->GetSurface(Section, &FirstVertex, &VerticesCount, &FirstIndex, &IndiciesCount, nullptr, nullptr, nullptr);
-
-				for (int32 VertexIdx = FirstVertex; VertexIdx < FirstVertex + VerticesCount;)
+				for (; VertexIdx < FirstVertex + VerticesCount; ++VertexIdx)
 				{
-					// Find a span with the same VertexMorphResourceId to amortise the cost of finding 
-					// in the loaded resources map. It is expected to find large consecutive mesh sections pointing to
-					// the same loaded resource.
-					
-					const int32 SpanStart = VertexIdx++;
-					const uint16 CurrentResourceId = VertexMorphsResourceIdView[SpanStart];
-					
-					// Vertex with no morphs are marked with TNumericLimits<uint16>::Max(), skip vertex if the case.
-					if (CurrentResourceId == TNumericLimits<uint16>::Max())
+					const int32 VertexResourceId = VertexMorphsResourceIdView[VertexIdx];
+					// we can skip vertices with no morph without breaking the span.
+					if (VertexResourceId == TNumericLimits<uint16>::Max())
 					{
 						continue;
 					}
 
-					for (; VertexIdx < FirstVertex + VerticesCount; ++VertexIdx)
+					if (CurrentResourceId != VertexResourceId)
 					{
-						const int32 VertexResourceId = VertexMorphsResourceIdView[VertexIdx];
-						// we can skip vertices with no morph without breaking the span.
-						if (VertexResourceId == TNumericLimits<uint16>::Max())
-						{
-							continue;
-						}
-
-						if (CurrentResourceId != VertexResourceId)
-						{
-							break;
-						}
+						break;
 					}
-					const int32 SpanEnd = VertexIdx;
+				}
+				const int32 SpanEnd = VertexIdx;
 
-					const TArray<FMorphTargetVertexData>* MorphTargetReconstructionData = ResourceIdToVertexDataMap.Find(CurrentResourceId); 
-					
-					if (!MorphTargetReconstructionData)
-					{
-						++NumNotFoundLoadedMorphsResources;
-						continue;
-					}
+				const TArray<FMorphTargetVertexData>* MorphTargetReconstructionData = ResourceIdToVertexDataMap.Find(CurrentResourceId);
 
-					const TArray<FMorphTargetVertexData>& SpanMorphData = *MorphTargetReconstructionData;
+				if (!MorphTargetReconstructionData)
+				{
+					++NumNotFoundLoadedMorphsResources;
+					continue;
+				}
 
+				const TArray<FMorphTargetVertexData>& SpanMorphData = *MorphTargetReconstructionData;
 
-					// This assumes the number of vertex in an span will be large compared to the number of
-					// morphs. Maybe the allocation could be done in a different pass as an optimization.
-					FMemory::Memzero(UsedMorphSet.GetData(), UsedMorphSet.Num());
-					for (int32 SpanVertexIdx = SpanStart; SpanVertexIdx < SpanEnd; ++SpanVertexIdx)
-					{
-						const int32 MorphCount = VertexMorphsCountView[SpanVertexIdx];
 
-						TArrayView<const FMorphTargetVertexData> MorphsVertexDataView = MakeArrayView(
-								SpanMorphData.GetData() + VertexMorphsInfoIndexView[SpanVertexIdx], 
-								MorphCount);
+				// This assumes the number of vertex in an span will be large compared to the number of
+				// morphs. Maybe the allocation could be done in a different pass as an optimization.
+				FMemory::Memzero(UsedMorphSet.GetData(), UsedMorphSet.Num());
+				for (int32 SpanVertexIdx = SpanStart; SpanVertexIdx < SpanEnd; ++SpanVertexIdx)
+				{
+					const int32 MorphCount = VertexMorphsCountView[SpanVertexIdx];
 
-						for (const FMorphTargetVertexData& MorphVertexData : MorphsVertexDataView)
-						{
-							check((uint32)UsedMorphSet.Num() > MorphVertexData.MorphNameIndex);
-							UsedMorphSet[MorphVertexData.MorphNameIndex] = 1;
-						}
-					}
+					TArrayView<const FMorphTargetVertexData> MorphsVertexDataView = MakeArrayView(
+						SpanMorphData.GetData() + VertexMorphsInfoIndexView[SpanVertexIdx],
+						MorphCount);
 
-					const int32 NumMorphs = UsedMorphSet.Num();
-					for (int32 MorphIndex = 0; MorphIndex < NumMorphs; ++MorphIndex)
+					for (const FMorphTargetVertexData& MorphVertexData : MorphsVertexDataView)
 					{
-						if (!UsedMorphSet[MorphIndex])
-						{
-							continue;
-						}
-
-						if (MorphsData[MorphIndex].IsEmpty())
-						{
-							MorphsData[MorphIndex].SetNum(NumLODs);
-						}
+						check((uint32)UsedMorphSet.Num() > MorphVertexData.MorphNameIndex);
+						UsedMorphSet[MorphVertexData.MorphNameIndex] = 1;
 					}
+				}
 
-					for (int32 SpanVertexIdx = SpanStart; SpanVertexIdx < SpanEnd; ++SpanVertexIdx)
+				const int32 NumMorphs = UsedMorphSet.Num();
+				for (int32 MorphIndex = 0; MorphIndex < NumMorphs; ++MorphIndex)
+				{
+					if (!UsedMorphSet[MorphIndex])
 					{
-						const uint16 MorphCount = VertexMorphsCountView[SpanVertexIdx];
-						if (MorphCount == 0)
-						{
-							continue;
-						}
-
-						TArrayView<const FMorphTargetVertexData> MorphsVertexDataView = MakeArrayView(
-								SpanMorphData.GetData() + VertexMorphsInfoIndexView[SpanVertexIdx], 
-								MorphCount);
-
-						for (const FMorphTargetVertexData& SourceVertex : MorphsVertexDataView)
-						{
-							FMorphTargetLODModel& DestMorphLODModel = MorphsData[SourceVertex.MorphNameIndex][LODIndex];
-
-							DestMorphLODModel.Vertices.Emplace(
-									FMorphTargetDelta{ SourceVertex.PositionDelta, SourceVertex.TangentZDelta, static_cast<uint32>(SpanVertexIdx) });
+						continue;
+					}
 
-							++SectionMorphTargetVertices[SourceVertex.MorphNameIndex];
-						}
+					if (MorphsData[MorphIndex].IsEmpty())
+					{
+						MorphsData[MorphIndex].SetNum(OperationData->NumLODsAvailable);
 					}
 				}
 
-				const int32 SectionMorphTargetsNum = SectionMorphTargetVertices.Num();
-				for (int32 MorphIdx = 0; MorphIdx < SectionMorphTargetsNum; ++MorphIdx)
+				for (int32 SpanVertexIdx = SpanStart; SpanVertexIdx < SpanEnd; ++SpanVertexIdx)
 				{
-					if (SectionMorphTargetVertices[MorphIdx] > 0)
+					const uint16 MorphCount = VertexMorphsCountView[SpanVertexIdx];
+					if (MorphCount == 0)
+					{
+						continue;
+					}
+
+					TArrayView<const FMorphTargetVertexData> MorphsVertexDataView = MakeArrayView(
+						SpanMorphData.GetData() + VertexMorphsInfoIndexView[SpanVertexIdx],
+						MorphCount);
+
+					for (const FMorphTargetVertexData& SourceVertex : MorphsVertexDataView)
 					{
-						FMorphTargetLODModel& MorphTargetLodModel = MorphsData[MorphIdx][LODIndex];
+						FMorphTargetLODModel& DestMorphLODModel = MorphsData[SourceVertex.MorphNameIndex][LODIndex];
 
-						MorphTargetLodModel.SectionIndices.Add(Section);
-						MorphTargetLodModel.NumVertices += SectionMorphTargetVertices[MorphIdx];
+						DestMorphLODModel.Vertices.Emplace(
+							FMorphTargetDelta{ SourceVertex.PositionDelta, SourceVertex.TangentZDelta, static_cast<uint32>(SpanVertexIdx) });
+
+						++SectionMorphTargetVertices[SourceVertex.MorphNameIndex];
 					}
 				}
 			}
 
-			if (NumNotFoundLoadedMorphsResources > 0)
+			const int32 SectionMorphTargetsNum = SectionMorphTargetVertices.Num();
+			for (int32 MorphIdx = 0; MorphIdx < SectionMorphTargetsNum; ++MorphIdx)
 			{
-				UE_LOG(LogMutable, Warning, TEXT("Needed realtime morph reconstruction data was not loaded properly. Some realtime morphs may not work correctly."));
-			}
-			
-			// Generate the SkeletalMesh data structures. The previous step could be done in an async task.
-			SkeletalMesh->GetMorphTargets().Empty();
-			const int32 NumMorphs = MorphTargetNames.Num(); 
-			for (int32 I = 0; I < NumMorphs; ++I)
-			{
-				if (MorphsData[I].IsEmpty())
+				if (SectionMorphTargetVertices[MorphIdx] > 0)
 				{
-					continue;
-				} 
+					FMorphTargetLODModel& MorphTargetLodModel = MorphsData[MorphIdx][LODIndex];
 
-				UMorphTarget* NewMorphTarget = NewObject<UMorphTarget>(SkeletalMesh, MorphTargetNames[I]);
-				NewMorphTarget->BaseSkelMesh = SkeletalMesh;
-				NewMorphTarget->GetMorphLODModels() = MoveTemp(MorphsData[I]);
-				SkeletalMesh->GetMorphTargets().Add(NewMorphTarget);
+					MorphTargetLodModel.SectionIndices.Add(Section);
+					MorphTargetLodModel.NumVertices += SectionMorphTargetVertices[MorphIdx];
+				}
 			}
 		}
-		else
+
+		if (NumNotFoundLoadedMorphsResources > 0)
 		{
-			MUTABLE_CPUPROFILER_SCOPE(BuildOrCopyMorphTargetsData_Copy);
+			UE_LOG(LogMutable, Warning, TEXT("Needed realtime morph reconstruction data was not loaded properly. Some realtime morphs may not work correctly."));
+		}
 
-			// Copy data from the last valid LOD generated from either the current mesh or the previous mesh
-			const int32 NumMorphTargets = SkeletalMesh->GetMorphTargets().Num();
-			for (int32 MorphTargetIndex = 0; MorphTargetIndex < NumMorphTargets; ++MorphTargetIndex)
+		// Generate the SkeletalMesh data structures. The previous step could be done in an async task.
+		MorphTargets.Empty();
+		const int32 NumMorphs = MorphTargetNames.Num();
+		for (int32 I = 0; I < NumMorphs; ++I)
+		{
+			if (MorphsData[I].IsEmpty())
 			{
-				SkeletalMesh->GetMorphTargets()[MorphTargetIndex]->GetMorphLODModels()[LODIndex] = 
-						SkeletalMesh->GetMorphTargets()[MorphTargetIndex]->GetMorphLODModels()[LastValidLODIndex];
+				continue;
 			}
+
+			UMorphTarget* NewMorphTarget = NewObject<UMorphTarget>(SkeletalMesh, MorphTargetNames[I]);
+			NewMorphTarget->BaseSkelMesh = SkeletalMesh;
+			NewMorphTarget->GetMorphLODModels() = MoveTemp(MorphsData[I]);
+			MorphTargets.Add(NewMorphTarget);
 		}
 	}
 
-	SkeletalMesh->InitMorphTargets();
+	// Copy MorphTargets from the FirstGeneratedLOD to the LODs below
+	for (int32 LODIndex = OperationData->FirstLODAvailable; LODIndex < FirstGeneratedLOD; ++LODIndex)
+	{
+		MUTABLE_CPUPROFILER_SCOPE(CopyMorphTargetsData);
+
+		const int32 NumMorphTargets = MorphTargets.Num();
+		for (int32 MorphTargetIndex = 0; MorphTargetIndex < NumMorphTargets; ++MorphTargetIndex)
+		{
+			MorphTargets[MorphTargetIndex]->GetMorphLODModels()[LODIndex] =
+				MorphTargets[MorphTargetIndex]->GetMorphLODModels()[FirstGeneratedLOD];
+		}
+	}
 
+	SkeletalMesh->InitMorphTargets();
 }
 
 namespace 
@@ -4627,73 +4620,60 @@ bool UCustomizableInstancePrivate::BuildOrCopyRenderData(const TSharedRef<FUpdat
 {
 	MUTABLE_CPUPROFILER_SCOPE(UCustomizableInstancePrivate::BuildOrCopyRenderData);
 
+	FSkeletalMeshRenderData* RenderData = SkeletalMesh->GetResourceForRendering();
+	check(RenderData);
+
 	UCustomizableObject* CustomizableObject = Public->GetCustomizableObject();
 	check(CustomizableObject);
 	const FModelResources& ModelResources = CustomizableObject->GetPrivate()->GetModelResources();
 
-	int32 LastValidLODIndex = OperationData->NumLODsAvailable - 1;
-	for (int32 LODIndex = LastValidLODIndex; LODIndex >= FirstLODAvailable; --LODIndex)
+	const int32 FirstGeneratedLOD = FMath::Max((int32)OperationData->GetRequestedLODs()[ComponentIndex], OperationData->GetMinLOD());
+	for (int32 LODIndex = FirstGeneratedLOD; LODIndex < OperationData->NumLODsAvailable; ++LODIndex)
 	{
-		MUTABLE_CPUPROFILER_SCOPE(BuildOrCopyRenderData_LODLoop);
+		MUTABLE_CPUPROFILER_SCOPE(BuildRenderData);
 
 		const FInstanceUpdateData::FLOD& LOD = OperationData->InstanceUpdateData.LODs[LODIndex];
-		FInstanceUpdateData::FComponent* Component = nullptr;
-		
-		if (ComponentIndex >=  LOD.ComponentCount)
+
+		if (ComponentIndex >= LOD.ComponentCount)
 		{
 			// Interrupt the generation if the LOD is empty and it should have been generated.
-			if (LODIndex >= OperationData->GetMinLOD() && LODIndex < OperationData->NumLODsAvailable)
-			{
-				UE_LOG(LogMutable, Warning, TEXT("Building instance: generated mesh [%s] has LOD [%d] with no component.")
-					, *SkeletalMesh->GetName()
-					, LODIndex);
+			UE_LOG(LogMutable, Warning, TEXT("Building instance: generated mesh [%s] has LOD [%d] with no component.")
+				, *SkeletalMesh->GetName()
+				, LODIndex);
 
-				// End with failure
-				return false;
-			}
-		}
-		else
-		{
-			Component = &OperationData->InstanceUpdateData.Components[LOD.FirstComponent + ComponentIndex];
+			// End with failure
+			return false;
 		}
 
-		if (!Component || !Component->bGenerated || Component->SurfaceCount == 0)
-		{
-			TRACE_CPUPROFILER_EVENT_SCOPE_TEXT(*FString::Printf(TEXT("BuildOrCopyRenderData_CopyData: From LOD %d to LOD %d"), LastValidLODIndex, LODIndex));
-
-			// Render Data will be reused from the previously generated component
-			UnrealConversionUtils::CopySkeletalMeshLODRenderData(SkeletalMesh, SkeletalMesh, LastValidLODIndex, LODIndex);
-			continue;
-		}
+		FInstanceUpdateData::FComponent& Component = OperationData->InstanceUpdateData.Components[LOD.FirstComponent + ComponentIndex];
 
 		// There could be components without a mesh in LODs
-		if (!Component->Mesh)
+		if (!Component.bGenerated || !Component.Mesh || Component.SurfaceCount == 0)
 		{
 			UE_LOG(LogMutable, Warning, TEXT("Building instance: generated mesh [%s] has LOD [%d] of component [%d] with no mesh.")
 				, *SkeletalMesh->GetName()
 				, LODIndex
-				, Component->Id);
+				, Component.Id);
 
 			// End with failure
 			return false;
 		}
 
-		TRACE_CPUPROFILER_EVENT_SCOPE_TEXT(*FString::Printf(TEXT("BuildOrCopyRenderData_BuildData: LOD %d"), LODIndex));
+		TRACE_CPUPROFILER_EVENT_SCOPE_TEXT(*FString::Printf(TEXT("BuildRenderData: Component %d, LOD %d"), ComponentIndex, LODIndex));
 
-		SetLastMeshId(Component->Id, LODIndex, Component->MeshID);
+		SetLastMeshId(Component.Id, LODIndex, Component.MeshID);
 
-		FSkeletalMeshLODRenderData& LODResource = SkeletalMesh->GetResourceForRendering()->LODRenderData[LODIndex];
+		FSkeletalMeshLODRenderData& LODResource = RenderData->LODRenderData[LODIndex];
 
 		UnrealConversionUtils::SetupRenderSections(
 			LODResource,
-			Component->Mesh,
-			LODIndex,
+			Component.Mesh,
 			OperationData->InstanceUpdateData.BoneMaps,
-			Component->FirstBoneMap);
+			Component.FirstBoneMap);
 
 		UnrealConversionUtils::CopyMutableVertexBuffers(
 			LODResource,
-			Component->Mesh,
+			Component.Mesh,
 			SkeletalMesh->GetLODInfo(LODIndex)->bAllowCPUAccess);
 
 
@@ -4720,14 +4700,14 @@ bool UCustomizableInstancePrivate::BuildOrCopyRenderData(const TSharedRef<FUpdat
 		}
 
 		// Update active and required bones
-		LODResource.ActiveBoneIndices.Append(Component->ActiveBones);
-		LODResource.RequiredBones.Append(Component->ActiveBones);
+		LODResource.ActiveBoneIndices.Append(Component.ActiveBones);
+		LODResource.RequiredBones.Append(Component.ActiveBones);
 
 		if (!ModelResources.SkinWeightProfilesInfo.IsEmpty())
 		{
 			bool bHasSkinWeightProfiles = false;
 
-			const mu::FMeshBufferSet& MutableMeshVertexBuffers = Component->Mesh->GetVertexBuffers();
+			const mu::FMeshBufferSet& MutableMeshVertexBuffers = Component.Mesh->GetVertexBuffers();
 			
 			const int32 SkinWeightProfilesCount = ModelResources.SkinWeightProfilesInfo.Num();
 			for (int32 ProfileIndex = 0; ProfileIndex < SkinWeightProfilesCount; ++ProfileIndex)
@@ -4769,14 +4749,14 @@ bool UCustomizableInstancePrivate::BuildOrCopyRenderData(const TSharedRef<FUpdat
 		}
 
 		// Copy indices.
-		if (!UnrealConversionUtils::CopyMutableIndexBuffers(LODResource, Component->Mesh))
+		if (!UnrealConversionUtils::CopyMutableIndexBuffers(LODResource, Component.Mesh))
 		{
 			// End with failure
 			return false;
 		}
 
 		// Update LOD and streaming data
-		const FMutableRefLODRenderData& RefLODRenderData = ModelResources.ReferenceSkeletalMeshesData[Component->Id].LODData[LODIndex].RenderData;
+		const FMutableRefLODRenderData& RefLODRenderData = ModelResources.ReferenceSkeletalMeshesData[Component.Id].LODData[LODIndex].RenderData;
 		LODResource.bIsLODOptional = RefLODRenderData.bIsLODOptional;
 		LODResource.bStreamedDataInlined = RefLODRenderData.bStreamedDataInlined;
 
@@ -4784,10 +4764,24 @@ bool UCustomizableInstancePrivate::BuildOrCopyRenderData(const TSharedRef<FUpdat
 		// USkeletalMesh::IsMaterialUsed checks this size to see if a material is being used. If it fails the textures used by it won't be included 
 		// in the map of textures to stream.
 		LODResource.BuffersSize = 1;
-
-		LastValidLODIndex = LODIndex;
 	}
 
+	// Copy LODRenderData from the FirstGeneratedLOD to the LODs below
+	for (int32 LODIndex = OperationData->FirstLODAvailable; LODIndex < FirstGeneratedLOD; ++LODIndex)
+	{
+		TRACE_CPUPROFILER_EVENT_SCOPE_TEXT(*FString::Printf(TEXT("CopyRenderData: From LOD %d to LOD %d"), FirstGeneratedLOD, LODIndex));
+
+		// Render Data will be reused from the previously generated component
+		FSkeletalMeshLODRenderData& SourceLODResource = RenderData->LODRenderData[FirstGeneratedLOD];
+		FSkeletalMeshLODRenderData& LODResource = RenderData->LODRenderData[LODIndex];
+
+		UnrealConversionUtils::CopySkeletalMeshLODRenderData(
+			LODResource,
+			SourceLODResource,
+			*SkeletalMesh,
+			SkeletalMesh->GetLODInfo(LODIndex)->bAllowCPUAccess
+		);
+	}
 
 	return true;
 }
@@ -5536,10 +5530,6 @@ void UCustomizableInstancePrivate::BuildMaterials(const TSharedRef<FUpdateContex
 
 	const FModelResources& ModelResources = CustomizableObject->GetPrivate()->GetModelResources();
 
-	// Find skipped LODs. The following valid LOD will be copied into them. 
-	TArray<bool> LODsSkipped;
-	LODsSkipped.SetNum(OperationData->NumLODsAvailable);
-
 	TArray<FGeneratedTexture> NewGeneratedTextures;
 
 	// Temp copy to allow reuse of MaterialInstances
@@ -5594,24 +5584,22 @@ void UCustomizableInstancePrivate::BuildMaterials(const TSharedRef<FUpdateContex
 
 		MUTABLE_CPUPROFILER_SCOPE(BuildMaterials_LODLoop);
 
-		for (int32 LODIndex = FirstLODAvailable; LODIndex < OperationData->NumLODsAvailable; LODIndex++)
+		const int32 FirstGeneratedLOD = FMath::Max((int32)OperationData->GetRequestedLODs()[ComponentIndex], OperationData->GetMinLOD());
+		for (int32 LODIndex = FirstGeneratedLOD; LODIndex < OperationData->NumLODsAvailable; LODIndex++)
 		{
 			const FInstanceUpdateData::FLOD& LOD = OperationData->InstanceUpdateData.LODs[LODIndex];
 			const FInstanceUpdateData::FComponent& Component = OperationData->InstanceUpdateData.Components[LOD.FirstComponent + ComponentIndex];
 
 			if (!Component.bGenerated)
 			{
-				LODsSkipped[LODIndex] = true;
 				continue;
 			}
 
-			if (SkeletalMesh->GetLODInfoArray().Num() != 0)
+			if (SkeletalMesh->GetLODInfoArray().IsValidIndex(LODIndex))
 			{
 				SkeletalMesh->GetLODInfoArray()[LODIndex].LODMaterialMap.Reset();
 			}
 
-			check(Component.bGenerated);
-
 			const FMutableRefSkeletalMeshData& RefSkeletalMeshData = ModelResources.ReferenceSkeletalMeshesData[Component.Id];
 
 			for (int32 SurfaceIndex = 0; SurfaceIndex < Component.SurfaceCount; ++SurfaceIndex)
@@ -6106,33 +6094,24 @@ void UCustomizableInstancePrivate::BuildMaterials(const TSharedRef<FUpdateContex
 		}
 
 		{
-			// Copy data from valid LODs into the skipped ones.
-			int32 LastValidLODIndex = OperationData->NumLODsAvailable - 1;
-			for (int32 LODIndex = LastValidLODIndex; LODIndex >= FirstLODAvailable; --LODIndex)
+			// Copy data from the FirstGeneratedLOD into the LODs below.
+			for (int32 LODIndex = OperationData->FirstLODAvailable; LODIndex < FirstGeneratedLOD; ++LODIndex)
 			{
-				// Copy information from the LastValidLODIndex
-				if (LODsSkipped[LODIndex])
-				{
-					SkeletalMesh->GetLODInfoArray()[LODIndex].LODMaterialMap = SkeletalMesh->GetLODInfoArray()[LastValidLODIndex].LODMaterialMap;
+				SkeletalMesh->GetLODInfoArray()[LODIndex].LODMaterialMap = SkeletalMesh->GetLODInfoArray()[FirstGeneratedLOD].LODMaterialMap;
 
-					TIndirectArray<FSkeletalMeshLODRenderData>& LODRenderData = SkeletalMesh->GetResourceForRendering()->LODRenderData;
+				TIndirectArray<FSkeletalMeshLODRenderData>& LODRenderData = SkeletalMesh->GetResourceForRendering()->LODRenderData;
 
-					const int32 NumRenderSections = LODRenderData[LODIndex].RenderSections.Num();
-					check(NumRenderSections == LODRenderData[LastValidLODIndex].RenderSections.Num());
+				const int32 NumRenderSections = LODRenderData[LODIndex].RenderSections.Num();
+				check(NumRenderSections == LODRenderData[FirstGeneratedLOD].RenderSections.Num());
 
-					if (NumRenderSections == LODRenderData[LastValidLODIndex].RenderSections.Num())
+				if (NumRenderSections == LODRenderData[FirstGeneratedLOD].RenderSections.Num())
+				{
+					for (int32 RenderSectionIndex = 0; RenderSectionIndex < NumRenderSections; ++RenderSectionIndex)
 					{
-						for (int32 RenderSectionIndex = 0; RenderSectionIndex < NumRenderSections; ++RenderSectionIndex)
-						{
-							const int32 MaterialIndex = LODRenderData[LastValidLODIndex].RenderSections[RenderSectionIndex].MaterialIndex;
-							LODRenderData[LODIndex].RenderSections[RenderSectionIndex].MaterialIndex = MaterialIndex;
-						}
+						const int32 MaterialIndex = LODRenderData[FirstGeneratedLOD].RenderSections[RenderSectionIndex].MaterialIndex;
+						LODRenderData[LODIndex].RenderSections[RenderSectionIndex].MaterialIndex = MaterialIndex;
 					}
 				}
-				else
-				{
-					LastValidLODIndex = LODIndex;
-				}
 			}
 		}
 
diff --git a/Engine/Plugins/Experimental/Mutable/Source/CustomizableObject/Private/MuCO/UnrealConversionUtils.cpp b/Engine/Plugins/Experimental/Mutable/Source/CustomizableObject/Private/MuCO/UnrealConversionUtils.cpp
--- a/Engine/Plugins/Experimental/Mutable/Source/CustomizableObject/Private/MuCO/UnrealConversionUtils.cpp
+++ b/Engine/Plugins/Experimental/Mutable/Source/CustomizableObject/Private/MuCO/UnrealConversionUtils.cpp
@@ -130,7 +130,6 @@ namespace UnrealConversionUtils
 	void SetupRenderSections(
 		FSkeletalMeshLODRenderData& LODResource,
 		const mu::MeshPtrConst InMutableMesh,
-		const int32 MeshLODIndex,
 		const TArray<uint16>& InBoneMap,
 		const int32 InFirstBoneMapIndex)
 	{
@@ -441,37 +440,22 @@ namespace UnrealConversionUtils
 	}
 
 	
-	 void CopySkeletalMeshLODRenderData(const USkeletalMesh* SrcSkeletalMesh,	USkeletalMesh* DestSkeletalMesh, int32 SrcLODIndex, int32 DestLODIndex)
+	 void CopySkeletalMeshLODRenderData(
+		 FSkeletalMeshLODRenderData& LODResource,
+		 FSkeletalMeshLODRenderData& SourceLODResource,
+		 const USkeletalMesh& SkeletalMesh,
+		 const bool bAllowCPUAccess
+	 )
 	 {
 		 MUTABLE_CPUPROFILER_SCOPE(CopySkeletalMeshLODRenderData);
 
-		 if (!SrcSkeletalMesh || !DestSkeletalMesh)
-		 {
-			 return;
-		 }
-
-		 if (!SrcSkeletalMesh->IsValidLODIndex(SrcLODIndex) || !DestSkeletalMesh->IsValidLODIndex(DestLODIndex))
-		 {
-			 // Trying to copy from/to an invalid LOD index
-			 check(false);
-			 return;
-		 }
-
-		 const FSkeletalMeshLODRenderData& SrcLODModel = SrcSkeletalMesh->GetResourceForRendering()->LODRenderData[SrcLODIndex];
-		 FSkeletalMeshLODRenderData& DestLODModel = DestSkeletalMesh->GetResourceForRendering()->LODRenderData[DestLODIndex];
-
 		 // Copying render sections
 		 {
-			 const FSkeletalMeshLODInfo& SrcLODInfo = SrcSkeletalMesh->GetLODInfoArray()[SrcLODIndex];
-			 FSkeletalMeshLODInfo& DestLODInfo = DestSkeletalMesh->GetLODInfoArray()[DestLODIndex];
-
-			 DestLODInfo.LODMaterialMap = SrcLODInfo.LODMaterialMap;
-
-			 const int32 SurfaceCount = SrcLODModel.RenderSections.Num();
+			 const int32 SurfaceCount = SourceLODResource.RenderSections.Num();
 			 for (int32 SurfaceIndex = 0; SurfaceIndex < SurfaceCount; ++SurfaceIndex)
 			 {
-				 const FSkelMeshRenderSection& SrcSection = SrcSkeletalMesh->GetResourceForRendering()->LODRenderData[SrcLODIndex].RenderSections[SurfaceIndex];
-				 FSkelMeshRenderSection* DestSection = new(DestSkeletalMesh->GetResourceForRendering()->LODRenderData[DestLODIndex].RenderSections) FSkelMeshRenderSection();
+				 const FSkelMeshRenderSection& SrcSection = SourceLODResource.RenderSections[SurfaceIndex];
+				 FSkelMeshRenderSection* DestSection = new(LODResource.RenderSections) FSkelMeshRenderSection();
 
 				 DestSection->DuplicatedVerticesBuffer.Init(1, TMap<int, TArray<int32>>());
 				 DestSection->bDisabled = SrcSection.bDisabled;
@@ -484,21 +468,17 @@ namespace UnrealConversionUtils
 					 DestSection->MaxBoneInfluences = SrcSection.MaxBoneInfluences;
 					 DestSection->NumVertices = SrcSection.NumVertices;
 					 DestSection->BoneMap = SrcSection.BoneMap;
-
-					 //TODO(Max): MTBL-1779
-					 //DestSection->bCastShadow = SrcSection.bCastShadow;
+					 DestSection->bCastShadow = SrcSection.bCastShadow;
 				 }
 			 }
 		 }
 
-		 const FStaticMeshVertexBuffers& SrcStaticVertexBuffer = SrcLODModel.StaticVertexBuffers;
-		 FStaticMeshVertexBuffers& DestStaticVertexBuffer = DestLODModel.StaticVertexBuffers;
+		 const FStaticMeshVertexBuffers& SrcStaticVertexBuffer = SourceLODResource.StaticVertexBuffers;
+		 FStaticMeshVertexBuffers& DestStaticVertexBuffer = LODResource.StaticVertexBuffers;
 
 		 const int32 NumVertices = SrcStaticVertexBuffer.PositionVertexBuffer.GetNumVertices();
 		 const int32 NumTexCoords = SrcStaticVertexBuffer.StaticMeshVertexBuffer.GetNumTexCoords();
 
-		 bool bAllowCPUAccess = SrcSkeletalMesh->GetLODInfoArray()[SrcLODIndex].bAllowCPUAccess;
-
 		 // Copying Static Vertex Buffers
 		 {
 			 // Position buffer
@@ -513,10 +493,8 @@ namespace UnrealConversionUtils
 			 FMemory::Memcpy(DestStaticVertexBuffer.StaticMeshVertexBuffer.GetTexCoordData(), SrcStaticVertexBuffer.StaticMeshVertexBuffer.GetTexCoordData(), DestStaticVertexBuffer.StaticMeshVertexBuffer.GetTexCoordSize());
 
 			 // Color buffer
-			 if (SrcSkeletalMesh->GetHasVertexColors())
+			 if (LODResource.StaticVertexBuffers.ColorVertexBuffer.GetNumVertices() > 0)
 			 {
-				 DestSkeletalMesh->SetHasVertexColors(true);
-
 				 DestStaticVertexBuffer.ColorVertexBuffer.Init(NumVertices);
 				 FMemory::Memcpy(DestStaticVertexBuffer.ColorVertexBuffer.GetVertexData(), SrcStaticVertexBuffer.ColorVertexBuffer.GetVertexData(), NumVertices * DestStaticVertexBuffer.ColorVertexBuffer.GetStride());
 			 }
@@ -524,8 +502,8 @@ namespace UnrealConversionUtils
 
 		 // Copying Skin Buffers
 		 {
-			 const FSkinWeightVertexBuffer& SrcSkinWeightBuffer = SrcLODModel.SkinWeightVertexBuffer;
-			 FSkinWeightVertexBuffer& DestSkinWeightBuffer = DestLODModel.SkinWeightVertexBuffer;
+			 const FSkinWeightVertexBuffer& SrcSkinWeightBuffer = SourceLODResource.SkinWeightVertexBuffer;
+			 FSkinWeightVertexBuffer& DestSkinWeightBuffer = LODResource.SkinWeightVertexBuffer;
 
 			 int32 NumBoneInfluences = SrcSkinWeightBuffer.GetDataVertexBuffer()->GetMaxBoneInfluences();
 			 int32 NumBones = SrcSkinWeightBuffer.GetDataVertexBuffer()->GetNumBoneWeights();
@@ -550,42 +528,60 @@ namespace UnrealConversionUtils
 
 		 // Copying Skin Weight Profiles Buffers
 		 {
-			 int32 NumSkinWeightProfiles = SrcSkeletalMesh->GetSkinWeightProfiles().Num();
+			 const int32 NumSkinWeightProfiles = SkeletalMesh.GetSkinWeightProfiles().Num();
 			 for (int32 ProfileIndex = 0; ProfileIndex < NumSkinWeightProfiles; ++ProfileIndex)
 			 {
-				 DestSkeletalMesh->AddSkinWeightProfile(SrcSkeletalMesh->GetSkinWeightProfiles()[ProfileIndex]);
-
-				 const FRuntimeSkinWeightProfileData* SrcProfile = SrcLODModel.SkinWeightProfilesData.GetOverrideData(SrcSkeletalMesh->GetSkinWeightProfiles()[ProfileIndex].Name);
-				 FRuntimeSkinWeightProfileData& DestProfile = DestLODModel.SkinWeightProfilesData.AddOverrideData(SrcSkeletalMesh->GetSkinWeightProfiles()[ProfileIndex].Name);
-
-				 DestProfile.NumWeightsPerVertex = SrcProfile->NumWeightsPerVertex;
-				 DestProfile.b16BitBoneIndices = SrcProfile->b16BitBoneIndices;
-
-				 DestProfile.BoneIDs = SrcProfile->BoneIDs;
-				 DestProfile.BoneWeights = SrcProfile->BoneWeights;
-				 DestProfile.VertexIndexToInfluenceOffset = SrcProfile->VertexIndexToInfluenceOffset;
+				 const FName& ProfileName = SkeletalMesh.GetSkinWeightProfiles()[ProfileIndex].Name;
+				 
+				 const FRuntimeSkinWeightProfileData* SourceProfile = SourceLODResource.SkinWeightProfilesData.GetOverrideData(ProfileName);
+				 FRuntimeSkinWeightProfileData& DestProfile = LODResource.SkinWeightProfilesData.AddOverrideData(ProfileName);
+				 
+				 DestProfile = *SourceProfile;
 			 }
 		 }
 
 		 // Copying Indices
 		 {
-			 if (SrcLODModel.MultiSizeIndexContainer.IsIndexBufferValid())
+			 if (SourceLODResource.MultiSizeIndexContainer.IsIndexBufferValid())
 			 {
-				 int32 IndexCount = SrcLODModel.MultiSizeIndexContainer.GetIndexBuffer()->Num();
-				 int32 ElementSize = SrcLODModel.MultiSizeIndexContainer.GetDataTypeSize();
+				 int32 IndexCount = SourceLODResource.MultiSizeIndexContainer.GetIndexBuffer()->Num();
+				 int32 ElementSize = SourceLODResource.MultiSizeIndexContainer.GetDataTypeSize();
 
-				 const void* Data = SrcSkeletalMesh->GetResourceForRendering()->LODRenderData[SrcLODIndex].MultiSizeIndexContainer.GetIndexBuffer()->GetPointerTo(0);
+				 const void* Data = SourceLODResource.MultiSizeIndexContainer.GetIndexBuffer()->GetPointerTo(0);
 
-				 DestLODModel.MultiSizeIndexContainer.CreateIndexBuffer(ElementSize);
-				 DestLODModel.MultiSizeIndexContainer.GetIndexBuffer()->Insert(0, IndexCount);
-				 FMemory::Memcpy(DestLODModel.MultiSizeIndexContainer.GetIndexBuffer()->GetPointerTo(0), Data, IndexCount * ElementSize);
+				 LODResource.MultiSizeIndexContainer.CreateIndexBuffer(ElementSize);
+				 LODResource.MultiSizeIndexContainer.GetIndexBuffer()->Insert(0, IndexCount);
+				 FMemory::Memcpy(LODResource.MultiSizeIndexContainer.GetIndexBuffer()->GetPointerTo(0), Data, IndexCount * ElementSize);
 			 }
 		 }
 
-		 DestLODModel.ActiveBoneIndices.Append(SrcLODModel.ActiveBoneIndices);
-		 DestLODModel.RequiredBones.Append(SrcLODModel.RequiredBones);
-		 DestLODModel.bIsLODOptional = SrcLODModel.bIsLODOptional;
-		 DestLODModel.bStreamedDataInlined = SrcLODModel.bStreamedDataInlined;
-		 DestLODModel.BuffersSize = SrcLODModel.BuffersSize;
+		 LODResource.ActiveBoneIndices.Append(SourceLODResource.ActiveBoneIndices);
+		 LODResource.RequiredBones.Append(SourceLODResource.RequiredBones);
+		 LODResource.bIsLODOptional = SourceLODResource.bIsLODOptional;
+		 LODResource.bStreamedDataInlined = SourceLODResource.bStreamedDataInlined;
+		 LODResource.BuffersSize = SourceLODResource.BuffersSize;
+	}
+
+
+	CUSTOMIZABLEOBJECT_API void UpdateSkeletalMeshLODRenderDataBuffersSize(FSkeletalMeshLODRenderData& LODResource)
+	{
+		LODResource.BuffersSize = 0;
+		
+		// Add VertexBuffers' size
+		LODResource.BuffersSize += LODResource.StaticVertexBuffers.PositionVertexBuffer.GetAllocatedSize();
+		LODResource.BuffersSize += LODResource.StaticVertexBuffers.StaticMeshVertexBuffer.GetResourceSize();
+		LODResource.BuffersSize += LODResource.StaticVertexBuffers.ColorVertexBuffer.GetAllocatedSize();
+		LODResource.BuffersSize += LODResource.SkinWeightVertexBuffer.GetVertexDataSize();
+
+		// Add Optional VertexBuffers' size
+		LODResource.BuffersSize += LODResource.ClothVertexBuffer.GetVertexDataSize();
+		LODResource.BuffersSize += LODResource.SkinWeightProfilesData.GetResourcesSize();
+		LODResource.BuffersSize += LODResource.MorphTargetVertexInfoBuffers.GetMorphDataSizeInBytes();
+
+		// Add IndexBuffer's size
+		if (LODResource.MultiSizeIndexContainer.IsIndexBufferValid())
+		{
+			LODResource.BuffersSize += LODResource.MultiSizeIndexContainer.GetIndexBuffer()->GetResourceDataSize();
+		}
 	}
 }
diff --git a/Engine/Plugins/Experimental/Mutable/Source/CustomizableObjectEditor/Private/MuCOE/MutableMeshPreviewUtils.cpp b/Engine/Plugins/Experimental/Mutable/Source/CustomizableObjectEditor/Private/MuCOE/MutableMeshPreviewUtils.cpp
--- a/Engine/Plugins/Experimental/Mutable/Source/CustomizableObjectEditor/Private/MuCOE/MutableMeshPreviewUtils.cpp
+++ b/Engine/Plugins/Experimental/Mutable/Source/CustomizableObjectEditor/Private/MuCOE/MutableMeshPreviewUtils.cpp
@@ -384,7 +384,6 @@ namespace MutableMeshPreviewUtils
 			UnrealConversionUtils::SetupRenderSections(
 				LODResource,
 				InMutableMesh,
-				0,
 				BoneMap,
 				0);
 
