[Mutable] Fix static analysis issues in OpImageLayer.

#rnx

diff --git a/Engine/Plugins/Experimental/Mutable/Source/MutableRuntime/Private/MuR/OpImageLayer.h b/Engine/Plugins/Experimental/Mutable/Source/MutableRuntime/Private/MuR/OpImageLayer.h
--- a/Engine/Plugins/Experimental/Mutable/Source/MutableRuntime/Private/MuR/OpImageLayer.h
+++ b/Engine/Plugins/Experimental/Mutable/Source/MutableRuntime/Private/MuR/OpImageLayer.h
@@ -908,10 +908,13 @@ namespace mu
 
 				// Copy the unblended channels
 				// \TODO: unnecessary when doing it in-place?
-				for (int32 C = 0; C < UnblendedChannels; ++C)
+				if constexpr (UnblendedChannels > 0)
 				{
-					DestBuf[BASE_CHANNEL_STRIDE * I + CHANNELS_TO_BLEND + C] 
-						= BaseBuf[BASE_CHANNEL_STRIDE * I + CHANNELS_TO_BLEND + C];
+					for (int32 C = 0; C < UnblendedChannels; ++C)
+					{
+						DestBuf[BASE_CHANNEL_STRIDE * I + CHANNELS_TO_BLEND + C] 
+							= BaseBuf[BASE_CHANNEL_STRIDE * I + CHANNELS_TO_BLEND + C];
+					}
 				}
 			}
 		};
@@ -1133,9 +1136,12 @@ namespace mu
 					}
 					// Copy the unblended channels
 					// \TODO: unnecessary when doing it in-place?
-					for (int32 C = 0; C < UnblendedChannels; ++C)
+					if constexpr (UnblendedChannels > 0)
 					{
-						DestBuf[BASE_CHANNEL_STRIDE * I + CHANNELS_TO_BLEND + C] = BaseBuf[BASE_CHANNEL_STRIDE * I + CHANNELS_TO_BLEND + C];
+						for (int32 C = 0; C < UnblendedChannels; ++C)
+						{
+							DestBuf[BASE_CHANNEL_STRIDE * I + CHANNELS_TO_BLEND + C] = BaseBuf[BASE_CHANNEL_STRIDE * I + CHANNELS_TO_BLEND + C];
+						}
 					}
 				}
 			};
@@ -1203,9 +1209,12 @@ namespace mu
 
 								// Copy the unblended channels
 								// \TODO: unnecessary when doing it in-place?
-								for (int32 C = 0; C < UnblendedChannels; ++C)
+								if constexpr (UnblendedChannels > 0)
 								{
-									DestBuf[BASE_CHANNEL_STRIDE * I + CHANNELS_TO_BLEND + C] = BaseBuf[BASE_CHANNEL_STRIDE * I + CHANNELS_TO_BLEND + C];
+									for (int32 C = 0; C < UnblendedChannels; ++C)
+									{
+										DestBuf[BASE_CHANNEL_STRIDE * I + CHANNELS_TO_BLEND + C] = BaseBuf[BASE_CHANNEL_STRIDE * I + CHANNELS_TO_BLEND + C];
+									}
 								}
 							}
 						}
@@ -1231,9 +1240,12 @@ namespace mu
 
 								// Copy the unblended channels
 								// \TODO: unnecessary when doing it in-place?
-								for (int32 C = 0; C < UnblendedChannels; ++C)
+								if constexpr (UnblendedChannels > 0)
 								{
-									DestBuf[BASE_CHANNEL_STRIDE * I + CHANNELS_TO_BLEND + C] = BaseBuf[BASE_CHANNEL_STRIDE * I + CHANNELS_TO_BLEND + C];
+									for (int32 C = 0; C < UnblendedChannels; ++C)
+									{
+										DestBuf[BASE_CHANNEL_STRIDE * I + CHANNELS_TO_BLEND + C] = BaseBuf[BASE_CHANNEL_STRIDE * I + CHANNELS_TO_BLEND + C];
+									}
 								}
 							}
 						}
@@ -1272,9 +1284,12 @@ namespace mu
 
 							// Copy the unblended channels
 							// \TODO: unnecessary when doing it in-place?
-							for (int32 C = 0; C < UnblendedChannels; ++C)
+							if constexpr (UnblendedChannels > 0)
 							{
-								DestBuf[BASE_CHANNEL_STRIDE * I + CHANNELS_TO_BLEND + C] = BaseBuf[BASE_CHANNEL_STRIDE * I + CHANNELS_TO_BLEND + C];
+								for (int32 C = 0; C < UnblendedChannels; ++C)
+								{
+									DestBuf[BASE_CHANNEL_STRIDE * I + CHANNELS_TO_BLEND + C] = BaseBuf[BASE_CHANNEL_STRIDE * I + CHANNELS_TO_BLEND + C];
+								}
 							}
 						}
 
@@ -1375,9 +1390,12 @@ namespace mu
 					}
 					// Copy the unblended channels
 					// \TODO: unnecessary when doing it in-place?
-					for (int32 C = 0; C < UnblendedChannels; ++C)
+					if constexpr (UnblendedChannels > 0)
 					{
-						DestBuf[BASE_CHANNEL_STRIDE * I + CHANNELS_TO_BLEND + C] = BaseBuf[BASE_CHANNEL_STRIDE * I + CHANNELS_TO_BLEND + C];
+						for (int32 C = 0; C < UnblendedChannels; ++C)
+						{
+							DestBuf[BASE_CHANNEL_STRIDE * I + CHANNELS_TO_BLEND + C] = BaseBuf[BASE_CHANNEL_STRIDE * I + CHANNELS_TO_BLEND + C];
+						}
 					}
 				}
 			}
@@ -1401,9 +1419,12 @@ namespace mu
 					}
 					// Copy the unblended channels
 					// \TODO: unnecessary when doing it in-place?
-					for (int32 C = 0; C < UnblendedChannels; ++C)
+					if constexpr (UnblendedChannels > 0)
 					{
-						DestBuf[BASE_CHANNEL_STRIDE * I + CHANNELS_TO_BLEND + C] = BaseBuf[BASE_CHANNEL_STRIDE * I + CHANNELS_TO_BLEND + C];
+						for (int32 C = 0; C < UnblendedChannels; ++C)
+						{
+							DestBuf[BASE_CHANNEL_STRIDE * I + CHANNELS_TO_BLEND + C] = BaseBuf[BASE_CHANNEL_STRIDE * I + CHANNELS_TO_BLEND + C];
+						}
 					}
 				}
 			}
