[Mutable] Basic implementation of UCustomizableObjectSkeletalMesh and FCustomizableObjectMeshStreamIn classes to allow streaming of mesh LODs.

#rnx

diff --git a/Engine/Plugins/Experimental/Mutable/Source/CustomizableObject/Classes/MuCO/CustomizableObjectSkeletalMesh.h b/Engine/Plugins/Experimental/Mutable/Source/CustomizableObject/Classes/MuCO/CustomizableObjectSkeletalMesh.h
new file mode 100644
--- /dev/null
+++ b/Engine/Plugins/Experimental/Mutable/Source/CustomizableObject/Classes/MuCO/CustomizableObjectSkeletalMesh.h
@@ -0,0 +1,61 @@
+// Copyright Epic Games, Inc. All Rights Reserved.
+
+#pragma once
+
+#include "CoreMinimal.h"
+#include "Engine/SkeletalMesh.h"
+
+#include "MuR/Ptr.h"
+
+#include "CustomizableObjectSkeletalMesh.generated.h"
+
+namespace mu
+{
+	class Model;
+	class Parameters;
+	typedef uint64 FResourceID;
+}
+
+class UCustomizableObjectInstance;
+class FUpdateContextPrivate;
+
+class FCustomizableObjectMeshStreamIn;
+
+/**
+ * CustomizableObjectSkeletalMesh is a class inheriting from USkeletalMesh with the sole purpose of allowing skeletal meshes generated by Mutable
+ * to stream-in Mesh LODs. 
+ */
+UCLASS()
+class CUSTOMIZABLEOBJECT_API UCustomizableObjectSkeletalMesh : public USkeletalMesh
+{
+	GENERATED_BODY()
+
+	friend FCustomizableObjectMeshStreamIn;
+
+public: // Public Methods
+
+	static UCustomizableObjectSkeletalMesh* CreateSkeletalMesh(const TSharedRef<FUpdateContextPrivate>& InOperationData, const UCustomizableObjectInstance& InInstance, const int32 ComponentIndex);
+
+private: // Private Methods
+
+	//~ Begin UStreamableRenderAsset Interface.
+	virtual bool StreamIn(int32 NewMipCount, bool bHighPrio) override;
+	//~ End UStreamableRenderAsset Interface.
+
+
+public: // Public Variables
+
+private: // Private Variables
+	FString CustomizableObjectPathName;
+	FString InstancePathName;
+
+	TSharedPtr<mu::Model, ESPMode::ThreadSafe> Model;
+	
+	/** Instance parameters at the time of creation. */
+	mu::Ptr<mu::Parameters> Parameters;
+	int32 State = -1;
+
+	/** Mesh IDs per LOD */
+	TArray<mu::FResourceID> MeshIDs;
+
+};
diff --git a/Engine/Plugins/Experimental/Mutable/Source/CustomizableObject/Private/MuCO/CustomizableObjectMeshUpdate.cpp b/Engine/Plugins/Experimental/Mutable/Source/CustomizableObject/Private/MuCO/CustomizableObjectMeshUpdate.cpp
new file mode 100644
--- /dev/null
+++ b/Engine/Plugins/Experimental/Mutable/Source/CustomizableObject/Private/MuCO/CustomizableObjectMeshUpdate.cpp
@@ -0,0 +1,262 @@
+// Copyright Epic Games, Inc. All Rights Reserved.
+
+/*=============================================================================
+CustomizableObjectMeshUpdate.cpp: Helpers to stream in CustomizableObject skeletal mesh LODs.
+=============================================================================*/
+
+
+#include "MuCO/CustomizableObjectMeshUpdate.h"
+
+#include "MuCO/CustomizableObjectSystem.h"
+#include "MuCO/CustomizableObjectSystemPrivate.h"
+#include "MuCO/CustomizableObjectSkeletalMesh.h"
+#include "MuCO/UnrealConversionUtils.h"
+
+#include "MuR/Model.h"
+#include "MuR/MeshBufferSet.h"
+
+#include "Engine/SkeletalMesh.h"
+#include "Streaming/RenderAssetUpdate.inl"
+#include "Rendering/SkeletalMeshRenderData.h"
+
+template class TRenderAssetUpdate<FSkelMeshUpdateContext>;
+
+
+FCustomizableObjectMeshStreamIn::FCustomizableObjectMeshStreamIn(const UCustomizableObjectSkeletalMesh* InMesh, bool bInHighPrio, bool bInRenderThread)
+	: FSkeletalMeshStreamIn(InMesh),
+	bHighPriority(bInHighPrio),
+	bRenderThread(bInRenderThread)
+{
+	OperationData = MakeShared<FMutableMeshOperationData>();
+
+	// This must run in the mutable thread.
+	check(UCustomizableObjectSystem::IsCreated());
+	OperationData->System = UCustomizableObjectSystem::GetInstance()->GetPrivate()->MutableSystem;
+
+	OperationData->Model = InMesh->Model;
+	OperationData->Parameters = InMesh->Parameters;
+	OperationData->State = InMesh->State;
+
+	OperationData->MeshIDs = InMesh->MeshIDs;
+	OperationData->Meshes.SetNum(InMesh->MeshIDs.Num());
+
+	OperationData->CurrentFirstLODIdx = CurrentFirstLODIdx;
+	OperationData->PendingFirstLODIdx = PendingFirstLODIdx;
+
+	PushTask(FContext(InMesh, TT_None), TT_Async, SRA_UPDATE_CALLBACK(DoInitiate), TT_None, nullptr);
+}
+
+void FCustomizableObjectMeshStreamIn::DoInitiate(const FContext& Context)
+{
+	check(Context.CurrentThread == TT_Async);
+
+	if (IsCancelled())
+	{
+		return;
+	}
+
+	// Launch MutableTask
+	RequestMeshUpdate(Context);
+
+	PushTask(Context, TT_Async, SRA_UPDATE_CALLBACK(DoConvertResources), TT_Async, SRA_UPDATE_CALLBACK(DoCancelMeshUpdate));
+}
+
+void FCustomizableObjectMeshStreamIn::DoConvertResources(const FContext& Context)
+{
+	check(Context.CurrentThread == TT_Async);
+
+	ConvertMesh(Context);
+
+	const EThreadType TThread = bRenderThread ? TT_Render : TT_Async;
+	const EThreadType CThread = (EThreadType)Context.CurrentThread;
+	PushTask(Context, TThread, SRA_UPDATE_CALLBACK(DoCreateBuffers), CThread, SRA_UPDATE_CALLBACK(DoCancel));
+}
+
+void FCustomizableObjectMeshStreamIn::DoCreateBuffers(const FContext& Context)
+{
+	if (bRenderThread)
+	{
+		CreateBuffers_RenderThread(Context);
+	}
+	else
+	{
+		CreateBuffers_Async(Context);
+	}
+	check(!TaskSynchronization.GetValue());
+	PushTask(Context, TT_Render, SRA_UPDATE_CALLBACK(DoFinishUpdate), (EThreadType)Context.CurrentThread, SRA_UPDATE_CALLBACK(DoCancel));
+}
+
+void FCustomizableObjectMeshStreamIn::DoCancelMeshUpdate(const FContext& Context)
+{
+	CancelMeshUpdate(Context);
+	PushTask(Context, TT_None, nullptr, (EThreadType)Context.CurrentThread, SRA_UPDATE_CALLBACK(DoCancel));
+}
+
+
+namespace impl
+{
+	void Task_Mutable_UpdateMesh(const TSharedPtr<FMutableMeshOperationData> OperationData)
+	{
+		MUTABLE_CPUPROFILER_SCOPE(Task_Mutable_UpdateMesh);
+
+		// This runs in a worker thread.
+		check(OperationData.IsValid());
+		check(OperationData->System.get());
+		check(OperationData->Model);
+		check(OperationData->Parameters.get());
+
+		if (!OperationData.IsValid())
+		{
+			return;
+		}
+
+		mu::SystemPtr System = OperationData->System;
+		const TSharedPtr<mu::Model, ESPMode::ThreadSafe> Model = OperationData->Model;
+
+#if WITH_EDITOR
+		// Recompiling a CO in the editor will invalidate the previously generated Model. Check that it is valid before accessing the streamed data.
+		if (!Model || !Model->IsValid())
+		{
+			return;
+		}
+#endif
+
+		// For now, we are forcing the recreation of mutable-side instances with every update.
+		mu::Instance::ID InstanceID = System->NewInstance(Model);
+		UE_LOG(LogMutable, Verbose, TEXT("Creating Mutable instance with id [%d] for a mesh update"), InstanceID);
+
+		// LOD mask, set to all ones to build all LODs
+		const uint32 LODMask = 0xFFFFFFFF;
+
+		// Main instance generation step
+		const mu::Instance* Instance = System->BeginUpdate(InstanceID, OperationData->Parameters, OperationData->State, LODMask);
+		check(Instance);
+
+		// Generate the required meshes
+		for (int32 Index = OperationData->PendingFirstLODIdx; Index < OperationData->CurrentFirstLODIdx; ++Index)
+		{
+			const mu::FResourceID MeshID = OperationData->MeshIDs[Index];
+			OperationData->Meshes[Index] = System->GetMesh(InstanceID, MeshID);
+		}
+
+		// End update
+		System->EndUpdate(InstanceID);
+		System->ReleaseInstance(InstanceID);
+
+		if (CVarClearWorkingMemoryOnUpdateEnd.GetValueOnAnyThread())
+		{
+			System->ClearWorkingMemory();
+		}
+	}
+} // namespace
+
+
+void FCustomizableObjectMeshStreamIn::RequestMeshUpdate(const FContext& Context)
+{
+	if (!UCustomizableObjectSystem::IsActive())
+	{
+		Abort();
+		return;
+	}
+
+	check(OperationData.IsValid());
+
+#if WITH_EDITOR
+	// Recompiling a CO in the editor will invalidate the previously generated Model. Check that it is valid before accessing the streamed data.
+	if (!OperationData->Model || !OperationData->Model->IsValid())
+	{
+		Abort();
+		return;
+	}
+#endif
+
+	UCustomizableObjectSystemPrivate* CustomizableObjectSystem = UCustomizableObjectSystem::GetInstance()->GetPrivate();
+	
+	TaskSynchronization.Increment();
+
+	TRefCountPtr<FCustomizableObjectMeshStreamIn> RefThis = this;
+	TSharedPtr<FMutableMeshOperationData> SharedOperationData = OperationData;
+
+	MutableTaskId = CustomizableObjectSystem->MutableTaskGraph.AddMutableThreadTaskLowPriority(
+		TEXT("Mutable_MeshUpdate"),
+		[SharedOperationData, RefThis]()
+		{
+			impl::Task_Mutable_UpdateMesh(SharedOperationData);
+
+			// At this point task synchronization would hold the number of pending requests.
+			RefThis->TaskSynchronization.Decrement();
+
+			// The tick here is intended to schedule the success or cancel callback.
+			// Using TT_None ensure gets which could create a dead lock.
+			RefThis->Tick(FSkeletalMeshUpdate::TT_None);
+		});
+}
+
+
+void FCustomizableObjectMeshStreamIn::CancelMeshUpdate(const FContext& Context)
+{
+	UCustomizableObjectSystemPrivate* CustomizableObjectSystem = UCustomizableObjectSystem::GetInstance()->GetPrivate();
+	if (CustomizableObjectSystem && MutableTaskId != FMutableTaskGraph::INVALID_ID)
+	{
+		// Cancel task if not launched yet.
+		const bool bMutableTaskCancelledBeforeRun = CustomizableObjectSystem->MutableTaskGraph.CancelMutableThreadTaskLowPriority(MutableTaskId);
+		if (bMutableTaskCancelledBeforeRun)
+		{
+			// Clear MeshUpdate data
+			OperationData = nullptr;
+
+			// At this point task synchronization would hold the number of pending requests.
+			TaskSynchronization.Decrement();
+			check(TaskSynchronization.GetValue() == 0);
+		}
+	}
+	else
+	{
+		check(TaskSynchronization.GetValue() == 0);
+	}
+
+	// The tick here is intended to schedule the success or cancel callback.
+	// Using TT_None ensure gets which could create a dead lock.
+	Tick(FSkeletalMeshUpdate::TT_None);
+}
+
+
+void FCustomizableObjectMeshStreamIn::ConvertMesh(const FContext& Context)
+{
+	MUTABLE_CPUPROFILER_SCOPE(ConvertMesh);
+
+	check(!TaskSynchronization.GetValue());
+
+	const USkeletalMesh* Mesh = Context.Mesh;
+	FSkeletalMeshRenderData* RenderData = Context.RenderData;
+	if (IsCancelled() || !Mesh || !RenderData)
+	{
+		return;
+	}
+
+	for (int32 LODIndex = PendingFirstLODIdx; LODIndex < CurrentFirstLODIdx; ++LODIndex)
+	{
+		mu::MeshPtrConst MutableMesh = OperationData->Meshes[LODIndex];
+
+		if (!MutableMesh)
+		{
+			check(false);
+			Abort();
+			return;
+		}
+
+		FSkeletalMeshLODRenderData& LODResource = *Context.LODResourcesView[LODIndex];
+		UnrealConversionUtils::CopyMutableVertexBuffers(LODResource, MutableMesh, true);
+		UnrealConversionUtils::CopyMutableIndexBuffers(LODResource, MutableMesh);
+
+		// TODO PERE: Implement missing buffers
+		// SkinWeightProfilesData
+		// MorphTargets (MorphTargetVertexInfoBuffers)
+		// Clothing (ClothVertexBuffer)
+
+		UnrealConversionUtils::UpdateSkeletalMeshLODRenderDataBuffersSize(LODResource);
+	}
+
+	// Clear MeshUpdate data
+	OperationData = nullptr;
+}
\ No newline at end of file
diff --git a/Engine/Plugins/Experimental/Mutable/Source/CustomizableObject/Private/MuCO/CustomizableObjectMeshUpdate.h b/Engine/Plugins/Experimental/Mutable/Source/CustomizableObject/Private/MuCO/CustomizableObjectMeshUpdate.h
new file mode 100644
--- /dev/null
+++ b/Engine/Plugins/Experimental/Mutable/Source/CustomizableObject/Private/MuCO/CustomizableObjectMeshUpdate.h
@@ -0,0 +1,86 @@
+// Copyright Epic Games, Inc. All Rights Reserved.
+
+/*=============================================================================
+SkeletalMeshUpdate.h: Helpers to stream in and out skeletal mesh LODs.
+=============================================================================*/
+
+#pragma once
+
+#include "CoreMinimal.h"
+#include "Streaming/SkeletalMeshUpdate.h"
+
+#include "MuR/System.h"
+#include "MuR/Mesh.h"
+
+class UCustomizableObjectSkeletalMesh;
+
+namespace mu 
+{ 
+	class Parameters;
+}
+
+struct FMutableMeshOperationData
+{
+	/** Mutable */
+	mu::Ptr<mu::System> System;
+	TSharedPtr<mu::Model, ESPMode::ThreadSafe> Model;
+	mu::Ptr<const mu::Parameters> Parameters;
+	int32 State = -1;
+
+	// IDs of the meshes to generate per LOD.
+	TArray<mu::FResourceID> MeshIDs;
+
+	// Meshes generated per LOD
+	TArray<mu::MeshPtrConst> Meshes;
+
+	/** SkeletalMesh */
+	/** The resident first LOD resource index.With domain = [0, ResourceState.NumLODs[.NOT THE ASSET LOD INDEX! */ 
+	int32 CurrentFirstLODIdx = INDEX_NONE;
+
+	/** The requested first LOD resource index. With domain = [0, ResourceState.NumLODs[. NOT THE ASSET LOD INDEX! */
+	int32 PendingFirstLODIdx = INDEX_NONE;
+};
+
+extern template class TRenderAssetUpdate<FSkelMeshUpdateContext>;
+
+/** 
+* This class provides a framework for loading the LODs of CustomizableObject skeletal meshes.
+* Each thread essentially calls Tick() until the job is done.
+* The object can be safely deleted when IsCompleted() returns true.
+*/
+class FCustomizableObjectMeshStreamIn : public FSkeletalMeshStreamIn
+{
+public:
+	FCustomizableObjectMeshStreamIn(const UCustomizableObjectSkeletalMesh* InMesh, bool bHighPrio, bool bRenderThread);
+
+private:
+
+	void DoInitiate(const FContext& Context);
+
+	void DoConvertResources(const FContext& Context);
+
+	void DoCreateBuffers(const FContext& Context);
+
+	void DoCancelMeshUpdate(const FContext& Context);
+
+	/** Creates a MeshUpdate task to generate the meshes for the LODs to stream in. */
+	void RequestMeshUpdate(const FContext& Context);
+
+	/** Cancels the MeshUpdate task. */
+	void CancelMeshUpdate(const FContext& Context);
+
+	/** Converts from mu::Mesh to FSkeletalMeshLODRenderData. */
+	void ConvertMesh(const FContext& Context);
+
+	/** TODO: Stream priority. */
+	bool bHighPriority = false;
+	
+	/** Used to create the buffers in the render thread or not. */
+	bool bRenderThread = false;
+
+	/** Context of the MeshUpdate. */
+	TSharedPtr<FMutableMeshOperationData> OperationData;
+
+	/** MeshUpdate task Id to cancel the task if the stream in task is aborted. */
+	uint32 MutableTaskId = 0;
+};
diff --git a/Engine/Plugins/Experimental/Mutable/Source/CustomizableObject/Private/MuCO/CustomizableObjectSkeletalMesh.cpp b/Engine/Plugins/Experimental/Mutable/Source/CustomizableObject/Private/MuCO/CustomizableObjectSkeletalMesh.cpp
new file mode 100644
--- /dev/null
+++ b/Engine/Plugins/Experimental/Mutable/Source/CustomizableObject/Private/MuCO/CustomizableObjectSkeletalMesh.cpp
@@ -0,0 +1,61 @@
+// Copyright Epic Games, Inc. All Rights Reserved.
+
+#include "MuCO/CustomizableObjectSkeletalMesh.h"
+
+#include "MuCO/CustomizableObjectSystem.h"
+#include "MuCO/CustomizableObjectSystemPrivate.h"
+
+#include "MuCO/CustomizableObject.h"
+#include "MuCO/CustomizableObjectInstance.h"
+#include "MuCO/CustomizableObjectMeshUpdate.h"
+
+UCustomizableObjectSkeletalMesh* UCustomizableObjectSkeletalMesh::CreateSkeletalMesh(const TSharedRef<FUpdateContextPrivate>& OperationData, const UCustomizableObjectInstance& Instance, const int32 InComponentIndex)
+{
+	UCustomizableObject* CustomizableObject = Instance.GetCustomizableObject();
+	check(CustomizableObject);
+
+	UCustomizableObjectSkeletalMesh* OutSkeletalMesh = NewObject<UCustomizableObjectSkeletalMesh>();
+	
+	// Debug info
+	OutSkeletalMesh->CustomizableObjectPathName = GetNameSafe(CustomizableObject);
+	OutSkeletalMesh->InstancePathName = Instance.GetName();
+
+
+	// Init properties
+	OutSkeletalMesh->Model = CustomizableObject->GetPrivate()->GetModel();
+
+	OutSkeletalMesh->Parameters = OperationData->Parameters;
+	OutSkeletalMesh->State = OperationData->GetCapturedDescriptor().GetState();
+	
+	OutSkeletalMesh->MeshIDs.Init(MAX_uint64, MAX_MESH_LOD_COUNT);
+	
+	for (int32 LODIndex = OperationData->FirstLODAvailable; LODIndex < OperationData->NumLODsAvailable; ++LODIndex)
+	{
+		const FInstanceUpdateData::FLOD& LOD = OperationData->InstanceUpdateData.LODs[LODIndex];
+
+		for (int32 ComponentIndex = 0; ComponentIndex < LOD.ComponentCount; ++ComponentIndex)
+		{
+			const FInstanceUpdateData::FComponent& Component = OperationData->InstanceUpdateData.Components[LOD.FirstComponent + ComponentIndex];
+
+			if (Component.Id == InComponentIndex)
+			{
+				OutSkeletalMesh->MeshIDs[LODIndex] = Component.MeshID;
+			}
+		}
+	}
+
+	return OutSkeletalMesh;
+}
+
+
+bool UCustomizableObjectSkeletalMesh::StreamIn(int32 NewMipCount, bool bHighPrio)
+{
+	check(IsInGameThread());
+	if (!HasPendingInitOrStreaming() && CachedSRRState.StreamIn(NewMipCount))
+	{
+		PendingUpdate = new FCustomizableObjectMeshStreamIn(this, bHighPrio, !GRHISupportsAsyncTextureCreation);
+
+		return !PendingUpdate->IsCancelled();
+	}
+	return false;
+}
