[mutable] Fixed some Instance editor bugs:
+ Instance Editor won't crash while modifying a projector with the option "Only  Runtime" checked.
+ Now instance will be updated while moving a float slider.

#rnx

diff --git a/Engine/Plugins/Experimental/Mutable/Source/CustomizableObjectEditor/Private/MuCOE/CustomizableInstanceDetails.cpp b/Engine/Plugins/Experimental/Mutable/Source/CustomizableObjectEditor/Private/MuCOE/CustomizableInstanceDetails.cpp
--- a/Engine/Plugins/Experimental/Mutable/Source/CustomizableObjectEditor/Private/MuCOE/CustomizableInstanceDetails.cpp
+++ b/Engine/Plugins/Experimental/Mutable/Source/CustomizableObjectEditor/Private/MuCOE/CustomizableInstanceDetails.cpp
@@ -244,9 +244,13 @@ void FCustomizableInstanceDetails::Refresh() const
 	}
 }
 
+
 void FCustomizableInstanceDetails::InstanceUpdated(UCustomizableObjectInstance* Instance) const
 {
-	Refresh();
+	if (!bUpdatingSlider)
+	{
+		Refresh();
+	}
 }
 
 
@@ -511,24 +515,24 @@ void FCustomizableInstanceDetails::OnUseUISectionsSelectionChanged(ECheckBoxStat
 
 struct FParameterInfo
 {
-	int32 ParameterIndex = -1;
-	int32 ParameterUIOrder = 0;
-	FString ParameterName = "";
+	int32 ParamIndexInObject = -1;
+	int32 ParamUIOrder = 0;
+	FString ParamName = "";
 
 	bool operator==(const FParameterInfo& Other)
 	{
-		return ParameterName == Other.ParameterName;
+		return ParamName == Other.ParamName;
 	}
 
 	friend bool operator<(const FParameterInfo& A, const FParameterInfo& B)
 	{
-		if (A.ParameterUIOrder != B.ParameterUIOrder)
+		if (A.ParamUIOrder != B.ParamUIOrder)
 		{
-			return A.ParameterUIOrder < B.ParameterUIOrder;
+			return A.ParamUIOrder < B.ParamUIOrder;
 		}
 		else
 		{
-			return A.ParameterName < B.ParameterName;
+			return A.ParamName < B.ParamName;
 		}
 	}
 };
@@ -558,53 +562,53 @@ void FCustomizableInstanceDetails::GenerateParametersView(IDetailCategoryBuilder
 		for (int32 ParamIndexInState = 0; ParamIndexInState < NumStateParameters; ++ParamIndexInState)
 		{
 			FParameterInfo ParameterSortInfo;
-			ParameterSortInfo.ParameterIndex = CustomizableObject->GetStateParameterIndex(CustomInstance->GetPrivate()->GetState(), ParamIndexInState);
+			ParameterSortInfo.ParamIndexInObject = CustomizableObject->GetStateParameterIndex(CustomInstance->GetPrivate()->GetState(), ParamIndexInState);
 
-			if (CustomInstance->IsParameterRelevant(ParameterSortInfo.ParameterIndex) && IsVisible(ParameterSortInfo.ParameterIndex))
+			if (CustomInstance->IsParameterRelevant(ParameterSortInfo.ParamIndexInObject) && IsVisible(ParameterSortInfo.ParamIndexInObject))
 			{
-				ParameterSortInfo.ParameterUIOrder = CustomizableObject->GetParameterUIMetadataFromIndex(ParameterSortInfo.ParameterIndex).ParamUIMetadata.UIOrder;
-				ParameterSortInfo.ParameterName = CustomizableObject->GetParameterName(ParameterSortInfo.ParameterIndex);
+				ParameterSortInfo.ParamUIOrder = CustomizableObject->GetParameterUIMetadataFromIndex(ParameterSortInfo.ParamIndexInObject).ParamUIMetadata.UIOrder;
+				ParameterSortInfo.ParamName = CustomizableObject->GetParameterName(ParameterSortInfo.ParamIndexInObject);
 				ParametersTree.Add(ParameterSortInfo);
 			}
 		}
 
 		ParametersTree.Sort();
 
-		for (int32 ParamIndexInObject = 0; ParamIndexInObject < ParametersTree.Num(); ++ParamIndexInObject)
+		for (int32 ParamIndex = 0; ParamIndex < ParametersTree.Num(); ++ParamIndex)
 		{
-			const FParameterInfo& ParamInfo = ParametersTree[ParamIndexInObject];
+			const FParameterInfo& ParamInfo = ParametersTree[ParamIndex];
 			
 			if (CustomInstance->GetPrivate()->bShowUISections)
 			{
-				IDetailGroup* CurrentSection = GenerateParameterSection(ParamIndexInObject, DetailsCategory);
+				IDetailGroup* CurrentSection = GenerateParameterSection(ParamInfo.ParamIndexInObject, DetailsCategory);
 				check(CurrentSection);
 
-				if (!IsMultidimensionalProjector(ParamIndexInObject))
+				if (!IsMultidimensionalProjector(ParamInfo.ParamIndexInObject))
 				{
-					GenerateWidgetRow(CurrentSection->AddWidgetRow(), ParamInfo.ParameterName, ParamInfo.ParameterIndex);
+					GenerateWidgetRow(CurrentSection->AddWidgetRow(), ParamInfo.ParamName, ParamInfo.ParamIndexInObject);
 				}
 				else
 				{
-					IDetailGroup* ProjectorGroup = &CurrentSection->AddGroup(FName(*ParamInfo.ParameterName), FText::GetEmpty());
+					IDetailGroup* ProjectorGroup = &CurrentSection->AddGroup(FName(*ParamInfo.ParamName), FText::GetEmpty());
 
 					// Call Order between the following lines maters.
-					ParentsGroups.Add(ParamInfo.ParameterName, ProjectorGroup);
-					GenerateWidgetRow(ProjectorGroup->HeaderRow(), ParamInfo.ParameterName, ParamIndexInObject);
+					ParentsGroups.Add(ParamInfo.ParamName, ProjectorGroup);
+					GenerateWidgetRow(ProjectorGroup->HeaderRow(), ParamInfo.ParamName, ParamInfo.ParamIndexInObject);
 				}
 			}
 			else
 			{
-				if (!IsMultidimensionalProjector(ParamIndexInObject))
+				if (!IsMultidimensionalProjector(ParamInfo.ParamIndexInObject))
 				{
-					GenerateWidgetRow(DetailsCategory.AddCustomRow(FText::FromString(ParamInfo.ParameterName)), ParamInfo.ParameterName, ParamInfo.ParameterIndex);
+					GenerateWidgetRow(DetailsCategory.AddCustomRow(FText::FromString(ParamInfo.ParamName)), ParamInfo.ParamName, ParamInfo.ParamIndexInObject);
 				}
 				else
 				{
-					IDetailGroup* ProjectorGroup = &DetailsCategory.AddGroup(FName(*ParamInfo.ParameterName), FText::GetEmpty());
+					IDetailGroup* ProjectorGroup = &DetailsCategory.AddGroup(FName(*ParamInfo.ParamName), FText::GetEmpty());
 
 					// Call Order between the following lines maters.
-					ParentsGroups.Add(ParamInfo.ParameterName, ProjectorGroup);
-					GenerateWidgetRow(ProjectorGroup->HeaderRow(), ParamInfo.ParameterName, ParamIndexInObject);
+					ParentsGroups.Add(ParamInfo.ParamName, ProjectorGroup);
+					GenerateWidgetRow(ProjectorGroup->HeaderRow(), ParamInfo.ParamName, ParamInfo.ParamIndexInObject);
 				}
 			}
 		}
@@ -617,11 +621,11 @@ void FCustomizableInstanceDetails::GenerateParametersView(IDetailCategoryBuilder
 		for (int32 ParamIndexInObject = 0; ParamIndexInObject < NumObjectParameter; ++ParamIndexInObject)
 		{
 			FParameterInfo ParameterSortInfo;
-			ParameterSortInfo.ParameterIndex = ParamIndexInObject;
-			if ((!CustomInstance->GetPrivate()->bShowOnlyRelevantParameters || CustomInstance->IsParameterRelevant(ParameterSortInfo.ParameterIndex)) && IsVisible(ParameterSortInfo.ParameterIndex))
+			ParameterSortInfo.ParamIndexInObject = ParamIndexInObject;
+			if ((!CustomInstance->GetPrivate()->bShowOnlyRelevantParameters || CustomInstance->IsParameterRelevant(ParameterSortInfo.ParamIndexInObject)) && IsVisible(ParameterSortInfo.ParamIndexInObject))
 			{
-				ParameterSortInfo.ParameterUIOrder = CustomizableObject->GetParameterUIMetadataFromIndex(ParameterSortInfo.ParameterIndex).ParamUIMetadata.UIOrder;
-				ParameterSortInfo.ParameterName = CustomizableObject->GetParameterName(ParameterSortInfo.ParameterIndex);
+				ParameterSortInfo.ParamUIOrder = CustomizableObject->GetParameterUIMetadataFromIndex(ParameterSortInfo.ParamIndexInObject).ParamUIMetadata.UIOrder;
+				ParameterSortInfo.ParamName = CustomizableObject->GetParameterName(ParameterSortInfo.ParamIndexInObject);
 				ParametersTree.Add(ParameterSortInfo);
 			}
 		}
@@ -630,14 +634,14 @@ void FCustomizableInstanceDetails::GenerateParametersView(IDetailCategoryBuilder
 	
 		for (int32 ParamIndexInObject = 0; ParamIndexInObject < ParametersTree.Num(); ++ParamIndexInObject)
 		{
-			FillChildrenMap(ParametersTree[ParamIndexInObject].ParameterIndex);
+			FillChildrenMap(ParametersTree[ParamIndexInObject].ParamIndexInObject);
 		}
 	
 		for (int32 ParamIndexInObject = 0; ParamIndexInObject < ParametersTree.Num(); ++ParamIndexInObject)
 		{
-			if (!ParamHasParent.Find(ParametersTree[ParamIndexInObject].ParameterIndex))
+			if (!ParamHasParent.Find(ParametersTree[ParamIndexInObject].ParamIndexInObject))
 			{
-				RecursivelyAddParamAndChildren(ParametersTree[ParamIndexInObject].ParameterIndex, "", DetailsCategory);
+				RecursivelyAddParamAndChildren(ParametersTree[ParamIndexInObject].ParamIndexInObject, "", DetailsCategory);
 			}
 		}
 	}
@@ -964,6 +968,7 @@ TSharedRef<SWidget> FCustomizableInstanceDetails::GenerateFloatWidget(const int3
 		.MinValue(CustomizableObject->GetParameterUIMetadataFromIndex(ParamIndexInObject).ParamUIMetadata.MinimumValue)
 		.MaxValue(CustomizableObject->GetParameterUIMetadataFromIndex(ParamIndexInObject).ParamUIMetadata.MaximumValue)
 		.OnValueChanged(this, &FCustomizableInstanceDetails::OnFloatParameterChanged, ParamName, -1)
+		.OnBeginSliderMovement(this, &FCustomizableInstanceDetails::OnFloatParameterSliderBegin)
 		.OnEndSliderMovement(this, &FCustomizableInstanceDetails::OnFloatParameterSliderEnd, ParamName, -1);
 }
 
@@ -989,16 +994,29 @@ float FCustomizableInstanceDetails::GetFloatParameterValue(FString ParamName, in
 
 void FCustomizableInstanceDetails::OnFloatParameterChanged(float Value, FString ParamName, int32 RangeIndex)
 {
+	CustomInstance->PreEditChange(nullptr);
+	CustomInstance->GetPrivate()->SetSelectedParameterProfileDirty();
+
 	CustomInstance->SetFloatParameterSelectedOption(ParamName, Value, RangeIndex);
+	CustomInstance->UpdateSkeletalMeshAsync(true, true);
+	
+	CustomInstance->PostEditChange();
+}
+
+
+void FCustomizableInstanceDetails::OnFloatParameterSliderBegin()
+{
+	bUpdatingSlider = true;
 }
 
 
 void FCustomizableInstanceDetails::OnFloatParameterSliderEnd(float Value, FString ParamName, int32 RangeIndex)
 {
-	CustomInstance->SetFloatParameterSelectedOption(ParamName, Value, RangeIndex);
+	bUpdatingSlider = false;
 
-	CustomInstance->GetPrivate()->SetSelectedParameterProfileDirty();
 	CustomInstance->PreEditChange(nullptr);
+	CustomInstance->GetPrivate()->SetSelectedParameterProfileDirty();
+	CustomInstance->SetFloatParameterSelectedOption(ParamName, Value, RangeIndex);
 	CustomInstance->UpdateSkeletalMeshAsync(true, true);
 	CustomInstance->PostEditChange();
 
@@ -1382,7 +1400,7 @@ TSharedRef<SWidget> FCustomizableInstanceDetails::GenerateMultidimensionalProjec
 		const int32 NumPoseValues = CustomizableObject->GetIntParameterNumOptions(PoseSwitchEnumParamIndexInObject);
 
 		TArray<FString> PoseOptionNamesAttribute;
-		FString PoseValue = CustomInstance->GetIntParameterSelectedOption(ParamName, -1);
+		FString PoseValue = CustomInstance->GetIntParameterSelectedOption(PoseSwitchEnumParamName, -1);
 		int32 PoseValueIndex = 0;
 
 		for (int32 j = 0; j < NumPoseValues; ++j)
@@ -1609,6 +1627,7 @@ TSharedRef<SWidget> FCustomizableInstanceDetails::GenerateMultidimensionalProjec
 				.MaxValue(1.0f)
 				.Value(this, &FCustomizableInstanceDetails::GetFloatParameterValue, OpacitySliderParamName, RangeIndex)
 				.OnValueChanged(this, &FCustomizableInstanceDetails::OnFloatParameterChanged, OpacitySliderParamName, RangeIndex)
+				.OnBeginSliderMovement(this, &FCustomizableInstanceDetails::OnFloatParameterSliderBegin)
 				.OnEndSliderMovement(this, &FCustomizableInstanceDetails::OnFloatParameterSliderEnd, OpacitySliderParamName, RangeIndex)
 				.Font(LayoutBuilder.Pin()->GetDetailFont())
 			]
diff --git a/Engine/Plugins/Experimental/Mutable/Source/CustomizableObjectEditor/Private/MuCOE/CustomizableInstanceDetails.h b/Engine/Plugins/Experimental/Mutable/Source/CustomizableObjectEditor/Private/MuCOE/CustomizableInstanceDetails.h
--- a/Engine/Plugins/Experimental/Mutable/Source/CustomizableObjectEditor/Private/MuCOE/CustomizableInstanceDetails.h
+++ b/Engine/Plugins/Experimental/Mutable/Source/CustomizableObjectEditor/Private/MuCOE/CustomizableInstanceDetails.h
@@ -81,6 +81,7 @@ private:
 	TSharedRef<SWidget> GenerateFloatWidget(const int32 ParamIndexInObject);
 	float GetFloatParameterValue(const FString ParamName, int32 RangeIndex) const;
 	void OnFloatParameterChanged(float Value, const FString ParamName, int32 RangeIndex);
+	void OnFloatParameterSliderBegin();
 	void OnFloatParameterSliderEnd(float Value, const FString ParamName, int32 RangeIndex);
  
 	// Texture Parameters Functions
@@ -156,5 +157,8 @@ private:
 	/** Array with all the possible multilayer projector texture options */
 	TArray<TArray<TSharedPtr<FString>>> ProjectorTextureOptions;
 
+	/** True when a slider is being edited*/
+	bool bUpdatingSlider = false;
+
 };
 
