[mutable] Fixed crash when table nodes had an empty data table with a texture column connected to a material node.
#rnx
#tests Editor

diff --git a/Engine/Plugins/Experimental/Mutable/Source/MutableTools/Private/MuT/CodeGenerator.cpp b/Engine/Plugins/Experimental/Mutable/Source/MutableTools/Private/MuT/CodeGenerator.cpp
--- a/Engine/Plugins/Experimental/Mutable/Source/MutableTools/Private/MuT/CodeGenerator.cpp
+++ b/Engine/Plugins/Experimental/Mutable/Source/MutableTools/Private/MuT/CodeGenerator.cpp
@@ -1320,6 +1320,13 @@ namespace mu
 								GenerateImage(ImageOptions, Result, pImageNode);
 								Ptr<ASTOp> blockAd = Result.op;
 
+								if (!blockAd)
+								{
+									// The GenerateImage(...) above has failed, skip this block
+									result.surfaceOp = nullptr;
+									continue;
+								}
+
 								// Calculate the desc of the generated block
 								FImageDesc BlockDesc = blockAd->GetImageDesc();
 
@@ -1498,8 +1505,8 @@ namespace mu
 								imageAd = fop;
 							}
 
-							// Apply mipmap and format if necessary
-							if (mipmapNode)
+							// Apply mipmap and format if necessary, skip if format is IF_NONE (possibly because a block was skipped above)
+							if (mipmapNode && FinalFormat != EImageFormat::IF_NONE)
 							{
 								Ptr<ASTOpImageMipmap> mop = new ASTOpImageMipmap();
 
diff --git a/Engine/Plugins/Experimental/Mutable/Source/MutableTools/Private/MuT/CodeGenerator.h b/Engine/Plugins/Experimental/Mutable/Source/MutableTools/Private/MuT/CodeGenerator.h
--- a/Engine/Plugins/Experimental/Mutable/Source/MutableTools/Private/MuT/CodeGenerator.h
+++ b/Engine/Plugins/Experimental/Mutable/Source/MutableTools/Private/MuT/CodeGenerator.h
@@ -748,9 +748,17 @@ namespace mu
             GeneratedTables.Add(CacheKey, Variable );
         }
 
+		int32 NumRows = NodeTable->GetPrivate()->Rows.Num();
+
         // Verify that the table column is the right type
         int32 ColIndex = NodeTable->FindColumn( node.ColumnName );
-        if ( ColIndex<0 )
+
+		if (NumRows == 0)
+		{
+			m_pErrorLog->GetPrivate()->Add("The table has no rows.", ELMT_ERROR, node.m_errorContext);
+			return nullptr;
+		}
+        else if (ColIndex < 0)
         {
             m_pErrorLog->GetPrivate()->Add("Table column not found.", ELMT_ERROR, node.m_errorContext);
             return nullptr;
@@ -764,14 +772,12 @@ namespace mu
 
         // Create the switch to cover all the options
         Ptr<ASTOp> lastSwitch;
-        int32 Rows = NodeTable->GetPrivate()->Rows.Num();
-
         Ptr<ASTOpSwitch> SwitchOp = new ASTOpSwitch();
 		SwitchOp->type = OPTYPE;
 		SwitchOp->variable = Variable;
 		SwitchOp->def = GenerateDefaultTableValue(TYPE);
 
-		for (int32 i = 0; i < Rows; ++i)
+		for (int32 i = 0; i < NumRows; ++i)
         {
             check(NodeTable->GetPrivate()->Rows[i].Id <= 0xFFFF);
             auto Condition = (uint16)NodeTable->GetPrivate()->Rows[i].Id;
