[Mutable] Split CodeRunner in smaller tasks.

- CodeRunner profile data in now a optional passed context.
- CodeRunner can be run inline passing the bForceInlineExecution flag.
- GetImageInline, GetMeshInline and GetImageDescInline are provider as a replacement for the old functionality.
- GetImage, GetMesh, and GetImageDesc now return a task with result.

#rnx

diff --git a/Engine/Plugins/Experimental/Mutable/Source/CustomizableObject/Private/MuCO/BusyWaits_Deprecated.cpp b/Engine/Plugins/Experimental/Mutable/Source/CustomizableObject/Private/MuCO/BusyWaits_Deprecated.cpp
--- a/Engine/Plugins/Experimental/Mutable/Source/CustomizableObject/Private/MuCO/BusyWaits_Deprecated.cpp
+++ b/Engine/Plugins/Experimental/Mutable/Source/CustomizableObject/Private/MuCO/BusyWaits_Deprecated.cpp
@@ -33,7 +33,7 @@ namespace impl_deprecated
 
 			// This should only be done when using progressive images, since GetImageDesc does some actual processing.
 			{
-				System->GetImageDesc(OperationData->InstanceID, Image.ImageID, ImageDesc);
+				System->GetImageDescInline(OperationData->InstanceID, Image.ImageID, ImageDesc);
 
 				const uint16 MaxTextureSizeToGenerate = static_cast<uint16>(CustomizableObjectSystemPrivateData->MaxTextureSizeToGenerate);
 				const uint16 MaxSize = FMath::Max(ImageDesc.m_size[0], ImageDesc.m_size[1]);
@@ -80,7 +80,7 @@ namespace impl_deprecated
 				}
 				else
 				{
-					Image.Image = System->GetImage(OperationData->InstanceID, Image.ImageID, Image.BaseMip + MipsToSkip, Image.BaseLOD);
+					Image.Image = System->GetImageInline(OperationData->InstanceID, Image.ImageID, Image.BaseMip + MipsToSkip, Image.BaseLOD);
 				}
 
 				check(Image.Image);
@@ -230,7 +230,7 @@ namespace impl_deprecated
 
 					if(bGenerateLOD)
 					{
-						Component.Mesh = System->GetMesh(OperationData->InstanceID, Component.MeshID);
+						Component.Mesh = System->GetMeshInline(OperationData->InstanceID, Component.MeshID);
 					}
 				}
 
@@ -309,7 +309,7 @@ namespace impl_deprecated
 									Image.bIsPassThrough = true;
 
 									// Since it's known it's a pass-through texture there is no need to cache or convert it so we can generate it here already.
-									Image.Image = System->GetImage(OperationData->InstanceID, Image.ImageID, 0, 0);
+									Image.Image = System->GetImageInline(OperationData->InstanceID, Image.ImageID, 0, 0);
 									check(Image.Image->IsReference());
 
 									uint32 ReferenceID = Image.Image->GetReferencedTexture();
diff --git a/Engine/Plugins/Experimental/Mutable/Source/CustomizableObject/Private/MuCO/CustomizableObjectMipDataProvider.cpp b/Engine/Plugins/Experimental/Mutable/Source/CustomizableObject/Private/MuCO/CustomizableObjectMipDataProvider.cpp
--- a/Engine/Plugins/Experimental/Mutable/Source/CustomizableObject/Private/MuCO/CustomizableObjectMipDataProvider.cpp
+++ b/Engine/Plugins/Experimental/Mutable/Source/CustomizableObject/Private/MuCO/CustomizableObjectMipDataProvider.cpp
@@ -192,7 +192,7 @@ namespace impl
 					{
 						MUTABLE_CPUPROFILER_SCOPE(GetImage);
 
-						Image = System->GetImage(InstanceID, MipImageID, ImageRef.BaseMip + OperationData->MipsToSkip, ImageRef.LOD);
+						Image = System->GetImageInline(InstanceID, MipImageID, ImageRef.BaseMip + OperationData->MipsToSkip, ImageRef.LOD);
 					}
 
 					check(Image);
diff --git a/Engine/Plugins/Experimental/Mutable/Source/MutableRuntime/Internal/MuR/SystemPrivate.h b/Engine/Plugins/Experimental/Mutable/Source/MutableRuntime/Internal/MuR/SystemPrivate.h
--- a/Engine/Plugins/Experimental/Mutable/Source/MutableRuntime/Internal/MuR/SystemPrivate.h
+++ b/Engine/Plugins/Experimental/Mutable/Source/MutableRuntime/Internal/MuR/SystemPrivate.h
@@ -1121,7 +1121,7 @@ namespace mu
 
     /** Struct to manage all the memory allocated for resources used during mutable operation. */
     struct FWorkingMemoryManager
-    {
+    {	
 		using MemoryCounter = MemoryCounters::FMemoryTrackerInternalMemoryCounter;
 
 		template<class Type>
@@ -1789,38 +1789,69 @@ namespace mu
 
 
 #if !(UE_BUILD_SHIPPING || UE_BUILD_TEST)
+		static constexpr uint64 InvalidRunnerId = 0;
+		std::atomic<uint64> DebugCurrentRunnerId {1};
+		
 		/** Temp variable with the ID of the thread running code, for debugging. */
 		uint32 DebugRunnerThreadID = FThread::InvalidThreadId;
+		uint64 DebugRunnerID = InvalidRunnerId;
 #endif
 
 		/** This is a development-only check to make sure calls to resource management happen in the correct thread. */
-		inline void BeginRunnerThread()
+		FORCEINLINE void BeginRunnerThread()
 		{
 #if !(UE_BUILD_SHIPPING || UE_BUILD_TEST)
 			// If this check fails it means have not set up correctly all the paths to debug threading for resource management.
 			check(DebugRunnerThreadID == FThread::InvalidThreadId);
+			check(DebugRunnerID == InvalidRunnerId)
 
 			DebugRunnerThreadID = FPlatformTLS::GetCurrentThreadId();
+			DebugRunnerID = ++DebugCurrentRunnerId;
 #endif
 		}
 
+		/** This is a development-only check to make sure calls to resource management happen in the correct thread. */
+		FORCEINLINE void ResetRunnerThread()
+		{
+#if !(UE_BUILD_SHIPPING || UE_BUILD_TEST)
+			// If this check fails it means have not set up correctly all the paths to debug threading for resource management.
+			check(DebugRunnerThreadID == FThread::InvalidThreadId);
+			check(DebugRunnerID != InvalidRunnerId)
+
+			DebugRunnerThreadID = FPlatformTLS::GetCurrentThreadId();
+#endif
+		}
+
+		/** This is a development-only check to make sure calls to resource management happen in the correct thread. */
+		FORCEINLINE void InvalidateRunnerThread()
+		{
+#if !(UE_BUILD_SHIPPING || UE_BUILD_TEST)
+			// If this check fails it means have not set up correctly all the paths to debug threading for resource management.
+			check(DebugRunnerThreadID != FThread::InvalidThreadId);
+			check(DebugRunnerID != InvalidRunnerId)
+
+			DebugRunnerThreadID = FThread::InvalidThreadId;
+#endif
+		}
 
 		/** This is a development-only check to make sure calls to resource management happen in the correct thread. */
-		inline void CheckRunnerThread()
+		FORCEINLINE void CheckRunnerThread()
 		{
 #if !(UE_BUILD_SHIPPING || UE_BUILD_TEST)
 			// If this check fails it means have not set up correctly all the paths to debug threading for resource management.
 			check(DebugRunnerThreadID != FThread::InvalidThreadId);
+			check(DebugRunnerID != InvalidRunnerId);
 
 			// If this check fails it means we are doing resource management from a thread that is not the CodeRunner::RunCode
 			// thread, and this is not allowed.
-			check(DebugRunnerThreadID== FPlatformTLS::GetCurrentThreadId());
+			check(DebugRunnerThreadID == FPlatformTLS::GetCurrentThreadId());
+			check(DebugRunnerID == DebugCurrentRunnerId);
 #endif
 		}
 
 
 		/** This is a development-only check to make sure calls to resource management happen in the correct thread. */
-		inline void EndRunnerThread()
+		FORCEINLINE void EndRunnerThread()
 		{
 			CurrentInstanceCache->CheckHitCountsCleared();
 
@@ -1833,8 +1864,10 @@ namespace mu
 #if !(UE_BUILD_SHIPPING || UE_BUILD_TEST)
 			// If this check fails it means have not set up correctly all the paths to debug threading for resource management.
 			check(DebugRunnerThreadID != FThread::InvalidThreadId);
+			check(DebugRunnerID != InvalidRunnerId);
 
 			DebugRunnerThreadID = FThread::InvalidThreadId;
+			DebugRunnerID = InvalidRunnerId;
 #endif
 		}
 
@@ -1844,6 +1877,7 @@ namespace mu
 	/** */
     class System::Private
     {
+		friend class System;
     public:
 
         Private( Ptr<Settings>, const TSharedPtr<ExtensionDataStreamer>& );
diff --git a/Engine/Plugins/Experimental/Mutable/Source/MutableRuntime/Private/MuR/CodeRunner.cpp b/Engine/Plugins/Experimental/Mutable/Source/MutableRuntime/Private/MuR/CodeRunner.cpp
--- a/Engine/Plugins/Experimental/Mutable/Source/MutableRuntime/Private/MuR/CodeRunner.cpp
+++ b/Engine/Plugins/Experimental/Mutable/Source/MutableRuntime/Private/MuR/CodeRunner.cpp
@@ -106,8 +106,21 @@ static FAutoConsoleVariableRef CVarUseImageTransformVectorImpl (
 namespace mu
 {
 
-    //---------------------------------------------------------------------------------------------
-    CodeRunner::CodeRunner(const Ptr<const Settings>& InSettings,
+	TSharedRef<CodeRunner> CodeRunner::Create(
+		const Ptr<const Settings>& InSettings,
+		class System::Private* InSystem,
+		EExecutionStrategy InExecutionStrategy,
+		const TSharedPtr<const Model>& InModel,
+		const Parameters* InParams,
+		OP::ADDRESS At,
+		uint32 InLODMask, uint8 ExecutionOptions, int32 InImageLOD, FScheduledOp::EType Type)
+	{
+		return MakeShared<CodeRunner>(FPrivateToken {},
+				InSettings, InSystem, InExecutionStrategy, InModel, InParams, At, InLODMask, ExecutionOptions, InImageLOD, Type);
+	}
+
+    CodeRunner::CodeRunner(FPrivateToken PrivateToken, 
+		const Ptr<const Settings>& InSettings,
 		class System::Private* InSystem,
 		EExecutionStrategy InExecutionStrategy,
 		const TSharedPtr<const Model>& InModel,
@@ -115,6 +128,7 @@ namespace mu
 		OP::ADDRESS at,
 		uint32 InLodMask, uint8 executionOptions, int32 InImageLOD, FScheduledOp::EType Type )
 		: m_pSettings(InSettings)
+		, RunnerCompletionEvent(TEXT("CodeRunnerCompletioneEventInit"))
 		, ExecutionStrategy(InExecutionStrategy)
 		, m_pSystem(InSystem)
 		, m_pModel(InModel)
@@ -129,10 +143,13 @@ namespace mu
    		{
 			GetMemory().IncreaseHitCount(FCacheAddress(at, 0, executionOptions));
 		}
-    	
+    
+		// Start with a completed Event. This is checked at StartRun() to make sure StartRun is not called while there is 
+		// a Run in progress.
+		RunnerCompletionEvent.Trigger();
 
 		ImageLOD = InImageLOD;
-
+	
 		// Push the root operation
 		FScheduledOp rootOp;
 		rootOp.At = at;
@@ -5018,8 +5035,8 @@ namespace mu
             }
             break;
         }
-    }
-    
+    }	
+	
 	//---------------------------------------------------------------------------------------------
     Ptr<RangeIndex> CodeRunner::BuildCurrentOpRangeIndex( const FScheduledOp& item, const Parameters* pParams, const Model* pModel, int32 parameterIndex )
     {
@@ -5347,7 +5364,6 @@ namespace mu
         }
     }
 
-
     //---------------------------------------------------------------------------------------------
     void CodeRunner::RunCode_Scalar(const FScheduledOp& item, const Parameters* pParams, const Model* pModel )
     {
@@ -5464,7 +5480,6 @@ namespace mu
         }
     }
 
-
     //---------------------------------------------------------------------------------------------
     void CodeRunner::RunCode_String(const FScheduledOp& item, const Parameters* pParams, const Model* pModel )
     {
diff --git a/Engine/Plugins/Experimental/Mutable/Source/MutableRuntime/Private/MuR/CodeRunner.h b/Engine/Plugins/Experimental/Mutable/Source/MutableRuntime/Private/MuR/CodeRunner.h
--- a/Engine/Plugins/Experimental/Mutable/Source/MutableRuntime/Private/MuR/CodeRunner.h
+++ b/Engine/Plugins/Experimental/Mutable/Source/MutableRuntime/Private/MuR/CodeRunner.h
@@ -36,15 +36,36 @@ namespace  mu
 	class RangeIndex;
 
     /** Code execution of the mutable virtual machine. */
-    class CodeRunner
+    class CodeRunner : public TSharedFromThis<CodeRunner>
     {
-    public:
-		CodeRunner(const Ptr<const Settings>&, class System::Private*, 
+		// The private token allows only members or friends to call MakeShared.
+		struct FPrivateToken { explicit FPrivateToken() = default; };
+
+	public:
+		static TSharedRef<CodeRunner> Create(
+				const Ptr<const Settings>&, 
+				class System::Private*, 
+				EExecutionStrategy,
+				const TSharedPtr<const Model>&, 
+				const Parameters* pParams,
+				OP::ADDRESS at, uint32 lodMask, uint8 executionOptions, int32 InImageLOD, FScheduledOp::EType);
+
+		// Private constructor to prevent stack allocation. In general we can not call AsShared() if the lifetime is
+		// bounded.
+		explicit CodeRunner(FPrivateToken, 
+			const Ptr<const Settings>&, 
+			class System::Private*, 
 			EExecutionStrategy,
-			const TSharedPtr<const Model>&, const Parameters* pParams,
-			OP::ADDRESS at, uint32 lodMask, uint8 executionOptions, int32 InImageLOD, FScheduledOp::EType );
+			const TSharedPtr<const Model>&, 
+			const Parameters* pParams,
+			OP::ADDRESS at, uint32 lodMask, uint8 executionOptions, int32 InImageLOD, FScheduledOp::EType);
 
     protected:
+		struct FProfileContext
+		{
+			uint32 NumRunOps = 0;
+			uint32 RunOpsPerType[int32(OP_TYPE::COUNT)] = {};
+		};
 
         /** Type of data sometimes stored in the code runner heap to pass info between operation stages. */
         struct FScheduledOpData
@@ -135,15 +156,19 @@ namespace  mu
     protected:
         //! Heap of intermediate data pushed by some instructions and referred by others.
         //! It is not released until no operations are pending.
-		TArray< FScheduledOpData > m_heapData;
-		TArray< FImageDesc > m_heapImageDesc;
+		TArray<FScheduledOpData> m_heapData;
+		TArray<FImageDesc> m_heapImageDesc;
 
 		/** Only used for correct mip skipping with external images. It is the LOD for which the image is build. */
 		int32 ImageLOD;
 
+		UE::Tasks::FTaskEvent RunnerCompletionEvent;
+
+		void Run(TUniquePtr<FProfileContext>&& ProfileContext, bool bForceInlineExecution);
+		void AbortRun();
 	public:
 
-		void Run();
+		UE::Tasks::FTask StartRun(bool bForceInlineExecution);
 
 		//!
 		void GetImageDescResult(FImageDesc& OutDesc);
diff --git a/Engine/Plugins/Experimental/Mutable/Source/MutableRuntime/Private/MuR/CodeRunnerTasked.cpp b/Engine/Plugins/Experimental/Mutable/Source/MutableRuntime/Private/MuR/CodeRunnerTasked.cpp
--- a/Engine/Plugins/Experimental/Mutable/Source/MutableRuntime/Private/MuR/CodeRunnerTasked.cpp
+++ b/Engine/Plugins/Experimental/Mutable/Source/MutableRuntime/Private/MuR/CodeRunnerTasked.cpp
@@ -46,15 +46,6 @@
 
 DECLARE_CYCLE_STAT(TEXT("MutableCoreTask"), STAT_MutableCoreTask, STATGROUP_Game);
 
-
-int32 CoreRunnerTaskPriority = 5;
-
-static FAutoConsoleVariableRef CVarCoreRunnerTaskPriority(
-	TEXT("mutable.CoreRunnerTaskPriority"),
-	CoreRunnerTaskPriority,
-	TEXT("0 AnyThread. 1 to 6, from hight to low priority."),
-	ECVF_Default);
-
 namespace mu
 {
 
@@ -273,21 +264,42 @@ namespace mu
 		IssuedTasks.Add(TaskToIssue);
 	}
 
+	UE::Tasks::FTask CodeRunner::StartRun(bool bForceInlineExecution)
+	{
+		check(RunnerCompletionEvent.IsCompleted());
 
-    //---------------------------------------------------------------------------------------------
-    void CodeRunner::Run()
-    {
 		bUnrecoverableError = false;
 
 		m_heapData.SetNum(0, EAllowShrinking::No);
 		m_heapImageDesc.SetNum(0, EAllowShrinking::No);
-
-		// Profiling data. TODO: make optional
+		
+		RunnerCompletionEvent = UE::Tasks::FTaskEvent(TEXT("CodeRunnerCompletionEvent"));
+		
 		bool bProfile = false;
-		uint32 NumRunOps = 0;
-		uint32 RunOpsPerType[size_t(OP_TYPE::COUNT)] = {};
+		
+		TUniquePtr<FProfileContext> ProfileContext = bProfile ? MakeUnique<FProfileContext>() : nullptr;
+		Run(MoveTemp(ProfileContext), bForceInlineExecution);
+
+		check(!bForceInlineExecution || RunnerCompletionEvent.IsCompleted());
+
+		return RunnerCompletionEvent;
+	}
+
+	void CodeRunner::AbortRun()
+	{
+		bUnrecoverableError = true;
+		RunnerCompletionEvent.Trigger();
+	}
+
+    void CodeRunner::Run(TUniquePtr<FProfileContext>&& ProfileContext, bool bForceInlineExecution)
+    {
+		MUTABLE_CPUPROFILER_SCOPE(TEXT("CodeRunner_Run"));
+
+		check(!RunnerCompletionEvent.IsCompleted());
 
-        while( !OpenTasks.IsEmpty() || !ClosedTasks.IsEmpty() || !IssuedTasks.IsEmpty())
+		const double StartSeconds = FPlatformTime::Seconds();
+
+        while(!OpenTasks.IsEmpty() || !ClosedTasks.IsEmpty() || !IssuedTasks.IsEmpty())
         {
 			UpdateTraces();
 			// Debug: log the amount of tasks that we'd be able to run concurrently:
@@ -308,12 +320,12 @@ namespace mu
 			//	UE_LOG(LogMutableCore, Log, TEXT("Tasks: %5d open, %5d issued, %5d closed, %d closed ready"), OpenTasks.Num(), IssuedTasks.Num(), ClosedTasks.Num(), ClosedReady);
 			//}
 
-			for (int Index = 0; Index < IssuedTasks.Num(); )
+			for (int32 Index = 0; Index < IssuedTasks.Num(); )
 			{
 				check(IssuedTasks[Index]);
 
-				bool WorkDone = IssuedTasks[Index]->IsComplete(this);
-				if (WorkDone)
+				bool bWorkDone = IssuedTasks[Index]->IsComplete(this);
+				if (bWorkDone)
 				{
 					const FScheduledOp& item = IssuedTasks[Index]->Op;
 					IssuedTasks[Index]->Complete(this);
@@ -326,7 +338,7 @@ namespace mu
 						ScheduledStagePerOp[item] = 0;
 					}
 
-					IssuedTasks.RemoveAt(Index,EAllowShrinking::No); // with swap? changes order of execution.
+					IssuedTasks.RemoveAt(Index, EAllowShrinking::No); // with swap? changes order of execution.
 				}
 				else
 				{
@@ -379,7 +391,7 @@ namespace mu
 				}
 
 				// Special processing in case it is an ImageDesc operation
-				if ( item.Type==FScheduledOp::EType::ImageDesc )
+				if (item.Type == FScheduledOp::EType::ImageDesc)
 				{
 					RunCodeImageDesc(item, m_pParams, m_pModel.Get(), m_lodMask);
 					continue;
@@ -402,8 +414,7 @@ namespace mu
 						LaunchIssuedTask(IssuedTask, bFailed);
 						if (bFailed)
 						{
-							bUnrecoverableError = true;
-							return;
+							return AbortRun();
 						}
 					}
 					else
@@ -424,10 +435,10 @@ namespace mu
 					}
 				}
 
-				if (bProfile)
+				if (ProfileContext)
 				{
-					++NumRunOps;
-					RunOpsPerType[size_t(m_pModel->GetPrivate()->m_program.GetOpType(item.At))]++;
+					++ProfileContext->NumRunOps;
+					++ProfileContext->RunOpsPerType[int32(m_pModel->GetPrivate()->m_program.GetOpType(item.At))];
 				}
 			}
 
@@ -442,14 +453,13 @@ namespace mu
 				LaunchIssuedTask(TaskToIssue, bFailed);
 				if (bFailed)
 				{
-					bUnrecoverableError = true;
-					return;
+					return AbortRun();
 				}
 			}
 
 			// Look for a closed task with dependencies satisfied and move them to the open task list.
 			bool bSomeWasReady = false;
-			for (int Index = 0; Index<ClosedTasks.Num(); )
+			for (int32 Index = 0; Index < ClosedTasks.Num(); )
 			{
 				bool Ready = true;
 				for ( const FCacheAddress& Dep: ClosedTasks[Index].Deps )
@@ -485,7 +495,7 @@ namespace mu
 			if (bDeadLock)
 			{
 				// Log the task graph
-				for (int Index = 0; Index < ClosedTasks.Num(); ++Index)
+				for (int32 Index = 0; Index < ClosedTasks.Num(); ++Index)
 				{
 					FString TaskDesc = FString::Printf(TEXT("Closed task %d-%d-%d depends on : "), ClosedTasks[Index].Op.At, ClosedTasks[Index].Op.ExecutionIndex, ClosedTasks[Index].Op.Stage );
 					for (const FCacheAddress& Dep : ClosedTasks[Index].Deps)
@@ -501,7 +511,7 @@ namespace mu
 				check(false);
 
 				// This should never happen but if it does, abort the code execution.
-				return;
+				return AbortRun();
 			}
 
 			// If at this point there is no open op and we haven't finished, we need to wait for an issued op to complete.
@@ -509,56 +519,109 @@ namespace mu
 			{
 				MUTABLE_CPUPROFILER_SCOPE(CodeRunner_WaitIssued);
 
-				for (int32 IssuedIndex = 0; IssuedIndex<IssuedTasks.Num(); ++IssuedIndex)
+				for (int32 IssuedIndex = 0; IssuedIndex < IssuedTasks.Num(); ++IssuedIndex)
 				{
 					if (IssuedTasks[IssuedIndex]->Event.IsValid())
 					{
-						if (CVarTaskGraphBusyWait->GetBool())
+						if (!bForceInlineExecution)
 						{
-							IssuedTasks[IssuedIndex]->Event.BusyWait();							
+							TArray<UE::Tasks::FTask, TInlineAllocator<8>> IssuedTasksCompletionEvents;
+							IssuedTasksCompletionEvents.Reserve(IssuedTasks.Num());
+
+							for (TSharedPtr<FIssuedTask>& IssuedTask : IssuedTasks)
+							{
+								IssuedTasksCompletionEvents.Add(IssuedTask->Event);
+							}
+
+							m_pSystem->WorkingMemoryManager.InvalidateRunnerThread();
+
+							UE::Tasks::Launch(TEXT("CodeRunnerFromIssuedTasksTask"),
+								[Runner = AsShared(), ProfileContext = MoveTemp(ProfileContext)]() mutable
+								{
+									Runner->m_pSystem->WorkingMemoryManager.ResetRunnerThread();
+
+									constexpr bool bForceInlineExecution = false;
+									Runner->Run(MoveTemp(ProfileContext), bForceInlineExecution);
+								},
+								UE::Tasks::Prerequisites(UE::Tasks::Any(IssuedTasksCompletionEvents)),
+								UE::Tasks::ETaskPriority::Inherit);
+							
+							return;
 						}
 						else
 						{
-							IssuedTasks[IssuedIndex]->Event.Wait();
+							if (CVarTaskGraphBusyWait->GetBool())
+							{
+								IssuedTasks[IssuedIndex]->Event.BusyWait();							
+							}
+							else
+							{
+								IssuedTasks[IssuedIndex]->Event.Wait();
+							}
+
+							break;
 						}
-
-						break;
 					}
 				}
 			}
+
+			if (!bForceInlineExecution)
+			{
+				// TODO: Move the timeout somewhere else more accessible, maybe a cvar.
+				const FTimespan Timeout = FTimespan::FromMilliseconds(2.0);
+				const FTimespan Duration = FTimespan::FromSeconds(FPlatformTime::Seconds() - StartSeconds);
+				if (Duration > Timeout)
+				{
+					m_pSystem->WorkingMemoryManager.InvalidateRunnerThread();
+
+					UE::Tasks::Launch(TEXT("CodeRunnerFromTimeoutTask"),
+						[Runner = AsShared(), ProfileContext = MoveTemp(ProfileContext)]() mutable
+						{
+							Runner->m_pSystem->WorkingMemoryManager.ResetRunnerThread();
+
+							constexpr bool bForceInlineExecution = false;
+							Runner->Run(MoveTemp(ProfileContext), bForceInlineExecution);
+						},
+						UE::Tasks::ETaskPriority::Inherit);
+					
+					return;
+				}
+			}
 		}
 
-		if (bProfile)
+		if (ProfileContext)
 		{
 			UE_LOG(LogMutableCore, Log, TEXT("Mutable Heap Bytes: %d"), m_heapData.Num()* m_heapData.GetTypeSize());
-			UE_LOG(LogMutableCore, Log, TEXT("Ran ops : %5d "), NumRunOps);
+			UE_LOG(LogMutableCore, Log, TEXT("Ran ops : %5d "), ProfileContext->NumRunOps);
 
-			constexpr int HistogramSize = 8;
+			constexpr int32 HistogramSize = 8;
 			int32 MostCommonOps[HistogramSize] = {};
-			for (size_t i=0; i<size_t(OP_TYPE::COUNT); ++i)
+			for (int32 OpIndex = 0; OpIndex < int32(OP_TYPE::COUNT); ++OpIndex)
 			{
-				for (int32 h=0; h<HistogramSize;++h)
+				for (int32 HistIndex = 0; HistIndex < HistogramSize; ++HistIndex)
 				{
-					if (RunOpsPerType[i] > RunOpsPerType[MostCommonOps[h]])
+					if (ProfileContext->RunOpsPerType[OpIndex] > ProfileContext->RunOpsPerType[MostCommonOps[HistIndex]])
 					{
 						// Displace others
-						int32 ElementsToMove = HistogramSize - h - 1;
+						int32 ElementsToMove = HistogramSize - HistIndex - 1;
 						if (ElementsToMove > 0)
 						{
-							FMemory::Memcpy(&MostCommonOps[h + 1], &MostCommonOps[h], sizeof(int32)* ElementsToMove);
+							FMemory::Memcpy(&MostCommonOps[HistIndex + 1], &MostCommonOps[HistIndex], sizeof(int32)*ElementsToMove);
 						}
 						// Set new value
-						MostCommonOps[h] = i;
+						MostCommonOps[HistIndex] = OpIndex;
 						break;
 					}
 				}
 			}
 
-			for (int h = 0; h < HistogramSize; ++h)
+			for (int32 HistIndex = 0; HistIndex < HistogramSize; ++HistIndex)
 			{
-				UE_LOG(LogMutableCore, Log, TEXT("    op %4d, %4d times."), MostCommonOps[h], RunOpsPerType[MostCommonOps[h]]);
+				UE_LOG(LogMutableCore, Log, TEXT("    op %4d, %4d times."), MostCommonOps[HistIndex], ProfileContext->RunOpsPerType[MostCommonOps[HistIndex]]);
 			}
 		}
+
+		RunnerCompletionEvent.Trigger();
     }
 
 	
diff --git a/Engine/Plugins/Experimental/Mutable/Source/MutableRuntime/Private/MuR/System.cpp b/Engine/Plugins/Experimental/Mutable/Source/MutableRuntime/Private/MuR/System.cpp
--- a/Engine/Plugins/Experimental/Mutable/Source/MutableRuntime/Private/MuR/System.cpp
+++ b/Engine/Plugins/Experimental/Mutable/Source/MutableRuntime/Private/MuR/System.cpp
@@ -332,8 +332,7 @@ namespace mu
 	}
 
 
-	//---------------------------------------------------------------------------------------------
-	Ptr<const Image> System::GetImage(Instance::ID instanceID, FResourceID ImageId, int32 MipsToSkip, int32 InImageLOD)
+	Ptr<const Image> System::GetImageInline(Instance::ID instanceID, FResourceID ImageId, int32 MipsToSkip, int32 InImageLOD)
 	{
 		LLM_SCOPE_BYNAME(TEXT("MutableRuntime"));
 		MUTABLE_CPUPROFILER_SCOPE(SystemGetImage);
@@ -359,6 +358,65 @@ namespace mu
 		return pResult;
 	}
 
+	UE::Tasks::TTask<Ptr<const Image>> System::GetImage(Instance::ID instanceID, FResourceID ImageId, int32 MipsToSkip, int32 InImageLOD)
+	{
+		LLM_SCOPE_BYNAME(TEXT("MutableRuntime"));
+		MUTABLE_CPUPROFILER_SCOPE(SystemGetImage);
+
+		Ptr<const Image> pResult;
+
+		// Find the live instance
+		FLiveInstance* pLiveInstance = m_pD->FindLiveInstance(instanceID);
+		check(pLiveInstance);
+		m_pD->WorkingMemoryManager.CurrentInstanceCache = pLiveInstance->Cache;
+
+		OP::ADDRESS RootAddress = GetResourceIDRoot(ImageId);
+		
+		m_pD->WorkingMemoryManager.BeginRunnerThread();
+		
+		mu::OP_TYPE OpType = pLiveInstance->Model->GetPrivate()->m_program.GetOpType(RootAddress);
+		if (GetOpDataType(OpType) != DT_IMAGE)
+		{
+			m_pD->WorkingMemoryManager.EndRunnerThread();
+			m_pD->WorkingMemoryManager.CurrentInstanceCache = nullptr;
+
+			return UE::Tasks::MakeCompletedTask<Ptr<const Image>>(
+					new mu::Image(16, 16, 1, EImageFormat::IF_RGBA_UBYTE, EInitializationType::Black));
+		}
+		
+		TSharedRef<CodeRunner> Runner = CodeRunner::Create(
+				m_pD->Settings, m_pD, EExecutionStrategy::MinimizeMemory, pLiveInstance->Model, pLiveInstance->OldParameters.get(), RootAddress, System::AllLODs, MipsToSkip, InImageLOD, FScheduledOp::EType::Full);
+	
+		
+		constexpr bool bForceInlineExecution = false;
+		UE::Tasks::FTask RunnerCompletionEvent = Runner->StartRun(bForceInlineExecution);
+
+		return UE::Tasks::Launch(TEXT("System::GetImageResultTask"),
+				[SystemPrivate = m_pD, Runner, RootAddress, MipsToSkip]() -> Ptr<const Image>
+				{
+					Ptr<const Image> Result;
+
+					SystemPrivate->bUnrecoverableError = Runner->bUnrecoverableError;
+					if (!Runner->bUnrecoverableError)
+					{
+						Result = SystemPrivate->WorkingMemoryManager.LoadImage(FCacheAddress(RootAddress, 0, MipsToSkip), true);
+					}
+
+					if (!Result)
+					{
+						Result = new mu::Image(16, 16, 1, EImageFormat::IF_RGBA_UBYTE, EInitializationType::Black);
+					}
+
+					SystemPrivate->WorkingMemoryManager.EndRunnerThread();
+					SystemPrivate->WorkingMemoryManager.CurrentInstanceCache = nullptr;
+					
+					return Result;
+				},
+				UE::Tasks::Prerequisites(RunnerCompletionEvent),
+				UE::Tasks::ETaskPriority::Inherit,
+				UE::Tasks::EExtendedTaskPriority::Inline);
+	}
+
 
 	// Temporarily make the Image DescCache clear at every image because otherwise it makes some textures 
 	// not evaluate their layout and be of size 0 and 0 lods, making them incorrectly evaluate MipsToSkip
@@ -368,9 +426,7 @@ namespace mu
 		TEXT("If different than 0, clear the image desc cache at every image."),
 		ECVF_Scalability);
 
-
-	//---------------------------------------------------------------------------------------------
-	void System::GetImageDesc(Instance::ID instanceID, FResourceID ImageId, FImageDesc& OutDesc)
+	void System::GetImageDescInline(Instance::ID instanceID, FResourceID ImageId, FImageDesc& OutDesc)
 	{
 		LLM_SCOPE_BYNAME(TEXT("MutableRuntime"));
 		MUTABLE_CPUPROFILER_SCOPE(SystemGetImageDesc);
@@ -401,9 +457,14 @@ namespace mu
 			m_pD->WorkingMemoryManager.BeginRunnerThread();
 					
 			int8 executionOptions = 0;
-			CodeRunner Runner(m_pD->Settings, m_pD, EExecutionStrategy::MinimizeMemory, pLiveInstance->Model, pLiveInstance->OldParameters.get(), RootAddress, System::AllLODs, executionOptions, 0, FScheduledOp::EType::ImageDesc);
-			Runner.Run();
-			Runner.GetImageDescResult(OutDesc);
+			TSharedRef<CodeRunner> Runner = CodeRunner::Create(
+					m_pD->Settings, m_pD, EExecutionStrategy::MinimizeMemory, pLiveInstance->Model, pLiveInstance->OldParameters.get(), RootAddress, System::AllLODs, executionOptions, 0, FScheduledOp::EType::ImageDesc);
+			
+			constexpr bool bForceInlineExecution = true;
+			UE::Tasks::FTask CompletionEvent = Runner->StartRun(bForceInlineExecution);
+			check(CompletionEvent.IsCompleted());
+
+			Runner->GetImageDescResult(OutDesc);
 
 			m_pD->WorkingMemoryManager.EndRunnerThread();
 		}
@@ -412,13 +473,70 @@ namespace mu
 	}
 
 
-    //---------------------------------------------------------------------------------------------
-    MeshPtrConst System::GetMesh( Instance::ID instanceID, FResourceID MeshId )
+	UE::Tasks::TTask<FImageDesc> System::GetImageDesc(Instance::ID instanceID, FResourceID ImageId)
+	{
+		LLM_SCOPE_BYNAME(TEXT("MutableRuntime"));
+		MUTABLE_CPUPROFILER_SCOPE(SystemGetImageDesc);
+
+		// Find the live instance
+		FLiveInstance* pLiveInstance = m_pD->FindLiveInstance(instanceID);
+		check(pLiveInstance);
+		m_pD->WorkingMemoryManager.CurrentInstanceCache = pLiveInstance->Cache;
+
+		OP::ADDRESS RootAddress = GetResourceIDRoot(ImageId);
+
+		const mu::Model* Model = pLiveInstance->Model.Get();
+		const mu::FProgram& program = Model->GetPrivate()->m_program;
+
+		// TODO: It should be possible to reuse this data if cleared in the correct places only, together with m_heapImageDesc.
+		int32 VarValue = CVarClearImageDescCache.GetValueOnAnyThread();
+		if (VarValue != 0)
+		{
+			m_pD->WorkingMemoryManager.CurrentInstanceCache->ClearDescCache();
+		}
+
+		m_pD->WorkingMemoryManager.BeginRunnerThread();
+		mu::OP_TYPE OpType = program.GetOpType(RootAddress);
+		if (GetOpDataType(OpType) != DT_IMAGE)
+		{
+			m_pD->WorkingMemoryManager.EndRunnerThread();
+			m_pD->WorkingMemoryManager.CurrentInstanceCache = nullptr;
+
+			return UE::Tasks::MakeCompletedTask<FImageDesc>(); 
+		}
+
+		// GetImageDesc may call normal execution paths where meshes are computed.
+		int8 ExecutionOptions = 0;
+		TSharedRef<CodeRunner> Runner = CodeRunner::Create(
+				m_pD->Settings, m_pD, EExecutionStrategy::MinimizeMemory, pLiveInstance->Model, pLiveInstance->OldParameters.get(), RootAddress, System::AllLODs, ExecutionOptions, 0, FScheduledOp::EType::ImageDesc);
+		
+		constexpr bool bForceInlineExecution = false;
+		UE::Tasks::FTask RunnerCompletionEvent = Runner->StartRun(bForceInlineExecution);
+		
+		return UE::Tasks::Launch(TEXT("System::GetImageDescResultTask"),
+				[SystemPrivate = m_pD, Runner]() -> FImageDesc
+				{
+					FImageDesc Result;
+					Runner->GetImageDescResult(Result);
+
+					SystemPrivate->WorkingMemoryManager.EndRunnerThread();
+					SystemPrivate->WorkingMemoryManager.CurrentInstanceCache = nullptr;
+
+					return Result;
+				},
+				UE::Tasks::Prerequisites(RunnerCompletionEvent),
+				UE::Tasks::ETaskPriority::Inherit,
+				UE::Tasks::EExtendedTaskPriority::Inline);
+
+	}
+
+
+    Ptr<const Mesh> System::GetMeshInline(Instance::ID instanceID, FResourceID MeshId)
     {
 		LLM_SCOPE_BYNAME(TEXT("MutableRuntime"));
 		MUTABLE_CPUPROFILER_SCOPE(SystemGetMesh);
 
-		MeshPtrConst pResult;
+		Ptr<const Mesh> Result;
 
 		// Find the live instance
 		FLiveInstance* pLiveInstance = m_pD->FindLiveInstance(instanceID);
@@ -426,16 +544,72 @@ namespace mu
 		m_pD->WorkingMemoryManager.CurrentInstanceCache = pLiveInstance->Cache;
 
 		OP::ADDRESS RootAddress = GetResourceIDRoot(MeshId);
-		pResult = m_pD->BuildMesh(pLiveInstance->Model, pLiveInstance->OldParameters.get(), RootAddress);
+		Result = m_pD->BuildMesh(pLiveInstance->Model, pLiveInstance->OldParameters.get(), RootAddress);
 
 		// If the mesh is null it means empty, but we still need to return a valid one
-		if (!pResult)
+		if (!Result)
 		{
-			pResult = new Mesh();
+			Result = new Mesh();
 		}
 
 		m_pD->WorkingMemoryManager.CurrentInstanceCache = nullptr;
-		return pResult;
+		return Result;
+	}
+
+	UE::Tasks::TTask<Ptr<const Mesh>> System::GetMesh(Instance::ID instanceID, FResourceID MeshId)
+    {
+		LLM_SCOPE_BYNAME(TEXT("MutableRuntime"));
+		MUTABLE_CPUPROFILER_SCOPE(SystemGetImage);
+
+		Ptr<const Mesh> ResultMesh;
+
+		// Find the live instance
+		FLiveInstance* pLiveInstance = m_pD->FindLiveInstance(instanceID);
+		check(pLiveInstance);
+		m_pD->WorkingMemoryManager.CurrentInstanceCache = pLiveInstance->Cache;
+		m_pD->WorkingMemoryManager.BeginRunnerThread();
+		
+		OP::ADDRESS RootAddress = GetResourceIDRoot(MeshId);
+			
+		mu::OP_TYPE OpType = pLiveInstance->Model->GetPrivate()->m_program.GetOpType(RootAddress);
+		if (GetOpDataType(OpType) != DT_MESH)
+		{
+			m_pD->WorkingMemoryManager.EndRunnerThread();
+			m_pD->WorkingMemoryManager.CurrentInstanceCache = nullptr;
+
+			return UE::Tasks::MakeCompletedTask<Ptr<const Mesh>>(new Mesh());
+		}
+		
+		TSharedRef<CodeRunner> Runner = CodeRunner::Create(
+				m_pD->Settings, m_pD, EExecutionStrategy::MinimizeMemory, pLiveInstance->Model, pLiveInstance->OldParameters.get(), RootAddress, System::AllLODs, 0, 0, FScheduledOp::EType::Full);
+		
+		constexpr bool bForceInlineExecution = false;
+		UE::Tasks::FTask RunnerCompletionEvent = Runner->StartRun(bForceInlineExecution);
+
+		return UE::Tasks::Launch(TEXT("System::GetMeshResultTask"),
+				[SystemPrivate = m_pD, Runner, RootAddress]() -> Ptr<const Mesh>
+				{
+					Ptr<const Mesh> Result;
+
+					SystemPrivate->bUnrecoverableError = Runner->bUnrecoverableError;
+					if (!Runner->bUnrecoverableError)
+					{
+						Result = SystemPrivate->WorkingMemoryManager.LoadMesh(FCacheAddress(RootAddress, 0, 0), true);
+					}
+
+					if (!Result)
+					{
+						Result = new Mesh();
+					}
+
+					SystemPrivate->WorkingMemoryManager.EndRunnerThread();
+					SystemPrivate->WorkingMemoryManager.CurrentInstanceCache = nullptr;
+					
+					return Result;
+				},
+				UE::Tasks::Prerequisites(RunnerCompletionEvent),
+				UE::Tasks::ETaskPriority::Inherit,
+				UE::Tasks::EExtendedTaskPriority::Inline);
 	}
 
 
@@ -682,10 +856,14 @@ namespace mu
 	void System::Private::RunCode(const TSharedPtr<const Model>& InModel,
 		const Parameters* InParameters, OP::ADDRESS InCodeRoot, uint32 InLODs, uint8 executionOptions, int32 InImageLOD)
 	{
-		CodeRunner Runner(Settings, this, EExecutionStrategy::MinimizeMemory, InModel, InParameters, InCodeRoot, InLODs,
+		TSharedRef<CodeRunner> Runner = CodeRunner::Create(Settings, this, EExecutionStrategy::MinimizeMemory, InModel, InParameters, InCodeRoot, InLODs,
 			executionOptions, InImageLOD, FScheduledOp::EType::Full);
-		Runner.Run();
-		bUnrecoverableError = Runner.bUnrecoverableError;
+		
+		constexpr bool bForceInlineExecutution = true;
+		UE::Tasks::FTask RunnerCompletionEvent = Runner->StartRun(bForceInlineExecutution);
+		check(RunnerCompletionEvent.IsCompleted());
+
+		bUnrecoverableError = Runner->bUnrecoverableError;
 	}
 
 
@@ -788,7 +966,6 @@ namespace mu
 	}
 
 	
-	//---------------------------------------------------------------------------------------------
 	Ptr<const Image> System::Private::BuildImage(const TSharedPtr<const Model>& pModel,
 		const Parameters* Params, OP::ADDRESS at, int32 MipsToSkip, int32 InImageLOD)
 	{
diff --git a/Engine/Plugins/Experimental/Mutable/Source/MutableRuntime/Public/MuR/System.h b/Engine/Plugins/Experimental/Mutable/Source/MutableRuntime/Public/MuR/System.h
--- a/Engine/Plugins/Experimental/Mutable/Source/MutableRuntime/Public/MuR/System.h
+++ b/Engine/Plugins/Experimental/Mutable/Source/MutableRuntime/Public/MuR/System.h
@@ -142,7 +142,7 @@ namespace mu
         //! no longer needed.
         //! \param pModel Model to build an instance of
         //! \return An identifier that is always bigger than 0.
-        Instance::ID NewInstance(const TSharedPtr<const Model>& Model );
+        Instance::ID NewInstance(const TSharedPtr<const Model>& Model);
 
         //! \brief Update an instance with a new parameter set and/or state.
         //!
@@ -156,34 +156,47 @@ namespace mu
         //! \return the instance data with all the LOD, components, and ids to generate the meshes 
         //! and images.  The returned Instance is only valid until the next call to EndUpdate with 
         //! the same instanceID parameter.
-        const Instance* BeginUpdate( Instance::ID InstanceID,
-                                     const Ptr<const Parameters>& Params,
-                                     int32 StateIndex,
-                                     uint32 LodMask );
+        const Instance* BeginUpdate(Instance::ID InstanceID,
+                                    const Ptr<const Parameters>& Params,
+                                    int32 StateIndex,
+                                    uint32 LodMask);
 
 		//! Only valid between BeginUpdate and EndUpdate
 		//! Calculate the description of an image, without generating it.
-		void GetImageDesc(Instance::ID InstanceID, FResourceID ImageId, FImageDesc& OutDesc );
+		UE::Tasks::TTask<FImageDesc> GetImageDesc(Instance::ID InstanceID, FResourceID ImageId);
 
 		//! Only valid between BeginUpdate and EndUpdate
 		//! \param MipsToSkip Number of mips to skip compared from the full image.
 		//! If 0, all mip levels will be generated. If more levels than possible to discard are specified, 
 		//! the image will still contain a minimum number of mips specified at model compile time.
-		Ptr<const Image> GetImage(Instance::ID InstanceID, FResourceID ImageId, int32 MipsToSkip = 0, int32 LOD = 0);
+		UE::Tasks::TTask<Ptr<const Image>> GetImage(Instance::ID InstanceID, FResourceID ImageId, int32 MipsToSkip = 0, int32 LOD = 0);
 
         //! Only valid between BeginUpdate and EndUpdate
-        Ptr<const Mesh> GetMesh(Instance::ID InstanceID, FResourceID MeshId);
+		UE::Tasks::TTask<Ptr<const Mesh>> GetMesh(Instance::ID InstanceID, FResourceID MeshId);
+
+		//! Only valid between BeginUpdate and EndUpdate
+		//! Calculate the description of an image, without generating it.
+		void GetImageDescInline(Instance::ID InstanceID, FResourceID ImageId, FImageDesc& OutDesc);
+
+		//! Only valid between BeginUpdate and EndUpdate
+		//! \param MipsToSkip Number of mips to skip compared from the full image.
+		//! If 0, all mip levels will be generated. If more levels than possible to discard are specified, 
+		//! the image will still contain a minimum number of mips specified at model compile time.
+		Ptr<const Image> GetImageInline(Instance::ID InstanceID, FResourceID ImageId, int32 MipsToSkip = 0, int32 LOD = 0);
+
+        //! Only valid between BeginUpdate and EndUpdate
+		Ptr<const Mesh> GetMeshInline(Instance::ID InstanceID, FResourceID MeshId);
 
         //! Invalidate and free the last Instance data returned by a call to BeginUpdate with
         //! the same instance index. After a call to this method, that Instance cannot be used any
         //! more and its content is undefined.
         //! \param instance The index of the instance whose last data will be invalidated.
-        void EndUpdate( Instance::ID InstanceID);
+        void EndUpdate(Instance::ID InstanceID);
 
         //! Completely destroy an instance. After a call to this method the given instance cannot be
         //! updated any more, and its resources may have been freed.
         //! \param instance The id of the instance to destroy.
-        void ReleaseInstance( Instance::ID InstanceID);
+        void ReleaseInstance(Instance::ID InstanceID);
 
 		//! Calculate the relevancy of every parameter. Some parameters may be unused depending on
 		//! the values of other parameters. This method will set to true the flags for parameters
