[mutable] Added log to expose time required by the AR search operation to be ran.

#tests local UE commandlet execution

diff --git a/Engine/Plugins/Experimental/Mutable/Source/MutableValidation/Private/MuV/ValidationUtils.cpp b/Engine/Plugins/Experimental/Mutable/Source/MutableValidation/Private/MuV/ValidationUtils.cpp
--- a/Engine/Plugins/Experimental/Mutable/Source/MutableValidation/Private/MuV/ValidationUtils.cpp
+++ b/Engine/Plugins/Experimental/Mutable/Source/MutableValidation/Private/MuV/ValidationUtils.cpp
@@ -15,8 +15,13 @@ void PrepareAssetRegistry()
 {
 	FAssetRegistryModule& AssetRegistryModule = FModuleManager::LoadModuleChecked<FAssetRegistryModule>(AssetRegistryConstants::ModuleName);
 	UE_LOG(LogMutable,Display,TEXT("Searching all assets (this will take some time)..."));
+	
+	const double AssetRegistrySearchStartSeconds = FPlatformTime::Seconds();
 	AssetRegistryModule.Get().SearchAllAssets(true /* bSynchronousSearch */);
-	UE_LOG(LogMutable,Display,TEXT("Asset searching completed!"));
+	const double AssetRegistrySearchEndSeconds = FPlatformTime::Seconds() - AssetRegistrySearchStartSeconds;
+	UE_LOG(LogMutable, Log, TEXT("(double) asset_registry_search_time_s : %f "), AssetRegistrySearchEndSeconds);
+
+	UE_LOG(LogMutable,Display,TEXT("Asset searching completed in \"%f\" seconds!"), AssetRegistrySearchEndSeconds);
 }
 
 
