[Mutable] Fix BC5 to RGBA image conversion.

#rnx

diff --git a/Engine/Plugins/Experimental/Mutable/Source/MutableRuntime/Private/MuR/OpImagePixelFormat.cpp b/Engine/Plugins/Experimental/Mutable/Source/MutableRuntime/Private/MuR/OpImagePixelFormat.cpp
--- a/Engine/Plugins/Experimental/Mutable/Source/MutableRuntime/Private/MuR/OpImagePixelFormat.cpp
+++ b/Engine/Plugins/Experimental/Mutable/Source/MutableRuntime/Private/MuR/OpImagePixelFormat.cpp
@@ -834,7 +834,7 @@ void FImageOperator::ImagePixelFormat(bool& bOutSuccess, int32 CompressionQualit
 			case EImageFormat::IF_BC5:
 			{
 				ImagePixelFormatInternal::BlockDecompressionLODRangeHelper(
-						Result, Base, ResultLODBegin, BaseLODBegin, NumLODs, miro::BC4_to_RGBA);
+						Result, Base, ResultLODBegin, BaseLODBegin, NumLODs, miro::BC5_to_RGBA);
 				break;
 			}
 			case EImageFormat::IF_ASTC_4x4_RGB_LDR:
