[mutable] Added a new option to Customizable Objects (DisableTableMaterialsParentCheck) to let materials in table nodes have different parent materials than the default one in the material node they are connected to. If these arbitrary materials don't have the channels Mutable is overriding, there will be no effect and those channels will be silently ignored.
#rnx
#tests PIE, Client

diff --git a/Engine/Plugins/Experimental/Mutable/Source/CustomizableObject/Classes/MuCO/CustomizableObject.h b/Engine/Plugins/Experimental/Mutable/Source/CustomizableObject/Classes/MuCO/CustomizableObject.h
--- a/Engine/Plugins/Experimental/Mutable/Source/CustomizableObject/Classes/MuCO/CustomizableObject.h
+++ b/Engine/Plugins/Experimental/Mutable/Source/CustomizableObject/Classes/MuCO/CustomizableObject.h
@@ -573,6 +573,12 @@ public:
 	UPROPERTY(EditAnywhere, Category = CompileOptions)
 	bool bEnableAssetUserDataMerge = false;
 
+	// Disabling the Table Materials parent material check will let the user use any material regardless of its parent when connecting a material from a table column to a material node.
+	// Warning! But it will not check if the table material channels connected to the Material node exist in the actual material used in the Instance, and will fail silently at runtime 
+	// when setting the value of those channels if they don't exist.
+	UPROPERTY(EditAnywhere, Category = CompileOptions)
+	bool bDisableTableMaterialsParentCheck = false;
+
 	// Options when compiling this customizable object (see EMutableCompileMeshType declaration for info)
 	UPROPERTY(EditAnywhere, Category = CompileOptions)
 	EMutableCompileMeshType MeshCompileType = EMutableCompileMeshType::LocalAndChildren;
diff --git a/Engine/Plugins/Experimental/Mutable/Source/CustomizableObjectEditor/Private/MuCOE/GenerateMutableSource/GenerateMutableSourceTable.cpp b/Engine/Plugins/Experimental/Mutable/Source/CustomizableObjectEditor/Private/MuCOE/GenerateMutableSource/GenerateMutableSourceTable.cpp
--- a/Engine/Plugins/Experimental/Mutable/Source/CustomizableObjectEditor/Private/MuCOE/GenerateMutableSource/GenerateMutableSourceTable.cpp
+++ b/Engine/Plugins/Experimental/Mutable/Source/CustomizableObjectEditor/Private/MuCOE/GenerateMutableSource/GenerateMutableSourceTable.cpp
@@ -460,7 +460,10 @@ bool FillTableColumn(const UCustomizableObjectNodeTable* TableNode,	mu::TablePtr
 
 			GenerationContext.AddParticipatingObject(*ReferenceMaterial);
 
-			if (!Material || ReferenceMaterial->GetMaterial() != Material->GetMaterial())
+			const bool bTableMaterialCheckDisabled = GenerationContext.Object->bDisableTableMaterialsParentCheck;
+			const bool bMaterialParentMismatch = !bTableMaterialCheckDisabled && ReferenceMaterial->GetMaterial() != Material->GetMaterial();
+
+			if (!Material || bMaterialParentMismatch)
 			{
 				FText Warning;
 
