[Mutable] Set MipDataProvider to use the tasks split CodeRunner.

- CVarEnableNewSplitMutableTask  added to disable the split code path.
- Change benchmark utility game thread task to use ExecuteOnGameThread function.

#preflght 65f2012c3d1ec391ceeba8cb
#rnx

diff --git a/Engine/Plugins/Experimental/Mutable/Source/CustomizableObject/Classes/MuCO/CustomizableObjectSystem.h b/Engine/Plugins/Experimental/Mutable/Source/CustomizableObject/Classes/MuCO/CustomizableObjectSystem.h
--- a/Engine/Plugins/Experimental/Mutable/Source/CustomizableObject/Classes/MuCO/CustomizableObjectSystem.h
+++ b/Engine/Plugins/Experimental/Mutable/Source/CustomizableObject/Classes/MuCO/CustomizableObjectSystem.h
@@ -45,6 +45,7 @@ extern TAutoConsoleVariable<bool> CVarEnableMeshCache;
 
 extern TAutoConsoleVariable<bool> CVarRollbackFixModelDiskStreamerDataRace;
 
+extern TAutoConsoleVariable<bool> CVarEnableNewSplitMutableTask;
 
 #if WITH_EDITOR
 
diff --git a/Engine/Plugins/Experimental/Mutable/Source/CustomizableObject/Private/MuCO/BusyWaits_Deprecated.cpp b/Engine/Plugins/Experimental/Mutable/Source/CustomizableObject/Private/MuCO/BusyWaits_Deprecated.cpp
--- a/Engine/Plugins/Experimental/Mutable/Source/CustomizableObject/Private/MuCO/BusyWaits_Deprecated.cpp
+++ b/Engine/Plugins/Experimental/Mutable/Source/CustomizableObject/Private/MuCO/BusyWaits_Deprecated.cpp
@@ -4,6 +4,9 @@
 #include "MuCO/CustomizableObjectInstancePrivate.h"
 #include "MuCO/CustomizableObjectSystem.h"
 #include "MuCO/CustomizableObjectSystemPrivate.h"
+#include "MuCO/CustomizableObjectMipDataProvider.h"
+
+#include "MuR/Model.h"
 #include "MuR/MutableTrace.h"
 
 
@@ -228,7 +231,7 @@ namespace impl_deprecated
 
 					Component.MeshID = Instance->GetMeshId(MutableLODIndex, ComponentIndex, 0);
 
-					if(bGenerateLOD)
+					if (bGenerateLOD)
 					{
 						Component.Mesh = System->GetMeshInline(OperationData->InstanceID, Component.MeshID);
 					}
@@ -310,6 +313,7 @@ namespace impl_deprecated
 
 									// Since it's known it's a pass-through texture there is no need to cache or convert it so we can generate it here already.
 									Image.Image = System->GetImageInline(OperationData->InstanceID, Image.ImageID, 0, 0);
+
 									check(Image.Image->IsReference());
 
 									uint32 ReferenceID = Image.Image->GetReferencedTexture();
@@ -422,3 +426,172 @@ namespace impl_deprecated
 #endif
 	}
 }
+
+namespace CustomizableObjectMipDataProvider::ImplDeprecated
+{
+	void Task_Mutable_UpdateImage(TSharedPtr<FMutableImageOperationData> OperationData)
+	{
+		MUTABLE_CPUPROFILER_SCOPE(Task_Mutable_UpdateImage);
+		const double StartTime = FPlatformTime::Seconds();
+		
+		// Cache memory used when starting the update of the image
+		OperationData->ImageUpdateStartBytes = mu::FGlobalMemoryCounter::GetCounter();
+		mu::FGlobalMemoryCounter::Zero();
+		
+		// Any external texture that may be needed for this update will be requested from Mutable Core's GetImage
+		// which will safely access the GlobalExternalImages map, and then just get the cached image or issue a disk read
+
+		// This runs in a worker thread.
+		check(OperationData.IsValid());
+		check(OperationData->UpdateContext->GetSystem().get());
+		check(OperationData->UpdateContext->GetModel());
+		check(OperationData->UpdateContext->GetParameters().get());
+
+		if (OperationData.IsValid())
+		{
+			mu::SystemPtr System = OperationData->UpdateContext->GetSystem();
+			const TSharedPtr<mu::Model, ESPMode::ThreadSafe> Model = OperationData->UpdateContext->GetModel();
+
+#if WITH_EDITOR
+			// Recompiling a CO in the editor will invalidate the previously generated Model. Check that it is valid before accessing the streamed data.
+			if (Model && Model->IsValid())
+#endif
+			{
+
+				// For now, we are forcing the recreation of mutable-side instances with every update.
+				mu::Instance::ID InstanceID = System->NewInstance(Model);
+				UE_LOG(LogMutable, Verbose, TEXT("Creating Mutable instance with id [%d] for a single UpdateImage"), InstanceID)
+
+					const mu::Instance* Instance = nullptr;
+
+				// Main instance generation step
+				{
+					// LOD mask, set to all ones to build all LODs
+					uint32 LODMask = 0xFFFFFFFF;
+
+					Instance = System->BeginUpdate(InstanceID, OperationData->UpdateContext->GetParameters(), OperationData->UpdateContext->GetState(), LODMask);
+					check(Instance);
+				}
+
+
+				// Generate the required image
+				{
+					MUTABLE_CPUPROFILER_SCOPE(RequestedImage);
+
+					const FMutableImageReference& ImageRef = OperationData->RequestedImage;
+
+					int32 SurfaceIndex = Instance->FindSurfaceById(ImageRef.LOD, ImageRef.Component, ImageRef.SurfaceId);
+					check(SurfaceIndex >= 0);
+
+					// This ID may be different than the ID obtained the first time the image was generated, because the mutable
+					// runtime cannot remember all the resources it has built, and only remembers a fixed amount.
+					mu::FResourceID MipImageID = Instance->GetImageId(ImageRef.LOD, ImageRef.Component, SurfaceIndex, ImageRef.Image);
+
+
+					mu::Ptr<const mu::Image> ResultImage;
+					{
+						MUTABLE_CPUPROFILER_SCOPE(GetImage);
+
+						ResultImage = System->GetImageInline(InstanceID, MipImageID, ImageRef.BaseMip + OperationData->MipsToSkip, ImageRef.LOD);
+					}
+
+					check(ResultImage);
+
+					int32 FullMipCount = ResultImage->GetMipmapCount(ResultImage->GetSizeX(), ResultImage->GetSizeY());
+					int32 RealMipCount = ResultImage->GetLODCount();
+
+					bool bForceMipchain =
+						// Did we fail to generate the entire mipchain (if we have mips at all)?
+						(RealMipCount != 1) && (RealMipCount != FullMipCount);
+
+					if (bForceMipchain)
+					{
+						MUTABLE_CPUPROFILER_SCOPE(GetImage_MipFix);
+
+						UE_LOG(LogMutable, Warning, TEXT("Mutable generated an incomplete mip chain for image."));
+
+						// Force the right number of mips. The missing data will be black.
+						mu::Ptr<mu::Image> NewImage = new mu::Image(ResultImage->GetSizeX(), ResultImage->GetSizeY(), FullMipCount, ResultImage->GetFormat(), mu::EInitializationType::Black);
+
+						// Formats with BytesPerBlock == 0 will not allocate memory. This type of images are not expected here.
+						check(!NewImage->DataStorage.IsEmpty());
+
+						for (int32 L = 0; L < RealMipCount; ++L)
+						{
+							TArrayView<uint8> DestView = NewImage->DataStorage.GetLOD(L);
+							TArrayView<const uint8> SrcView = ResultImage->DataStorage.GetLOD(L);
+
+							check(DestView.Num() == SrcView.Num());
+							FMemory::Memcpy(DestView.GetData(), SrcView.GetData(), DestView.Num());
+						}
+
+						ResultImage = NewImage;
+					}
+
+					OperationData->Result = ResultImage;
+				}
+
+				// End update
+				{
+					MUTABLE_CPUPROFILER_SCOPE(EndUpdate);
+					System->EndUpdate(InstanceID);
+					System->ReleaseInstance(InstanceID);
+
+					if (CVarClearWorkingMemoryOnUpdateEnd.GetValueOnAnyThread())
+					{
+						System->ClearWorkingMemory();
+					}
+				}
+			}
+		}
+
+		if (CVarEnableBenchmark.GetValueOnAnyThread())
+		{
+			double Time = FPlatformTime::Seconds() - StartTime;
+			// Report the peak memory used by the operation
+			const int64 PeakMemory = mu::FGlobalMemoryCounter::GetPeak();
+			// Report the peak memory used during the operation (operation + baseline)
+			const int64 RealMemoryPeak = PeakMemory + OperationData->ImageUpdateStartBytes;
+
+			const FString& CustomizableObjectPathName = OperationData->UpdateContext->GetCustomizableObjectPathName();
+			const FString& InstancePathName = OperationData->UpdateContext->GetInstancePathName();
+			
+			FFunctionGraphTask::CreateAndDispatchWhenReady(
+			[CustomizableObjectPathName, InstancePathName, Time, PeakMemory, RealMemoryPeak]()
+			{
+				if (!UCustomizableObjectSystem::IsCreated()) // We are shutting down
+				{
+					return;	
+				}
+				
+				UCustomizableObjectSystem* System = UCustomizableObjectSystem::GetInstance();
+				if (!System)
+				{
+					return;
+				}
+
+				System->GetPrivate()->LogBenchmarkUtil.FinishUpdateImage(CustomizableObjectPathName, InstancePathName, Time, PeakMemory, RealMemoryPeak);
+			},
+			TStatId{},
+			nullptr,
+			ENamedThreads::GameThread);
+		}
+		
+		{
+			// The request could be cancelled in parallel from CancelCounterSafely and its value be changed
+			// between reading it and actually running Decrement() and RescheduleCallback(), so lock
+			FScopeLock Lock(&OperationData->CounterTaskLock);
+
+			if (OperationData->Counter) // If the request has been cancelled the counter will be null
+			{
+				// Make the FMutableTextureMipDataProvider continue
+				OperationData->Counter->Decrement();
+
+				if (OperationData->Counter->GetValue() == 0)
+				{
+					OperationData->RescheduleCallback();
+				}
+			}
+		}
+	}
+}
diff --git a/Engine/Plugins/Experimental/Mutable/Source/CustomizableObject/Private/MuCO/BusyWaits_Deprecated.h b/Engine/Plugins/Experimental/Mutable/Source/CustomizableObject/Private/MuCO/BusyWaits_Deprecated.h
--- a/Engine/Plugins/Experimental/Mutable/Source/CustomizableObject/Private/MuCO/BusyWaits_Deprecated.h
+++ b/Engine/Plugins/Experimental/Mutable/Source/CustomizableObject/Private/MuCO/BusyWaits_Deprecated.h
@@ -7,6 +7,7 @@
 #include "Templates/SharedPointer.h"
 
 class FUpdateContextPrivate;
+struct FMutableImageOperationData;
 namespace mu
 {
 	class Model;
@@ -26,4 +27,11 @@ namespace impl_deprecated
 
 	// This runs in a worker thread.
 	void Task_Mutable_Update_GetMesh(const TSharedRef<FUpdateContextPrivate>& OperationData, const TSharedPtr<mu::Model>& Model);
+
+}
+
+namespace CustomizableObjectMipDataProvider::ImplDeprecated
+{
+	// This runs in the mutable thread.
+	void Task_Mutable_UpdateImage(TSharedPtr<FMutableImageOperationData> OperationData);
 }
diff --git a/Engine/Plugins/Experimental/Mutable/Source/CustomizableObject/Private/MuCO/CustomizableObjectMipDataProvider.cpp b/Engine/Plugins/Experimental/Mutable/Source/CustomizableObject/Private/MuCO/CustomizableObjectMipDataProvider.cpp
--- a/Engine/Plugins/Experimental/Mutable/Source/CustomizableObject/Private/MuCO/CustomizableObjectMipDataProvider.cpp
+++ b/Engine/Plugins/Experimental/Mutable/Source/CustomizableObject/Private/MuCO/CustomizableObjectMipDataProvider.cpp
@@ -7,9 +7,15 @@
 #include "MuR/Model.h"
 #include "TextureResource.h"
 #include "UnrealMutableImageProvider.h"
+#include "MuCO/BusyWaits_Deprecated.h"
+
+#include "Containers/Ticker.h"
 
 #include UE_INLINE_GENERATED_CPP_BY_NAME(CustomizableObjectMipDataProvider)
 
+
+#define UE_MUTABLE_MIPDATA_PROVIDER_UPDATE_IMAGE_REGION		TEXT("MipDataProvider_Mutable_UpdateImage")
+
 UMutableTextureMipDataProviderFactory::UMutableTextureMipDataProviderFactory(const FObjectInitializer& ObjectInitializer)
 	: Super(ObjectInitializer)
 {
@@ -124,11 +130,10 @@ void FMutableTextureMipDataProvider::Init(const FTextureUpdateContext& Context,
 }
 
 
-namespace impl
+namespace Impl
 {
 	void Task_Mutable_UpdateImage(TSharedPtr<FMutableImageOperationData> OperationData)
 	{
-		MUTABLE_CPUPROFILER_SCOPE(Task_Mutable_UpdateImage);
 		const double StartTime = FPlatformTime::Seconds();
 		
 		if (CVarEnableBenchmark.GetValueOnAnyThread())
@@ -147,62 +152,86 @@ namespace impl
 		check(OperationData->UpdateContext->GetModel());
 		check(OperationData->UpdateContext->GetParameters().get());
 
-		if (OperationData.IsValid())
+		if (!OperationData.IsValid())
 		{
-			mu::SystemPtr System = OperationData->UpdateContext->GetSystem();
-			const TSharedPtr<mu::Model, ESPMode::ThreadSafe> Model = OperationData->UpdateContext->GetModel();
+			return;
+		}
 
-#if WITH_EDITOR
-			// Recompiling a CO in the editor will invalidate the previously generated Model. Check that it is valid before accessing the streamed data.
-			if (Model && Model->IsValid())
-#endif
-			{
+		TRACE_BEGIN_REGION(UE_MUTABLE_MIPDATA_PROVIDER_UPDATE_IMAGE_REGION);
+		
+		mu::Ptr<mu::System> System = OperationData->UpdateContext->GetSystem();
+		const TSharedPtr<mu::Model, ESPMode::ThreadSafe> Model = OperationData->UpdateContext->GetModel();
 
-				// For now, we are forcing the recreation of mutable-side instances with every update.
-				mu::Instance::ID InstanceID = System->NewInstance(Model);
-				UE_LOG(LogMutable, Verbose, TEXT("Creating Mutable instance with id [%d] for a single UpdateImage"), InstanceID)
+		auto EndUpdateImage = [](TSharedPtr<FMutableImageOperationData>& OperationData)
+		{
+			// The request could be cancelled in parallel from CancelCounterSafely and its value be changed
+			// between reading it and actually running Decrement() and RescheduleCallback(), so lock
+			FScopeLock Lock(&OperationData->CounterTaskLock);
 
-					const mu::Instance* Instance = nullptr;
+			if (OperationData->Counter) // If the request has been cancelled the counter will be null
+			{
+				// Make the FMutableTextureMipDataProvider continue
+				OperationData->Counter->Decrement();
 
-				// Main instance generation step
+				if (OperationData->Counter->GetValue() == 0)
 				{
-					// LOD mask, set to all ones to build all LODs
-					uint32 LODMask = 0xFFFFFFFF;
-
-					Instance = System->BeginUpdate(InstanceID, OperationData->UpdateContext->GetParameters(), OperationData->UpdateContext->GetState(), LODMask);
-					check(Instance);
+					OperationData->RescheduleCallback();
 				}
+			}
+			
+			TRACE_END_REGION(UE_MUTABLE_MIPDATA_PROVIDER_UPDATE_IMAGE_REGION);
+		};
 
+#if WITH_EDITOR
+		// Recompiling a CO in the editor will invalidate the previously generated Model. Check that it is valid before accessing the streamed data.
+		if (!(Model && Model->IsValid()))
+		{
+			EndUpdateImage(OperationData);
+			return;
+		}
+#endif
 
-				// Generate the required image
-				{
-					MUTABLE_CPUPROFILER_SCOPE(RequestedImage);
+		// For now, we are forcing the recreation of mutable-side instances with every update.
+		mu::Instance::ID InstanceID = System->NewInstance(Model);
+		UE_LOG(LogMutable, Verbose, TEXT("Creating Mutable instance with id [%d] for a single UpdateImage"), InstanceID)
 
-					const FMutableImageReference& ImageRef = OperationData->RequestedImage;
+		const mu::Instance* Instance = nullptr;
 
-					int32 SurfaceIndex = Instance->FindSurfaceById(ImageRef.LOD, ImageRef.Component, ImageRef.SurfaceId);
-					check(SurfaceIndex >= 0);
+		// Main instance generation step
+		{
+			// LOD mask, set to all ones to build all LODs
+			uint32 LODMask = 0xFFFFFFFF;
 
-					// This ID may be different than the ID obtained the first time the image was generated, because the mutable
-					// runtime cannot remember all the resources it has built, and only remembers a fixed amount.
-					mu::FResourceID MipImageID = Instance->GetImageId(ImageRef.LOD, ImageRef.Component, SurfaceIndex, ImageRef.Image);
+			Instance = System->BeginUpdate(InstanceID, OperationData->UpdateContext->GetParameters(), OperationData->UpdateContext->GetState(), LODMask);
+			check(Instance);
+		}
 
+		const FMutableImageReference& ImageRef = OperationData->RequestedImage;
 
-					mu::ImagePtrConst Image;
-					{
-						MUTABLE_CPUPROFILER_SCOPE(GetImage);
+		int32 SurfaceIndex = Instance->FindSurfaceById(ImageRef.LOD, ImageRef.Component, ImageRef.SurfaceId);
+		check(SurfaceIndex >= 0);
 
-						Image = System->GetImageInline(InstanceID, MipImageID, ImageRef.BaseMip + OperationData->MipsToSkip, ImageRef.LOD);
-					}
+		// This ID may be different than the ID obtained the first time the image was generated, because the mutable
+		// runtime cannot remember all the resources it has built, and only remembers a fixed amount.
+		mu::FResourceID MipImageID = Instance->GetImageId(ImageRef.LOD, ImageRef.Component, SurfaceIndex, ImageRef.Image);
+
+		UE::Tasks::TTask<mu::Ptr<const mu::Image>> GetImageTask = 
+				System->GetImage(InstanceID, MipImageID, ImageRef.BaseMip + OperationData->MipsToSkip, ImageRef.LOD);
+
+		UE::Tasks::AddNested(UE::Tasks::Launch(TEXT("MipDataProvider_EndUpdateImagesTask"),
+				[System, OperationData, InstanceID, StartTime, GetImageTask, EndUpdateImage]() mutable
+				{
+					check(GetImageTask.IsCompleted());
+
+					mu::Ptr<const mu::Image> ResultImage = GetImageTask.GetResult();
 
-					check(Image);
+					check(ResultImage);
 
-					int32 FullMipCount = Image->GetMipmapCount(Image->GetSizeX(), Image->GetSizeY());
-					int32 RealMipCount = Image->GetLODCount();
+					int32 FullMipCount = ResultImage->GetMipmapCount(ResultImage->GetSizeX(), ResultImage->GetSizeY());
+					int32 RealMipCount = ResultImage->GetLODCount();
 
-					bool bForceMipchain =
-						// Did we fail to generate the entire mipchain (if we have mips at all)?
-						(RealMipCount != 1) && (RealMipCount != FullMipCount);
+					// Did we fail to generate the entire mipchain (if we have mips at all)?
+					bool bForceMipchain = (RealMipCount != 1) && (RealMipCount != FullMipCount);
 
 					if (bForceMipchain)
 					{
@@ -211,7 +240,7 @@ namespace impl
 						UE_LOG(LogMutable, Warning, TEXT("Mutable generated an incomplete mip chain for image."));
 
 						// Force the right number of mips. The missing data will be black.
-						mu::Ptr<mu::Image> NewImage = new mu::Image(Image->GetSizeX(), Image->GetSizeY(), FullMipCount, Image->GetFormat(), mu::EInitializationType::Black);
+						mu::Ptr<mu::Image> NewImage = new mu::Image(ResultImage->GetSizeX(), ResultImage->GetSizeY(), FullMipCount, ResultImage->GetFormat(), mu::EInitializationType::Black);
 
 						// Formats with BytesPerBlock == 0 will not allocate memory. This type of images are not expected here.
 						check(!NewImage->DataStorage.IsEmpty());
@@ -219,82 +248,65 @@ namespace impl
 						for (int32 L = 0; L < RealMipCount; ++L)
 						{
 							TArrayView<uint8> DestView = NewImage->DataStorage.GetLOD(L);
-							TArrayView<const uint8> SrcView = Image->DataStorage.GetLOD(L);
+							TArrayView<const uint8> SrcView = ResultImage->DataStorage.GetLOD(L);
 
 							check(DestView.Num() == SrcView.Num());
 							FMemory::Memcpy(DestView.GetData(), SrcView.GetData(), DestView.Num());
 						}
 
-						Image = NewImage;
+						ResultImage = NewImage;
 					}
 
-					OperationData->Result = Image;
-				}
-
-				// End update
-				{
-					MUTABLE_CPUPROFILER_SCOPE(EndUpdate);
-					System->EndUpdate(InstanceID);
-					System->ReleaseInstance(InstanceID);
+					OperationData->Result = ResultImage;
 
-					if (CVarClearWorkingMemoryOnUpdateEnd.GetValueOnAnyThread())
+					// End update
 					{
-						System->ClearWorkingMemory();
-					}
-				}
-			}
-		}
+						MUTABLE_CPUPROFILER_SCOPE(EndUpdate);
+						System->EndUpdate(InstanceID);
+						System->ReleaseInstance(InstanceID);
 
-		if (CVarEnableBenchmark.GetValueOnAnyThread())
-		{
-			double Time = FPlatformTime::Seconds() - StartTime;
-			// Report the peak memory used by the operation
-			const int64 PeakMemory = mu::FGlobalMemoryCounter::GetPeak();
-			// Report the peak memory used during the operation (operation + baseline)
-			const int64 RealMemoryPeak = PeakMemory + OperationData->ImageUpdateStartBytes;
-
-			const FString& CustomizableObjectPathName = OperationData->UpdateContext->GetCustomizableObjectPathName();
-			const FString& InstancePathName = OperationData->UpdateContext->GetInstancePathName();
-			
-			FFunctionGraphTask::CreateAndDispatchWhenReady(
-			[CustomizableObjectPathName, InstancePathName, Time, PeakMemory, RealMemoryPeak]()
-			{
-				if (!UCustomizableObjectSystem::IsCreated()) // We are shutting down
-				{
-					return;	
-				}
-				
-				UCustomizableObjectSystem* System = UCustomizableObjectSystem::GetInstance();
-				if (!System)
-				{
-					return;
-				}
-
-				System->GetPrivate()->LogBenchmarkUtil.FinishUpdateImage(CustomizableObjectPathName, InstancePathName, Time, PeakMemory, RealMemoryPeak);
-			},
-			TStatId{},
-			nullptr,
-			ENamedThreads::GameThread);
-		}
-		
-		{
-			// The request could be cancelled in parallel from CancelCounterSafely and its value be changed
-			// between reading it and actually running Decrement() and RescheduleCallback(), so lock
-			FScopeLock Lock(&OperationData->CounterTaskLock);
+						if (CVarClearWorkingMemoryOnUpdateEnd.GetValueOnAnyThread())
+						{
+							System->ClearWorkingMemory();
+						}
+					}
 
-			if (OperationData->Counter) // If the request has been cancelled the counter will be null
-			{
-				// Make the FMutableTextureMipDataProvider continue
-				OperationData->Counter->Decrement();
+					if (CVarEnableBenchmark.GetValueOnAnyThread())
+					{
+						double Time = FPlatformTime::Seconds() - StartTime;
+						// Report the peak memory used by the operation
+						const int64 PeakMemory = mu::FGlobalMemoryCounter::GetPeak();
+						// Report the peak memory used during the operation (operation + baseline)
+						const int64 RealMemoryPeak = PeakMemory + OperationData->ImageUpdateStartBytes;
+
+						const FString& CustomizableObjectPathName = OperationData->UpdateContext->GetCustomizableObjectPathName();
+						const FString& InstancePathName = OperationData->UpdateContext->GetInstancePathName();
+					
+						ExecuteOnGameThread(UE_SOURCE_LOCATION, 
+						[CustomizableObjectPathName, InstancePathName, Time, PeakMemory, RealMemoryPeak]()
+						{
+							if (!UCustomizableObjectSystem::IsCreated()) // We are shutting down
+							{
+								return;	
+							}
+							
+							UCustomizableObjectSystem* System = UCustomizableObjectSystem::GetInstance();
+							if (!System)
+							{
+								return;
+							}
+
+							System->GetPrivate()->LogBenchmarkUtil.FinishUpdateImage(CustomizableObjectPathName, InstancePathName, Time, PeakMemory, RealMemoryPeak);
+						});
+					}
 
-				if (OperationData->Counter->GetValue() == 0)
-				{
-					OperationData->RescheduleCallback();
-				}
-			}
-		}
+					EndUpdateImage(OperationData);
+				},
+				UE::Tasks::Prerequisites(GetImageTask),
+				UE::Tasks::ETaskPriority::Inherit,
+				UE::Tasks::EExtendedTaskPriority::Inline)); // MipDataProvider_EndUpdateImagesTaskGetImage.
 	}
-} // namespace
+} // namespace Impl
 
 
 int32 FMutableTextureMipDataProvider::GetMips(const FTextureUpdateContext& Context, int32 StartingMipIndex, const FTextureMipInfoArray& MipInfos, const FTextureUpdateSyncOptions& SyncOptions)
@@ -380,12 +392,19 @@ int32 FMutableTextureMipDataProvider::GetMips(const FTextureUpdateContext& Conte
 			OperationData->Counter = SyncOptions.Counter;
 			OperationData->RescheduleCallback = SyncOptions.RescheduleCallback;
 
-			TSharedPtr<FMutableImageOperationData> LocalOperationData = OperationData;
 			MutableTaskId = CustomizableObjectSystem->MutableTaskGraph.AddMutableThreadTaskLowPriority(
 				TEXT("Mutable_MipUpdate"),
-				[LocalOperationData]()
+				[OperationData = this->OperationData]()
 				{
-					impl::Task_Mutable_UpdateImage(LocalOperationData);
+					if (CVarEnableNewSplitMutableTask.GetValueOnAnyThread())
+					{
+						Impl::Task_Mutable_UpdateImage(OperationData);
+					}
+					else
+					{
+						using namespace CustomizableObjectMipDataProvider;
+						ImplDeprecated::Task_Mutable_UpdateImage(OperationData);
+					}
 				});
 		}
 
diff --git a/Engine/Plugins/Experimental/Mutable/Source/CustomizableObject/Private/MuCO/CustomizableObjectSystem.cpp b/Engine/Plugins/Experimental/Mutable/Source/CustomizableObject/Private/MuCO/CustomizableObjectSystem.cpp
--- a/Engine/Plugins/Experimental/Mutable/Source/CustomizableObject/Private/MuCO/CustomizableObjectSystem.cpp
+++ b/Engine/Plugins/Experimental/Mutable/Source/CustomizableObject/Private/MuCO/CustomizableObjectSystem.cpp
@@ -112,6 +112,12 @@ TAutoConsoleVariable<bool> CVarEnableMeshCache(
 	TEXT("Enables or disables the reuse of meshes."),
 	ECVF_Scalability);
 
+TAutoConsoleVariable<bool> CVarEnableNewSplitMutableTask(
+	TEXT("mutable.EnableNewSplitMutableTask"),
+	true,
+	TEXT("Enables or disables the then new split GetImages and GetMesh tasks that remove BusyWaits."),
+	ECVF_Scalability);
+
 int32 UCustomizableObjectSystemPrivate::SkeletalMeshMinLodQualityLevel = -1;
 
 static void CVarMutableSinkFunction()
