[mutable] Fixed Static Analysis warning 443069 in GenerateMutableSourceTable.cpp due to CL 32217795.
#rnx

diff --git a/Engine/Plugins/Experimental/Mutable/Source/CustomizableObjectEditor/Private/MuCOE/GenerateMutableSource/GenerateMutableSourceTable.cpp b/Engine/Plugins/Experimental/Mutable/Source/CustomizableObjectEditor/Private/MuCOE/GenerateMutableSource/GenerateMutableSourceTable.cpp
--- a/Engine/Plugins/Experimental/Mutable/Source/CustomizableObjectEditor/Private/MuCOE/GenerateMutableSource/GenerateMutableSourceTable.cpp
+++ b/Engine/Plugins/Experimental/Mutable/Source/CustomizableObjectEditor/Private/MuCOE/GenerateMutableSource/GenerateMutableSourceTable.cpp
@@ -461,7 +461,8 @@ bool FillTableColumn(const UCustomizableObjectNodeTable* TableNode,	mu::TablePtr
 			GenerationContext.AddParticipatingObject(*ReferenceMaterial);
 
 			const bool bTableMaterialCheckDisabled = GenerationContext.Object->bDisableTableMaterialsParentCheck;
-			const bool bMaterialParentMismatch = !bTableMaterialCheckDisabled && ReferenceMaterial->GetMaterial() != Material->GetMaterial();
+			const bool bMaterialParentMismatch = !bTableMaterialCheckDisabled && Material 
+												 && ReferenceMaterial->GetMaterial() != Material->GetMaterial();
 
 			if (!Material || bMaterialParentMismatch)
 			{
