[Mutable] Enable edition of 'LOD Custom Settings'.

#rnx

diff --git a/Engine/Plugins/Experimental/Mutable/Source/CustomizableObjectEditor/Classes/MuCOE/Nodes/CustomizableObjectNodeObject.h b/Engine/Plugins/Experimental/Mutable/Source/CustomizableObjectEditor/Classes/MuCOE/Nodes/CustomizableObjectNodeObject.h
--- a/Engine/Plugins/Experimental/Mutable/Source/CustomizableObjectEditor/Classes/MuCOE/Nodes/CustomizableObjectNodeObject.h
+++ b/Engine/Plugins/Experimental/Mutable/Source/CustomizableObjectEditor/Classes/MuCOE/Nodes/CustomizableObjectNodeObject.h
@@ -198,7 +198,7 @@ public:
 
 	// Array of bones to remove from the mesh.All influences assigned to these bones will be transferred to the closest valid bone.
 	// Selected per component and LOD. Bones will be accumulated down the line.
-	UPROPERTY(VisibleAnywhere, Category = CustomizableObject)
+	UPROPERTY(EditAnywhere, Category = CustomizableObject)
 	TArray<FComponentSettings> ComponentSettings;
 	
     // To avoid any no properly saved GUIDs
