[mutable] Fixed instance bake issues, there were some vertex reduction and normal artifacts added by UE's mesh build/reduction system, when the mesh should be left completely untouched.
#rnx
#tests Editor

diff --git a/Engine/Plugins/Experimental/Mutable/Source/CustomizableObject/Private/MuCO/CustomizableObjectInstance.cpp b/Engine/Plugins/Experimental/Mutable/Source/CustomizableObject/Private/MuCO/CustomizableObjectInstance.cpp
--- a/Engine/Plugins/Experimental/Mutable/Source/CustomizableObject/Private/MuCO/CustomizableObjectInstance.cpp
+++ b/Engine/Plugins/Experimental/Mutable/Source/CustomizableObject/Private/MuCO/CustomizableObjectInstance.cpp
@@ -3238,9 +3238,36 @@ void UCustomizableInstancePrivate::InitSkeletalMeshData(const TSharedRef<FUpdate
 			// Disable LOD simplification when baking instances
 			LODInfo.ReductionSettings.NumOfTrianglesPercentage = 1.f;
 			LODInfo.ReductionSettings.NumOfVertPercentage = 1.f;
+			LODInfo.ReductionSettings.MaxNumOfTriangles = TNumericLimits<uint32>::Max();
+			LODInfo.ReductionSettings.MaxNumOfVerts = TNumericLimits<uint32>::Max();
+			LODInfo.ReductionSettings.bRecalcNormals = 0;
+			LODInfo.ReductionSettings.WeldingThreshold = TNumericLimits<float>::Min();
+			LODInfo.ReductionSettings.bMergeCoincidentVertBones = 0;
+			LODInfo.ReductionSettings.bImproveTrianglesForCloth = 0;
 
 #if WITH_EDITORONLY_DATA
+			// If any of these structures changes sizes/fields, add the new reduction/build settings to the LODInfo in this section
+			// and set them to NOT change or do anything to the skeletal mesh.
+			static_assert(sizeof(FSkeletalMeshOptimizationSettings) == 120);
+			static_assert(sizeof(FSkeletalMeshBuildSettings) == 24);
+
+			LODInfo.ReductionSettings.MaxNumOfTrianglesPercentage = TNumericLimits<uint32>::Max();
+			LODInfo.ReductionSettings.MaxNumOfVertsPercentage = TNumericLimits<uint32>::Max();
+
+			LODInfo.BuildSettings.bRecomputeNormals = false;
+			LODInfo.BuildSettings.bRecomputeTangents = false;
+			LODInfo.BuildSettings.bUseMikkTSpace = false;
+			LODInfo.BuildSettings.bComputeWeightedNormals = false;
+			LODInfo.BuildSettings.bRemoveDegenerates = false;
+			LODInfo.BuildSettings.bUseHighPrecisionTangentBasis = false;
+			LODInfo.BuildSettings.bUseHighPrecisionSkinWeights = false;
 			LODInfo.BuildSettings.bUseFullPrecisionUVs = true;
+			LODInfo.BuildSettings.bUseBackwardsCompatibleF16TruncUVs = false;
+			LODInfo.BuildSettings.ThresholdPosition = TNumericLimits<float>::Min();
+			LODInfo.BuildSettings.ThresholdTangentNormal = TNumericLimits<float>::Min();
+			LODInfo.BuildSettings.ThresholdUV = TNumericLimits<float>::Min();
+			LODInfo.BuildSettings.MorphThresholdPosition = TNumericLimits<float>::Min();
+			LODInfo.BuildSettings.BoneInfluenceLimit = 0;
 #endif
 			LODInfo.LODMaterialMap.SetNumZeroed(1);
 		}
