[Mutable] Fix clothing physics asset tags when generating mutable source meshes.

- An incorrect asset could be picked when generating the instance.

#rnx

diff --git a/Engine/Plugins/Experimental/Mutable/Source/CustomizableObject/Private/MuCO/CustomizableObjectInstance.cpp b/Engine/Plugins/Experimental/Mutable/Source/CustomizableObject/Private/MuCO/CustomizableObjectInstance.cpp
--- a/Engine/Plugins/Experimental/Mutable/Source/CustomizableObject/Private/MuCO/CustomizableObjectInstance.cpp
+++ b/Engine/Plugins/Experimental/Mutable/Source/CustomizableObject/Private/MuCO/CustomizableObjectInstance.cpp
@@ -4990,31 +4990,36 @@ UE::Tasks::FTask UCustomizableInstancePrivate::LoadAdditionalAssetsAndDataAsync(
 						}
 					}
 				}
-				else if (Tag.RemoveFromStart("__ClothPhysicsAsset:"))
+				else if (Tag.RemoveFromStart("__ClothPA:"))
 				{
 					FString AssetIndexString, PhysicsAssetIndexString;
 
-					if (Tag.Split(TEXT("_AssetIdx_"), &AssetIndexString, &PhysicsAssetIndexString) && AssetIndexString.IsNumeric())
+					if (Tag.Split(TEXT("_"), &AssetIndexString, &PhysicsAssetIndexString))
 					{
-						const int32 AssetIndex = FCString::Atoi(*AssetIndexString);
-						const int32 PhysicsAssetIndex = FCString::Atoi(*PhysicsAssetIndexString);
+						if (AssetIndexString.IsNumeric() && PhysicsAssetIndexString.IsNumeric())
+						{
+							const int32 AssetIndex = FCString::Atoi(*AssetIndexString);
+							const int32 PhysicsAssetIndex = FCString::Atoi(*PhysicsAssetIndexString);
 
-						const TSoftObjectPtr<UPhysicsAsset>& PhysicsAsset = ModelResources.PhysicsAssets.IsValidIndex(PhysicsAssetIndex) ? ModelResources.PhysicsAssets[PhysicsAssetIndex] : nullptr;
+							const TSoftObjectPtr<UPhysicsAsset>& PhysicsAsset = ModelResources.PhysicsAssets.IsValidIndex(PhysicsAssetIndex) 
+									? ModelResources.PhysicsAssets[PhysicsAssetIndex] 
+									: nullptr;
 
-						// The entry should always be in the map
-						if (!PhysicsAsset.IsNull())
-						{
-							if (PhysicsAsset.Get())
+							// The entry should always be in the map
+							if (!PhysicsAsset.IsNull())
 							{
-								if (ClothingPhysicsAssets.IsValidIndex(AssetIndex))
+								if (PhysicsAsset.Get())
 								{
-									ClothingPhysicsAssets[AssetIndex] = PhysicsAsset.Get();
+									if (ClothingPhysicsAssets.IsValidIndex(AssetIndex))
+									{
+										ClothingPhysicsAssets[AssetIndex] = PhysicsAsset.Get();
+									}
+								}
+								else
+								{
+									ComponentData->ClothingPhysicsAssetsToStream.Emplace(AssetIndex, PhysicsAssetIndex);
+									AssetsToStream.Add(PhysicsAsset.ToSoftObjectPath());
 								}
-							}
-							else
-							{
-								ComponentData->ClothingPhysicsAssetsToStream.Emplace(AssetIndex, PhysicsAssetIndex);
-								AssetsToStream.Add(PhysicsAsset.ToSoftObjectPath());
 							}
 						}
 					}
diff --git a/Engine/Plugins/Experimental/Mutable/Source/CustomizableObjectEditor/Private/MuCOE/GenerateMutableSource/GenerateMutableSourceMesh.cpp b/Engine/Plugins/Experimental/Mutable/Source/CustomizableObjectEditor/Private/MuCOE/GenerateMutableSource/GenerateMutableSourceMesh.cpp
--- a/Engine/Plugins/Experimental/Mutable/Source/CustomizableObjectEditor/Private/MuCOE/GenerateMutableSource/GenerateMutableSourceMesh.cpp
+++ b/Engine/Plugins/Experimental/Mutable/Source/CustomizableObjectEditor/Private/MuCOE/GenerateMutableSource/GenerateMutableSourceMesh.cpp
@@ -2989,14 +2989,10 @@ mu::NodeMeshPtr GenerateMutableSourceMesh(const UEdGraphPin* Pin,
 
 						check(AssetIndex != INDEX_NONE);
 
-						TSoftObjectPtr<UPhysicsAsset> PhysicsAsset = ClothingAssetCommon->PhysicsAsset.Get();
-
-						FString ClothPhysicsAssetTag = FString("__ClothPhysicsAsset:") + FString::Printf(TEXT("%d_AssetIdx_"), AssetIndex) + PhysicsAsset.ToString();
-
 						GenerationContext.AddParticipatingObject(*ClothingAssetCommon->PhysicsAsset);
 						
-						GenerationContext.PhysicsAssets.AddUnique(ClothingAssetCommon->PhysicsAsset);
-
+						const int32 PhysicsAssetIndex = GenerationContext.PhysicsAssets.AddUnique(ClothingAssetCommon->PhysicsAsset);
+						FString ClothPhysicsAssetTag = FString::Printf(TEXT("__ClothPA:%d_%d"), AssetIndex, PhysicsAssetIndex);
 						AddTagToMutableMeshUnique(*MutableMesh, ClothPhysicsAssetTag);
 					}
 				}
