[Mutable] Fix ASTOpSwitch GetImageDesc so it returns an aggregated descriptor.

- Now BlankLayout ImageDesc code generation will use the bReturnBestOption flag.

#rnx

diff --git a/Engine/Plugins/Experimental/Mutable/Source/MutableTools/Private/MuT/ASTOpSwitch.cpp b/Engine/Plugins/Experimental/Mutable/Source/MutableTools/Private/MuT/ASTOpSwitch.cpp
--- a/Engine/Plugins/Experimental/Mutable/Source/MutableTools/Private/MuT/ASTOpSwitch.cpp
+++ b/Engine/Plugins/Experimental/Mutable/Source/MutableTools/Private/MuT/ASTOpSwitch.cpp
@@ -193,20 +193,20 @@ namespace mu
 	}
 
 	//-------------------------------------------------------------------------------------------------
-	FImageDesc ASTOpSwitch::GetImageDesc(bool returnBestOption, class FGetImageDescContext* context) const
+	FImageDesc ASTOpSwitch::GetImageDesc(bool bReturnBestOption, class FGetImageDescContext* Context) const
 	{
-		FImageDesc res;
+		FImageDesc Result;
 
 		// Local context in case it is necessary
-		FGetImageDescContext localContext;
-		if (!context)
+		FGetImageDescContext LocalContext;
+		if (!Context)
 		{
-			context = &localContext;
+			Context = &LocalContext;
 		}
 		else
 		{
 			// Cached result?
-			FImageDesc* PtrValue = context->m_results.Find(this);
+			FImageDesc* PtrValue = Context->m_results.Find(this);
 			if (PtrValue)
 			{
 				return *PtrValue;
@@ -219,69 +219,70 @@ namespace mu
 		// In some places this will force re-formatting of the image.
 		// The code optimiser will take care then of moving the format operations down to each
 		// branch and remove the unnecessary ones.
-		FImageDesc candidate;
-		bool sameSize = true;
-		bool sameFormat = true;
-		bool sameLods = true;
-		bool first = true;
+		FImageDesc Candidate;
+
+		bool bSameSize = true;
+		bool bSameFormat = true;
+		bool bSameLods = true;
+		bool bFirst = true;
 
 		if (def)
 		{
-			FImageDesc childDesc = def->GetImageDesc(returnBestOption, context);
-			candidate = childDesc;
-			first = false;
+			FImageDesc ChildDesc = def->GetImageDesc(bReturnBestOption, Context);
+			Candidate = ChildDesc;
+			bFirst = false;
 		}
 
-		for (int i = 0; i < cases.Num(); ++i)
+		for (int32 CaseIndex = 0; CaseIndex < cases.Num(); ++CaseIndex)
 		{
-			if (cases[i].branch)
+			if (cases[CaseIndex].branch)
 			{
-				FImageDesc childDesc = cases[i].branch->GetImageDesc(returnBestOption, context);
-				if (first)
+				FImageDesc ChildDesc = cases[CaseIndex].branch->GetImageDesc(bReturnBestOption, Context);
+				if (bFirst)
 				{
-					candidate = childDesc;
-					first = false;
+					Candidate = ChildDesc;
+					bFirst = false;
 				}
 				else
 				{
-					sameSize = (candidate.m_size == childDesc.m_size);
-					sameFormat = (candidate.m_format == childDesc.m_format);
-					sameLods = (candidate.m_lods == childDesc.m_lods);
+					bSameSize = bSameSize && (Candidate.m_size == ChildDesc.m_size);
+					bSameFormat = bSameFormat && (Candidate.m_format == ChildDesc.m_format);
+					bSameLods = bSameLods && (Candidate.m_lods == ChildDesc.m_lods);
 
-					if (returnBestOption)
+					if (bReturnBestOption)
 					{
-						candidate.m_format =
-							GetMostGenericFormat(candidate.m_format, childDesc.m_format);
+						Candidate.m_format = GetMostGenericFormat(Candidate.m_format, ChildDesc.m_format);
 					}
 				}
 			}
 		}
 
+		Result = Candidate;
 
-		res = candidate;
-
-		if (!sameFormat && !returnBestOption)
+		// In case of ReturnBestOption the first valid case will be used to determine size and lods.
+		// Format will be the most generic from all cases.
+		if (!bSameFormat && !bReturnBestOption)
 		{
-			res.m_format = EImageFormat::IF_NONE;
+			Result.m_format = EImageFormat::IF_NONE;
 		}
 
-		if (!sameSize && !returnBestOption)
+		if (!bSameSize && !bReturnBestOption)
 		{
-			res.m_size = FImageSize(0, 0);
+			Result.m_size = FImageSize(0, 0);
 		}
 
-		if (!sameLods && !returnBestOption)
+		if (!bSameLods && !bReturnBestOption)
 		{
-			res.m_lods = 0;
+			Result.m_lods = 0;
 		}
 
 		// Cache the result
-		if (context)
+		if (Context)
 		{
-			context->m_results.Add(this, res);
+			Context->m_results.Add(this, Result);
 		}
 
-		return res;
+		return Result;
 	}
 
 
diff --git a/Engine/Plugins/Experimental/Mutable/Source/MutableTools/Private/MuT/CodeGenerator.cpp b/Engine/Plugins/Experimental/Mutable/Source/MutableTools/Private/MuT/CodeGenerator.cpp
--- a/Engine/Plugins/Experimental/Mutable/Source/MutableTools/Private/MuT/CodeGenerator.cpp
+++ b/Engine/Plugins/Experimental/Mutable/Source/MutableTools/Private/MuT/CodeGenerator.cpp
@@ -1306,7 +1306,7 @@ namespace mu
 								}
 							};
 
-							for (int b = 0; b < pLayout->GetBlockCount(); ++b)
+							for (int32 BlockIndex = 0; BlockIndex < pLayout->GetBlockCount(); ++BlockIndex)
 							{
 								// Generate the image
 								FImageGenerationOptions ImageOptions;
@@ -1315,7 +1315,7 @@ namespace mu
 								ImageOptions.RectSize = { 0,0 };
 								ImageOptions.ActiveTags = node.m_tags;
 								ImageOptions.LayoutToApply = pLayout;
-								ImageOptions.LayoutBlockId = pLayout->m_blocks[b].m_id;
+								ImageOptions.LayoutBlockId = pLayout->m_blocks[BlockIndex].m_id;
 								FImageGenerationResult Result;
 								GenerateImage(ImageOptions, Result, pImageNode);
 								Ptr<ASTOp> blockAd = Result.op;
@@ -1327,14 +1327,15 @@ namespace mu
 									continue;
 								}
 
-								// Calculate the desc of the generated block
-								FImageDesc BlockDesc = blockAd->GetImageDesc();
+								// Calculate the desc of the generated block.
+								constexpr bool bReturnBestOption = true;
+								FImageDesc BlockDesc = blockAd->GetImageDesc(bReturnBestOption, nullptr);
 
 								// Block in layout grid units (cells)
 								box< UE::Math::TIntVector2<uint16> > rectInCells;
 								pLayout->GetBlock
 								(
-									b,
+									BlockIndex,
 									&rectInCells.min[0], &rectInCells.min[1],
 									&rectInCells.size[0], &rectInCells.size[1]
 								);
@@ -1354,7 +1355,7 @@ namespace mu
 										if (const NodePatchImage* pPatch = e.node->m_textures[t].m_pPatch.get())
 										{
 											// Is the current block to be patched?
-											if (pPatch->GetPrivate()->m_blocks.Contains(b))
+											if (pPatch->GetPrivate()->m_blocks.Contains(BlockIndex))
 											{
 												blockAd = GenerateImageBlockPatch(blockAd, pPatch, e.condition, ImageOptions);
 											}
@@ -1379,8 +1380,8 @@ namespace mu
 								composeOp->BlockImage = blockAd;
 
 								// Set the absolute block index.
-								check(pLayout->m_blocks[b].m_id >= 0);
-								composeOp->BlockIndex = pLayout->m_blocks[b].m_id;
+								check(pLayout->m_blocks[BlockIndex].m_id >= 0);
+								composeOp->BlockIndex = pLayout->m_blocks[BlockIndex].m_id;
 
 								imageAd = composeOp;
 							}
