[mutable] Some custom maths cleanup.

Remove some of Mutable's legacy math functions in favor of Unreal's.

#rnx

diff --git a/Engine/Plugins/Experimental/Mutable/Source/CustomizableObjectEditor/Private/MuCOE/GenerateMutableSource/GenerateMutableSourceColor.cpp b/Engine/Plugins/Experimental/Mutable/Source/CustomizableObjectEditor/Private/MuCOE/GenerateMutableSource/GenerateMutableSourceColor.cpp
--- a/Engine/Plugins/Experimental/Mutable/Source/CustomizableObjectEditor/Private/MuCOE/GenerateMutableSource/GenerateMutableSourceColor.cpp
+++ b/Engine/Plugins/Experimental/Mutable/Source/CustomizableObjectEditor/Private/MuCOE/GenerateMutableSource/GenerateMutableSourceColor.cpp
@@ -71,7 +71,7 @@ mu::NodeColourPtr GenerateMutableSourceColor(const UEdGraphPin* Pin, FMutableGra
 
 		ColorNode->SetName(TypedNodeColorParam->ParameterName);
 		ColorNode->SetUid(GenerationContext.GetNodeIdUnique(Node).ToString());
-		ColorNode->SetDefaultValue(TypedNodeColorParam->DefaultValue.R, TypedNodeColorParam->DefaultValue.G, TypedNodeColorParam->DefaultValue.B, TypedNodeColorParam->DefaultValue.A);
+		ColorNode->SetDefaultValue(TypedNodeColorParam->DefaultValue);
 
 		GenerationContext.ParameterUIDataMap.Add(TypedNodeColorParam->ParameterName, FParameterUIData(
 			TypedNodeColorParam->ParameterName,
diff --git a/Engine/Plugins/Experimental/Mutable/Source/CustomizableObjectEditor/Private/MuCOE/GenerateMutableSource/GenerateMutableSourceModifier.cpp b/Engine/Plugins/Experimental/Mutable/Source/CustomizableObjectEditor/Private/MuCOE/GenerateMutableSource/GenerateMutableSourceModifier.cpp
--- a/Engine/Plugins/Experimental/Mutable/Source/CustomizableObjectEditor/Private/MuCOE/GenerateMutableSource/GenerateMutableSourceModifier.cpp
+++ b/Engine/Plugins/Experimental/Mutable/Source/CustomizableObjectEditor/Private/MuCOE/GenerateMutableSource/GenerateMutableSourceModifier.cpp
@@ -156,14 +156,7 @@ mu::NodeModifierPtr GenerateMutableSourceModifier(const UEdGraphPin * Pin, FMuta
 				TransformMesh->SetSource(ClipMesh.get());
 
 				FMatrix Matrix = TypedNodeClipMesh->Transform.ToMatrixWithScale();
-				const float M[16] = {
-				  	float(Matrix.M[0][0]),float(Matrix.M[1][0]),float(Matrix.M[2][0]),float(Matrix.M[3][0]),
-					float(Matrix.M[0][1]),float(Matrix.M[1][1]),float(Matrix.M[2][1]),float(Matrix.M[3][1]),
-					float(Matrix.M[0][2]),float(Matrix.M[1][2]),float(Matrix.M[2][2]),float(Matrix.M[3][2]),
-					float(Matrix.M[0][3]),float(Matrix.M[1][3]),float(Matrix.M[2][3]),float(Matrix.M[3][3])
-				};
-
-				TransformMesh->SetTransform(M);
+				TransformMesh->SetTransform(FMatrix44f(Matrix));
 				ClipMesh = TransformMesh;
 			}
 
diff --git a/Engine/Plugins/Experimental/Mutable/Source/CustomizableObjectEditor/Private/MuCOE/GenerateMutableSource/GenerateMutableSourceProjector.cpp b/Engine/Plugins/Experimental/Mutable/Source/CustomizableObjectEditor/Private/MuCOE/GenerateMutableSource/GenerateMutableSourceProjector.cpp
--- a/Engine/Plugins/Experimental/Mutable/Source/CustomizableObjectEditor/Private/MuCOE/GenerateMutableSource/GenerateMutableSourceProjector.cpp
+++ b/Engine/Plugins/Experimental/Mutable/Source/CustomizableObjectEditor/Private/MuCOE/GenerateMutableSource/GenerateMutableSourceProjector.cpp
@@ -39,10 +39,10 @@ mu::NodeProjectorPtr GenerateMutableSourceProjector(const UEdGraphPin* Pin, FMut
 			case mu::PROJECTOR_TYPE::PLANAR:
 				ProjectorNode->SetValue(
 					ProjectorType,
-					TypedNodeConst->Value.Position[0], TypedNodeConst->Value.Position[1], TypedNodeConst->Value.Position[2],
-					TypedNodeConst->Value.Direction[0], TypedNodeConst->Value.Direction[1], TypedNodeConst->Value.Direction[2],
-					TypedNodeConst->Value.Up[0], TypedNodeConst->Value.Up[1], TypedNodeConst->Value.Up[2],
-					TypedNodeConst->Value.Scale[0], TypedNodeConst->Value.Scale[1], TypedNodeConst->Value.Scale[2],
+					TypedNodeConst->Value.Position,
+					TypedNodeConst->Value.Direction,
+					TypedNodeConst->Value.Up,
+					TypedNodeConst->Value.Scale,
 					TypedNodeConst->Value.Angle);
 				break;
 			
@@ -54,10 +54,10 @@ mu::NodeProjectorPtr GenerateMutableSourceProjector(const UEdGraphPin* Pin, FMut
 					const float Height = TypedNodeConst->Value.Scale[2];
 					ProjectorNode->SetValue(
 						ProjectorType,
-						TypedNodeConst->Value.Position[0], TypedNodeConst->Value.Position[1], TypedNodeConst->Value.Position[2],
-						-TypedNodeConst->Value.Direction[0], -TypedNodeConst->Value.Direction[1], -TypedNodeConst->Value.Direction[2],
-						-TypedNodeConst->Value.Up[0], -TypedNodeConst->Value.Up[1], -TypedNodeConst->Value.Up[2],
-						-Height, Radius, Radius,
+						TypedNodeConst->Value.Position,
+						-TypedNodeConst->Value.Direction,
+						-TypedNodeConst->Value.Up,
+						FVector3f( -Height, Radius, Radius ),
 						TypedNodeConst->Value.Angle);
 				}
 				break;
@@ -65,10 +65,10 @@ mu::NodeProjectorPtr GenerateMutableSourceProjector(const UEdGraphPin* Pin, FMut
 			case mu::PROJECTOR_TYPE::WRAPPING:
 				ProjectorNode->SetValue(
 					ProjectorType,
-					TypedNodeConst->Value.Position[0], TypedNodeConst->Value.Position[1], TypedNodeConst->Value.Position[2],
-					TypedNodeConst->Value.Direction[0], TypedNodeConst->Value.Direction[1], TypedNodeConst->Value.Direction[2],
-					TypedNodeConst->Value.Up[0], TypedNodeConst->Value.Up[1], TypedNodeConst->Value.Up[2],
-					TypedNodeConst->Value.Scale[0], TypedNodeConst->Value.Scale[1], TypedNodeConst->Value.Scale[2],
+					TypedNodeConst->Value.Position,
+					TypedNodeConst->Value.Direction,
+					TypedNodeConst->Value.Up,
+					TypedNodeConst->Value.Scale,
 					TypedNodeConst->Value.Angle);
 				break;
 			
@@ -92,10 +92,10 @@ mu::NodeProjectorPtr GenerateMutableSourceProjector(const UEdGraphPin* Pin, FMut
 		case 0:
 			ProjectorNode->SetDefaultValue(
 				mu::PROJECTOR_TYPE::PLANAR,
-				TypedNodeParam->DefaultValue.Position[0], TypedNodeParam->DefaultValue.Position[1], TypedNodeParam->DefaultValue.Position[2],
-				TypedNodeParam->DefaultValue.Direction[0], TypedNodeParam->DefaultValue.Direction[1], TypedNodeParam->DefaultValue.Direction[2],
-				TypedNodeParam->DefaultValue.Up[0], TypedNodeParam->DefaultValue.Up[1], TypedNodeParam->DefaultValue.Up[2],
-				TypedNodeParam->DefaultValue.Scale[0], TypedNodeParam->DefaultValue.Scale[1], TypedNodeParam->DefaultValue.Scale[2],
+				TypedNodeParam->DefaultValue.Position,
+				TypedNodeParam->DefaultValue.Direction,
+				TypedNodeParam->DefaultValue.Up,
+				TypedNodeParam->DefaultValue.Scale,
 				TypedNodeParam->DefaultValue.Angle);
 			break;
 
@@ -107,10 +107,10 @@ mu::NodeProjectorPtr GenerateMutableSourceProjector(const UEdGraphPin* Pin, FMut
 			float Height = TypedNodeParam->DefaultValue.Scale[2];
 			ProjectorNode->SetDefaultValue(
 				mu::PROJECTOR_TYPE::CYLINDRICAL,
-				TypedNodeParam->DefaultValue.Position[0], TypedNodeParam->DefaultValue.Position[1], TypedNodeParam->DefaultValue.Position[2],
-				-TypedNodeParam->DefaultValue.Direction[0], -TypedNodeParam->DefaultValue.Direction[1], -TypedNodeParam->DefaultValue.Direction[2],
-				-TypedNodeParam->DefaultValue.Up[0], -TypedNodeParam->DefaultValue.Up[1], -TypedNodeParam->DefaultValue.Up[2],
-				-Height, Radius, Radius,
+				TypedNodeParam->DefaultValue.Position,
+				-TypedNodeParam->DefaultValue.Direction,
+				-TypedNodeParam->DefaultValue.Up,
+				FVector3f( -Height, Radius, Radius ),
 				TypedNodeParam->DefaultValue.Angle);
 			break;
 		}
@@ -118,10 +118,10 @@ mu::NodeProjectorPtr GenerateMutableSourceProjector(const UEdGraphPin* Pin, FMut
 		case 2:
 			ProjectorNode->SetDefaultValue(
 				mu::PROJECTOR_TYPE::WRAPPING,
-				TypedNodeParam->DefaultValue.Position[0], TypedNodeParam->DefaultValue.Position[1], TypedNodeParam->DefaultValue.Position[2],
-				TypedNodeParam->DefaultValue.Direction[0], TypedNodeParam->DefaultValue.Direction[1], TypedNodeParam->DefaultValue.Direction[2],
-				TypedNodeParam->DefaultValue.Up[0], TypedNodeParam->DefaultValue.Up[1], TypedNodeParam->DefaultValue.Up[2],
-				TypedNodeParam->DefaultValue.Scale[0], TypedNodeParam->DefaultValue.Scale[1], TypedNodeParam->DefaultValue.Scale[2],
+				TypedNodeParam->DefaultValue.Position,
+				TypedNodeParam->DefaultValue.Direction,
+				TypedNodeParam->DefaultValue.Up,
+				TypedNodeParam->DefaultValue.Scale,
 				TypedNodeParam->DefaultValue.Angle);
 			break;
 
diff --git a/Engine/Plugins/Experimental/Mutable/Source/CustomizableObjectEditor/Private/MuCOE/SMutableCodeViewer.cpp b/Engine/Plugins/Experimental/Mutable/Source/CustomizableObjectEditor/Private/MuCOE/SMutableCodeViewer.cpp
--- a/Engine/Plugins/Experimental/Mutable/Source/CustomizableObjectEditor/Private/MuCOE/SMutableCodeViewer.cpp
+++ b/Engine/Plugins/Experimental/Mutable/Source/CustomizableObjectEditor/Private/MuCOE/SMutableCodeViewer.cpp
@@ -2838,7 +2838,7 @@ void SMutableCodeViewer::PreviewMutableCurve(const mu::Curve* Curve)
 }
 
 // TODO: Implement matrix viewer
-void SMutableCodeViewer::PreviewMutableMatrix(const mu::mat4f* Mat)
+void SMutableCodeViewer::PreviewMutableMatrix(const FMatrix44f& Mat)
 {
 	UE_LOG(LogMutable, Warning, TEXT("Previewer for Mutable Matrices not yet implemented"))
 }
diff --git a/Engine/Plugins/Experimental/Mutable/Source/CustomizableObjectEditor/Private/MuCOE/SMutableCodeViewer.h b/Engine/Plugins/Experimental/Mutable/Source/CustomizableObjectEditor/Private/MuCOE/SMutableCodeViewer.h
--- a/Engine/Plugins/Experimental/Mutable/Source/CustomizableObjectEditor/Private/MuCOE/SMutableCodeViewer.h
+++ b/Engine/Plugins/Experimental/Mutable/Source/CustomizableObjectEditor/Private/MuCOE/SMutableCodeViewer.h
@@ -796,7 +796,7 @@ public:
 	void PreviewMutableSkeleton(mu::SkeletonPtrConst Skeleton);
 	void PreviewMutableString(const FString& InString);
 	void PreviewMutableProjector(const mu::FProjector* Projector);
-	void PreviewMutableMatrix(const mu::mat4f* Mat);
+	void PreviewMutableMatrix(const FMatrix44f& Mat);
 	void PreviewMutableShape(const mu::FShape* Shape);
 	void PreviewMutableCurve(const mu::Curve* Curve);
 	
diff --git a/Engine/Plugins/Experimental/Mutable/Source/CustomizableObjectEditor/Private/MuCOE/SMutableConstantsWidget.cpp b/Engine/Plugins/Experimental/Mutable/Source/CustomizableObjectEditor/Private/MuCOE/SMutableConstantsWidget.cpp
--- a/Engine/Plugins/Experimental/Mutable/Source/CustomizableObjectEditor/Private/MuCOE/SMutableConstantsWidget.cpp
+++ b/Engine/Plugins/Experimental/Mutable/Source/CustomizableObjectEditor/Private/MuCOE/SMutableConstantsWidget.cpp
@@ -1265,10 +1265,10 @@ void SMutableConstantsWidget::LoadConstantMatrices()
 	for (int32 MatrixIndex = 0; MatrixIndex < ConstantsCount; MatrixIndex++)
 	{
 		TSharedPtr<FMutableConstantMatrixElement> ConstantMatrixElement = MakeShared<FMutableConstantMatrixElement>();
-		ConstantMatrixElement->Matrix = &(MutableProgramPtr->m_constantMatrices[MatrixIndex]);
+		ConstantMatrixElement->Matrix = MutableProgramPtr->m_constantMatrices[MatrixIndex];
 		ConstantMatrixElement->IndexOnSourceVector = MatrixIndex;
 		
-		Archive << *ConstantMatrixElement->Matrix;
+		Archive << ConstantMatrixElement->Matrix;
 		
 		ConstantMatrixElements.Add(ConstantMatrixElement);
 	}
diff --git a/Engine/Plugins/Experimental/Mutable/Source/CustomizableObjectEditor/Private/MuCOE/SMutableConstantsWidget.h b/Engine/Plugins/Experimental/Mutable/Source/CustomizableObjectEditor/Private/MuCOE/SMutableConstantsWidget.h
--- a/Engine/Plugins/Experimental/Mutable/Source/CustomizableObjectEditor/Private/MuCOE/SMutableConstantsWidget.h
+++ b/Engine/Plugins/Experimental/Mutable/Source/CustomizableObjectEditor/Private/MuCOE/SMutableConstantsWidget.h
@@ -84,7 +84,7 @@ struct FMutableConstantProjectorElement : public FMutableConstantElement
 */
 struct FMutableConstantMatrixElement : public FMutableConstantElement
 { 
-	const mu::mat4f* Matrix = nullptr;
+	FMatrix44f Matrix;
 };
 
 /**
diff --git a/Engine/Plugins/Experimental/Mutable/Source/MutableRuntime/Internal/MuR/ConvertData.h b/Engine/Plugins/Experimental/Mutable/Source/MutableRuntime/Internal/MuR/ConvertData.h
--- a/Engine/Plugins/Experimental/Mutable/Source/MutableRuntime/Internal/MuR/ConvertData.h
+++ b/Engine/Plugins/Experimental/Mutable/Source/MutableRuntime/Internal/MuR/ConvertData.h
@@ -5,6 +5,7 @@
 #include "MuR/Mesh.h"
 #include "MuR/Platform.h"
 #include "MuR/MutableMath.h"
+#include "Math/Float16.h"
 
 
 namespace mu
@@ -47,8 +48,8 @@ namespace mu
 
 			case MBF_FLOAT16:
 			{
-				const float16* pTypedSource = reinterpret_cast<const float16*>(pSource);
-				pTypedResult[channel] = halfToFloat(pTypedSource[channel]);
+				const FFloat16* pTypedSource = reinterpret_cast<const FFloat16*>(pSource);
+				pTypedResult[channel] = pTypedSource[channel];
 				break;
 			}
 
@@ -193,8 +194,8 @@ namespace mu
 
 			case MBF_FLOAT16:
 			{
-				const float16* pTypedSource = reinterpret_cast<const float16*>( pSource );
-				pTypedResult[channel] = halfToFloat(pTypedSource[channel]);
+				const FFloat16* pTypedSource = reinterpret_cast<const FFloat16*>( pSource );
+				pTypedResult[channel] = pTypedSource[channel];
 				break;
 			}
 
@@ -317,20 +318,20 @@ namespace mu
 		//-----------------------------------------------------------------------------------------
 		case MBF_FLOAT16:
 		{
-			float16* pTypedResult = reinterpret_cast<float16*>( pResult );
+			FFloat16* pTypedResult = reinterpret_cast<FFloat16*>( pResult );
 
 			switch ( sourceFormat )
 			{
 			case MBF_FLOAT32:
 			{
 				const float* pTypedSource = reinterpret_cast<const float*>( pSource );
-				pTypedResult[channel] = floatToHalf(pTypedSource[channel]);
+				pTypedResult[channel] = pTypedSource[channel];
 				break;
 			}
 
 			case MBF_FLOAT16:
 			{
-				const float16* pTypedSource = reinterpret_cast<const float16*>( pSource );
+				const FFloat16* pTypedSource = reinterpret_cast<const FFloat16*>( pSource );
 				pTypedResult[channel] = pTypedSource[channel];
 				break;
 			}
@@ -338,42 +339,42 @@ namespace mu
 			case MBF_UINT32:
 			{
                 const uint32* pTypedSource = reinterpret_cast<const uint32*>( pSource );
-				pTypedResult[channel] = floatToHalf((float)(pTypedSource[channel]));
+				pTypedResult[channel] = (float)(pTypedSource[channel]);
 				break;
 			}
 
 			case MBF_INT32:
 			{
                 const int32* pTypedSource = reinterpret_cast<const int32*>( pSource );
-				pTypedResult[channel] = floatToHalf((float)(pTypedSource[channel]));
+				pTypedResult[channel] = (float)(pTypedSource[channel]);
 				break;
 			}
 
 			case MBF_UINT16:
 			{
                 const uint16* pTypedSource = reinterpret_cast<const uint16*>( pSource );
-				pTypedResult[channel] = floatToHalf((float)(pTypedSource[channel]));
+				pTypedResult[channel] = (float)(pTypedSource[channel]);
 				break;
 			}
 
 			case MBF_INT16:
 			{
                 const int16* pTypedSource = reinterpret_cast<const int16*>( pSource );
-				pTypedResult[channel] = floatToHalf((float)(pTypedSource[channel]));
+				pTypedResult[channel] = (float)(pTypedSource[channel]);
 				break;
 			}
 
 			case MBF_UINT8:
 			{
                 const uint8* pTypedSource = reinterpret_cast<const uint8*>( pSource );
-				pTypedResult[channel] = floatToHalf((float)(pTypedSource[channel]));
+				pTypedResult[channel] = (float)(pTypedSource[channel]);
 				break;
 			}
 
 			case MBF_INT8:
 			{
                 const int8* pTypedSource = reinterpret_cast<const int8*>( pSource );
-				pTypedResult[channel] = floatToHalf((float)(pTypedSource[channel]));
+				pTypedResult[channel] = (float)(pTypedSource[channel]);
 				break;
 			}
 
@@ -407,13 +408,13 @@ namespace mu
 
 			case MBF_FLOAT16:
 			{
-				const float16* pTypedSource = reinterpret_cast<const float16*>( pSource );
+				const FFloat16* pTypedSource = reinterpret_cast<const FFloat16*>( pSource );
                 pTypedResult[channel] = (uint8)
                         FMath::Min<uint32>(
 							0xFF,
                             FMath::Max<uint32>(
 								0,
-                                (uint32)halfToFloat(pTypedSource[channel])
+                                (uint32)float(pTypedSource[channel])
 								)
 							);
 				break;
@@ -529,13 +530,13 @@ namespace mu
 
 			case MBF_FLOAT16:
 			{
-				const float16* pTypedSource = reinterpret_cast<const float16*>( pSource );
+				const FFloat16* pTypedSource = reinterpret_cast<const FFloat16*>( pSource );
                 pTypedResult[channel] = (uint16)
                         FMath::Min<uint32>(
 							0xFF,
                             FMath::Max<uint32>(
 								0,
-                                (uint32)halfToFloat(pTypedSource[channel])
+                                (uint32)float(pTypedSource[channel])
 								)
 							);
 				break;
@@ -637,13 +638,13 @@ namespace mu
 
 			case MBF_FLOAT16:
 			{
-				const float16* pTypedSource = reinterpret_cast<const float16*>( pSource );
+				const FFloat16* pTypedSource = reinterpret_cast<const FFloat16*>( pSource );
 				pTypedResult[channel] =
                         FMath::Min<uint32>(
 							0xFF,
                             FMath::Max<uint32>(
 								0,
-                                (uint32)halfToFloat(pTypedSource[channel])
+                                (uint32)float(pTypedSource[channel])
 								)
 							);
 				break;
@@ -730,13 +731,13 @@ namespace mu
 
 			case MBF_FLOAT16:
 			{
-				const float16* pTypedSource = reinterpret_cast<const float16*>( pSource );
+				const FFloat16* pTypedSource = reinterpret_cast<const FFloat16*>( pSource );
                 pTypedResult[channel] = (int8)
                         FMath::Min<int32>(
 							127,
                             FMath::Max<int32>(
 								-128,
-                                (int32)halfToFloat(pTypedSource[channel])
+                                (int32)float(pTypedSource[channel])
 								)
 							);
 				break;
@@ -787,13 +788,13 @@ namespace mu
 
 			case MBF_FLOAT16:
 			{
-				const float16* pTypedSource = reinterpret_cast<const float16*>( pSource );
+				const FFloat16* pTypedSource = reinterpret_cast<const FFloat16*>( pSource );
                 pTypedResult[channel] = (int16)
                         FMath::Min<int32>(
 							32767,
                             FMath::Max<int32>(
 								-32768,
-                                (int32)halfToFloat(pTypedSource[channel])
+                                (int32)float(pTypedSource[channel])
 								)
 							);
 				break;
@@ -878,8 +879,8 @@ namespace mu
 
 			case MBF_FLOAT16:
 			{
-				const float16* pTypedSource = reinterpret_cast<const float16*>( pSource );
-                pTypedResult[channel] = (int32)halfToFloat(pTypedSource[channel]);
+				const FFloat16* pTypedSource = reinterpret_cast<const FFloat16*>( pSource );
+                pTypedResult[channel] = (int32)float(pTypedSource[channel]);
 				break;
 			}
 
@@ -978,13 +979,13 @@ namespace mu
 
 			case MBF_FLOAT16:
 			{
-				const float16* pTypedSource = reinterpret_cast<const float16*>( pSource );
+				const FFloat16* pTypedSource = reinterpret_cast<const FFloat16*>( pSource );
                 pTypedResult[channel] = (uint8)
                         FMath::Min<uint32>(
 							0xFF,
                             FMath::Max<uint32>(
 								0,
-                                (uint32)(((float)0xFF)*halfToFloat(pTypedSource[channel])+0.5f)
+                                (uint32)(((float)0xFF)*float(pTypedSource[channel])+0.5f)
 								)
 							);
 				break;
@@ -1043,13 +1044,13 @@ namespace mu
 
 			case MBF_FLOAT16:
 			{
-				const float16* pTypedSource = reinterpret_cast<const float16*>( pSource );
+				const FFloat16* pTypedSource = reinterpret_cast<const FFloat16*>( pSource );
                 pTypedResult[channel] = (uint16)
                         FMath::Min<uint32>(
 							0xFFFF,
                             FMath::Max<uint32>(
 								0,
-                                (uint32)(((float)0xFFFF)*halfToFloat(pTypedSource[channel])+0.5f)
+                                (uint32)(((float)0xFFFF)*float(pTypedSource[channel])+0.5f)
 								)
 							);
 				break;
@@ -1094,13 +1095,13 @@ namespace mu
 
 			case MBF_FLOAT16:
 			{
-				const float16* pTypedSource = reinterpret_cast<const float16*>( pSource );
+				const FFloat16* pTypedSource = reinterpret_cast<const FFloat16*>( pSource );
                 pTypedResult[channel] = (uint32)
                         FMath::Min<uint32>(
 							0xFFFFFFFF,
                             FMath::Max<uint32>(
 								0,
-                                (uint32)(((float)0xFFFFFFFF)*halfToFloat(pTypedSource[channel])+0.5f)
+                                (uint32)(((float)0xFFFFFFFF)*float(pTypedSource[channel])+0.5f)
 								)
 							);
 				break;
@@ -1145,13 +1146,13 @@ namespace mu
 
 			case MBF_FLOAT16:
 			{
-				const float16* pTypedSource = reinterpret_cast<const float16*>( pSource );
+				const FFloat16* pTypedSource = reinterpret_cast<const FFloat16*>( pSource );
                 pTypedResult[channel] = (int8)
                         FMath::Min<int32>(
 							127,
                             FMath::Max<int32>(
 								-128,
-                                (int32)(128.0f*halfToFloat(pTypedSource[channel])+0.5f)
+                                (int32)(128.0f*float(pTypedSource[channel])+0.5f)
 								)
 							);
 				break;
@@ -1196,13 +1197,13 @@ namespace mu
 
 			case MBF_FLOAT16:
 			{
-				const float16* pTypedSource = reinterpret_cast<const float16*>( pSource );
+				const FFloat16* pTypedSource = reinterpret_cast<const FFloat16*>( pSource );
                 pTypedResult[channel] = (int16)
                         FMath::Min<int32>(
 							32767,
                             FMath::Max<int32>(
 								-32768,
-                                (int32)(32768.0f*halfToFloat(pTypedSource[channel])+0.5f)
+                                (int32)(32768.0f*float(pTypedSource[channel])+0.5f)
 								)
 							);
 				break;
@@ -1241,8 +1242,8 @@ namespace mu
 
 			case MBF_FLOAT16:
 			{
-				const float16* pTypedSource = reinterpret_cast<const float16*>( pSource );
-                pTypedResult[channel] = (int32)(2147483648.0f*halfToFloat(pTypedSource[channel])+0.5f);
+				const FFloat16* pTypedSource = reinterpret_cast<const FFloat16*>( pSource );
+                pTypedResult[channel] = (int32)(2147483648.0f*float(pTypedSource[channel])+0.5f);
 				break;
 			}
 
diff --git a/Engine/Plugins/Experimental/Mutable/Source/MutableRuntime/Internal/MuR/MeshPrivate.h b/Engine/Plugins/Experimental/Mutable/Source/MutableRuntime/Internal/MuR/MeshPrivate.h
--- a/Engine/Plugins/Experimental/Mutable/Source/MutableRuntime/Internal/MuR/MeshPrivate.h
+++ b/Engine/Plugins/Experimental/Mutable/Source/MutableRuntime/Internal/MuR/MeshPrivate.h
@@ -340,14 +340,10 @@ namespace mu
 			}
 		}
 
-		inline vec<CTYPE,COMPONENTS>& operator*()
+		// \TODO: Replace this with safer sized-typed access.
+		inline CTYPE* operator*()
 		{
-			return *reinterpret_cast< vec<CTYPE,COMPONENTS>* >( m_pBuf );
-		}
-
-		inline vec<CTYPE,COMPONENTS>& operator[]( size_t i )
-		{
-			return *((*this)+int(i));
+			return reinterpret_cast<CTYPE*>(m_pBuf);
 		}
 
 		inline MeshBufferIterator<FORMAT,CTYPE,COMPONENTS> operator+(int c)
@@ -521,25 +517,13 @@ namespace mu
             ConvertData( 0, &res, MBF_UINT32, ptr(), m_format );
             return res;
         }
-        
-        vec<int32_t,8> GetAsVec8i() const
-        {
-            vec<int32_t, 8> res;
-            for ( int c = 0; c < FMath::Min( m_components, 8 ); ++c )
-            {
-                ConvertData( c, &res[0], MBF_INT32, ptr(), m_format );
-            }
-            return res;
-        }
-
-		vec<int32_t, 12> GetAsVec12i() const
+ 
+		void GetAsInt32Vec(int32* Data, int32 Count) const
 		{
-			vec<int32_t, 12> res;
-			for (int c = 0; c < FMath::Min(m_components, 12); ++c)
+			for (int32 c = 0; c < FMath::Min(m_components, Count); ++c)
 			{
-				ConvertData(c, &res[0], MBF_INT32, ptr(), m_format);
+				ConvertData(c, Data, MBF_INT32, ptr(), m_format);
 			}
-			return res;
 		}
 
         inline UntypedMeshBufferIteratorConst operator+(int c) const
@@ -623,14 +607,9 @@ namespace mu
 			}
 		}
 
-		inline const vec<CTYPE,COMPONENTS>& operator*() const
-		{
-			return *reinterpret_cast< const vec<CTYPE,COMPONENTS>* >( m_pBuf );
-		}
-
-		inline const vec<CTYPE,COMPONENTS>& operator[]( size_t i ) const
+		inline const CTYPE* operator*() const
 		{
-			return *((*this)+int(i));
+			return reinterpret_cast<const CTYPE*>(m_pBuf);
 		}
 
 		inline MeshBufferIteratorConst<FORMAT,CTYPE,COMPONENTS> operator+(int c) const
diff --git a/Engine/Plugins/Experimental/Mutable/Source/MutableRuntime/Internal/MuR/ModelPrivate.h b/Engine/Plugins/Experimental/Mutable/Source/MutableRuntime/Internal/MuR/ModelPrivate.h
--- a/Engine/Plugins/Experimental/Mutable/Source/MutableRuntime/Internal/MuR/ModelPrivate.h
+++ b/Engine/Plugins/Experimental/Mutable/Source/MutableRuntime/Internal/MuR/ModelPrivate.h
@@ -247,7 +247,7 @@ namespace mu
 		TArray<FProjector> m_constantProjectors;
 
         //! Constant matrices, usually used for transforms
-		TArray<mat4f> m_constantMatrices;
+		TArray<FMatrix44f> m_constantMatrices;
 
 		//! Constant shapes
 		TArray<FShape> m_constantShapes;
@@ -505,7 +505,7 @@ namespace mu
             return index;
         }
 
-        OP::ADDRESS AddConstant( const mat4f& m )
+        OP::ADDRESS AddConstant( const FMatrix44f& m )
         {
             // Ensure unique
             for ( SIZE_T i=0; i<m_constantMatrices.Num(); ++i)
diff --git a/Engine/Plugins/Experimental/Mutable/Source/MutableRuntime/Internal/MuR/ParametersPrivate.h b/Engine/Plugins/Experimental/Mutable/Source/MutableRuntime/Internal/MuR/ParametersPrivate.h
--- a/Engine/Plugins/Experimental/Mutable/Source/MutableRuntime/Internal/MuR/ParametersPrivate.h
+++ b/Engine/Plugins/Experimental/Mutable/Source/MutableRuntime/Internal/MuR/ParametersPrivate.h
@@ -201,11 +201,11 @@ namespace Private
 	struct FShape
 	{
 		// Transform
-		vec3f position;
-		vec3f up;
-		vec3f side;
+		FVector3f position;
+		FVector3f up;
+		FVector3f side;
 		
-		vec3f size;
+		FVector3f size;
 
 		// 
 		enum class Type : uint8
diff --git a/Engine/Plugins/Experimental/Mutable/Source/MutableRuntime/Internal/MuR/SerialisationPrivate.h b/Engine/Plugins/Experimental/Mutable/Source/MutableRuntime/Internal/MuR/SerialisationPrivate.h
--- a/Engine/Plugins/Experimental/Mutable/Source/MutableRuntime/Internal/MuR/SerialisationPrivate.h
+++ b/Engine/Plugins/Experimental/Mutable/Source/MutableRuntime/Internal/MuR/SerialisationPrivate.h
@@ -305,20 +305,7 @@ namespace mu
 		}
 	}
 
-	MUTABLE_DEFINE_POD_SERIALISABLE(vec2f);
-	MUTABLE_DEFINE_POD_SERIALISABLE(vec3f);
-	MUTABLE_DEFINE_POD_SERIALISABLE(vec4f);
-	MUTABLE_DEFINE_POD_SERIALISABLE(mat3f);
-	MUTABLE_DEFINE_POD_SERIALISABLE(mat4f);
-	MUTABLE_DEFINE_POD_SERIALISABLE(vec2<int>);
-
 	// Unreal POD Serializables
-	MUTABLE_DEFINE_POD_SERIALISABLE(FUintVector2);
-	MUTABLE_DEFINE_POD_SERIALISABLE(UE::Math::TIntVector2<uint16>);
-	MUTABLE_DEFINE_POD_SERIALISABLE(UE::Math::TIntVector2<int16>);
-	MUTABLE_DEFINE_POD_SERIALISABLE(FVector4f);
-	MUTABLE_DEFINE_POD_SERIALISABLE(UE::Math::TVector4<float>);
-
 	MUTABLE_IMPLEMENT_POD_VECTOR_SERIALISABLE(float);
 	MUTABLE_IMPLEMENT_POD_VECTOR_SERIALISABLE(double);
 	MUTABLE_IMPLEMENT_POD_VECTOR_SERIALISABLE(uint8);
@@ -329,19 +316,15 @@ namespace mu
 	MUTABLE_IMPLEMENT_POD_VECTOR_SERIALISABLE(int16);
 	MUTABLE_IMPLEMENT_POD_VECTOR_SERIALISABLE(int32);
 	MUTABLE_IMPLEMENT_POD_VECTOR_SERIALISABLE(int64);
-
-	MUTABLE_IMPLEMENT_POD_VECTOR_SERIALISABLE(vec2f);
-	MUTABLE_IMPLEMENT_POD_VECTOR_SERIALISABLE(vec3f);
-	MUTABLE_IMPLEMENT_POD_VECTOR_SERIALISABLE(vec4f);
-	MUTABLE_IMPLEMENT_POD_VECTOR_SERIALISABLE(mat3f);
-	MUTABLE_IMPLEMENT_POD_VECTOR_SERIALISABLE(mat4f);
-	MUTABLE_IMPLEMENT_POD_VECTOR_SERIALISABLE(vec2<int>);
 	MUTABLE_IMPLEMENT_POD_VECTOR_SERIALISABLE(TCHAR);
 
+	MUTABLE_IMPLEMENT_POD_VECTOR_SERIALISABLE(FIntVector2);
 	MUTABLE_IMPLEMENT_POD_VECTOR_SERIALISABLE(FUintVector2);
 	MUTABLE_IMPLEMENT_POD_VECTOR_SERIALISABLE(UE::Math::TIntVector2<uint16>);
 	MUTABLE_IMPLEMENT_POD_VECTOR_SERIALISABLE(UE::Math::TIntVector2<int16>);
+	MUTABLE_IMPLEMENT_POD_VECTOR_SERIALISABLE(FVector2f);
 	MUTABLE_IMPLEMENT_POD_VECTOR_SERIALISABLE(FVector4f);
+	MUTABLE_IMPLEMENT_POD_VECTOR_SERIALISABLE(FMatrix44f);
 
 	
 	//---------------------------------------------------------------------------------------------
@@ -501,23 +484,4 @@ namespace mu
 	}
 
 
-	//---------------------------------------------------------------------------------------------
-	template< typename T >
-	inline void operator<< (OutputArchive& arch, const UE::Math::TVector4<T>& v)
-	{
-		arch << v.X;
-		arch << v.Y;
-		arch << v.Z;
-		arch << v.W;
-	}
-
-	template< typename T >
-	inline void operator>> (InputArchive& arch, UE::Math::TVector4<T>& v)
-	{
-		arch >> v.X;
-		arch >> v.Y;
-		arch >> v.Z;
-		arch >> v.W;
-	}
-
 }
diff --git a/Engine/Plugins/Experimental/Mutable/Source/MutableRuntime/Private/MuR/BlockCompression/Miro/Miro.cpp b/Engine/Plugins/Experimental/Mutable/Source/MutableRuntime/Private/MuR/BlockCompression/Miro/Miro.cpp
--- a/Engine/Plugins/Experimental/Mutable/Source/MutableRuntime/Private/MuR/BlockCompression/Miro/Miro.cpp
+++ b/Engine/Plugins/Experimental/Mutable/Source/MutableRuntime/Private/MuR/BlockCompression/Miro/Miro.cpp
@@ -6546,9 +6546,9 @@ namespace impl
 				blue = red - blue;
 			}
 			if (majComp == 1)
-				std::swap(red, green);
+				Swap(red, green);
 			else if (majComp == 2)
-				std::swap(red, blue);
+				Swap(red, blue);
 			e0 = UVec4(FMath::Clamp(red - scale, 0, 0xfff),
 				FMath::Clamp(green - scale, 0, 0xfff),
 				FMath::Clamp(blue - scale, 0, 0xfff),
diff --git a/Engine/Plugins/Experimental/Mutable/Source/MutableRuntime/Private/MuR/CodeRunner.cpp b/Engine/Plugins/Experimental/Mutable/Source/MutableRuntime/Private/MuR/CodeRunner.cpp
--- a/Engine/Plugins/Experimental/Mutable/Source/MutableRuntime/Private/MuR/CodeRunner.cpp
+++ b/Engine/Plugins/Experimental/Mutable/Source/MutableRuntime/Private/MuR/CodeRunner.cpp
@@ -1947,12 +1947,12 @@ namespace mu
             		
                 Ptr<const Mesh> Source = LoadMesh(FCacheAddress(args.source,item));
 
-                const mat4f& mat = Program.m_constantMatrices[args.matrix];
+                const FMatrix44f& mat = Program.m_constantMatrices[args.matrix];
 
 				Ptr<Mesh> Result = CreateMesh(Source ? Source->GetDataSize() : 0);
 
 				bool bOutSuccess = false;
-                MeshTransform(Result.get(), Source.get(), ToUnreal(mat), bOutSuccess);
+                MeshTransform(Result.get(), Source.get(), mat, bOutSuccess);
 
 				if (!bOutSuccess)
 				{
@@ -2003,8 +2003,8 @@ namespace mu
                 // Should be an ellipse
                 const FShape& morphShape = Program.m_constantShapes[args.morphShape];
 
-                const mu::vec3f& origin = morphShape.position;
-                const mu::vec3f& normal = morphShape.up;
+                const FVector3f& origin = morphShape.position;
+                const FVector3f& normal = morphShape.up;
 
                 if (args.vertexSelectionType == OP::MeshClipMorphPlaneArgs::VS_SHAPE)
                 {
diff --git a/Engine/Plugins/Experimental/Mutable/Source/MutableRuntime/Private/MuR/Mesh.cpp b/Engine/Plugins/Experimental/Mutable/Source/MutableRuntime/Private/MuR/Mesh.cpp
--- a/Engine/Plugins/Experimental/Mutable/Source/MutableRuntime/Private/MuR/Mesh.cpp
+++ b/Engine/Plugins/Experimental/Mutable/Source/MutableRuntime/Private/MuR/Mesh.cpp
@@ -1008,7 +1008,7 @@ void Mesh::GetVertexMap
                 for ( int ov=0; ov<bucketVertexCount; ++ov )
                 {
                     int otherVertexIndex = buckets[bucket][ov];
-                    vec3<float> p = itopBegin[ otherVertexIndex ];
+                    FVector3f p = (itopBegin+otherVertexIndex).GetAsVec3f();
 
                     bool same = true;
                     for ( int d=0; same && d<3; ++d )
diff --git a/Engine/Plugins/Experimental/Mutable/Source/MutableRuntime/Private/MuR/OpImageGradient.h b/Engine/Plugins/Experimental/Mutable/Source/MutableRuntime/Private/MuR/OpImageGradient.h
--- a/Engine/Plugins/Experimental/Mutable/Source/MutableRuntime/Private/MuR/OpImageGradient.h
+++ b/Engine/Plugins/Experimental/Mutable/Source/MutableRuntime/Private/MuR/OpImageGradient.h
@@ -21,7 +21,7 @@ namespace mu
 		for ( int i=0; i< sizeX; ++i )
 		{
 			float delta = float(i)/float(sizeX -1);
-			vec3<float> c = c0 * (1.0f-delta) + c1 * delta;
+			FVector3f c = c0 * (1.0f-delta) + c1 * delta;
 
             uint8 colour[4];
             colour[0] = (uint8)FMath::Max( 0, FMath::Min( 255, int(c[0]*255.0f) ) );
diff --git a/Engine/Plugins/Experimental/Mutable/Source/MutableRuntime/Private/MuR/OpImageNormalCombine.h b/Engine/Plugins/Experimental/Mutable/Source/MutableRuntime/Private/MuR/OpImageNormalCombine.h
--- a/Engine/Plugins/Experimental/Mutable/Source/MutableRuntime/Private/MuR/OpImageNormalCombine.h
+++ b/Engine/Plugins/Experimental/Mutable/Source/MutableRuntime/Private/MuR/OpImageNormalCombine.h
@@ -21,16 +21,17 @@ namespace mu
 		const float blendedGf = (static_cast<float>((blended >> 8)  & 0xFF) / 255.0f) * 2.0f - 1.0f;
 		const float blendedBf = (static_cast<float>((blended >> 16) & 0xFF) / 255.0f) * 2.0f - 1.0f;
 
-		vec3f a = vec3f(baseRf, baseGf, baseBf);
-		vec3f b = vec3f(-blendedRf, -blendedGf, blendedBf);
+		FVector3f a = FVector3f(baseRf, baseGf, baseBf);
+		FVector3f b = FVector3f(-blendedRf, -blendedGf, blendedBf);
 
 		// The original does not normalize but if not done results don't look good due to signal clipping.
-		vec3f n = normalise_approx(a*dot(a, b) - b*a.z());
+		FVector3f n = (a* FVector3f::DotProduct(a, b) - b*a.Z);
+		n.Normalize();
 
 		return 
-			(mu::clamp(static_cast<unsigned>((n.x() + 1.0f) * 0.5f * 255.0f), 0u, 255u) << 0) |
-			(mu::clamp(static_cast<unsigned>((n.y() + 1.0f) * 0.5f * 255.0f), 0u, 255u) << 8) |
-			(mu::clamp(static_cast<unsigned>((n.z() + 1.0f) * 0.5f * 255.0f), 0u, 255u) << 16);
+			(FMath::Clamp(static_cast<unsigned>((n.X + 1.0f) * 0.5f * 255.0f), 0u, 255u) << 0) |
+			(FMath::Clamp(static_cast<unsigned>((n.Y + 1.0f) * 0.5f * 255.0f), 0u, 255u) << 8) |
+			(FMath::Clamp(static_cast<unsigned>((n.Z + 1.0f) * 0.5f * 255.0f), 0u, 255u) << 16);
 	}
 
 	inline unsigned CombineNormalMasked(unsigned base, unsigned blended, unsigned mask)
@@ -48,17 +49,17 @@ namespace mu
 
 		const float maskRf = (static_cast<float>(mask & 0xFF) / 255.0f);
 
-		vec3f a = vec3f(baseRf, baseGf, baseBf + 1.0f);
-		vec3f b = vec3f(-blendedRf, -blendedGf, blendedBf);
+		FVector3f a = FVector3f(baseRf, baseGf, baseBf + 1.0f);
+		FVector3f b = FVector3f(-blendedRf, -blendedGf, blendedBf);
 
 		// The original does not normalize but if not done results don't look good due to signal clipping.
-		vec3f n = mu::normalise_approx( 
-					mu::lerp( vec3f(baseRf, baseGf, baseBf), a*dot(a, b) - b*a.z(), maskRf ) );
+		FVector3f n = ( FMath::Lerp( FVector3f(baseRf, baseGf, baseBf), a*FVector3f::DotProduct(a, b) - b*a.Z, maskRf ) );
+		n.Normalize();
 
 		return 
-			(mu::clamp(static_cast<unsigned>((n.x() + 1.0f) * 0.5f * 255.0f), 0u, 255u) << 0) |
-			(mu::clamp(static_cast<unsigned>((n.y() + 1.0f) * 0.5f * 255.0f), 0u, 255u) << 8) |
-			(mu::clamp(static_cast<unsigned>((n.z() + 1.0f) * 0.5f * 255.0f), 0u, 255u) << 16);
+			(FMath::Clamp(static_cast<unsigned>((n.X + 1.0f) * 0.5f * 255.0f), 0u, 255u) << 0) |
+			(FMath::Clamp(static_cast<unsigned>((n.Y + 1.0f) * 0.5f * 255.0f), 0u, 255u) << 8) |
+			(FMath::Clamp(static_cast<unsigned>((n.Z + 1.0f) * 0.5f * 255.0f), 0u, 255u) << 16);
 	}
 
 	inline void ImageNormalCombine(Image* pResult, const Image* pBase, FVector4f col)
diff --git a/Engine/Plugins/Experimental/Mutable/Source/MutableRuntime/Private/MuR/OpImageProject.cpp b/Engine/Plugins/Experimental/Mutable/Source/MutableRuntime/Private/MuR/OpImageProject.cpp
--- a/Engine/Plugins/Experimental/Mutable/Source/MutableRuntime/Private/MuR/OpImageProject.cpp
+++ b/Engine/Plugins/Experimental/Mutable/Source/MutableRuntime/Private/MuR/OpImageProject.cpp
@@ -17,6 +17,7 @@
 
 #include "Async/ParallelFor.h"
 #include "Containers/Array.h"
+#include "Intersection/IntrRay3Triangle3.h"
 #include "Math/IntPoint.h"
 #include "Math/UnrealMathSSE.h"
 #include "Misc/AssertionMacros.h"
@@ -2055,24 +2056,23 @@ void MeshProject_Optimised_Cylindrical(const FOptimizedVertex* pVertices, int ve
     // TODO: support for non uniform scale?
     float radius = projectorScale[1];
     float height = projectorScale[0];
-    mat3f worldToCylinder = mat3f(vec3f(projectorDirection), vec3f(projectorSide), vec3f(projectorUp) );
-    //worldToCylinder = worldToCylinder.GetTransposed();
+	FMatrix44f WorldToCylinder(projectorDirection, projectorSide, projectorUp,FVector3f(0,0,0));
 
     for ( int v=0; v<vertexCount; ++v )
     {
         // Cylinder is along the X axis
 
         // Project
-		vec3f relPos = vec3f(pVertices[v].Position - projectorPosition);
-        vec3f vertexPos_cylinder = worldToCylinder.Transform(relPos);
+		FVector3f RelPos = pVertices[v].Position - projectorPosition;
+		FVector3f VertexPos_Cylinder = WorldToCylinder.InverseTransformPosition(RelPos);
 
         // This final projection needs to be done per pixel
-        float x = vertexPos_cylinder.x() / height;
-        float r2 = vertexPos_cylinder.y()*vertexPos_cylinder.y()
-                + vertexPos_cylinder.z()*vertexPos_cylinder.z();
+        float x = VertexPos_Cylinder.X / height;
+        float r2 = VertexPos_Cylinder.Y * VertexPos_Cylinder.Y
+                + VertexPos_Cylinder.Z * VertexPos_Cylinder.Z;
         projectedPositions[v].pos0 = x;
-        projectedPositions[v].pos1 = vertexPos_cylinder.y() / radius;
-        projectedPositions[v].pos2 = vertexPos_cylinder.z() / radius;
+        projectedPositions[v].pos1 = VertexPos_Cylinder.Y / radius;
+        projectedPositions[v].pos2 = VertexPos_Cylinder.Z / radius;
         uint32 planeMask =
                 ((x<0.0f)<<0) |
                 ((x>1.0f)<<1) |
@@ -2193,6 +2193,9 @@ void MeshProject_Optimised_Wrapping(const Mesh* pMesh,
         vertToFacesMap.Add(collapsedVertexMap[i2], f);
     }
 
+	FRay3f Ray(projectorPosition, projectorDirection, false);
+	UE::Geometry::FIntrRay3Triangle3f Intersector(Ray, UE::Geometry::FTriangle3f());
+
     // Trace a ray in the projection direction to find the face that will be projected planarly and be the root of the unfolding
     // Also build face connectivity information
     for (int f = 0; f < faceCount; ++f)
@@ -2203,16 +2206,17 @@ void MeshProject_Optimised_Wrapping(const Mesh* pMesh,
 
 		FVector3f rayStart = projectorPosition;
 		FVector3f rayEnd = projectorPosition + projectorDirection * maxDist;
-		FVector3f aux_out_intersection;
-        int out_intersected_vert, out_intersected_edge_v0, out_intersected_edge_v1;
-        float t;
+
+		Intersector.Triangle = UE::Geometry::FTriangle3f(pVertices[i0].Position, pVertices[i1].Position, pVertices[i2].Position);
+
 
         rayLength = (rayEnd - rayStart).Length();
 
-        bool intersects = rayIntersectsFace2(rayStart, rayEnd, pVertices[i0].Position, pVertices[i1].Position, pVertices[i2].Position,
-            aux_out_intersection, out_intersected_vert, out_intersected_edge_v0, out_intersected_edge_v1, t);
+		bool bIntersects = Intersector.Find();
+		float t = Intersector.RayParameter;
+		FVector3f aux_out_intersection = Ray.PointAt(Intersector.RayParameter);
 
-        if (intersects && t < min_t)
+        if (bIntersects && t < min_t)
         {
             intersectedFace = f;
             min_t = t;
diff --git a/Engine/Plugins/Experimental/Mutable/Source/MutableRuntime/Private/MuR/OpLayoutPack.h b/Engine/Plugins/Experimental/Mutable/Source/MutableRuntime/Private/MuR/OpLayoutPack.h
--- a/Engine/Plugins/Experimental/Mutable/Source/MutableRuntime/Private/MuR/OpLayoutPack.h
+++ b/Engine/Plugins/Experimental/Mutable/Source/MutableRuntime/Private/MuR/OpLayoutPack.h
@@ -85,11 +85,11 @@ namespace mu
     {
         TArray< UE::Math::TIntVector2<uint16> > blocks;
 		TArray< FLayoutBlock > sorted;
-		TArray< vec2<int> > positions;
-		TArray< int > priorities;
-		TArray< vec2<int> > reductions;
-		TArray< int > ReduceBothAxes;
-		TArray< int > ReduceByTwo;
+		TArray< FIntVector2 > positions;
+		TArray< int32 > priorities;
+		TArray< FIntVector2 > reductions;
+		TArray< int32 > ReduceBothAxes;
+		TArray< int32 > ReduceByTwo;
     };
 
 
@@ -116,132 +116,6 @@ namespace mu
     }
 
 
-    // inline void LayoutPack( Layout* pResult,
-    //                         const Layout* pSourceLayout,
-    //                         SCRATCH_LAYOUT_PACK* scratch  )
-    // {
-    //     check( pResult->GetBlockCount() == pSourceLayout->GetBlockCount() );
-
-    //     int blockCount = pSourceLayout->GetBlockCount();
-    //     check( (int)scratch->blocks.Num()==blockCount );
-
-    //     // Look for the maximum block sizes on the layout and the total area
-    //     int maxX = 0;
-    //     int maxY = 0;
-    //     int area = 0;
-    //     for ( int index=0; index<blockCount; ++index )
-    //     {
-    //         box< vec2<int> > b;
-    //         pSourceLayout->GetBlock( index, &b.min[0], &b.min[1], &b.size[0], &b.size[1] );
-
-    //         maxX = FMath::Max( maxX, b.size[0] );
-    //         maxY = FMath::Max( maxY, b.size[1] );
-
-    //         area += b.size[0] * b.size[1];
-
-    //         scratch->blocks[index] = b.size;
-    //     }
-
-    //     // Grow until the area is big enough to fit all blocks. We always grow X first, because
-    //     // in case we cannot pack everything, we will grow Y with the current horizon algorithm.
-    //     maxX = ceilPow2( maxX );
-    //     maxY = ceilPow2( maxY );
-    //     while ( maxX*maxY<area )
-    //     {
-    //         if (maxX>maxY)
-    //         {
-    //             maxY*=2;
-    //         }
-    //         else
-    //         {
-    //             maxX*=2;
-    //         }
-    //     }
-
-    //     // Sort by height, area
-    //     check( (int)scratch->sorted.Num()==blockCount );
-    //     for ( int index=0; index<blockCount; ++index )
-    //     {
-    //         scratch->sorted[index] = FLayoutBlock( index, scratch->blocks[index], scratch->priorities[index] );
-    //     }
-    //     std::sort( scratch->sorted.begin(), scratch->sorted.end(), CompareBlocks );
-
-    //     check( (int)scratch->sorted.Num()==blockCount );
-
-    //     bool fits = false;
-    //     int iterations = 0;
-
-    //     while (!fits)
-    //     {
-    //         ++iterations;
-
-    //         // Pack with fixed horizontal size
-    //         check( maxX<256 );
-    //         int16_t horizon[256];
-    //         mutable_memset( horizon, 0, 256*sizeof(int16_t) );
-    //         maxY = 0;
-
-    //         for ( size_t p=0; p<scratch->sorted.Num(); ++p )
-    //         {
-    //             // Seek for the lowest span where the block fits
-    //             int currentLevel = std::numeric_limits<int>::max();
-    //             int currentX = 0;
-    //             for ( int x=0; x<=maxX-scratch->sorted[p].size[0]; ++x )
-    //             {
-    //                 int level = 0;
-    //                 for( int xs=x; xs<x+scratch->sorted[p].size[0]; ++xs )
-    //                 {
-    //                     level = FMath::Max( level, (int)horizon[xs] );
-    //                 }
-
-    //                 if (level<currentLevel)
-    //                 {
-    //                     currentLevel = level;
-    //                     currentX = x;
-    //                 }
-
-    //             }
-
-    //             check( currentX>=0 && currentX<=maxX-scratch->sorted[p].size[0] );
-
-    //             // Update horizon
-    //             for( int xs=currentX; xs<currentX+scratch->sorted[p].size[0]; ++xs )
-    //             {
-    //                 horizon[xs] = (uint16)(currentLevel+scratch->sorted[p].size[1]);
-    //             }
-
-    //             // Store
-    //             scratch->positions[ scratch->sorted[p].index ] = vec2<int>( currentX, currentLevel );
-    //             maxY = FMath::Max( maxY, currentLevel+scratch->sorted[p].size[1] );
-    //         }
-
-    //         maxY = ceilPow2( maxY );
-
-    //         // TODO: Adjust this value
-    //         if ( maxY <= maxX || iterations>5000 )
-    //         {
-    //             fits = true;
-    //         }
-    //         else
-    //         {
-    //             fits = !std::next_permutation( scratch->sorted.begin(), scratch->sorted.end(), CompareBlocks );
-    //         }
-    //     }
-
-    //     // Set data in the result
-    //     pResult->SetGridSize( maxX, maxY );
-    //     for ( int index=0; index<blockCount; ++index )
-    //     {
-    //         pResult->SetBlock
-    //             (
-    //                 index,
-    //                 scratch->positions[index][0], scratch->positions[index][1],
-    //                 scratch->blocks[index][0], scratch->blocks[index][1]
-    //             );
-    //     }
-    // }
-
-
 	inline void ReductionOperation(uint16& BlockSize, EReductionMethod ReductionMethod, int32 bReduceByTwo)
 	{
 		if (ReductionMethod == EReductionMethod::UNITARY_REDUCTION)
@@ -483,7 +357,7 @@ namespace mu
 			}
 
 			// Store
-			scratch->positions[scratch->sorted[best].index] = vec2<int>(bestX, bestLevel);
+			scratch->positions[scratch->sorted[best].index] = FIntVector2(bestX, bestLevel);
 			*maxY = FMath::Max(*maxY, uint16(bestLevel + scratch->sorted[best].size[1]) );
 
 			if (packStrategy == EPackStrategy::FIXED_LAYOUT && *maxY > layoutSizeY)
@@ -552,7 +426,7 @@ namespace mu
 			bool bReduceBothAxes, bReduceByTwo;
 			pSourceLayout->GetBlockOptions(index, p, bReduceBothAxes, bReduceByTwo);
 
-			vec2<int> reductions;
+			FIntVector2 reductions;
 			reductions[0] = 0;
 			reductions[1] = 0;
 
diff --git a/Engine/Plugins/Experimental/Mutable/Source/MutableRuntime/Private/MuR/OpMeshApplyLayout.h b/Engine/Plugins/Experimental/Mutable/Source/MutableRuntime/Private/MuR/OpMeshApplyLayout.h
--- a/Engine/Plugins/Experimental/Mutable/Source/MutableRuntime/Private/MuR/OpMeshApplyLayout.h
+++ b/Engine/Plugins/Experimental/Mutable/Source/MutableRuntime/Private/MuR/OpMeshApplyLayout.h
@@ -108,15 +108,15 @@ namespace mu
 			case MBF_FLOAT16:
 			{
 				// TODO: Optimise
-				float16* pUV = reinterpret_cast<float16*>( pVertices );
+				FFloat16* pUV = reinterpret_cast<FFloat16*>( pVertices );
 
 				FVector2f uv;
-				uv[0] = halfToFloat( pUV[0] );
-				uv[1] = halfToFloat( pUV[1] );
+				uv[0] = float( pUV[0] );
+				uv[1] = float( pUV[1] );
 				uv = uv * transforms[ relBlock ].size + transforms[ relBlock ].min;
 
-				pUV[0] = floatToHalf( uv[0] );
-				pUV[1] = floatToHalf( uv[1] );
+				pUV[0] = FFloat16( uv[0] );
+				pUV[1] = FFloat16( uv[1] );
 				break;
 			}
 
diff --git a/Engine/Plugins/Experimental/Mutable/Source/MutableRuntime/Private/MuR/OpMeshApplyPose.h b/Engine/Plugins/Experimental/Mutable/Source/MutableRuntime/Private/MuR/OpMeshApplyPose.h
--- a/Engine/Plugins/Experimental/Mutable/Source/MutableRuntime/Private/MuR/OpMeshApplyPose.h
+++ b/Engine/Plugins/Experimental/Mutable/Source/MutableRuntime/Private/MuR/OpMeshApplyPose.h
@@ -121,7 +121,7 @@ namespace mu
 
         for ( int v=0; v<vertexCount; ++v )
         {
-            FVector3f sourcePosition = ToUnreal(*itSource);
+            FVector3f sourcePosition = itSource.GetAsVec3f();
 			FVector3f position = FVector3f(0,0,0);
 
             float totalWeight = 0.0f;
@@ -152,7 +152,9 @@ namespace mu
                 position /= totalWeight;
             }
 
-            *itTarget = FromUnreal(position);
+			(*itTarget)[0] = position[0];
+			(*itTarget)[1] = position[1];
+			(*itTarget)[2] = position[2];
 
             ++itTarget;
             ++itSource;
diff --git a/Engine/Plugins/Experimental/Mutable/Source/MutableRuntime/Private/MuR/OpMeshClipMorphPlane.h b/Engine/Plugins/Experimental/Mutable/Source/MutableRuntime/Private/MuR/OpMeshClipMorphPlane.h
--- a/Engine/Plugins/Experimental/Mutable/Source/MutableRuntime/Private/MuR/OpMeshClipMorphPlane.h
+++ b/Engine/Plugins/Experimental/Mutable/Source/MutableRuntime/Private/MuR/OpMeshClipMorphPlane.h
@@ -10,9 +10,9 @@
 
 namespace mu
 {
-	inline bool PointInBoundingBox(const vec3f& point, const FShape& selectionShape)
+	inline bool PointInBoundingBox(const FVector3f& point, const FShape& selectionShape)
 	{
-		vec3f v = point - selectionShape.position;
+		FVector3f v = point - selectionShape.position;
 
 		for (int i = 0; i < 3; ++i)
 		{
@@ -25,15 +25,15 @@ namespace mu
 		return true;
 	}
 
-	inline bool VertexIsInMaxRadius(const vec3f& vertex, const vec3f& origin, float vertexSelectionBoneMaxRadius)
+	inline bool VertexIsInMaxRadius(const FVector3f& vertex, const FVector3f& origin, float vertexSelectionBoneMaxRadius)
 	{
 		if (vertexSelectionBoneMaxRadius < 0.f)
 		{
 			return true;
 		}
 
-		vec3f radiusVec = vertex - origin;
-		float radius = dot(radiusVec, radiusVec);
+		FVector3f radiusVec = vertex - origin;
+		float radius = FVector3f::DotProduct(radiusVec, radiusVec);
 
 		return radius < vertexSelectionBoneMaxRadius* vertexSelectionBoneMaxRadius;
 	}
@@ -69,7 +69,7 @@ namespace mu
 	//---------------------------------------------------------------------------------------------
 	//! Reference version
 	//---------------------------------------------------------------------------------------------
-	inline void MeshClipMorphPlane(Mesh* Result, const Mesh* pBase, const vec3f& origin, const vec3f& normal, float dist, float factor, float radius,
+	inline void MeshClipMorphPlane(Mesh* Result, const Mesh* pBase, const FVector3f& origin, const FVector3f& normal, float dist, float factor, float radius,
 		float radius2, float angle, const FShape& selectionShape, bool& bOutSuccess, const int32 BoneId = INDEX_NONE, float vertexSelectionBoneMaxRadius = -1.f)
 	{
 		bOutSuccess = true;
@@ -78,15 +78,15 @@ namespace mu
 		//float factor = 1.f;
 
 		// Generate vector perpendicular to normal for ellipse rotation reference base
-		vec3f aux_base(0.f, 1.f, 0.f);
+		FVector3f aux_base(0.f, 1.f, 0.f);
 
-		if (fabs(dot(normal, aux_base)) > 0.95f)		// fabs = absolute value
+		if (FMath::Abs(FVector3f::DotProduct(normal, aux_base)) > 0.95f)		// fabs = absolute value
 		{
-			aux_base = vec3f(0.f, 0.f, 1.f);
+			aux_base = FVector3f(0.f, 0.f, 1.f);
 		}
 		
-		vec3f origin_radius_vector = cross(normal, aux_base);		// PERPENDICULAR VECTOR TO THE PLANE normal and aux base
-		check(fabs(dot(normal, origin_radius_vector)) < 0.05f);
+		FVector3f origin_radius_vector = FVector3f::CrossProduct(normal, aux_base);		// PERPENDICULAR VECTOR TO THE PLANE normal and aux base
+		check(FMath::Abs(FVector3f::DotProduct(normal, origin_radius_vector)) < 0.05f);
 
         uint32 vcount = pBase->GetVertexBuffers().GetElementCount();
 
@@ -305,7 +305,7 @@ namespace mu
 				case MBS_POSITION:
                     for (uint32 v = 0; v < vcount; ++v)
 					{
-						vec3f vertex(0.0f, 0.0f, 0.0f);
+						FVector3f vertex(0.0f, 0.0f, 0.0f);
 						for (int i = 0; i < 3; ++i)
 						{
 							ConvertData(i, &vertex[0], MBF_FLOAT32, it.ptr(), it.GetFormat());
@@ -326,25 +326,25 @@ namespace mu
 
 						if (bIsVertexAffectedBone || bIsVertexAffectedNoShape || bIsVertexAffectedBoundingBox)
 						{
-							vec3f morph_plane_center = origin;					// MORPH PLANE POS relative to root of the selected bone
-							vec3f clip_plane_center = origin + normal * dist;	// CPLIPPING PLANE POS
-							vec3f aux_morph = vertex - morph_plane_center;		// MORPH PLANE --> CURRENT VERTEX
-							vec3f aux_clip = vertex - clip_plane_center;		// CLIPPING PLANE --> CURRENT VERTEX
+							FVector3f morph_plane_center = origin;					// MORPH PLANE POS relative to root of the selected bone
+							FVector3f clip_plane_center = origin + normal * dist;	// CPLIPPING PLANE POS
+							FVector3f aux_morph = vertex - morph_plane_center;		// MORPH PLANE --> CURRENT VERTEX
+							FVector3f aux_clip = vertex - clip_plane_center;		// CLIPPING PLANE --> CURRENT VERTEX
 
-							float dot_morph = dot(aux_morph, normal);			// ANGLE (MORPH PLANE TO VERTEX AND NORMAL)
-							float dot_cut = dot(aux_clip, normal);				// ANGLE (CLIPPING PLANE TO VERTEX AND NORMAL )
+							float dot_morph = FVector3f::DotProduct(aux_morph, normal);			// ANGLE (MORPH PLANE TO VERTEX AND NORMAL)
+							float dot_cut = FVector3f::DotProduct(aux_clip, normal);				// ANGLE (CLIPPING PLANE TO VERTEX AND NORMAL )
 
 							if (dot_morph >= 0.f || dot_cut >= 0.f)				// CHECK IF CLIPPING OR MORPH SHOULD BE COMPUTED FOR V VERTEX
 							{
-								vec3f current_center = morph_plane_center + normal * dot_morph;	// PROJECTED POINT FROM THE MROPH PLANE (THE CLOSER THE DOT VALUE OF NORMAL AND VERTEX THE FURTHER IT GOES )
-								vec3f radius_vector = vertex - current_center;					// 	
-								float radius_vector_len = length(radius_vector);
-								vec3f radius_vector_unit = radius_vector_len != 0.f ? radius_vector / radius_vector_len : vec3f(0.f, 0.f, 0.f); // UNITARY VECTOR THAT GOES FROM THE POINT TO THE VERTEX
+								FVector3f current_center = morph_plane_center + normal * dot_morph;	// PROJECTED POINT FROM THE MROPH PLANE (THE CLOSER THE DOT VALUE OF NORMAL AND VERTEX THE FURTHER IT GOES )
+								FVector3f radius_vector = vertex - current_center;					// 	
+								float radius_vector_len = radius_vector.Length();
+								FVector3f radius_vector_unit = radius_vector_len != 0.f ? radius_vector / radius_vector_len : FVector3f(0.f, 0.f, 0.f); // UNITARY VECTOR THAT GOES FROM THE POINT TO THE VERTEX
 
-								float angle_from_origin = acosf(dot(radius_vector_unit, origin_radius_vector));
+								float angle_from_origin = acosf(FVector3f::DotProduct(radius_vector_unit, origin_radius_vector));
 
 								// Cross product between the perpendicular vector from radius vector and origin radius and the normal vector
-								if (dot(cross(radius_vector_unit, origin_radius_vector), normal) < 0)
+								if (FVector3f::DotProduct(FVector3f::CrossProduct(radius_vector_unit, origin_radius_vector), normal) < 0)
 								{
 									angle_from_origin = -angle_from_origin;
 								}
@@ -355,16 +355,16 @@ namespace mu
 								float term2 = radius * sinf(angle_from_origin);
 								float ellipse_radius_at_angle = radius * radius2 / sqrtf(term1 * term1 + term2 * term2);
 
-								vec3f vertex_proj_ellipse = current_center + radius_vector_unit * ellipse_radius_at_angle;
-								//vec3f vertex_proj_ellipse = current_center + radius_vector_unit * radius;
+								FVector3f vertex_proj_ellipse = current_center + radius_vector_unit * ellipse_radius_at_angle;
+								//FVector3f vertex_proj_ellipse = current_center + radius_vector_unit * radius;
 
-								float morph_alpha = dist != 0.f && dot_morph <= dist ? clamp(0.f, 1.f, powf(dot_morph / dist, factor)) : 1.f;
+								float morph_alpha = dist != 0.f && dot_morph <= dist ? FMath::Clamp(powf(dot_morph / dist, factor),0.f, 1.f) : 1.f;
 
 								vertex = vertex * (1.f - morph_alpha) + vertex_proj_ellipse * morph_alpha;
 
 								if (dot_cut >= 0.f)		// CHECK IF THE VERTEX SHOULD BE CLIPPED
 								{
-									vec3f vert_displ = normal * -dot_cut;
+									FVector3f vert_displ = normal * -dot_cut;
 									vertex = vertex + vert_displ;
 									RemovedVertices[v] = true;
 								}
@@ -389,12 +389,12 @@ namespace mu
 		// Now remove all the faces from the result mesh that have all vertices removed
 		UntypedMeshBufferIteratorConst itBase(Result->GetIndexBuffers(), MBS_VERTEXINDEX);
 		UntypedMeshBufferIterator itDest(Result->GetIndexBuffers(), MBS_VERTEXINDEX);
-		int aFaceCount = Result->GetFaceCount();
+		int32 aFaceCount = Result->GetFaceCount();
 
 		UntypedMeshBufferIteratorConst ito(Result->GetIndexBuffers(), MBS_VERTEXINDEX);
 		for (int f = 0; f < aFaceCount; ++f)
 		{
-            vec3<uint32> ov;
+            FUint32Vector3 ov;
 			ov[0] = ito.GetAsUINT32(); ++ito;
 			ov[1] = ito.GetAsUINT32(); ++ito;
 			ov[2] = ito.GetAsUINT32(); ++ito;
@@ -414,11 +414,11 @@ namespace mu
 			itBase += 3;
 		}
 
-		std::size_t removedIndices = itBase - itDest;
+		SIZE_T removedIndices = itBase - itDest;
 		check(removedIndices % 3 == 0);
 
-		Result->GetFaceBuffers().SetElementCount(aFaceCount - (int)removedIndices / 3);
-		Result->GetIndexBuffers().SetElementCount(aFaceCount * 3 - (int)removedIndices);
+		Result->GetFaceBuffers().SetElementCount(aFaceCount - (int32)removedIndices / 3);
+		Result->GetIndexBuffers().SetElementCount(aFaceCount * 3 - (int32)removedIndices);
 
 		// TODO: Should redo/reorder the face buffer before SetElementCount since some deleted faces could be left and some remaining faces deleted.
 
diff --git a/Engine/Plugins/Experimental/Mutable/Source/MutableRuntime/Private/MuR/OpMeshClipWithMesh.cpp b/Engine/Plugins/Experimental/Mutable/Source/MutableRuntime/Private/MuR/OpMeshClipWithMesh.cpp
--- a/Engine/Plugins/Experimental/Mutable/Source/MutableRuntime/Private/MuR/OpMeshClipWithMesh.cpp
+++ b/Engine/Plugins/Experimental/Mutable/Source/MutableRuntime/Private/MuR/OpMeshClipWithMesh.cpp
@@ -236,88 +236,7 @@ namespace mu { namespace
     }
 
 
-    //---------------------------------------------------------------------------------------------
-    //! Reference version
-    //---------------------------------------------------------------------------------------------
-  //  inline int get_num_intersections( const FVector3f& vertex,
-  //                                    const FVector3f& ray,
-  //                                    const TArray<FVector3f>& vertices,
-  //                                    const TArray<uint32>& faces,
-  //                                    TArray<int>& collapsedVertexMap,
-		//							  TArray<uint8>& vertex_already_intersected,
-  //                                    TMap<TPair<int, int>, bool>& edge_already_intersected,
-		//							  float dynamic_epsilon
-  //                                    )
-  //  {
-  //  	
-  //      MUTABLE_CPUPROFILER_SCOPE(get_num_intersections);
-  //  	
-  //      int num_intersections = 0;
-		//FVector3f intersection;
-
-		//FMemory::Memzero( vertex_already_intersected.GetData(), vertex_already_intersected.Num());
-
-  //      edge_already_intersected.Empty();
-
-  //  #define get_collapsed_vertex(vertex_idx) (vertices[collapsedVertexMap[vertex_idx]])
-
-  //      size_t fcount = faces.Num()/3;
-
-  //      // Check vertex against all ClipMorph faces
-  //      for (uint32 face = 0; face < fcount; ++face)
-  //      {
-  //          uint32 vertex_indexs[3] = { faces[3 * face], faces[3 * face + 1], faces[3 * face + 2] };
-
-		//	FVector3f v0 = get_collapsed_vertex(vertex_indexs[0]);
-		//	FVector3f v1 = get_collapsed_vertex(vertex_indexs[1]);
-		//	FVector3f v2 = get_collapsed_vertex(vertex_indexs[2]);
-
-  //          int out_intersected_vert, out_intersected_edge_v0, out_intersected_edge_v1;
-
-  //          if (rayIntersectsFace(vertex, ray, v0, v1, v2, intersection, out_intersected_vert, out_intersected_edge_v0, out_intersected_edge_v1, dynamic_epsilon))
-  //          {
-  //              bool vertex_not_intersected_before = true;
-  //              bool edge_not_intersected_before = true;
-
-  //              if (out_intersected_vert >= 0)
-  //              {
-  //                  int collapsed_vert_index = collapsedVertexMap[vertex_indexs[out_intersected_vert]];
-  //                  vertex_not_intersected_before = !vertex_already_intersected[collapsed_vert_index];
-
-  //                  vertex_already_intersected[collapsed_vert_index] = true;
-  //              }
-
-  //              if (out_intersected_edge_v0 >= 0)
-  //              {
-  //                  int collapsed_edge_vert_index0 = collapsedVertexMap[vertex_indexs[out_intersected_edge_v0]];
-  //                  int collapsed_edge_vert_index1 = collapsedVertexMap[vertex_indexs[out_intersected_edge_v1]];
-
-  //                  TPair<int, int> key;
-  //                  key.Key = collapsed_edge_vert_index0 <= collapsed_edge_vert_index1 ? collapsed_edge_vert_index0 : collapsed_edge_vert_index1;
-  //                  key.Value = collapsed_edge_vert_index1 > collapsed_edge_vert_index0 ? collapsed_edge_vert_index1 : collapsed_edge_vert_index0;
-  //                  edge_not_intersected_before = edge_already_intersected.Find(key) != nullptr;
-
-  //                  if (edge_not_intersected_before)
-  //                  {
-  //                      edge_already_intersected.Add( key, true );
-  //                  }
-  //              }
-
-  //              if (vertex_not_intersected_before && edge_not_intersected_before)
-  //              {
-  //                  num_intersections++;
-  //              }
-  //          }
-  //      }
-
-  //      return num_intersections;
-  //  }
-
-
-    //---------------------------------------------------------------------------------------------
-    //! Core Geometry version
-    //---------------------------------------------------------------------------------------------
-    inline int GetNumIntersections( const FRay3f& Ray,
+    inline int32 GetNumIntersections( const FRay3f& Ray,
                                     const TArray<FVector3f>& Vertices,
                                     const TArray<uint32>& Faces,
                                     const TArray<int32>& CollapsedVertexMap,
@@ -980,21 +899,21 @@ namespace mu
 
         // Make a tolerance proportional to the mesh bounding box size
         // TODO: Use precomputed bounding box
-        box< vec3<float> > aabbox;
+        box< FVector3f > aabbox;
         if ( fragmentFaceCount > 0 )
         {
             MeshBufferIteratorConst<MBF_FLOAT32,float,3> itp( pFragment->GetVertexBuffers(), MBS_POSITION );
 
-            aabbox.min = *itp;
+            aabbox.min = itp.GetAsVec3f();
             ++itp;
 
             for ( int v=1; v<pFragment->GetVertexBuffers().GetElementCount(); ++v )
             {
-                aabbox.Bound( *itp );
+                aabbox.Bound(itp.GetAsVec3f());
                 ++itp;
             }
         }
-        float tolerance = 1e-5f * length(aabbox.size);
+        float tolerance = 1e-5f * aabbox.size.Length();
         Mesh::VERTEX_MATCH_MAP vertexMap;
         pFragment->GetVertexMap( *pBase, vertexMap, tolerance );
 
@@ -1047,7 +966,7 @@ namespace mu
         for ( int f=0; f<sourceFaceCount; ++f )
         {
             bool hasFace = false;
-            vec3<uint32> ov;
+            FUint32Vector3 ov;
             ov[0] = ito.GetAsUINT32(); ++ito;
             ov[1] = ito.GetAsUINT32(); ++ito;
             ov[2] = ito.GetAsUINT32(); ++ito;
@@ -1061,7 +980,7 @@ namespace mu
             {
                 int tf =  buckets[bucket][btf];
 
-                vec3<uint32> v;
+				FUint32Vector3 v;
                 v[0] = (itti+3*tf+0).GetAsUINT32();
                 v[1] = (itti+3*tf+1).GetAsUINT32();
                 v[2] = (itti+3*tf+2).GetAsUINT32();
@@ -1089,7 +1008,7 @@ namespace mu
         int aFaceCount = pBase->GetFaceCount();
         for (int f = 0; f < aFaceCount; ++f)
         {
-            vec3<uint32> ov;
+			FUint32Vector3 ov;
             ov[0] = itoi.GetAsUINT32(); ++itoi;
             ov[1] = itoi.GetAsUINT32(); ++itoi;
             ov[2] = itoi.GetAsUINT32(); ++itoi;
diff --git a/Engine/Plugins/Experimental/Mutable/Source/MutableRuntime/Private/MuR/OpMeshDifference.h b/Engine/Plugins/Experimental/Mutable/Source/MutableRuntime/Private/MuR/OpMeshDifference.h
--- a/Engine/Plugins/Experimental/Mutable/Source/MutableRuntime/Private/MuR/OpMeshDifference.h
+++ b/Engine/Plugins/Experimental/Mutable/Source/MutableRuntime/Private/MuR/OpMeshDifference.h
@@ -8,15 +8,6 @@
 namespace mu
 {
 
-#pragma pack(push,1)
-struct MORPH_VERTEX
-{
-    uint32_t index;
-	vec3<float> pos;
-	vec3<float> nor;
-};
-#pragma pack(pop)
-
 	//---------------------------------------------------------------------------------------------
 	//! Create a diff from the mesh vertices. The meshes must have the same amount of vertices.
 	//! If the channel list is empty all the channels will be compared.
@@ -204,126 +195,4 @@ struct MORPH_VERTEX
 	}
 
 
-	/*/
-	{
-		bool correct =
-				( pBase->GetVertexBuffers().GetElementCount()
-						== pTarget->GetVertexBuffers().GetElementCount() )
-				&&
-				( pBase->GetIndexCount() == pTarget->GetIndexCount()
-				&&
-				 pBase->GetVertexBuffers().GetElementCount()>0
-				);
-
-		MeshPtr pDest = new Mesh();
-
-		if (correct)
-		{
-			// For now, just make a difference from the position and normal buffers.
-			UINT32 vcount = pBase->GetVertexBuffers().GetElementCount();
-
-			// Base mesh position
-			int basePosElementSize = 0;
-			const UINT8* pBasePosBuf = 0;
-			GetMeshBuf( pBase, MBS_POSITION, MBF_FLOAT32, 3, pBasePosBuf, basePosElementSize );
-
-			// Base mesh normal
-			int baseNorElementSize = 0;
-			const UINT8* pBaseNorBuf = 0;
-			GetMeshBuf( pBase, MBS_NORMAL, MBF_FLOAT32, 3, pBaseNorBuf, baseNorElementSize );
-
-			// Target mesh position
-			int targetPosElementSize = 0;
-			const UINT8* pTargetPosBuf = 0;
-			GetMeshBuf( pTarget, MBS_POSITION, MBF_FLOAT32, 3, pTargetPosBuf, targetPosElementSize );
-
-			// Target mesh normal
-			int targetNorElementSize = 0;
-			const UINT8* pTargetNorBuf = 0;
-			GetMeshBuf( pTarget, MBS_NORMAL, MBF_FLOAT32, 3, pTargetNorBuf, targetNorElementSize );
-
-			// Actual diff
-			vector< MORPH_VERTEX > vertices;
-
-			for ( UINT32 vi=0; vi<vcount; ++vi )
-			{
-				MORPH_VERTEX v;
-
-				vec3<float> basePos = *reinterpret_cast< const vec3<float>* >( pBasePosBuf );
-				vec3<float> targetPos = *reinterpret_cast< const vec3<float>* >( pTargetPosBuf );
-				v.pos = targetPos - basePos;
-
-				vec3<float> baseNor = *reinterpret_cast< const vec3<float>* >( pBaseNorBuf );
-				vec3<float> targetNor = *reinterpret_cast< const vec3<float>* >( pTargetNorBuf );
-				v.nor = targetNor - baseNor;
-
-				if ( !( v.pos.AlmostNull() &&
-						v.nor.AlmostNull() ) )
-				{
-					v.index = vi;
-					vertices.push_back( v );
-				}
-
-				pBasePosBuf += basePosElementSize;
-				pTargetPosBuf += targetPosElementSize;
-				pBaseNorBuf += baseNorElementSize;
-				pTargetNorBuf += targetNorElementSize;
-			}
-
-			// Create the morph mesh
-			{
-				int morphCount = (int)vertices.size();
-				pDest->GetVertexBuffers().SetElementCount( morphCount );
-				pDest->GetVertexBuffers().SetBufferCount( 1 );
-
-				// Vertex indices
-				{
-					vector<MESH_BUFFER_SEMANTIC> semantic;
-					semantic.push_back( MBS_VERTEXINDEX );
-					semantic.push_back( MBS_POSITION );
-					semantic.push_back( MBS_NORMAL );
-
-					vector<int> semanticIndex;
-					semanticIndex.push_back( 0 );
-					semanticIndex.push_back( 0 );
-					semanticIndex.push_back( 0 );
-
-					vector<MESH_BUFFER_FORMAT> format;
-					format.push_back( MBF_UINT32 );
-					format.push_back( MBF_FLOAT32 );
-					format.push_back( MBF_FLOAT32 );
-
-					vector<int> components;
-					components.push_back( 1 );
-					components.push_back( 3 );
-					components.push_back( 3 );
-
-					vector<int> offset;
-					offset.push_back( 0 );
-					offset.push_back( 4 );
-					offset.push_back( 16 );
-
-					pDest->GetVertexBuffers().SetBuffer
-						(
-							0,
-							28,
-							3,
-							&semantic[0],
-							&semanticIndex[0],
-							&format[0],
-							&components[0],
-							&offset[0]
-						);
-
-					UINT8* pDestData = pDest->GetVertexBuffers().GetBufferData( 0 );
-					int size = morphCount*sizeof(MORPH_VERTEX);
-					memcpy( pDestData, &vertices[0], size );
-				}
-
-			}
-		}
-
-		return pDest;
-	}
-*/
 }
diff --git a/Engine/Plugins/Experimental/Mutable/Source/MutableRuntime/Private/MuR/OpMeshFormat.cpp b/Engine/Plugins/Experimental/Mutable/Source/MutableRuntime/Private/MuR/OpMeshFormat.cpp
--- a/Engine/Plugins/Experimental/Mutable/Source/MutableRuntime/Private/MuR/OpMeshFormat.cpp
+++ b/Engine/Plugins/Experimental/Mutable/Source/MutableRuntime/Private/MuR/OpMeshFormat.cpp
@@ -32,122 +32,123 @@ namespace mu
 #define MUTABLE_TANGENT_GENERATION_EPSILON_2    0.001f
 #define MUTABLE_TANGENT_MIN_AXIS_DIFFERENCE     0.999f
 
-		struct TVertex
-		{
-
-			TVertex() = default;
-
-			TVertex(const vec3<float>& p, const vec3<float>& n, const vec2<float>& t)
-			{
-				pos = p;
-				nor = n;
-				tex = t;
-			}
-
-			vec3<float> pos;
-			vec3<float> nor;
-			vec2<float> tex;
-
-			inline bool operator< (const TVertex& other) const
-			{
-				if (pos < other.pos)
-					return true;
-				if (other.pos < pos)
-					return false;
-
-				if (nor < other.nor)
-					return true;
-				if (other.nor < nor)
-					return false;
-
-				// Compare the texture coordinates with a particular precission.
-				vec2<int> uv0 = vec2<int>(int(tex[0] * MUTABLE_VERTEX_MERGE_TEX_RANGE),
-					int(tex[1] * MUTABLE_VERTEX_MERGE_TEX_RANGE));
-				vec2<int> uv1 = vec2<int>(int(other.tex[0] * MUTABLE_VERTEX_MERGE_TEX_RANGE),
-					int(other.tex[1] * MUTABLE_VERTEX_MERGE_TEX_RANGE));
-
-				if (uv0 < uv1)
-					return true;
-				if (uv1 < uv0)
-					return false;
-
-				return false;
-			}
-
-			inline bool operator== (const TVertex& other) const
-			{
-				return  !((*this) < other)
-					&&
-					!(other < (*this));
-			}
-		};
-
-
-		struct TFace
-		{
-			vec3<float> T;
-			vec3<float> B;
-			vec3<float> N;
-
-			TFace()
-				: T(0, 0, 0)
-				, B(0, 0, 0)
-				, N(0, 0, 0)
-			{}
-
-			TFace
-			(
-				const vec3<float>& v1,
-				const vec3<float>& v2,
-				const vec3<float>& v3,
-				const vec2<float>& w1,
-				const vec2<float>& w2,
-				const vec2<float>& w3
-			)
-			{
-				vec3<float> E1 = v2 - v1;
-				vec3<float> E2 = v3 - v1;
-
-				vec2<float> UV1 = w2 - w1;
-				vec2<float> UV2 = w3 - w1;
-
-				float  UVdet = UV1[0] * UV2[1] - UV2[0] * UV1[1];
-
-				N = normalise(cross(E1, E2));
-
-				if (!(fabs(UVdet) <= MUTABLE_TANGENT_GENERATION_EPSILON_1))
-				{
-					double r = 1.0 / UVdet;
-
-					T = vec3<float>
-						(
-							(UV2[1] * E1[0] - UV1[1] * E2[0]),
-							(UV2[1] * E1[1] - UV1[1] * E2[1]),
-							(UV2[1] * E1[2] - UV1[1] * E2[2])
-							);
-
-					B = vec3<float>
-						(
-							(UV1[0] * E2[0] - UV2[0] * E1[0]),
-							(UV1[0] * E2[1] - UV2[0] * E1[1]),
-							(UV1[0] * E2[2] - UV2[0] * E1[2])
-							);
-
-					T = T * float(r);
-					B = B * float(r);
-
-					T = normalise(T);
-					B = normalise(B);
-				}
-				else
-				{
-					T = vec3<float>(0, 0, 0);
-					B = vec3<float>(0, 0, 0);
-					N = vec3<float>(0, 0, 0);
-				}
-			}
-
-		};
+		//struct TVertex
+		//{
+
+		//	TVertex() = default;
+
+		//	TVertex(const FVector3f& p, const FVector3f& n, const FVector2f& t)
+		//	{
+		//		pos = p;
+		//		nor = n;
+		//		tex = t;
+		//	}
+
+		//	FVector3f pos;
+		//	FVector3f nor;
+		//	FVector2f tex;
+
+		//	inline bool operator< (const TVertex& other) const
+		//	{
+		//		if (pos < other.pos)
+		//			return true;
+		//		if (other.pos < pos)
+		//			return false;
+
+		//		if (nor < other.nor)
+		//			return true;
+		//		if (other.nor < nor)
+		//			return false;
+
+		//		// Compare the texture coordinates with a particular precission.
+		//		FIntVector2 uv0 = FIntVector2(int32(tex[0] * MUTABLE_VERTEX_MERGE_TEX_RANGE),
+		//			int32(tex[1] * MUTABLE_VERTEX_MERGE_TEX_RANGE));
+		//		FIntVector2 uv1 = FIntVector2(int32(other.tex[0] * MUTABLE_VERTEX_MERGE_TEX_RANGE),
+		//			int32(other.tex[1] * MUTABLE_VERTEX_MERGE_TEX_RANGE));
+
+		//		if (uv0 < uv1)
+		//			return true;
+		//		if (uv1 < uv0)
+		//			return false;
+
+		//		return false;
+		//	}
+
+		//	inline bool operator== (const TVertex& other) const
+		//	{
+		//		return  !((*this) < other)
+		//			&&
+		//			!(other < (*this));
+		//	}
+		//};
+
+
+		//struct TFace
+		//{
+		//	FVector3f T;
+		//	FVector3f B;
+		//	FVector3f N;
+
+		//	TFace()
+		//		: T(0, 0, 0)
+		//		, B(0, 0, 0)
+		//		, N(0, 0, 0)
+		//	{}
+
+		//	TFace
+		//	(
+		//		const FVector3f& v1,
+		//		const FVector3f& v2,
+		//		const FVector3f& v3,
+		//		const FVector2f& w1,
+		//		const FVector2f& w2,
+		//		const FVector2f& w3
+		//	)
+		//	{
+		//		FVector3f E1 = v2 - v1;
+		//		FVector3f E2 = v3 - v1;
+
+		//		FVector2f UV1 = w2 - w1;
+		//		FVector2f UV2 = w3 - w1;
+
+		//		float  UVdet = UV1[0] * UV2[1] - UV2[0] * UV1[1];
+
+		//		N = FVector3f::CrossProduct(E1, E2);
+		//		N.Normalize();
+
+		//		if (!(fabs(UVdet) <= MUTABLE_TANGENT_GENERATION_EPSILON_1))
+		//		{
+		//			double r = 1.0 / UVdet;
+
+		//			T = FVector3f
+		//				(
+		//					(UV2[1] * E1[0] - UV1[1] * E2[0]),
+		//					(UV2[1] * E1[1] - UV1[1] * E2[1]),
+		//					(UV2[1] * E1[2] - UV1[1] * E2[2])
+		//					);
+
+		//			B = FVector3f
+		//				(
+		//					(UV1[0] * E2[0] - UV2[0] * E1[0]),
+		//					(UV1[0] * E2[1] - UV2[0] * E1[1]),
+		//					(UV1[0] * E2[2] - UV2[0] * E1[2])
+		//					);
+
+		//			T = T * float(r);
+		//			B = B * float(r);
+
+		//			T.Normalize();
+		//			B.Normalize();
+		//		}
+		//		else
+		//		{
+		//			T = FVector3f(0, 0, 0);
+		//			B = FVector3f(0, 0, 0);
+		//			N = FVector3f(0, 0, 0);
+		//		}
+		//	}
+
+		//};
 
 	}
 
@@ -333,22 +334,20 @@ namespace mu
 								UntypedMeshBufferIteratorConst yIt(Source, MBS_BINORMAL, resultSemanticIndex);
 								UntypedMeshBufferIteratorConst zIt(Source, MBS_NORMAL, resultSemanticIndex);
 
-								mat3f mat;
 								xIt += v;
 								yIt += v;
 								zIt += v;
-								mat[0] = vec3f(xIt.GetAsVec3f());
-								mat[1] = vec3f(yIt.GetAsVec3f());
-								mat[2] = vec3f(zIt.GetAsVec3f());
+
+								FMatrix44f Mat(xIt.GetAsVec3f(), yIt.GetAsVec3f(), zIt.GetAsVec3f(), FVector3f(0, 0, 0));
 
 								uint8_t sign = 0;
 								if (resultFormat == MBF_PACKEDDIR8_W_TANGENTSIGN)
 								{
-									sign = mat.GetDeterminant() < 0 ? 0 : 255;
+									sign = Mat.RotDeterminant() < 0 ? 0 : 255;
 								}
 								else if (resultFormat == MBF_PACKEDDIRS8_W_TANGENTSIGN)
 								{
-									sign = mat.GetDeterminant() < 0 ? -128 : 127;
+									sign = Mat.RotDeterminant() < 0 ? -128 : 127;
 								}
 								pData[3] = sign;
 							}
@@ -597,7 +596,8 @@ namespace mu
 							{
 								// If MAX_TOTAL_INFLUENCES ever changed, the next line would no longer work or compile and 
 								// GetAsVec12i would need to be changed accordingly
-								vec<int32_t, MAX_TOTAL_INFLUENCES> va = it.GetAsVec12i();
+								int32 va[MAX_TOTAL_INFLUENCES];
+								it.GetAsInt32Vec(va, MAX_TOTAL_INFLUENCES);
 								for (int c = 0; c < it.GetComponents(); ++c)
 								{
 									maxBoneIndex = FMath::Max(maxBoneIndex, va[c]);
@@ -700,4 +700,3 @@ namespace mu
 	}
 
 }
-
diff --git a/Engine/Plugins/Experimental/Mutable/Source/MutableRuntime/Private/MuR/OpMeshInterpolate.h b/Engine/Plugins/Experimental/Mutable/Source/MutableRuntime/Private/MuR/OpMeshInterpolate.h
--- a/Engine/Plugins/Experimental/Mutable/Source/MutableRuntime/Private/MuR/OpMeshInterpolate.h
+++ b/Engine/Plugins/Experimental/Mutable/Source/MutableRuntime/Private/MuR/OpMeshInterpolate.h
@@ -26,33 +26,33 @@ namespace mu
 
 			// A mesh position
 			int aPosElementSize = 0;
-            uint8_t* pAPosBuf = 0;
+            uint8* pAPosBuf = 0;
 			GetMeshBuf( pDest.get(), MBS_POSITION, MBF_FLOAT32, 3, pAPosBuf, aPosElementSize );
 
 			// A mesh normal
 			int aNorElementSize = 0;
-            uint8_t* pANorBuf = 0;
+            uint8* pANorBuf = 0;
 			GetMeshBuf( pDest.get(), MBS_NORMAL, MBF_FLOAT32, 3, pANorBuf, aNorElementSize );
 
 			// B mesh position
 			int bPosElementSize = 0;
-            const uint8_t* pBPosBuf = 0;
+            const uint8* pBPosBuf = 0;
 			GetMeshBuf( pB, MBS_POSITION, MBF_FLOAT32, 3, pBPosBuf, bPosElementSize );
 
 			// B mesh normal
 			int bNorElementSize = 0;
-            const uint8_t* pBNorBuf = 0;
+            const uint8* pBNorBuf = 0;
 			GetMeshBuf( pB, MBS_NORMAL, MBF_FLOAT32, 3, pBNorBuf, bNorElementSize );
 
-            uint32_t vcount = pA->GetVertexBuffers().GetElementCount();
-            for ( uint32_t vi=0; vi<vcount; ++vi )
+            uint32 vcount = pA->GetVertexBuffers().GetElementCount();
+            for ( uint32 vi=0; vi<vcount; ++vi )
 			{
-				vec3<float>* aPos = reinterpret_cast< vec3<float>* >( pAPosBuf );
-				const vec3<float>* bPos = reinterpret_cast< const vec3<float>* >( pBPosBuf );
+				FVector3f* aPos = reinterpret_cast<FVector3f* >( pAPosBuf );
+				const FVector3f* bPos = reinterpret_cast< const FVector3f* >( pBPosBuf );
 				*aPos += (*bPos - *aPos) * factor;
 
-				vec3<float>* aNor = reinterpret_cast< vec3<float>* >( pANorBuf );
-				const vec3<float>* bNor = reinterpret_cast< const vec3<float>* >( pBNorBuf );
+				FVector3f* aNor = reinterpret_cast<FVector3f* >( pANorBuf );
+				const FVector3f* bNor = reinterpret_cast< const FVector3f* >( pBNorBuf );
 				*aNor += (*bNor - *aNor) * factor;
 
 				pAPosBuf += aPosElementSize;
diff --git a/Engine/Plugins/Experimental/Mutable/Source/MutableRuntime/Private/MuR/OpMeshRemove.cpp b/Engine/Plugins/Experimental/Mutable/Source/MutableRuntime/Private/MuR/OpMeshRemove.cpp
--- a/Engine/Plugins/Experimental/Mutable/Source/MutableRuntime/Private/MuR/OpMeshRemove.cpp
+++ b/Engine/Plugins/Experimental/Mutable/Source/MutableRuntime/Private/MuR/OpMeshRemove.cpp
@@ -98,9 +98,9 @@ namespace mu
                 for ( int32 f=0; f<IndexCount/3; ++f )
                 {
                     uint32 sourceIndices[3];
-                    sourceIndices[0] = itSource[0][0];
-                    sourceIndices[1] = itSource[1][0];
-                    sourceIndices[2] = itSource[2][0];
+                    sourceIndices[0] = (*itSource)[0];
+                    sourceIndices[1] = (*itSource)[1];
+                    sourceIndices[2] = (*itSource)[2];
 
 					check(sourceIndices[0] < RemovedVertexCount);
 					check(sourceIndices[1] < RemovedVertexCount);
@@ -146,9 +146,9 @@ namespace mu
                 for ( int32 f=0; f<indexCount/3; ++f )
                 {
                     uint16 sourceIndices[3];
-                    sourceIndices[0] = itSource[0][0];
-                    sourceIndices[1] = itSource[1][0];
-                    sourceIndices[2] = itSource[2][0];
+                    sourceIndices[0] = (*itSource)[0];
+                    sourceIndices[1] = (*itSource)[1];
+                    sourceIndices[2] = (*itSource)[2];
 
 					check(sourceIndices[0] < RemovedVertexCount);
 					check(sourceIndices[1] < RemovedVertexCount);
diff --git a/Engine/Plugins/Experimental/Mutable/Source/MutableRuntime/Private/MuR/Serialisation.cpp b/Engine/Plugins/Experimental/Mutable/Source/MutableRuntime/Private/MuR/Serialisation.cpp
--- a/Engine/Plugins/Experimental/Mutable/Source/MutableRuntime/Private/MuR/Serialisation.cpp
+++ b/Engine/Plugins/Experimental/Mutable/Source/MutableRuntime/Private/MuR/Serialisation.cpp
@@ -16,30 +16,25 @@ namespace mu
 	MUTABLE_IMPLEMENT_POD_SERIALISABLE(float);    
 	MUTABLE_IMPLEMENT_POD_SERIALISABLE(double);   
                                                   
-    MUTABLE_IMPLEMENT_POD_SERIALISABLE(int8);   
-    MUTABLE_IMPLEMENT_POD_SERIALISABLE(int16);  
-    MUTABLE_IMPLEMENT_POD_SERIALISABLE(int32);  
-    MUTABLE_IMPLEMENT_POD_SERIALISABLE(int64);  
-                                                 
-    MUTABLE_IMPLEMENT_POD_SERIALISABLE(uint8);  
-    MUTABLE_IMPLEMENT_POD_SERIALISABLE(uint16); 
-    MUTABLE_IMPLEMENT_POD_SERIALISABLE(uint32); 
-    MUTABLE_IMPLEMENT_POD_SERIALISABLE(uint64)
-
-	MUTABLE_IMPLEMENT_POD_SERIALISABLE(FGuid);
-
-	MUTABLE_IMPLEMENT_POD_SERIALISABLE(vec2f);                                        
-	MUTABLE_IMPLEMENT_POD_SERIALISABLE(vec3f);
-	MUTABLE_IMPLEMENT_POD_SERIALISABLE(vec4f);
-	MUTABLE_IMPLEMENT_POD_SERIALISABLE(mat3f);                                        
-	MUTABLE_IMPLEMENT_POD_SERIALISABLE(mat4f);                                        
-	MUTABLE_IMPLEMENT_POD_SERIALISABLE(vec2<int>);                                    
+    MUTABLE_IMPLEMENT_POD_SERIALISABLE( int8 );   
+    MUTABLE_IMPLEMENT_POD_SERIALISABLE( int16 );  
+    MUTABLE_IMPLEMENT_POD_SERIALISABLE( int32 );  
+    MUTABLE_IMPLEMENT_POD_SERIALISABLE( int64 );  
+                                                  
+    MUTABLE_IMPLEMENT_POD_SERIALISABLE( uint8 );  
+    MUTABLE_IMPLEMENT_POD_SERIALISABLE( uint16 ); 
+    MUTABLE_IMPLEMENT_POD_SERIALISABLE( uint32 ); 
+    MUTABLE_IMPLEMENT_POD_SERIALISABLE( uint64 )
 	
 	// Unreal POD Serializables                                                       
-	MUTABLE_IMPLEMENT_POD_SERIALISABLE(FUintVector2);                                 
-	MUTABLE_IMPLEMENT_POD_SERIALISABLE(UE::Math::TIntVector2<uint16>);                
-	MUTABLE_IMPLEMENT_POD_SERIALISABLE(UE::Math::TIntVector2<int16>); 
-	MUTABLE_IMPLEMENT_POD_SERIALISABLE(UE::Math::TVector4<float>); 
+	MUTABLE_IMPLEMENT_POD_SERIALISABLE(FGuid);
+	MUTABLE_IMPLEMENT_POD_SERIALISABLE(FUintVector2);
+	MUTABLE_IMPLEMENT_POD_SERIALISABLE(FIntVector2);
+	MUTABLE_IMPLEMENT_POD_SERIALISABLE(UE::Math::TIntVector2<uint16>);
+	MUTABLE_IMPLEMENT_POD_SERIALISABLE(UE::Math::TIntVector2<int16>);
+	MUTABLE_IMPLEMENT_POD_SERIALISABLE(FVector2f);
+	MUTABLE_IMPLEMENT_POD_SERIALISABLE(FVector4f);
+	MUTABLE_IMPLEMENT_POD_SERIALISABLE(FMatrix44f);
 
 	
     void operator<<(OutputArchive& arch, const FString& t)
diff --git a/Engine/Plugins/Experimental/Mutable/Source/MutableRuntime/Private/MuR/System.cpp b/Engine/Plugins/Experimental/Mutable/Source/MutableRuntime/Private/MuR/System.cpp
--- a/Engine/Plugins/Experimental/Mutable/Source/MutableRuntime/Private/MuR/System.cpp
+++ b/Engine/Plugins/Experimental/Mutable/Source/MutableRuntime/Private/MuR/System.cpp
@@ -56,9 +56,6 @@ static FAutoConsoleVariableRef CVarEnableDetailedMemoryBudgetExceededLogging (
 
 namespace mu
 {
-    static_assert( sizeof(mat4f) == 64, "UNEXPECTED_STRUCT_PACKING" );
-
-
 	MUTABLE_IMPLEMENT_ENUM_SERIALISABLE(ETextureCompressionStrategy);
 
 	TRACE_DECLARE_INT_COUNTER(MutableRuntime_LiveInstances,		TEXT("MutableRuntime/LiveInstances"));
diff --git a/Engine/Plugins/Experimental/Mutable/Source/MutableRuntime/Public/MuR/ImageTypes.h b/Engine/Plugins/Experimental/Mutable/Source/MutableRuntime/Public/MuR/ImageTypes.h
--- a/Engine/Plugins/Experimental/Mutable/Source/MutableRuntime/Public/MuR/ImageTypes.h
+++ b/Engine/Plugins/Experimental/Mutable/Source/MutableRuntime/Public/MuR/ImageTypes.h
@@ -11,7 +11,7 @@ namespace mu
 {
 	//!
 	using FImageSize = UE::Math::TIntVector2<uint16>;
-	using FImageRect = box<vec2<uint16>>;
+	using FImageRect = box<FImageSize>;
 
 	//! Pixel formats supported by the images.
 	//! \ingroup runtime
diff --git a/Engine/Plugins/Experimental/Mutable/Source/MutableRuntime/Public/MuR/MutableMath.h b/Engine/Plugins/Experimental/Mutable/Source/MutableRuntime/Public/MuR/MutableMath.h
--- a/Engine/Plugins/Experimental/Mutable/Source/MutableRuntime/Public/MuR/MutableMath.h
+++ b/Engine/Plugins/Experimental/Mutable/Source/MutableRuntime/Public/MuR/MutableMath.h
@@ -10,1020 +10,64 @@
 namespace mu
 {
 
-    //---------------------------------------------------------------------------------------------
-    //!
-    //---------------------------------------------------------------------------------------------
-    template<class T>
-    T clamp( const T mi, const T ma, const T v )
-    {
-        return FMath::Max( mi, FMath::Min( ma, v ) );
-    }
-
-
-	//---------------------------------------------------------------------------------------------
-	// 16-bit floating point number support.
-	//---------------------------------------------------------------------------------------------
-    typedef uint16 float16;
-
-	//---------------------------------------------------------------------------------------------
-    inline uint32 halfToFloatI( float16 y )
-	{
-		int s = (y >> 15) & 0x00000001;
-		int e = (y >> 10) & 0x0000001f;
-		int f =  y        & 0x000003ff;
-
-		if (e == 0)
-		{
-			if (f == 0)
-			{
-				return s << 31;
-			}
-			else
-			{
-				while (!(f & 0x00000400))
-				{
-					f <<= 1;
-					e -=  1;
-				}
-				e += 1;
-				f &= ~0x00000400;
-			}
-		}
-		else if (e == 31)
-		{
-			if (f == 0)
-			{
-				return (s << 31) | 0x7f800000;
-			}
-			else
-			{
-				return (s << 31) | 0x7f800000 | (f << 13);
-			}
-		}
-
-		e = e + (127 - 15);
-		f = f << 13;
-
-		return ((s << 31) | (e << 23) | f);
-	}
-
-
-	//---------------------------------------------------------------------------------------------
-	inline float halfToFloat( float16 y)
-	{
-		union
-		{
-			float f;
-            uint32 i;
-		} v;
-
-		v.i = halfToFloatI(y);
-		return v.f;
-	}
-
-
-	//---------------------------------------------------------------------------------------------
-    inline float16 floatToHalfI( uint32 i )
-	{
-		int s =  (i >> 16) & 0x00008000;
-		int e = ((i >> 23) & 0x000000ff) - (127 - 15);
-		int f =   i        & 0x007fffff;
-
-		if (e <= 0)
-		{
-			if (e < -10)
-			{
-				if (s)
-				{
-					return 0x8000;
-				}
-				else
-				{
-				   return 0;
-				}
-			}
-			f = (f | 0x00800000) >> (1 - e);
-			return float16( s | (f >> 13) );
-		}
-		else if (e == 0xff - (127 - 15))
-		{
-			if (f == 0)
-			{
-				return float16(s | 0x7c00 );
-			}
-			else
-			{
-				f >>= 13;
-				return float16(s | 0x7c00 | f | (f == 0));
-			}
-		}
-		else
-		{
-			if (e > 30)
-			{
-				return float16( s | 0x7c00 );
-			}
-			return float16(s | (e << 10) | (f >> 13) );
-		}
-	}
-
-
-	//---------------------------------------------------------------------------------------------
-	inline float16 floatToHalf(float i)
-	{
-		union
-		{
-			float f;
-            uint32 i;
-		} v;
-
-		v.f = i;
-		return floatToHalfI(v.i);
-	}
-
-
-
-	//---------------------------------------------------------------------------------------------
-	//! Generic vector template.
-	//---------------------------------------------------------------------------------------------
-	template <class SCALAR, int DIM>
-    class vec
-	{
-	protected:
-
-		SCALAR m[ DIM ];
-
-	public:
-
-		//-----------------------------------------------------------------------------------------
-		typedef vec<SCALAR,DIM> V;
-
-        //-----------------------------------------------------------------------------------------
-        inline vec()
-        {
-            for( int i=0; i<DIM; ++i )
-            {
-                m[i]=0;
-            }
-        }
-
-
-		//-----------------------------------------------------------------------------------------
-		static inline int GetDim()
-		{
-			return DIM;
-		}
-
-		//-----------------------------------------------------------------------------------------
-		inline SCALAR& operator[] (int i)
-		{
-			check( i>=0 && i <DIM );
-			return m[i];
-		}
-
-		//-----------------------------------------------------------------------------------------
-		inline const SCALAR& operator[] (int i) const
-		{
-			check( i>=0 && i <DIM );
-			return m[i];
-		}
-
-		//-----------------------------------------------------------------------------------------
-		inline bool operator==( const V& other ) const
-		{
-			bool res = true;
-			for( int i=0; i<DIM; ++i )
-			{
-				res &= m[i] == other.m[i];
-			}
-			return res;
-		}
-
-		//-----------------------------------------------------------------------------------------
-		inline bool operator!=( const V& other ) const
-		{
-			bool res = false;
-			for( int i=0; i<DIM; ++i )
-			{
-				res |= m[i] != other.m[i];
-			}
-			return res;
-		}
-
-		//-----------------------------------------------------------------------------------------
-		inline V operator+( const V& other ) const
-		{
-			V res;
-			for( int i=0; i<DIM; ++i )
-			{
-				res.m[i] = m[i] + other.m[i];
-			}
-			return res;
-		}
-
-		//-----------------------------------------------------------------------------------------
-		inline void operator+=( const V& other )
-		{
-			for( int i=0; i<DIM; ++i )
-			{
-				m[i] += other.m[i];
-			}
-		}
-
-        //-----------------------------------------------------------------------------------------
-        inline V operator-( const V& other ) const
-        {
-            V res;
-            for( int i=0; i<DIM; ++i )
-            {
-                res.m[i] = m[i] - other.m[i];
-            }
-            return res;
-        }
-
-        //-----------------------------------------------------------------------------------------
-        inline V operator-( SCALAR v ) const
-        {
-            V res;
-            for( int i=0; i<DIM; ++i )
-            {
-                res.m[i] = m[i] - v;
-            }
-            return res;
-        }
-
-		//-----------------------------------------------------------------------------------------
-		inline void operator-=( const V& other )
-		{
-			for( int i=0; i<DIM; ++i )
-			{
-                m[i] -= other.m[i];
-			}
-		}
-
-		//-----------------------------------------------------------------------------------------
-		inline V operator*( const V& other ) const
-		{
-			V res;
-			for( int i=0; i<DIM; ++i )
-			{
-				res.m[i] = m[i] * other.m[i];
-			}
-			return res;
-		}
-
-		//-----------------------------------------------------------------------------------------
-		inline V operator*( SCALAR s ) const
-		{
-			V res;
-			for( int i=0; i<DIM; ++i )
-			{
-				res.m[i] = m[i] * s;
-			}
-			return res;
-		}
-
-		//-----------------------------------------------------------------------------------------
-		inline V operator/( const V& other ) const
-		{
-			V res;
-			for( int i=0; i<DIM; ++i )
-			{
-				res.m[i] = m[i] / other.m[i];
-			}
-			return res;
-		}
-
-		//-----------------------------------------------------------------------------------------
-		inline V operator/( SCALAR s ) const
-		{
-			V res;
-			for( int i=0; i<DIM; ++i )
-			{
-				res.m[i] = m[i] / s;
-			}
-			return res;
-		}		
-
-        //-----------------------------------------------------------------------------------------
-        inline void operator*=( SCALAR s )
-        {
-            for( int i=0; i<DIM; ++i )
-            {
-                m[i] *= s;
-            }
-        }
-
-        //-----------------------------------------------------------------------------------------
-        inline void operator/=( SCALAR s )
-        {
-            for( int i=0; i<DIM; ++i )
-            {
-                m[i] /= s;
-            }
-        }
-
-        //-----------------------------------------------------------------------------------------
-		//! Per component min
-        //-----------------------------------------------------------------------------------------
-		inline static V min(const V& a, const V& b)
-		{
-			V r;
-			for (size_t i = 0; i < DIM; ++i)
-			{
-				r.m[i] = FMath::Min(a.m[i], b.m[i]);
-			}
-			return r;
-		}
-
-		inline static V max(const V& a, const V& b)
-		{
-			V r;
-			for (size_t i = 0; i < DIM; ++i)
-			{
-				r.m[i] = FMath::Max(a.m[i], b.m[i]);
-			}
-			return r;
-		}
-
-		//-----------------------------------------------------------------------------------------
-		//! Strict ordering of vectors.
-		//-----------------------------------------------------------------------------------------
-		inline bool
-		operator<(const V& other) const
-        {
-			bool equal = true;
-			bool res = false;
-			for( int i=0; equal && i<DIM; ++i )
-			{
-				if( m[i] < other.m[i] )
-				{
-					equal = false;
-					res = true;
-				}
-				else if( m[i] > other.m[i] )
-				{
-					equal = false;
-					res = false;
-				}
-			}
-			return res;
-		}
-
-
-        //-----------------------------------------------------------------------------------------
-        inline bool AllSmallerThan( const V& other ) const
-        {
-            bool res = true;
-            for( int i=0; i<DIM; ++i )
-            {
-                res &= m[i] < other.m[i];
-            }
-            return res;
-        }
-
-
-        //-----------------------------------------------------------------------------------------
-        inline bool AllSmallerOrEqualThan( const V& other ) const
-        {
-            bool res = true;
-            for( int i=0; i<DIM; ++i )
-            {
-                res &= m[i] <= other.m[i];
-            }
-            return res;
-        }
-
-		//-----------------------------------------------------------------------------------------
-		inline bool AllGreaterOrEqualThan( const V& other ) const
-		{
-			bool res = true;
-			for( int i=0; i<DIM; ++i )
-			{
-				res &= m[i] >= other.m[i];
-			}
-			return res;
-		}
-
-		//-----------------------------------------------------------------------------------------
-		inline bool AlmostEqual
-			(
-				const V& other,
-				SCALAR tol = std::numeric_limits<SCALAR>::min()
-			) const
-		{
-			bool res = true;
-			for( int i=0; i<DIM; ++i )
-			{
-				res &= fabs(m[i]-other.m[i])<tol;
-			}
-			return res;
-		}
-	};
-
-    //!
-    template <class SCALAR, int DIM>
-    inline vec<SCALAR,DIM> clamp( const vec<SCALAR,DIM>& v, SCALAR mi, SCALAR  ma )
-    {
-        vec<SCALAR,DIM> res;
-        for( int i=0; i<DIM; ++i )
-        {
-            res[i] = clamp( mi, ma, v[i] );
-        }
-        return res;
-    }
-
-    //!
-    template <class SCALAR, int DIM>
-    inline vec<SCALAR,DIM> lerp( const vec<SCALAR,DIM>& f0, const vec<SCALAR,DIM>& f1, float t )
-    {
-        const float s = 1.0f - t;
-        return f0 * s + f1 * t;
-    }
-
-	//!
-	template <class _SCALAR>
-	class vec2 : public vec<_SCALAR,2>
-	{
-	public:
-        typedef _SCALAR SCALAR;
-
-        //-----------------------------------------------------------------------------------------
-		inline vec2()
-		{
-			// Elements are cleared in the parent constructor
-		}
-
-		//-----------------------------------------------------------------------------------------
-		inline explicit vec2(SCALAR v)
-		{
-			this->m[0] = v;
-			this->m[1] = v;
-		}
-
-		//-----------------------------------------------------------------------------------------
-		inline vec2(SCALAR x, SCALAR y)
-		{
-			this->m[0] = x;
-			this->m[1] = y;
-		}
-
-		//-----------------------------------------------------------------------------------------
-		template<class SCALAR2>
-        inline vec2( const vec<SCALAR2,2>& v )
-        {
-			this->m[0] = v[0];
-			this->m[1] = v[1];
-		}
-
-        //-----------------------------------------------------------------------------------------
-        inline SCALAR& x()
-        {
-            return this->m[0];
-        }
-
-        //-----------------------------------------------------------------------------------------
-        inline SCALAR& y()
-        {
-            return this->m[1];
-        }
-
-        //-----------------------------------------------------------------------------------------
-        inline const SCALAR& x() const
-        {
-            return this->m[0];
-        }
-
-        //-----------------------------------------------------------------------------------------
-        inline const SCALAR& y() const
-        {
-            return this->m[1];
-        }
-	};
-
-    //!
-    typedef vec2<float> vec2f;
-
-	//!
-	template <class SCALAR>
-	class vec3 : public vec<SCALAR,3>
-	{
-	public:
-
-		//-----------------------------------------------------------------------------------------
-		inline vec3()
-		{
-			// Elements are cleared in the parent constructor
-		}
-
-		//-----------------------------------------------------------------------------------------
-		inline explicit vec3(SCALAR v)
-		{
-			this->m[0] = v;
-			this->m[1] = v;
-			this->m[2] = v;
-		}
-
-		//-----------------------------------------------------------------------------------------
-		inline vec3(SCALAR x, SCALAR y, SCALAR z)
-		{
-			this->m[0] = x;
-			this->m[1] = y;
-			this->m[2] = z;
-		}
-
-		//-----------------------------------------------------------------------------------------
-		inline vec3(const vec<SCALAR, 3>& other)
-		{
-			this->m[0] = other[0];
-			this->m[1] = other[1];
-			this->m[2] = other[2];
-		}				
-
-		//-----------------------------------------------------------------------------------------
-		inline vec3(const FVector3f& other)
-		{
-			this->m[0] = other[0];
-			this->m[1] = other[1];
-			this->m[2] = other[2];
-		}
-
-		//-----------------------------------------------------------------------------------------
-		inline vec3(const FVector4f& other)
-		{
-			this->m[0] = other[0];
-			this->m[1] = other[1];
-			this->m[2] = other[2];
-		}
-
-        //-----------------------------------------------------------------------------------------
-        inline void Set( SCALAR x, SCALAR y, SCALAR z )
-        {
-            this->m[0] = x;
-            this->m[1] = y;
-            this->m[2] = z;
-        }
-
-
-        //-----------------------------------------------------------------------------------------
-        inline SCALAR& x()
-        {
-            return this->m[0];
-        }
-
-        //-----------------------------------------------------------------------------------------
-        inline SCALAR& y()
-        {
-            return this->m[1];
-        }
-
-        //-----------------------------------------------------------------------------------------
-        inline SCALAR& z()
-        {
-            return this->m[2];
-        }
-
-        //-----------------------------------------------------------------------------------------
-        inline const SCALAR& x() const
-        {
-            return this->m[0];
-        }
-
-        //-----------------------------------------------------------------------------------------
-        inline const SCALAR& y() const
-        {
-            return this->m[1];
-        }
-
-        //-----------------------------------------------------------------------------------------
-        inline const SCALAR& z() const
-        {
-            return this->m[2];
-        }
-
-        //-----------------------------------------------------------------------------------------
-        inline vec2<SCALAR> xy() const
-        {
-            return vec2<SCALAR>( this->m[0], this->m[1] );
-        }
-
-	};
-
-	//!
-	typedef vec3<float> vec3f;
-
-	//!
-	template <class SCALAR>
-	class vec4 : public vec<SCALAR, 4>
-	{
-	public:
-
-		//-----------------------------------------------------------------------------------------
-		inline vec4()
-		{
-			// Elements are cleared in the parent constructor
-		}
-
-		//-----------------------------------------------------------------------------------------
-		inline explicit vec4(SCALAR v)
-		{
-			this->m[0] = v;
-			this->m[1] = v;
-			this->m[2] = v;
-			this->m[3] = v;
-		}
-
-		//-----------------------------------------------------------------------------------------
-		inline vec4(SCALAR x, SCALAR y, SCALAR z, SCALAR w)
-		{
-			this->m[0] = x;
-			this->m[1] = y;
-			this->m[2] = z;
-			this->m[3] = w;
-		}
-
-		//-----------------------------------------------------------------------------------------
-		inline vec4(const vec<SCALAR, 4>& other)
-		{
-			this->m[0] = other[0];
-			this->m[1] = other[1];
-			this->m[2] = other[2];
-			this->m[3] = other[3];
-		}
-
-		//-----------------------------------------------------------------------------------------
-		inline vec4(const FVector4f& other)
-		{
-			this->m[0] = other[0];
-			this->m[1] = other[1];
-			this->m[2] = other[2];
-			this->m[3] = other[3];
-		}
-
-		//-----------------------------------------------------------------------------------------
-		inline vec4(const vec<SCALAR, 3>& other)
-		{
-			this->m[0] = other[0];
-			this->m[1] = other[1];
-			this->m[2] = other[2];
-			this->m[3] = 1;
-		}
-
-		//-----------------------------------------------------------------------------------------
-		inline void Set(SCALAR x, SCALAR y, SCALAR z, SCALAR w)
-		{
-			this->m[0] = x;
-			this->m[1] = y;
-			this->m[2] = z;
-			this->m[3] = w;
-		}
-
-
-		//-----------------------------------------------------------------------------------------
-		inline SCALAR& x()
-		{
-			return this->m[0];
-		}
-
-		//-----------------------------------------------------------------------------------------
-		inline SCALAR& y()
-		{
-			return this->m[1];
-		}
-
-		//-----------------------------------------------------------------------------------------
-		inline SCALAR& z()
-		{
-			return this->m[2];
-		}
-
-		//-----------------------------------------------------------------------------------------
-		inline SCALAR& w()
-		{
-			return this->m[3];
-		}
-
-		//-----------------------------------------------------------------------------------------
-		inline const SCALAR& x() const
-		{
-			return this->m[0];
-		}
-
-		//-----------------------------------------------------------------------------------------
-		inline const SCALAR& y() const
-		{
-			return this->m[1];
-		}
-
-		//-----------------------------------------------------------------------------------------
-		inline const SCALAR& z() const
-		{
-			return this->m[2];
-		}
-
-		//-----------------------------------------------------------------------------------------
-		inline const SCALAR& w() const
-		{
-			return this->m[3];
-		}
-
-		//-----------------------------------------------------------------------------------------
-		inline vec2<SCALAR> xy() const
-		{
-			return vec2<SCALAR>(this->m[0], this->m[1]);
-		}
-
-		//-----------------------------------------------------------------------------------------
-		inline vec3<SCALAR> xyz() const
-		{
-			return vec3<SCALAR>(this->m[0], this->m[1], this->m[2]);
-		}
-
-	};
-
-	//!
-	typedef vec4<float> vec4f;
-
-
-	//---------------------------------------------------------------------------------------------
-	//! Cross product of 3D vectors
-	//---------------------------------------------------------------------------------------------
-	template <class SCALAR>
-	inline vec3<SCALAR> cross( const vec3<SCALAR>& a, const vec3<SCALAR>& b )
-	{
-		return vec3<SCALAR>
-				(
-					a[1] * b[2] - a[2] * b[1],
-					a[2] * b[0] - a[0] * b[2],
-					a[0] * b[1] - a[1] * b[0]
-				);
-	}
-
-
-	//---------------------------------------------------------------------------------------------
-	//! Dot product of 3D vectors
-	//---------------------------------------------------------------------------------------------
-	template <class SCALAR, int DIM >
-	inline SCALAR dot( const vec<SCALAR, DIM>& a, const vec<SCALAR,DIM>& b )
-	{
-		SCALAR res = 0;
-
-		for ( int i=0; i<DIM; ++i )
-		{
-			res += a[i]*b[i];
-		}
-
-		return res;
-	}
-
-
-	//-----------------------------------------------------------------------------------------
-	//! Length of a vector
-	//-----------------------------------------------------------------------------------------
-	template <class SCALAR, int DIM>
-	inline SCALAR length( const vec<SCALAR,DIM>& v )
-	{
-		SCALAR l = dot( v, v );
-		l = sqrtf(l);
-		return l;
-	}
-
-
-	//---------------------------------------------------------------------------------------------
-	//! Normalise a vector
-	//---------------------------------------------------------------------------------------------
-	template <class SCALAR, int DIM>
-	inline vec<SCALAR,DIM> normalise( const vec<SCALAR,DIM>& a )
-	{
-		float l = length( a );
-		if ( l > std::numeric_limits<SCALAR>::epsilon() )
-		{
-			return a/l;
-		}
-		return a;
-	}
-
-
-	//---------------------------------------------------------------------------------------------
-	//! Normalise a vector, assuming that its length won't be null
-	//---------------------------------------------------------------------------------------------
-	template <class SCALAR, int DIM>
-	inline vec<SCALAR,DIM> normalise_unsafe( const vec<SCALAR,DIM>& a )
-	{
-		float l = length( a );
-		return a/l;
-	}
-
-	//---------------------------------------------------------------------------------------------
-	//! Point to plane distance
-	//---------------------------------------------------------------------------------------------
-	inline float point_to_plane_dist(const FVector3f& point, const FVector3f& plane_origin, const FVector3f& plane_normal)
-	{
-		return FVector3f::DotProduct((point - plane_origin), plane_normal);
-	}
-
-	//---------------------------------------------------------------------------------------------
-	//! Ray plane intersection
-	//---------------------------------------------------------------------------------------------
-	inline FVector3f ray_plane_intersection(const FVector3f& ray_start, const FVector3f& ray_end, const FVector3f& plane_origin, const FVector3f&plane_normal)
-	{
-		return ray_start + (ray_end - ray_start) *	(FVector3f::DotProduct((plane_origin - ray_start), plane_normal) / FVector3f::DotProduct((ray_end - ray_start), plane_normal));
-	}
-
-    //constexpr float point_in_plane_eps = 0.0001f;
-
-    inline bool point_in_face(const FVector3f& InVtx, const FVector3f& normal, FVector3f vertices[3], int& out_intersected_vert, int& out_intersected_edge_v0, int& out_intersected_edge_v1, float point_in_plane_eps = 0.0001f)
-    {
-		FVector3f SidePlaneNormal;
-		FVector3f Side;
-        bool intersected_edge[3] = { false, false, false };
-        int intersected_edge_count = 0;
-
-        for (int x = 0; x < 3; ++x)
-        {
-            // Create plane perpendicular to both this side and the polygon's normal.
-            Side = vertices[x] - vertices[(x - 1 < 0) ? 3 - 1 : x - 1];
-            SidePlaneNormal = FVector3f::CrossProduct(Side, normal.GetSafeNormal());
-            SidePlaneNormal.Normalize();
-
-            float dist = point_to_plane_dist(InVtx, vertices[x], SidePlaneNormal);
-
-            if (dist > point_in_plane_eps)
-            {
-                return false; // If point is not behind all the planes created by this polys edges, it's outside the poly.
-            }
-            else if (dist > -point_in_plane_eps)
-            {
-                intersected_edge[x] = true;
-                intersected_edge_count++;
-            }
-        }
-
-        out_intersected_vert = -1;
-        out_intersected_edge_v0 = -1;
-        out_intersected_edge_v1 = -1;
-
-        if (intersected_edge_count > 0)
-        {
-            if (intersected_edge_count == 1) // Only one edge intersected
-            {
-                if (intersected_edge[0])
-                {
-                    out_intersected_edge_v0 = 0;
-                    out_intersected_edge_v1 = 2;
-                }
-                else if (intersected_edge[1])
-                {
-                    out_intersected_edge_v0 = 0;
-                    out_intersected_edge_v1 = 1;
-                }
-                else if (intersected_edge[2])
-                {
-                    out_intersected_edge_v0 = 1;
-                    out_intersected_edge_v1 = 2;
-                }
-            }
-            else if (intersected_edge_count == 2) // Two edges intersected means it actually hit a vertex
-            {
-                if (intersected_edge[0] && intersected_edge[1])
-                {
-                    out_intersected_vert = 0;
-                }
-                else if (intersected_edge[1] && intersected_edge[2])
-                {
-                    out_intersected_vert = 1;
-                }
-                else if (intersected_edge[2] && intersected_edge[0])
-                {
-                    out_intersected_vert = 2;
-                }
-            }
-        }
-
-        return true;
-    }
-
-    //---------------------------------------------------------------------------------------------
-    inline bool rayIntersectsFace( const FVector3f& ray_start, const FVector3f& ray_end,
-                                   const FVector3f& v0, const FVector3f& v1, const FVector3f& v2,
-								   FVector3f& out_intersection,
-                                   int& out_intersected_vert, int& out_intersected_edge_v0,
-                                   int& out_intersected_edge_v1, float point_in_face_epsilon = 0.0001f)
-    {
-		FVector3f normal = FVector3f::CrossProduct(v1 - v0, v2 - v0);
-
-        // Does the ray cross the plane?
-        const float DistStart = point_to_plane_dist(ray_start, v0, normal);
-        const float DistEnd = point_to_plane_dist(ray_end, v0, normal);
-
-        if ((DistStart < 0 && DistEnd < 0) || (DistStart > 0 && DistEnd > 0))
-        {
-            return false;
-        }
-
-        // Get the intersection of the line and the plane.
-        out_intersection = ray_plane_intersection(ray_start, ray_end, v0, normal);
-
-        // \todo: review the equality comparison. Is it doing anything without a tolerance?
-        if (out_intersection == ray_start || out_intersection == ray_end)
-        {
-            return false;
-        }
-
-        // Check if the intersection point is actually on the poly.
-		FVector3f vertices[3] = { v0, v1, v2 };
-
-        return point_in_face( out_intersection, normal, vertices,
-                              out_intersected_vert,
-                              out_intersected_edge_v0, out_intersected_edge_v1, point_in_face_epsilon);
-    }
-
-
-    //---------------------------------------------------------------------------------------------
-    inline bool ray_plane_intersection( const FVector3f& ray_start, const FVector3f& ray_end,
-                                        const FVector3f& plane_origin, const FVector3f&plane_normal,
-                                        FVector3f& out_intersection, float& out_t )
-    {
-        float den = FVector3f::DotProduct(ray_end - ray_start, plane_normal);
-
-        if (den == 0.f)
-        {
-            return false;
-        }
-
-        out_t = FVector3f::DotProduct((plane_origin - ray_start), plane_normal) / den;
-        out_intersection = ray_start + (ray_end - ray_start) * out_t;
-
-        return true;
-    }
-
-
-    //---------------------------------------------------------------------------------------------
-    inline bool rayIntersectsFace2(const FVector3f& ray_start, const FVector3f& ray_end,
-                                   const FVector3f& v0, const FVector3f& v1, const FVector3f& v2,
-		FVector3f& out_intersection,
-                                   int& out_intersected_vert,
-                                   int& out_intersected_edge_v0, int& out_intersected_edge_v1,
-                                   float& out_t )
-    {
-		FVector3f normal = FVector3f::CrossProduct(v1 - v0, v2 - v0);
-
-        // Does the ray cross the plane?
-        const float DistStart = point_to_plane_dist(ray_start, v0, normal);
-        const float DistEnd = point_to_plane_dist(ray_end, v0, normal);
+	template<class VECTOR>
+	inline int32 GetDim()
+	{
+		return VECTOR::GetDim();
+	}
 
-        if ((DistStart < 0 && DistEnd < 0) || (DistStart > 0 && DistEnd > 0))
-        {
-            return false;
-        }
+	template<>
+	inline int32 GetDim<FVector3f>()
+	{
+		return 3;
+	}
 
-        // Get the intersection of the line and the plane.
-        if (!ray_plane_intersection(ray_start, ray_end, v0, normal, out_intersection, out_t))
-        {
-            return false;
-        }
+	template<>
+	inline int32 GetDim<FIntVector2>()
+	{
+		return 2;
+	}
 
-        // \todo: review the equality comparison. Is it doing anything without a tolerance?
-        if (out_intersection == ray_start || out_intersection == ray_end)
-        {
-            return false;
-        }
+	template<>
+	inline int32 GetDim<UE::Math::TIntVector2<uint16>>()
+	{
+		return 2;
+	}
 
-        // Check if the intersection point is actually on the poly.
-		FVector3f vertices[3] = { v0, v1, v2 };
+	template<class VECTOR>
+	inline bool AllGreaterOrEqualThan(const VECTOR& v0, const VECTOR& v1)
+	{
+		return v0.AllGreaterOrEqualThan(v1);
+	}
 
-        return point_in_face( out_intersection, normal, vertices, out_intersected_vert,
-                              out_intersected_edge_v0, out_intersected_edge_v1 );
-    }
+	template<>
+	inline bool AllGreaterOrEqualThan<FIntVector2>(const FIntVector2& v0, const FIntVector2& v1)
+	{
+		return  (v0[0] >= v1[0]) && (v0[1] >= v1[1]);
+	}
 
-	//---------------------------------------------------------------------------------------------
-	//! Approximated inverse square root
-	//---------------------------------------------------------------------------------------------
-	inline float rsqrt_approx( float number )
+	template<class VECTOR>
+	inline bool AllSmallerThan(const VECTOR& v0, const VECTOR& v1)
 	{
-		union MUTABLE_FLOAT_INT
-		{
-            int32 i;
-			float y;
-		} fi;
-		float x2;
-		const float threehalfs = 1.5F;
+		return v0.AllSmallerThan(v1);
+	}
 
-		x2 = number * 0.5F;
-		fi.y  = number;
-		fi.i  = 0x5f3759df - ( fi.i >> 1 );
-		fi.y  = fi.y * ( threehalfs - ( x2 * fi.y * fi.y ) );
-		return fi.y;
+	template<>
+	inline bool AllSmallerThan<FIntVector2>(const FIntVector2& v0, const FIntVector2& v1)
+	{
+		return  (v0[0] < v1[0]) && (v0[1] < v1[1]);
 	}
 
+	template<class VECTOR>
+	inline bool AllSmallerOrEqualThan(const VECTOR& v0, const VECTOR& v1)
+	{
+		return v0.AllSmallerOrEqualThan(v1);
+	}
 
-	//---------------------------------------------------------------------------------------------
-	//! Normalise a vector, approximately
-	//---------------------------------------------------------------------------------------------
-	template <class SCALAR, int DIM>
-	inline vec<SCALAR,DIM> normalise_approx( const vec<SCALAR,DIM>& a )
+	template<>
+	inline bool AllSmallerOrEqualThan<FIntVector2>(const FIntVector2& v0, const FIntVector2& v1)
 	{
-		return a * rsqrt_approx( dot(a,a) );
+		return  (v0[0] <= v1[0]) && (v0[1] <= v1[1]);
 	}
 
 
@@ -1052,13 +96,13 @@ namespace mu
         //-----------------------------------------------------------------------------------------
         inline bool Contains( const VECTOR& v ) const
         {
-            return v.AllGreaterOrEqualThan(min) && v.AllSmallerThan(min+size);
+            return AllGreaterOrEqualThan(v,min) && AllSmallerThan(v, min+size);
         }
 
         //-----------------------------------------------------------------------------------------
         inline bool ContainsInclusive( const VECTOR& v ) const
         {
-            return v.AllGreaterOrEqualThan(min) && v.AllSmallerOrEqualThan(min+size);
+            return AllGreaterOrEqualThan(v, min) && AllSmallerOrEqualThan(v, min+size);
         }
 
         //-----------------------------------------------------------------------------------------
@@ -1067,7 +111,7 @@ namespace mu
         inline bool Intersects(const box<VECTOR>& other) const
         {
             bool result = true;
-            for (int d = 0; d < VECTOR::GetDim(); ++d)
+            for (int d = 0; d < GetDim<VECTOR>(); ++d)
             {
                 result = result &&
                          FMath::Abs((min[d] * 2 + size[d]) - (other.min[d] * 2 + other.size[d])) <=
@@ -1083,7 +127,7 @@ namespace mu
         inline bool IntersectsExclusive(const box<VECTOR>& other) const
         {
             bool result = true;
-            for (int d = 0; d < VECTOR::GetDim(); ++d)
+            for (int d = 0; d < GetDim<VECTOR>(); ++d)
             {
                 result = result &&
                          FMath::Abs((min[d] * 2 + size[d]) - (other.min[d] * 2 + other.size[d])) <
@@ -1096,15 +140,16 @@ namespace mu
         //-----------------------------------------------------------------------------------------
         //! If the boxes don't intersect, at least one size of the returned box will be negative.
         //-----------------------------------------------------------------------------------------
-        inline box<VECTOR> Intersect(const box<VECTOR>& other) const
+        inline box<FIntVector2> Intersect2i(const box<FIntVector2>& other) const
         {
-            box<VECTOR> result;
-            for (int d = 0; d < VECTOR::GetDim(); ++d)
+			constexpr int32 Dim = 2;
+			box<VECTOR> result;
+            for (int d = 0; d < Dim; ++d)
             {
                 result.min[d] = FMath::Max(min[d], other.min[d]);
                 result.size[d] =
-                    FMath::Min(typename VECTOR::SCALAR(min[d] + size[d] - result.min[d]),
-                             typename VECTOR::SCALAR(other.min[d] + other.size[d] - result.min[d]));
+                    FMath::Min(int32(min[d] + size[d] - result.min[d]),
+                               int32(other.min[d] + other.size[d] - result.min[d]));
             }
             return result;
         }
@@ -1118,7 +163,7 @@ namespace mu
 		//-----------------------------------------------------------------------------------------
 		inline void Bound(const VECTOR& v)
 		{
-			for (int d = 0; d < VECTOR::GetDim(); ++d)
+			for (int32 d = 0; d < GetDim<VECTOR>(); ++d)
 			{
 				if (v[d] < min[d])
 				{
@@ -1190,398 +235,5 @@ namespace mu
 	};
 
 
-
-
-	//---------------------------------------------------------------------------------------------
-	//! Generic square matrix template.
-	//---------------------------------------------------------------------------------------------
-	template <class SCALAR, int DIM>
-    class mat
-	{
-	protected:
-
-		//! Rows
-		vec<SCALAR,DIM> m[ DIM ];
-
-	public:
-
-		//-----------------------------------------------------------------------------------------
-		typedef vec<SCALAR,DIM> V;
-		typedef mat<SCALAR,DIM> M;
-
-		//-----------------------------------------------------------------------------------------
-		static inline int GetDim()
-		{
-			return DIM;
-		}
-
-		//-----------------------------------------------------------------------------------------
-		//! 0 matrix
-		//-----------------------------------------------------------------------------------------
-		mat()
-		{
-		}
-
-		//-----------------------------------------------------------------------------------------
-		//! Build from a float array
-		//-----------------------------------------------------------------------------------------
-		explicit mat( const float* p )
-		{
-			for (int j=0; j<DIM; ++j)
-			{
-				for (int i=0; i<DIM; ++i)
-				{
-					m[j][i] = *p;
-					++p;
-				}
-			}
-		}
-
-		//-----------------------------------------------------------------------------------------
-		//! Initialise from another matrix with different size, filling the unknown with identites.
-		//-----------------------------------------------------------------------------------------
-		template<int DIM2>
-		explicit mat( const mat<SCALAR,DIM2>& other )
-		{
-			for (int i=0; i<FMath::Min(DIM,DIM2); ++i)
-			{
-				for (int j=0; j<FMath::Min(DIM,DIM2) ;++j)
-				{
-					m[i][j] = other[i][j];
-				}
-			}
-
-			for (int i=DIM2; i<DIM; ++i)
-			{
-				m[i][i] = 1;
-			}
-		}
-
-		//-----------------------------------------------------------------------------------------
-		//! Make a scale matrix.
-		//-----------------------------------------------------------------------------------------
-		template<int DIM2>
-		inline static M Scale( const vec<SCALAR,DIM2>& s )
-		{
-			M r;
-			for (int i=0; i<FMath::Min(DIM,DIM2); ++i)
-			{
-				r.m[i][i] = s[i];
-			}
-
-			for (int i=DIM2; i<DIM; ++i)
-			{
-				r.m[i][i] = 1;
-			}
-
-			return r;
-		}
-
-		//-----------------------------------------------------------------------------------------
-		//! Make a translation matrix.
-		//-----------------------------------------------------------------------------------------
-		template<int DIM2>
-		static inline M Translate( const vec<SCALAR,DIM2>& pos )
-		{
-			M r = Identity();
-
-			for (int i=0; i<FMath::Min(DIM,DIM2); ++i)
-			{
-				r.m[i][DIM-1] = pos[i];
-			}
-
-			return r;
-		}
-
-        //-----------------------------------------------------------------------------------------
-        //! Return an identity matrix.
-        //-----------------------------------------------------------------------------------------
-        static inline M Identity()
-        {
-            M res;
-            for (int i=0; i<DIM;++i)
-            {
-                res[i][i] = 1;
-            }
-            return res;
-        }
-
-        //-----------------------------------------------------------------------------------------
-        //! Return an identity matrix.
-        //-----------------------------------------------------------------------------------------
-        inline void SetIdentity()
-        {
-            *this = M::Identity();
-        }
-
-		//-----------------------------------------------------------------------------------------
-		//! Return a row
-		//-----------------------------------------------------------------------------------------
-        inline V& operator[] (int i)
-		{
-			check( i>=0 && i <DIM );
-			return m[i];
-		}
-
-		//-----------------------------------------------------------------------------------------
-		//! Return a row
-		//-----------------------------------------------------------------------------------------
-		inline const V& GetRow(int i) const
-		{
-			check( i>=0 && i <DIM );
-			return m[i];
-		}
-
-		//-----------------------------------------------------------------------------------------
-		//! Return a column
-		//-----------------------------------------------------------------------------------------
-		inline V GetColumn(int i) const
-		{
-			check( i>=0 && i <DIM );
-			V r;
-			for ( int j=0; j<DIM; ++j )
-			{
-				r[j] = m[j][i];
-			}
-			return r;
-		}
-
-		//-----------------------------------------------------------------------------------------
-        inline const V& operator[] (int i) const
-		{
-			check( i>=0 && i <DIM );
-			return m[i];
-		}
-
-		//-----------------------------------------------------------------------------------------
-		inline bool operator==( const M& other ) const
-		{
-			bool res = true;
-			for( int i=0; i<DIM; ++i )
-			{
-				res &= m[i] == other.m[i];
-			}
-			return res;
-		}
-
-		//-----------------------------------------------------------------------------------------
-		inline bool operator!=( const M& other ) const
-		{
-			bool res = false;
-			for( int i=0; i<DIM; ++i )
-			{
-				res |= m[i] != other.m[i];
-			}
-			return res;
-		}
-
-		//-----------------------------------------------------------------------------------------
-		inline V operator+( const M& other ) const
-		{
-			V res;
-			for( int i=0; i<DIM; ++i )
-			{
-				res.m[i] = m[i] + other.m[i];
-			}
-			return res;
-		}
-
-		//-----------------------------------------------------------------------------------------
-		inline void operator+=( const M& other )
-		{
-			for( int i=0; i<DIM; ++i )
-			{
-				m[i] += other.m[i];
-			}
-		}
-
-		//-----------------------------------------------------------------------------------------
-		inline V operator-( const M& other ) const
-		{
-			V res;
-			for( int i=0; i<DIM; ++i )
-			{
-				res.m[i] = m[i] - other.m[i];
-			}
-			return res;
-		}
-
-		//-----------------------------------------------------------------------------------------
-		inline void operator-=( const M& other )
-		{
-			for( int i=0; i<DIM; ++i )
-			{
-				m[i] -= other.m[i];
-			}
-		}
-
-        //-----------------------------------------------------------------------------------------
-        M operator*( const mat<SCALAR,DIM>& m2 ) const
-        {
-            M r;
-
-            for ( int i=0; i<DIM; ++i )
-            {
-                for ( int j=0; j<DIM; ++j )
-                {
-                    r.m[i][j] = dot( m[i], m2.GetColumn(j) );
-                }
-            }
-
-            return r;
-        }
-
-		//-----------------------------------------------------------------------------------------
-		V Transform(const V& v) const
-		{
-			V r;
-
-			for (int i = 0; i < DIM; ++i)
-			{
-				r[i] = dot(m[i], v);
-			}
-
-			return r;
-		}
-
-		//-----------------------------------------------------------------------------------------
-		V operator*(const V& v) const
-		{
-			return Transform(v);
-		}
-
-        //-----------------------------------------------------------------------------------------
-        SCALAR GetDeterminant() const
-        {
-            SCALAR result=0;
-
-            for ( int i=0; i<DIM; ++i )
-            {
-                SCALAR accum=1;
-                for ( int j=0; j<DIM; ++j )
-                {
-                    accum *= m[(i+j)%DIM][j];
-                }
-                result+=accum;
-
-                accum=1;
-                for ( int j=0; j<DIM; ++j )
-                {
-                    accum *= m[(DIM+i-j)%DIM][j];
-                }
-                result-=accum;
-            }
-
-            return result;
-        }
-
-
-		//-----------------------------------------------------------------------------------------
-		//! Strict ordering of matrices.
-		//-----------------------------------------------------------------------------------------
-		inline bool operator<( const M& other ) const
-		{
-			bool equal = true;
-			bool res = false;
-			for( int i=0; equal && i<DIM; ++i )
-			{
-				if( m[i] < other.m[i] )
-				{
-					equal = false;
-					res = true;
-				}
-				else if( m[i] > other.m[i] )
-				{
-					equal = false;
-					res = false;
-				}
-			}
-			return res;
-		}
-
-
-		//-----------------------------------------------------------------------------------------
-		inline bool AlmostEqual
-			(
-				const V& other,
-				SCALAR tol = std::numeric_limits<SCALAR>::min()
-			) const
-		{
-			bool res = true;
-			for( int i=0; i<DIM; ++i )
-			{
-				res &= m[i].AlmostEqual(other.m[i],tol);
-			}
-			return res;
-		}
-
-	};
-
-
-	//!
-	template <class SCALAR>
-	class mat3 : public mat<SCALAR,3>
-	{
-	public:
-
-		mat3()
-		{
-		}
-
-        mat3( SCALAR a0, SCALAR a1, SCALAR a2,
-              SCALAR a3, SCALAR a4, SCALAR a5,
-              SCALAR a6, SCALAR a7, SCALAR a8 )
-        {
-            this->m[0][0] = a0;
-            this->m[0][1] = a1;
-            this->m[0][2] = a2;
-            this->m[1][0] = a3;
-            this->m[1][1] = a4;
-            this->m[1][2] = a5;
-            this->m[2][0] = a6;
-            this->m[2][1] = a7;
-            this->m[2][2] = a8;
-        }
-
-        mat3( vec3<SCALAR> r0, vec3<SCALAR> r1, vec3<SCALAR> r2 )
-        {
-            this->m[0] = r0;
-            this->m[1] = r1;
-            this->m[2] = r2;
-        }
-
-
-	};
-
-
-	//!
-	typedef mat3<float> mat3f;
-    typedef mat<float,4> mat4f;
-
-	//---------------------------------------------------------------------------------------------
-	// Temp while transitioning to UE math.
-	//---------------------------------------------------------------------------------------------
-	inline FMatrix44f ToUnreal(const mat4f& m)
-	{
-		FMatrix44f ueMat;
-		for (int i = 0; i < 4; ++i)
-			for (int j = 0; j < 4; ++j)
-				ueMat.M[i][j] = m[i][j];
-		return ueMat;
-	}
-
-
-	inline FVector3f ToUnreal(const vec3f& v)
-	{
-		return FVector3f(v.x(), v.y(), v.z());
-	}
-
-
-	inline vec3f FromUnreal(const FVector3f& v)
-	{
-		return vec3f(v.X, v.Y, v.Z);
-	}
-
-
 }
 
diff --git a/Engine/Plugins/Experimental/Mutable/Source/MutableRuntime/Public/MuR/Serialisation.h b/Engine/Plugins/Experimental/Mutable/Source/MutableRuntime/Public/MuR/Serialisation.h
--- a/Engine/Plugins/Experimental/Mutable/Source/MutableRuntime/Public/MuR/Serialisation.h
+++ b/Engine/Plugins/Experimental/Mutable/Source/MutableRuntime/Public/MuR/Serialisation.h
@@ -344,8 +344,16 @@ namespace mu
     MUTABLE_DEFINE_POD_SERIALISABLE(uint8);
     MUTABLE_DEFINE_POD_SERIALISABLE(uint16);
     MUTABLE_DEFINE_POD_SERIALISABLE(uint32);
-    MUTABLE_DEFINE_POD_SERIALISABLE(uint64);
-	
+	MUTABLE_DEFINE_POD_SERIALISABLE(uint64);
+
+	MUTABLE_DEFINE_POD_SERIALISABLE(FUintVector2);
+	MUTABLE_DEFINE_POD_SERIALISABLE(FIntVector2);
+	MUTABLE_DEFINE_POD_SERIALISABLE(UE::Math::TIntVector2<uint16>);
+	MUTABLE_DEFINE_POD_SERIALISABLE(UE::Math::TIntVector2<int16>);
+	MUTABLE_DEFINE_POD_SERIALISABLE(FVector2f);
+	MUTABLE_DEFINE_POD_SERIALISABLE(FVector4f);
+	MUTABLE_DEFINE_POD_SERIALISABLE(FMatrix44f);
+
 	MUTABLE_DEFINE_POD_SERIALISABLE(FGuid);
 
 	MUTABLE_DEFINE_POD_VECTOR_SERIALISABLE(float);
@@ -359,11 +367,9 @@ namespace mu
 	MUTABLE_DEFINE_POD_VECTOR_SERIALISABLE(int32);
 	MUTABLE_DEFINE_POD_VECTOR_SERIALISABLE(int64);
 
-	MUTABLE_DEFINE_POD_VECTOR_SERIALISABLE(vec2f);
-	MUTABLE_DEFINE_POD_VECTOR_SERIALISABLE(vec3f);
-	MUTABLE_DEFINE_POD_VECTOR_SERIALISABLE(mat3f);
-	MUTABLE_DEFINE_POD_VECTOR_SERIALISABLE(mat4f);
-	MUTABLE_DEFINE_POD_VECTOR_SERIALISABLE(vec2<int>);
+	MUTABLE_DEFINE_POD_VECTOR_SERIALISABLE(FVector2f);
+	MUTABLE_DEFINE_POD_VECTOR_SERIALISABLE(FMatrix44f);
+	MUTABLE_DEFINE_POD_VECTOR_SERIALISABLE(FIntVector2);
 	MUTABLE_DEFINE_POD_VECTOR_SERIALISABLE(TCHAR);
 
 	MUTABLE_DEFINE_POD_VECTOR_SERIALISABLE(FUintVector2);
diff --git a/Engine/Plugins/Experimental/Mutable/Source/MutableTools/Internal/MuT/NodeColourParameterPrivate.h b/Engine/Plugins/Experimental/Mutable/Source/MutableTools/Internal/MuT/NodeColourParameterPrivate.h
--- a/Engine/Plugins/Experimental/Mutable/Source/MutableTools/Internal/MuT/NodeColourParameterPrivate.h
+++ b/Engine/Plugins/Experimental/Mutable/Source/MutableTools/Internal/MuT/NodeColourParameterPrivate.h
@@ -18,7 +18,7 @@ namespace mu
 
 		static FNodeType s_type;
 
-		vec4<float> m_defaultValue;
+		FVector4f m_defaultValue;
 		FString m_name;
 		FString m_uid;
 
@@ -27,7 +27,7 @@ namespace mu
 		//!
 		void Serialise( OutputArchive& arch ) const
 		{
-            uint32_t ver = 4;
+            uint32 ver = 4;
 			arch << ver;
 
 			arch << m_defaultValue;
@@ -39,17 +39,17 @@ namespace mu
 		//!
 		void Unserialise( InputArchive& arch )
 		{
-            uint32_t ver;
+            uint32 ver;
 			arch >> ver;
             check(ver>=2&&ver<=4);
 
 			if (ver <= 3)
 			{
-				vec3<float> aux_defaultValue;
-				arch >> aux_defaultValue;
+				FVector3f TempDefaultValue;
+				arch >> TempDefaultValue;
 
-				m_defaultValue = aux_defaultValue;
-				m_defaultValue.w() = 1.f;
+				m_defaultValue = TempDefaultValue;
+				m_defaultValue[3] = 1.f;
 			}
 			else
 			{
diff --git a/Engine/Plugins/Experimental/Mutable/Source/MutableTools/Internal/MuT/NodeModifierMeshClipMorphPlanePrivate.h b/Engine/Plugins/Experimental/Mutable/Source/MutableTools/Internal/MuT/NodeModifierMeshClipMorphPlanePrivate.h
--- a/Engine/Plugins/Experimental/Mutable/Source/MutableTools/Internal/MuT/NodeModifierMeshClipMorphPlanePrivate.h
+++ b/Engine/Plugins/Experimental/Mutable/Source/MutableTools/Internal/MuT/NodeModifierMeshClipMorphPlanePrivate.h
@@ -34,8 +34,8 @@ namespace mu
         float m_factor;
 
         // Ellipse location
-        vec3f m_origin;
-        vec3f m_normal;
+		FVector3f m_origin;
+		FVector3f m_normal;
         float m_radius1, m_radius2, m_rotation;
 
         //! Typed of vertex selection
@@ -52,9 +52,9 @@ namespace mu
         } VERTEX_SELECTION;
 
         // Vertex selection box
-        uint8_t m_vertexSelectionType;
-        vec3f m_selectionBoxOrigin;
-        vec3f m_selectionBoxRadius;
+        uint8 m_vertexSelectionType;
+		FVector3f m_selectionBoxOrigin;
+		FVector3f m_selectionBoxRadius;
         uint16 m_vertexSelectionBone;
 
 		// Max distance a vertex can have to the bone in order to be affected. A negative value
@@ -66,7 +66,7 @@ namespace mu
         {
             NodeModifier::Private::Serialise(arch);
 
-			uint32_t ver = 3;
+			uint32 ver = 3;
             arch << ver;
 
             arch << m_origin;
@@ -88,7 +88,7 @@ namespace mu
         {
             NodeModifier::Private::Unserialise( arch );
 
-            uint32_t ver;
+            uint32 ver;
             arch >> ver;
             check(ver>=2);
 
diff --git a/Engine/Plugins/Experimental/Mutable/Source/MutableTools/Internal/MuT/NodeProjectorPrivate.h b/Engine/Plugins/Experimental/Mutable/Source/MutableTools/Internal/MuT/NodeProjectorPrivate.h
--- a/Engine/Plugins/Experimental/Mutable/Source/MutableTools/Internal/MuT/NodeProjectorPrivate.h
+++ b/Engine/Plugins/Experimental/Mutable/Source/MutableTools/Internal/MuT/NodeProjectorPrivate.h
@@ -26,16 +26,16 @@ namespace mu
 		static FNodeType s_type;
 
         PROJECTOR_TYPE m_type = PROJECTOR_TYPE::PLANAR;
-		vec3<float> m_position;
-		vec3<float> m_direction;
-		vec3<float> m_up;
-        vec3<float> m_scale;
+		FVector3f m_position;
+		FVector3f m_direction;
+		FVector3f m_up;
+		FVector3f m_scale;
         float m_projectionAngle = 0.0f;
 
 		//!
 		void Serialise( OutputArchive& arch ) const
 		{
-            uint32_t ver = 2;
+            uint32 ver = 2;
 			arch << ver;
 
             arch << m_type;
@@ -49,7 +49,7 @@ namespace mu
 		//!
 		void Unserialise( InputArchive& arch )
 		{
-            uint32_t ver;
+            uint32 ver;
 			arch >> ver;
             check(ver==2);
 
@@ -80,7 +80,7 @@ namespace mu
 		{
 			NodeProjectorConstant::Private::Serialise( arch );
 
-            uint32_t ver = 3;
+            uint32 ver = 3;
 			arch << ver;
 
 			arch << m_name;
@@ -93,7 +93,7 @@ namespace mu
 		{
 			NodeProjectorConstant::Private::Unserialise( arch );
 
-            uint32_t ver;
+            uint32 ver;
 			arch >> ver;
             check(ver<=3);
 
diff --git a/Engine/Plugins/Experimental/Mutable/Source/MutableTools/Internal/MuT/TablePrivate.h b/Engine/Plugins/Experimental/Mutable/Source/MutableTools/Internal/MuT/TablePrivate.h
--- a/Engine/Plugins/Experimental/Mutable/Source/MutableTools/Internal/MuT/TablePrivate.h
+++ b/Engine/Plugins/Experimental/Mutable/Source/MutableTools/Internal/MuT/TablePrivate.h
@@ -148,7 +148,7 @@ namespace mu
 					{
 						if (ver <= 1)
 						{
-							vec3<float> Value;
+							FVector3f Value;
 							arch >> Value;
 
 							v.Color = FVector4f(Value[0], Value[1], Value[2], 1.0f);
diff --git a/Engine/Plugins/Experimental/Mutable/Source/MutableTools/Private/MuT/ASTOpImageCompose.h b/Engine/Plugins/Experimental/Mutable/Source/MutableTools/Private/MuT/ASTOpImageCompose.h
--- a/Engine/Plugins/Experimental/Mutable/Source/MutableTools/Private/MuT/ASTOpImageCompose.h
+++ b/Engine/Plugins/Experimental/Mutable/Source/MutableTools/Private/MuT/ASTOpImageCompose.h
@@ -11,8 +11,7 @@
 
 namespace mu
 {
-struct FProgram;
-template <class SCALAR> class vec4;
+	struct FProgram;
 
 	class ASTOpImageCompose final : public ASTOp
 	{
diff --git a/Engine/Plugins/Experimental/Mutable/Source/MutableTools/Private/MuT/ASTOpImageCrop.cpp b/Engine/Plugins/Experimental/Mutable/Source/MutableTools/Private/MuT/ASTOpImageCrop.cpp
--- a/Engine/Plugins/Experimental/Mutable/Source/MutableTools/Private/MuT/ASTOpImageCrop.cpp
+++ b/Engine/Plugins/Experimental/Mutable/Source/MutableTools/Private/MuT/ASTOpImageCrop.cpp
@@ -289,13 +289,13 @@ namespace mu
 			Ptr<ASTOp> rectOp = typedPatch->patch.child();
 			ASTOp::FGetImageDescContext context;
 			FImageDesc patchDesc = rectOp->GetImageDesc(false, &context);
-			box<vec2<int16>> patchBox;
+			box<FIntVector2> patchBox;
 			patchBox.min[0] = typedPatch->location[0];
 			patchBox.min[1] = typedPatch->location[1];
 			patchBox.size[0] = patchDesc.m_size[0];
 			patchBox.size[1] = patchDesc.m_size[1];
 
-			box<vec2<int16>> cropBox;
+			box<FIntVector2> cropBox;
 			cropBox.min[0] = currentCropOp->Min[0];
 			cropBox.min[1] = currentCropOp->Min[1];
 			cropBox.size[0] = currentCropOp->Size[0];
@@ -313,7 +313,7 @@ namespace mu
 				Ptr<ASTOpImagePatch> newOp = mu::Clone<ASTOpImagePatch>(at);
 				newOp->base = Visit(newOp->base.child(), currentCropOp);
 
-				box<vec2<int16>> ibox = patchBox.Intersect(cropBox);
+				box<FIntVector2> ibox = patchBox.Intersect2i(cropBox);
 				check(ibox.size[0] > 0 && ibox.size[1] > 0);
 
 				Ptr<ASTOpImageCrop> patchCropOp = mu::Clone<ASTOpImageCrop>(currentCropOp);
@@ -336,13 +336,13 @@ namespace mu
 			// We can combine the two crops into a possibly smaller crop
 			const ASTOpImageCrop* childCrop = static_cast<const ASTOpImageCrop*>(at.get());
 
-			box<vec2<int16>> childCropBox;
+			box<FIntVector2> childCropBox;
 			childCropBox.min[0] = childCrop->Min[0];
 			childCropBox.min[1] = childCrop->Min[1];
 			childCropBox.size[0] = childCrop->Size[0];
 			childCropBox.size[1] = childCrop->Size[1];
 
-			box<vec2<int16>> cropBox;
+			box<FIntVector2> cropBox;
 			cropBox.min[0] = currentCropOp->Min[0];
 			cropBox.min[1] = currentCropOp->Min[1];
 			cropBox.size[0] = currentCropOp->Size[0];
@@ -350,12 +350,13 @@ namespace mu
 
 			// Compose the crops: in the final image the child crop is applied first and the
 			// current ctop is applied to the result. So the final crop box would be:
-			box<vec2<int16>> ibox;
+			box<FIntVector2> ibox;
 			ibox.min = childCropBox.min + cropBox.min;
-			ibox.size = vec2<int16>::min(cropBox.size, childCropBox.size);
-			check(cropBox.min.AllSmallerOrEqualThan(childCropBox.size));
-			check((cropBox.min + cropBox.size).AllSmallerOrEqualThan(childCropBox.size));
-			check((ibox.min + ibox.size).AllSmallerOrEqualThan(childCropBox.min + childCropBox.size));
+			ibox.size[0] = FMath::Min(cropBox.size[0], childCropBox.size[0]);
+			ibox.size[1] = FMath::Min(cropBox.size[1], childCropBox.size[1]);
+			//check(cropBox.min.AllSmallerOrEqualThan(childCropBox.size));
+			//check((cropBox.min + cropBox.size).AllSmallerOrEqualThan(childCropBox.size));
+			//check((ibox.min + ibox.size).AllSmallerOrEqualThan(childCropBox.min + childCropBox.size));
 
 			// This happens more often that one would think
 			if (ibox == childCropBox)
@@ -441,7 +442,7 @@ namespace mu
 			bool bRasterHasCrop = nop->UncroppedSizeX != 0;
 			if (!bRasterHasCrop)
 			{
-				box<vec2<int16>> cropBox;
+				box<FIntVector2> cropBox;
 				cropBox.min[0] = currentCropOp->Min[0];
 				cropBox.min[1] = currentCropOp->Min[1];
 				cropBox.size[0] = currentCropOp->Size[0];
diff --git a/Engine/Plugins/Experimental/Mutable/Source/MutableTools/Private/MuT/ASTOpImageCrop.h b/Engine/Plugins/Experimental/Mutable/Source/MutableTools/Private/MuT/ASTOpImageCrop.h
--- a/Engine/Plugins/Experimental/Mutable/Source/MutableTools/Private/MuT/ASTOpImageCrop.h
+++ b/Engine/Plugins/Experimental/Mutable/Source/MutableTools/Private/MuT/ASTOpImageCrop.h
@@ -12,7 +12,7 @@
 
 namespace mu
 {
-struct FProgram;
+	struct FProgram;
 
 	class ASTOpImageCrop final : public ASTOp
 	{
@@ -37,7 +37,7 @@ struct FProgram;
 		virtual FImageDesc GetImageDesc(bool returnBestOption, FGetImageDescContext* context) const override;
 		virtual Ptr<ImageSizeExpression> GetImageSizeExpression() const override;
 		virtual void GetLayoutBlockSize(int* pBlockX, int* pBlockY) override;
-		//TODO: virtual bool IsImagePlainConstant(vec4<float>& colour) const override;
+		//TODO: virtual bool IsImagePlainConstant(FVector4f& colour) const override;
 		virtual Ptr<ASTOp> OptimiseSemantic(const FModelOptimizationOptions& options, int32 Pass) const override;
 		virtual Ptr<ASTOp> OptimiseSink(const FModelOptimizationOptions& options, FOptimizeSinkContext& context) const override;
 
diff --git a/Engine/Plugins/Experimental/Mutable/Source/MutableTools/Private/MuT/ASTOpImageLayer.cpp b/Engine/Plugins/Experimental/Mutable/Source/MutableTools/Private/MuT/ASTOpImageLayer.cpp
--- a/Engine/Plugins/Experimental/Mutable/Source/MutableTools/Private/MuT/ASTOpImageLayer.cpp
+++ b/Engine/Plugins/Experimental/Mutable/Source/MutableTools/Private/MuT/ASTOpImageLayer.cpp
@@ -473,7 +473,7 @@ namespace mu
 				FImageRect maskUsage;
 				maskUsage.min[0] = (sourceMaskUsage.min[0] / blockSize) * blockSize;
 				maskUsage.min[1] = (sourceMaskUsage.min[1] / blockSize) * blockSize;
-				vec2<uint16> minOffset = sourceMaskUsage.min - maskUsage.min;
+				FImageSize minOffset = sourceMaskUsage.min - maskUsage.min;
 				maskUsage.size[0] = ((sourceMaskUsage.size[0] + minOffset[0] + blockSize - 1) / blockSize) * blockSize;
 				maskUsage.size[1] = ((sourceMaskUsage.size[1] + minOffset[1] + blockSize - 1) / blockSize) * blockSize;
 
diff --git a/Engine/Plugins/Experimental/Mutable/Source/MutableTools/Private/MuT/ASTOpImagePatch.h b/Engine/Plugins/Experimental/Mutable/Source/MutableTools/Private/MuT/ASTOpImagePatch.h
--- a/Engine/Plugins/Experimental/Mutable/Source/MutableTools/Private/MuT/ASTOpImagePatch.h
+++ b/Engine/Plugins/Experimental/Mutable/Source/MutableTools/Private/MuT/ASTOpImagePatch.h
@@ -37,7 +37,7 @@ struct FProgram;
 		FImageDesc GetImageDesc(bool returnBestOption, FGetImageDescContext* context) const override;
 		Ptr<ImageSizeExpression> GetImageSizeExpression() const override;
 		void GetLayoutBlockSize(int* pBlockX, int* pBlockY) override;
-		//TODO: bool IsImagePlainConstant(vec4<float>& colour) const override;
+		//TODO: bool IsImagePlainConstant(FVector4f& colour) const override;
 	};
 
 
diff --git a/Engine/Plugins/Experimental/Mutable/Source/MutableTools/Private/MuT/ASTOpMeshTransform.h b/Engine/Plugins/Experimental/Mutable/Source/MutableTools/Private/MuT/ASTOpMeshTransform.h
--- a/Engine/Plugins/Experimental/Mutable/Source/MutableTools/Private/MuT/ASTOpMeshTransform.h
+++ b/Engine/Plugins/Experimental/Mutable/Source/MutableTools/Private/MuT/ASTOpMeshTransform.h
@@ -23,7 +23,7 @@ struct FProgram;
 
 		ASTChild source;
 
-		mat4f matrix;
+		FMatrix44f matrix;
 
 	public:
 
diff --git a/Engine/Plugins/Experimental/Mutable/Source/MutableTools/Private/MuT/CodeGenerator.cpp b/Engine/Plugins/Experimental/Mutable/Source/MutableTools/Private/MuT/CodeGenerator.cpp
--- a/Engine/Plugins/Experimental/Mutable/Source/MutableTools/Private/MuT/CodeGenerator.cpp
+++ b/Engine/Plugins/Experimental/Mutable/Source/MutableTools/Private/MuT/CodeGenerator.cpp
@@ -2232,20 +2232,20 @@ namespace mu
 					morphShape.position = TypedNode->m_origin;
 					morphShape.up = TypedNode->m_normal;
 					// TODO: Move rotation to ellipse rotation reference base instead of passing it directly
-					morphShape.size = vec3f(TypedNode->m_radius1, TypedNode->m_radius2, TypedNode->m_rotation);
+					morphShape.size = FVector3f(TypedNode->m_radius1, TypedNode->m_radius2, TypedNode->m_rotation);
 
 					// Generate a "side" vector.
 					// \todo: make generic and move to the vector class
 					{
 						// Generate vector perpendicular to normal for ellipse rotation reference base
-						vec3f aux_base(0.f, 1.f, 0.f);
+						FVector3f aux_base(0.f, 1.f, 0.f);
 
-						if (fabs(dot(TypedNode->m_normal, aux_base)) > 0.95f)
+						if (FMath::Abs(FVector3f::DotProduct(TypedNode->m_normal, aux_base)) > 0.95f)
 						{
-							aux_base = vec3f(0.f, 0.f, 1.f);
+							aux_base = FVector3f(0.f, 0.f, 1.f);
 						}
 
-						morphShape.side = cross(TypedNode->m_normal, aux_base);
+						morphShape.side = FVector3f::CrossProduct(TypedNode->m_normal, aux_base);
 					}
 					op->morphShape = morphShape;
 				}
diff --git a/Engine/Plugins/Experimental/Mutable/Source/MutableTools/Private/MuT/CodeGenerator_Mesh.cpp b/Engine/Plugins/Experimental/Mutable/Source/MutableTools/Private/MuT/CodeGenerator_Mesh.cpp
--- a/Engine/Plugins/Experimental/Mutable/Source/MutableTools/Private/MuT/CodeGenerator_Mesh.cpp
+++ b/Engine/Plugins/Experimental/Mutable/Source/MutableTools/Private/MuT/CodeGenerator_Mesh.cpp
@@ -159,7 +159,7 @@ namespace mu
 	void GetUVIsland(TArray<FTriangle>& InTriangles,
 		const uint32 InFirstTriangle,
 		TArray<uint32>& OutTriangleIndices,
-		const TArray<vec2<float>>& InUVs,
+		const TArray<FVector2f>& InUVs,
 		const TMultiMap<int32, uint32>& InVertexToTriangleMap)
 	{
 		MUTABLE_CPUPROFILER_SCOPE(LayoutUV_GetUVIsland);
@@ -205,7 +205,7 @@ namespace mu
 					{
 						// Check if the vertex is in the same UV Island 
 						if (!SkipTrinalges[OtherTriangleIndex]
-							&& InUVs[Triangle.Indices[1]].AlmostEqual(InUVs[OtherTriangle.Indices[OtherIndex]], 0.00001))
+							&& InUVs[Triangle.Indices[1]].Equals(InUVs[OtherTriangle.Indices[OtherIndex]], 0.00001f))
 						{
 							OutTriangleIndices.Add(OtherTriangleIndex);
 							PendingTriangles.Add(OtherTriangleIndex);
@@ -220,7 +220,7 @@ namespace mu
 					{
 						// Check if the vertex is in the same UV Island 
 						if (!SkipTrinalges[OtherTriangleIndex]
-							&& InUVs[Triangle.Indices[2]].AlmostEqual(InUVs[OtherTriangle.Indices[OtherIndex]], 0.00001))
+							&& InUVs[Triangle.Indices[2]].Equals(InUVs[OtherTriangle.Indices[OtherIndex]], 0.00001f))
 						{
 							OutTriangleIndices.Add(OtherTriangleIndex);
 							PendingTriangles.Add(OtherTriangleIndex);
@@ -249,7 +249,7 @@ namespace mu
 					{
 						// Check if the vertex belong to the same UV island
 						if (!SkipTrinalges[OtherTriangleIndex]
-							&& InUVs[Triangle.Indices[2]].AlmostEqual(InUVs[OtherTriangle.Indices[OtherIndex]], 0.00001))
+							&& InUVs[Triangle.Indices[2]].Equals(InUVs[OtherTriangle.Indices[OtherIndex]], 0.00001f))
 						{
 							OutTriangleIndices.Add(OtherTriangleIndex);
 							PendingTriangles.Add(OtherTriangleIndex);
@@ -329,7 +329,7 @@ namespace mu
 
 		// 
 		TArray<uint16> BlockIds;
-		TArray<box<vec2<float>>> BlockRects;
+		TArray<box<FVector2f>> BlockRects;
 
 		BlockIds.SetNumUninitialized(NumBlocks);
 		BlockRects.SetNumUninitialized(NumBlocks);
@@ -345,7 +345,7 @@ namespace mu
 			uint16 MinX, MinY, SizeX, SizeY;
 			Layout->GetBlock(BlockIndex, &MinX, &MinY, &SizeX, &SizeY);
 
-			box<vec2<float>>& BlockRect = BlockRects[BlockIndex];
+			box<FVector2f>& BlockRect = BlockRects[BlockIndex];
 			BlockRect.min[0] = ((float)MinX) / (float)Grid.X;
 			BlockRect.min[1] = ((float)MinY) / (float)Grid.Y;
 			BlockRect.size[0] = ((float)SizeX) / (float)Grid.X;
@@ -399,17 +399,17 @@ namespace mu
 		(buffer, channel, &semantic, &semanticIndex, &format, &components, &offset);
 		check(semantic == MBS_TEXCOORDS);
 
-		const uint8_t* pData = currentLayoutMesh->GetVertexBuffers().GetBufferData(buffer);
+		uint8* pData = currentLayoutMesh->GetVertexBuffers().GetBufferData(buffer);
 		int elemSize = currentLayoutMesh->GetVertexBuffers().GetElementSize(buffer);
 		int channelOffset = currentLayoutMesh->GetVertexBuffers().GetChannelOffset(buffer, channel);
 		pData += channelOffset;
 
 
 		// Temp copy of the UVs
-		TArray<vec2<float>> TempUVs;
+		TArray<FVector2f> TempUVs;
 		TempUVs.SetNumUninitialized(NumVertices);
 
-		// Get a copy of the UVs as vec2<float> to work with them. 
+		// Get a copy of the UVs as FVector2f to work with them. 
 		{
 			bool bNonNormalizedUVs = false;
 			const bool bIsOverlayLayout = GeneratedLayout->GetLayoutPackingStrategy() == mu::EPackStrategy::OVERLAY_LAYOUT;
@@ -417,15 +417,15 @@ namespace mu
 			const uint8* pVertices = pData;
 			for (int32 VertexIndex = 0; VertexIndex < NumVertices; ++VertexIndex)
 			{
-				vec2<float>& UV = TempUVs[VertexIndex];
+				FVector2f& UV = TempUVs[VertexIndex];
 				if (format == MBF_FLOAT32)
 				{
-					UV = *((vec2<float>*)pVertices);
+					UV = *((FVector2f*)pVertices);
 				}
 				else if (format == MBF_FLOAT16)
 				{
-					float16* pUV = (float16*)pVertices;
-					UV = vec2<float>(halfToFloat(pUV[0]), halfToFloat(pUV[1]));
+					const FFloat16* pUV = reinterpret_cast<const FFloat16*>(pVertices);
+					UV = FVector2f(float(pUV[0]), float(pUV[1]));
 				}
 
 				// Check that UVs are normalized. If not, clamp the values and throw a warning.
@@ -589,7 +589,7 @@ namespace mu
 			const float MaxY = (((float)LayoutBlock.m_size.Y + LayoutBlock.m_min.Y) / (float)Grid.Y) - 2 * SmallNumber;
 
 			// Iterate triangles and clamp the UVs
-			vec2<float>* TempUVsData = TempUVs.GetData();
+			FVector2f* TempUVsData = TempUVs.GetData();
 			for (int32 TriangleIndex : TriangleIndices)
 			{
 				FTriangle& OtherTriangle = Triangles[TriangleIndex];
@@ -605,7 +605,7 @@ namespace mu
 
 					// Clamp UVs
 					const int32 UVIndex = OtherTriangle.Indices[VertexIndex];
-					vec2<float>& UV = TempUVs[UVIndex];
+					FVector2f& UV = TempUVs[UVIndex];
 					UV[0] = FMath::Clamp(UV[0], MinX, MaxX);
 					UV[1] = FMath::Clamp(UV[1], MinY, MaxY);
 					*(LayoutData + UVIndex) = BlockIndex;
@@ -621,7 +621,7 @@ namespace mu
 			TArray<float> UnassignedUVs;
 			UnassignedUVs.Reserve(NumVertices / 100);
 
-			const vec2<float>* UVs = TempUVs.GetData();
+			const FVector2f* UVs = TempUVs.GetData();
 			for (int32 VertexIndex = 0; VertexIndex < NumVertices; ++VertexIndex)
 			{
 				if (LayoutData[VertexIndex] == MAX_uint16)
@@ -645,12 +645,12 @@ namespace mu
 
 		// Format and copy UVs
 		{
-			uint8* pVertices = const_cast<uint8*>(pData);
-			vec2<float>* UVs = TempUVs.GetData();
+			uint8* pVertices = pData;
+			FVector2f* UVs = TempUVs.GetData();
 
 			for (int32 VertexIndex = 0; VertexIndex < NumVertices; ++VertexIndex)
 			{
-				vec2<float>* UV = &TempUVs[VertexIndex];
+				FVector2f* UV = &TempUVs[VertexIndex];
 
 				uint16& LayoutBlockIndex = LayoutData[VertexIndex];
 				if (BlockIds.IsValidIndex(LayoutBlockIndex))
@@ -670,14 +670,14 @@ namespace mu
 				// Copy UVs
 				if (format == MBF_FLOAT32)
 				{
-					vec2<float>* pUV = (vec2<float>*)pVertices;
+					FVector2f* pUV = reinterpret_cast<FVector2f*>(pVertices);
 					*pUV = *UV;
 				}
 				else if (format == MBF_FLOAT16)
 				{
-					float16* pUV = (float16*)pVertices;
-					pUV[0] = floatToHalf((*UV)[0]);
-					pUV[1] = floatToHalf((*UV)[1]);
+					FFloat16* pUV = reinterpret_cast<FFloat16*>(pVertices);
+					pUV[0] = FFloat16((*UV)[0]);
+					pUV[1] = FFloat16((*UV)[1]);
 				}
 
 				pVertices += elemSize;
@@ -1533,9 +1533,9 @@ namespace mu
         Ptr<ASTOpMeshTransform> op = new ASTOpMeshTransform();
 
         // Base
-        if (node.m_pSource)
+        if (node.Source)
         {
-            GenerateMesh(InOptions, OutResult, node.m_pSource);
+            GenerateMesh(InOptions, OutResult, node.Source);
             op->source = OutResult.meshOp;
         }
         else
@@ -1545,7 +1545,7 @@ namespace mu
                 ELMT_ERROR, node.m_errorContext);
         }
 
-        op->matrix = node.m_transform;
+        op->matrix = node.Transform;
 
         OutResult.meshOp = op;
     }
@@ -1579,20 +1579,20 @@ namespace mu
             op->morphShape.type = (uint8_t)FShape::Type::Ellipse;
             op->morphShape.position = node.m_origin;
             op->morphShape.up = node.m_normal;
-            op->morphShape.size = vec3f(node.m_radius1, node.m_radius2, node.m_rotation); // TODO: Move rotation to ellipse rotation reference base instead of passing it directly
+            op->morphShape.size = FVector3f(node.m_radius1, node.m_radius2, node.m_rotation); // TODO: Move rotation to ellipse rotation reference base instead of passing it directly
 
                                                                                       // Generate a "side" vector.
                                                                                       // \todo: make generic and move to the vector class
             {
                 // Generate vector perpendicular to normal for ellipse rotation reference base
-                vec3f aux_base(0.f, 1.f, 0.f);
+				FVector3f aux_base(0.f, 1.f, 0.f);
 
-                if (fabs(dot(node.m_normal, aux_base)) > 0.95f)
+                if (fabs(FVector3f::DotProduct(node.m_normal, aux_base)) > 0.95f)
                 {
-                    aux_base = vec3f(0.f, 0.f, 1.f);
+                    aux_base = FVector3f(0.f, 0.f, 1.f);
                 }
 
-                op->morphShape.side = cross(node.m_normal, aux_base);
+                op->morphShape.side = FVector3f::CrossProduct(node.m_normal, aux_base);
             }
         }
 
diff --git a/Engine/Plugins/Experimental/Mutable/Source/MutableTools/Private/MuT/NodeColourConstant.cpp b/Engine/Plugins/Experimental/Mutable/Source/MutableTools/Private/MuT/NodeColourConstant.cpp
--- a/Engine/Plugins/Experimental/Mutable/Source/MutableTools/Private/MuT/NodeColourConstant.cpp
+++ b/Engine/Plugins/Experimental/Mutable/Source/MutableTools/Private/MuT/NodeColourConstant.cpp
@@ -28,7 +28,7 @@ namespace mu
 	NodeColourConstantPtr NodeColourConstant::OldStaticUnserialise(InputArchive& arch)
 	{
 		NodeColourConstantPtr pResult = new NodeColourConstant();
-		vec3<float> Value;
+		FVector3f Value;
 
 		arch >> Value;
 		pResult->GetPrivate()->m_value = FVector4f(Value[0], Value[1], Value[2], 1.0f);
diff --git a/Engine/Plugins/Experimental/Mutable/Source/MutableTools/Private/MuT/NodeColourParameter.cpp b/Engine/Plugins/Experimental/Mutable/Source/MutableTools/Private/MuT/NodeColourParameter.cpp
--- a/Engine/Plugins/Experimental/Mutable/Source/MutableTools/Private/MuT/NodeColourParameter.cpp
+++ b/Engine/Plugins/Experimental/Mutable/Source/MutableTools/Private/MuT/NodeColourParameter.cpp
@@ -44,9 +44,9 @@ namespace mu
 
 
 	//---------------------------------------------------------------------------------------------
-	void NodeColourParameter::SetDefaultValue( float R, float G, float B, float A )
+	void NodeColourParameter::SetDefaultValue(FVector4f Value)
 	{
-		m_pD->m_defaultValue = vec4<float>( R, G, B, A );
+		m_pD->m_defaultValue = Value;
 	}
 
 
diff --git a/Engine/Plugins/Experimental/Mutable/Source/MutableTools/Private/MuT/NodeImageGradientPrivate.h b/Engine/Plugins/Experimental/Mutable/Source/MutableTools/Private/MuT/NodeImageGradientPrivate.h
--- a/Engine/Plugins/Experimental/Mutable/Source/MutableTools/Private/MuT/NodeImageGradientPrivate.h
+++ b/Engine/Plugins/Experimental/Mutable/Source/MutableTools/Private/MuT/NodeImageGradientPrivate.h
@@ -21,12 +21,12 @@ namespace mu
 
 		NodeColourPtr m_pColour0;
 		NodeColourPtr m_pColour1;
-		vec2<int> m_size = { 256,1 };		
+		FIntVector2 m_size = { 256,1 };		
 
 		//!
 		void Serialise( OutputArchive& arch ) const
 		{
-            uint32_t ver = 0;
+            uint32 ver = 0;
 			arch << ver;
 
 			arch << m_pColour0;
@@ -37,7 +37,7 @@ namespace mu
 		//!
 		void Unserialise( InputArchive& arch )
 		{
-            uint32_t ver;
+            uint32 ver;
 			arch >> ver;
 			check(ver==0);
 
diff --git a/Engine/Plugins/Experimental/Mutable/Source/MutableTools/Private/MuT/NodeLayout.cpp b/Engine/Plugins/Experimental/Mutable/Source/MutableTools/Private/MuT/NodeLayout.cpp
--- a/Engine/Plugins/Experimental/Mutable/Source/MutableTools/Private/MuT/NodeLayout.cpp
+++ b/Engine/Plugins/Experimental/Mutable/Source/MutableTools/Private/MuT/NodeLayout.cpp
@@ -180,7 +180,8 @@ namespace mu
 		if (pMesh && layoutIndex >=0 && gridSizeX+gridSizeY>0)
 		{
 			int indexCount = pMesh->GetIndexCount();
-			vector< vec2f > UVs(indexCount*2);
+			TArray< FVector2f > UVs;
+			UVs.SetNumUninitialized(indexCount * 2);
 
 			UntypedMeshBufferIteratorConst indexIt(pMesh->GetIndexBuffers(), MBS_VERTEXINDEX, 0);
 			UntypedMeshBufferIteratorConst texIt(pMesh->GetVertexBuffers(), MBS_TEXCOORDS, layoutIndex);
@@ -229,12 +230,12 @@ namespace mu
 			layout->SetMaxGridSize(gridSizeX, gridSizeY);
 			layout->SetLayoutPackingStrategy(EPackStrategy::RESIZABLE_LAYOUT);
 			
-			vector<box<vec2<int>>> blocks;
+			TArray<box<FIntVector2>> blocks;
 			
 			//Generating blocks
 			for (int i = 0; i < indexCount; ++i)
 			{
-				vec2<int> a, b;
+				FIntVector2 a, b;
 			
 				a[0] = (int)floor(UVs[i * 2][0] * gridSizeX);
 				a[1] = (int)floor(UVs[i * 2][1] * gridSizeY);
@@ -253,7 +254,7 @@ namespace mu
 				{
 					bool contains = false;
 			
-					for (size_t it = 0; it < blocks.size(); ++it)
+					for (int32 it = 0; it < blocks.Num(); ++it)
 					{
 						if (blocks[it].Contains(a) || blocks[it].Contains(b))
 						{
@@ -264,11 +265,11 @@ namespace mu
 					//There is no block that contains them 
 					if (!contains)
 					{
-						box<vec2<int>> currBlock;
+						box<FIntVector2> currBlock;
 						currBlock.min = a;
-						currBlock.size = vec2<int>(1, 1);
+						currBlock.size = FIntVector2(1, 1);
 			
-						blocks.push_back(currBlock);
+						blocks.Add(currBlock);
 					}
 				}
 				else //they are in different blocks
@@ -277,7 +278,7 @@ namespace mu
 					int idxB = -1;
 					
 					//Getting the blocks that contain them
-					for (size_t it = 0; it < blocks.size(); ++it)
+					for (int32 it = 0; it < blocks.Num(); ++it)
 					{
 						if (blocks[it].Contains(a))
 						{
@@ -292,41 +293,41 @@ namespace mu
 					//The blocks are not the same
 					if (idxA != idxB)
 					{
-						box<vec2<int>> currBlock;
+						box<FIntVector2> currBlock;
 						
 						//One of the blocks doesn't exist
 						if (idxA != -1 && idxB == -1)
 						{
 							currBlock.min = b;
-							currBlock.size = vec2<int>(1, 1);
+							currBlock.size = FIntVector2(1, 1);
 							blocks[idxA].Bound(currBlock);
 						}
 						else if (idxB != -1 && idxA == -1)
 						{
 							currBlock.min = a;
-							currBlock.size = vec2<int>(1, 1);
+							currBlock.size = FIntVector2(1, 1);
 							blocks[idxB].Bound(currBlock);
 						}
 						else //Both exist
 						{
 							blocks[idxA].Bound(blocks[idxB]);
-							blocks.erase(blocks.begin() + idxB);
+							blocks.RemoveAt(idxB);
 						}
 					}
 					else //the blocks doesn't exist
 					{
 						if (idxA == -1)
 						{
-							box<vec2<int>> currBlockA;
-							box<vec2<int>> currBlockB;
+							box<FIntVector2> currBlockA;
+							box<FIntVector2> currBlockB;
 			
 							currBlockA.min = a;
 							currBlockB.min = b;
-							currBlockA.size = vec2<int>(1, 1);
-							currBlockB.size = vec2<int>(1, 1);
+							currBlockA.size = FIntVector2(1, 1);
+							currBlockB.size = FIntVector2(1, 1);
 			
 							currBlockA.Bound(currBlockB);
-							blocks.push_back(currBlockA);
+							blocks.Add(currBlockA);
 						}
 					}
 				}
@@ -339,14 +340,14 @@ namespace mu
 			{
 				intersections = false;
 			
-				for (size_t i = 0; !intersections && i < blocks.size(); ++i)
+				for (int32 i = 0; !intersections && i < blocks.Num(); ++i)
 				{
-					for (size_t j = 0; j < blocks.size(); ++j)
+					for (int32 j = 0; j < blocks.Num(); ++j)
 					{
 						if (i != j && blocks[i].IntersectsExclusive(blocks[j]))
 						{
 							blocks[i].Bound(blocks[j]);
-							blocks.erase(blocks.begin() + j);
+							blocks.RemoveAt(j);
 							intersections = true;
 							break;
 						}
@@ -354,7 +355,7 @@ namespace mu
 				}
 			}
 			
-			int numBlocks = (int)blocks.size();
+			int32 numBlocks = blocks.Num();
 			
 			//Generating layout blocks
 			if (numBlocks > 0)
diff --git a/Engine/Plugins/Experimental/Mutable/Source/MutableTools/Private/MuT/NodeMeshClipMorphPlane.cpp b/Engine/Plugins/Experimental/Mutable/Source/MutableTools/Private/MuT/NodeMeshClipMorphPlane.cpp
--- a/Engine/Plugins/Experimental/Mutable/Source/MutableTools/Private/MuT/NodeMeshClipMorphPlane.cpp
+++ b/Engine/Plugins/Experimental/Mutable/Source/MutableTools/Private/MuT/NodeMeshClipMorphPlane.cpp
@@ -46,8 +46,8 @@ namespace mu
 	//---------------------------------------------------------------------------------------------
 	void NodeMeshClipMorphPlane::SetPlane(float centerX, float centerY, float centerZ, float normalX, float normalY, float normalZ)
 	{
-		m_pD->m_origin = vec3f(centerX, centerY, centerZ);
-		m_pD->m_normal = vec3f(normalX, normalY, normalZ);
+		m_pD->m_origin = FVector3f(centerX, centerY, centerZ);
+		m_pD->m_normal = FVector3f(normalX, normalY, normalZ);
 	}
 
 
@@ -70,8 +70,8 @@ namespace mu
 	void NodeMeshClipMorphPlane::SetVertexSelectionBox(float centerX, float centerY, float centerZ, float radiusX, float radiusY, float radiusZ)
 	{
 		m_pD->m_vertexSelectionType = Private::VS_SHAPE;
-		m_pD->m_selectionBoxOrigin = vec3f(centerX, centerY, centerZ);
-		m_pD->m_selectionBoxRadius = vec3f(radiusX, radiusY, radiusZ);
+		m_pD->m_selectionBoxOrigin = FVector3f(centerX, centerY, centerZ);
+		m_pD->m_selectionBoxRadius = FVector3f(radiusX, radiusY, radiusZ);
 	}
 
 	//---------------------------------------------------------------------------------------------
diff --git a/Engine/Plugins/Experimental/Mutable/Source/MutableTools/Private/MuT/NodeMeshClipMorphPlanePrivate.h b/Engine/Plugins/Experimental/Mutable/Source/MutableTools/Private/MuT/NodeMeshClipMorphPlanePrivate.h
--- a/Engine/Plugins/Experimental/Mutable/Source/MutableTools/Private/MuT/NodeMeshClipMorphPlanePrivate.h
+++ b/Engine/Plugins/Experimental/Mutable/Source/MutableTools/Private/MuT/NodeMeshClipMorphPlanePrivate.h
@@ -38,8 +38,8 @@ namespace mu
 		float m_factor;
 
 		// Ellipse location
-		vec3f m_origin;
-		vec3f m_normal;
+		FVector3f m_origin;
+		FVector3f m_normal;
 		float m_radius1, m_radius2, m_rotation;
 
 		//! Typed of vertex selection
@@ -57,9 +57,9 @@ namespace mu
 		} VERTEX_SELECTION;
 
 		// Vertex selection box 
-        uint8_t m_vertexSelectionType;
-		vec3f m_selectionBoxOrigin;
-		vec3f m_selectionBoxRadius;
+        uint8 m_vertexSelectionType;
+		FVector3f m_selectionBoxOrigin;
+		FVector3f m_selectionBoxRadius;
 		uint16 m_vertexSelectionBone;
 
 		TArray<FString> m_tags;
@@ -71,7 +71,7 @@ namespace mu
 		//!
 		void Serialise( OutputArchive& arch ) const
 		{
-            uint32_t ver = 5;
+            uint32 ver = 5;
 			arch << ver;
 
 			arch << m_pSource;
@@ -93,7 +93,7 @@ namespace mu
 		//!
 		void Unserialise( InputArchive& arch )
 		{
-            uint32_t ver;
+            uint32 ver;
 			arch >> ver;
             check(ver>=3&&ver<=5);
 
diff --git a/Engine/Plugins/Experimental/Mutable/Source/MutableTools/Private/MuT/NodeMeshTransform.cpp b/Engine/Plugins/Experimental/Mutable/Source/MutableTools/Private/MuT/NodeMeshTransform.cpp
--- a/Engine/Plugins/Experimental/Mutable/Source/MutableTools/Private/MuT/NodeMeshTransform.cpp
+++ b/Engine/Plugins/Experimental/Mutable/Source/MutableTools/Private/MuT/NodeMeshTransform.cpp
@@ -34,28 +34,28 @@ namespace mu
 	//---------------------------------------------------------------------------------------------
     NodeMeshPtr NodeMeshTransform::GetSource() const
 	{
-		return m_pD->m_pSource;
+		return m_pD->Source;
 	}
 
 
     //---------------------------------------------------------------------------------------------
     void NodeMeshTransform::SetSource( NodeMesh* p )
     {
-        m_pD->m_pSource = p;
+        m_pD->Source = p;
     }
 
 
     //---------------------------------------------------------------------------------------------
-    void NodeMeshTransform::SetTransform( const float* mat )
+    void NodeMeshTransform::SetTransform( const FMatrix44f& Value )
     {
-        memcpy( &m_pD->m_transform[0][0], mat, 16*sizeof(float) );
+        m_pD->Transform = Value;
     }
 
 
     //---------------------------------------------------------------------------------------------
-    void NodeMeshTransform::GetTransform( float* mat ) const
+	const FMatrix44f& NodeMeshTransform::GetTransform() const
     {
-        memcpy( mat, &m_pD->m_transform[0][0], 16*sizeof(float) );
+        return m_pD->Transform;
     }
 
 
@@ -64,9 +64,9 @@ namespace mu
 	{
 		NodeLayoutPtr pResult;
 
-		if ( m_pSource )
+		if ( Source )
 		{
-			NodeMesh::Private* pPrivate = static_cast<NodeMesh::Private*>( m_pSource->GetBasePrivate() );
+			NodeMesh::Private* pPrivate = static_cast<NodeMesh::Private*>( Source->GetBasePrivate() );
 			pResult = pPrivate->GetLayout( index );
 		}
 
diff --git a/Engine/Plugins/Experimental/Mutable/Source/MutableTools/Private/MuT/NodeMeshTransformPrivate.h b/Engine/Plugins/Experimental/Mutable/Source/MutableTools/Private/MuT/NodeMeshTransformPrivate.h
--- a/Engine/Plugins/Experimental/Mutable/Source/MutableTools/Private/MuT/NodeMeshTransformPrivate.h
+++ b/Engine/Plugins/Experimental/Mutable/Source/MutableTools/Private/MuT/NodeMeshTransformPrivate.h
@@ -7,6 +7,7 @@
 #include "MuT/AST.h"
 
 #include "MuR/MutableMath.h"
+#include "MuR/Serialisation.h"
 
 
 namespace mu
@@ -19,28 +20,28 @@ namespace mu
 
 		static FNodeType s_type;
 
-		NodeMeshPtr m_pSource;
-        mat4f m_transform;
+		NodeMeshPtr Source;
+		FMatrix44f Transform;
 
 		//!
 		void Serialise( OutputArchive& arch ) const
 		{
-            uint32_t ver = 0;
+            uint32 ver = 0;
 			arch << ver;
 
-			arch << m_pSource;
-            arch << m_transform;
+			arch << Source;
+            arch << Transform;
 		}
 
 		//!
 		void Unserialise( InputArchive& arch )
 		{
-            uint32_t ver;
+            uint32 ver;
 			arch >> ver;
 			check(ver==0);
 
-			arch >> m_pSource;
-            arch >> m_transform;
+			arch >> Source;
+            arch >> Transform;
 		}
 
 		// NodeMesh::Private interface
diff --git a/Engine/Plugins/Experimental/Mutable/Source/MutableTools/Private/MuT/NodeModifierMeshClipMorphPlane.cpp b/Engine/Plugins/Experimental/Mutable/Source/MutableTools/Private/MuT/NodeModifierMeshClipMorphPlane.cpp
--- a/Engine/Plugins/Experimental/Mutable/Source/MutableTools/Private/MuT/NodeModifierMeshClipMorphPlane.cpp
+++ b/Engine/Plugins/Experimental/Mutable/Source/MutableTools/Private/MuT/NodeModifierMeshClipMorphPlane.cpp
@@ -32,8 +32,8 @@ namespace mu
     void NodeModifierMeshClipMorphPlane::SetPlane( float centerX, float centerY, float centerZ,
                                                    float normalX, float normalY, float normalZ)
 	{
-		m_pD->m_origin = vec3f(centerX, centerY, centerZ);
-		m_pD->m_normal = vec3f(normalX, normalY, normalZ);
+		m_pD->m_origin = FVector3f(centerX, centerY, centerZ);
+		m_pD->m_normal = FVector3f(normalX, normalY, normalZ);
 	}
 
 
@@ -56,8 +56,8 @@ namespace mu
 	void NodeModifierMeshClipMorphPlane::SetVertexSelectionBox(float centerX, float centerY, float centerZ, float radiusX, float radiusY, float radiusZ)
 	{
 		m_pD->m_vertexSelectionType = Private::VS_SHAPE;
-		m_pD->m_selectionBoxOrigin = vec3f(centerX, centerY, centerZ);
-		m_pD->m_selectionBoxRadius = vec3f(radiusX, radiusY, radiusZ);
+		m_pD->m_selectionBoxOrigin = FVector3f(centerX, centerY, centerZ);
+		m_pD->m_selectionBoxRadius = FVector3f(radiusX, radiusY, radiusZ);
 	}
 
 	//---------------------------------------------------------------------------------------------
diff --git a/Engine/Plugins/Experimental/Mutable/Source/MutableTools/Private/MuT/NodeProjectorConstant.cpp b/Engine/Plugins/Experimental/Mutable/Source/MutableTools/Private/MuT/NodeProjectorConstant.cpp
--- a/Engine/Plugins/Experimental/Mutable/Source/MutableTools/Private/MuT/NodeProjectorConstant.cpp
+++ b/Engine/Plugins/Experimental/Mutable/Source/MutableTools/Private/MuT/NodeProjectorConstant.cpp
@@ -30,48 +30,32 @@ namespace mu
 	//---------------------------------------------------------------------------------------------
 	// Own Interface
 	//---------------------------------------------------------------------------------------------
-    void NodeProjectorConstant::GetValue( PROJECTOR_TYPE* pType,
-                                          float* pPosX, float* pPosY, float* pPosZ,
-										  float* pDirX, float* pDirY, float* pDirZ,
-										  float* pUpX, float* pUpY, float* pUpZ,
-                                          float* pScaleU, float* pScaleV, float* pScaleW,
-                                          float* pProjectionAngle ) const
+    void NodeProjectorConstant::GetValue( PROJECTOR_TYPE* OutType,
+		FVector3f* OutPos,
+		FVector3f* OutDir,
+		FVector3f* OutUp,
+		FVector3f* OutScale,
+		float* OutProjectionAngle ) const
 	{
-        if (pType) *pType = m_pD->m_type;
+        if (OutType) *OutType = m_pD->m_type;
 
-        if (pPosX) *pPosX = m_pD->m_position[0];
-        if (pPosY) *pPosY = m_pD->m_position[1];
-        if (pPosZ) *pPosZ = m_pD->m_position[2];
+        if (OutPos) *OutPos = m_pD->m_position;
+        if (OutDir) *OutDir = m_pD->m_direction;
+        if (OutUp) *OutUp = m_pD->m_up;
+        if (OutScale) *OutScale = m_pD->m_scale;
 
-        if (pDirX) *pDirX = m_pD->m_direction[0];
-        if (pDirY) *pDirY = m_pD->m_direction[1];
-        if (pDirZ) *pDirZ = m_pD->m_direction[2];
-
-        if (pUpX) *pUpX = m_pD->m_up[0];
-        if (pUpY) *pUpY = m_pD->m_up[1];
-        if (pUpZ) *pUpZ = m_pD->m_up[2];
-
-        if (pScaleU) *pScaleU = m_pD->m_scale[0];
-        if (pScaleV) *pScaleV = m_pD->m_scale[1];
-        if (pScaleW) *pScaleW = m_pD->m_scale[2];
-
-        if (pProjectionAngle) *pProjectionAngle = m_pD->m_projectionAngle;
+        if (OutProjectionAngle) *OutProjectionAngle = m_pD->m_projectionAngle;
     }
 
 
 	//---------------------------------------------------------------------------------------------
-    void NodeProjectorConstant::SetValue( PROJECTOR_TYPE type,
-                                          float posX, float posY, float posZ,
-										  float dirX, float dirY, float dirZ,
-										  float upX, float upY, float upZ,
-                                          float scaleU, float scaleV, float scaleW,
-                                          float projectionAngle )
+    void NodeProjectorConstant::SetValue( PROJECTOR_TYPE type, FVector3f pos, FVector3f dir, FVector3f up, FVector3f scale, float projectionAngle )
 	{
         m_pD->m_type = type;
-		m_pD->m_position = vec3<float>( posX, posY, posZ );
-		m_pD->m_direction = vec3<float>( dirX, dirY, dirZ );
-		m_pD->m_up = vec3<float>( upX, upY, upZ );
-        m_pD->m_scale = vec3<float>( scaleU, scaleV, scaleW );
+		m_pD->m_position = pos;
+		m_pD->m_direction = dir;
+		m_pD->m_up = up;
+        m_pD->m_scale = scale;
         m_pD->m_projectionAngle = projectionAngle;
 	}
 
diff --git a/Engine/Plugins/Experimental/Mutable/Source/MutableTools/Private/MuT/NodeProjectorParameter.cpp b/Engine/Plugins/Experimental/Mutable/Source/MutableTools/Private/MuT/NodeProjectorParameter.cpp
--- a/Engine/Plugins/Experimental/Mutable/Source/MutableTools/Private/MuT/NodeProjectorParameter.cpp
+++ b/Engine/Plugins/Experimental/Mutable/Source/MutableTools/Private/MuT/NodeProjectorParameter.cpp
@@ -53,17 +53,17 @@ namespace mu
 
 	//---------------------------------------------------------------------------------------------
     void NodeProjectorParameter::SetDefaultValue( PROJECTOR_TYPE type,
-                                                  float posX, float posY, float posZ,
-												  float dirX, float dirY, float dirZ,
-												  float upX, float upY, float upZ,
-                                                  float scaleU, float scaleV, float scaleW,
-                                                  float projectionAngle )
+		FVector3f pos,
+		FVector3f dir,
+		FVector3f up,
+		FVector3f scale,
+		float projectionAngle )
 	{
         m_pD->m_type = type;
-        m_pD->m_position = vec3<float>( posX, posY, posZ );
-		m_pD->m_direction = vec3<float>( dirX, dirY, dirZ );
-		m_pD->m_up = vec3<float>( upX, upY, upZ );
-        m_pD->m_scale = vec3<float>( scaleU, scaleV, scaleW );
+        m_pD->m_position = pos;
+		m_pD->m_direction = dir;
+		m_pD->m_up = up;
+        m_pD->m_scale = scale;
         m_pD->m_projectionAngle = projectionAngle;
     }
 
diff --git a/Engine/Plugins/Experimental/Mutable/Source/MutableTools/Public/MuT/NodeColourParameter.h b/Engine/Plugins/Experimental/Mutable/Source/MutableTools/Public/MuT/NodeColourParameter.h
--- a/Engine/Plugins/Experimental/Mutable/Source/MutableTools/Public/MuT/NodeColourParameter.h
+++ b/Engine/Plugins/Experimental/Mutable/Source/MutableTools/Public/MuT/NodeColourParameter.h
@@ -54,7 +54,7 @@ namespace mu
 		void SetUid( const FString& );
 
 		//! Set the default value of the parameter.
-		void SetDefaultValue( float R, float G, float B, float A );
+		void SetDefaultValue( FVector4f );
 
         //! Set the number of ranges (dimensions) for this parameter.
         //! By default a parameter has 0 ranges, meaning it only has one value.
diff --git a/Engine/Plugins/Experimental/Mutable/Source/MutableTools/Public/MuT/NodeMeshTransform.h b/Engine/Plugins/Experimental/Mutable/Source/MutableTools/Public/MuT/NodeMeshTransform.h
--- a/Engine/Plugins/Experimental/Mutable/Source/MutableTools/Public/MuT/NodeMeshTransform.h
+++ b/Engine/Plugins/Experimental/Mutable/Source/MutableTools/Public/MuT/NodeMeshTransform.h
@@ -6,7 +6,7 @@
 #include "MuR/RefCounted.h"
 #include "MuT/Node.h"
 #include "MuT/NodeMesh.h"
-
+#include "Math/MathFwd.h"
 
 namespace mu
 {
@@ -49,8 +49,8 @@ namespace mu
 		void SetSource( NodeMesh* );
 
         //! \param pMat4x4 pointer to 16 floats representing the 4x4 matrix to apply
-        void SetTransform( const float* pMat4x4 );
-        void GetTransform( float* pMat4x4 ) const;
+        void SetTransform(const FMatrix44f&);
+		const FMatrix44f& GetTransform() const;
 
 		//-----------------------------------------------------------------------------------------
 		// Interface pattern
diff --git a/Engine/Plugins/Experimental/Mutable/Source/MutableTools/Public/MuT/NodeProjector.h b/Engine/Plugins/Experimental/Mutable/Source/MutableTools/Public/MuT/NodeProjector.h
--- a/Engine/Plugins/Experimental/Mutable/Source/MutableTools/Public/MuT/NodeProjector.h
+++ b/Engine/Plugins/Experimental/Mutable/Source/MutableTools/Public/MuT/NodeProjector.h
@@ -104,20 +104,15 @@ namespace mu
 		//-----------------------------------------------------------------------------------------
 
 		//! Get the value that this node returns
-        void GetValue( PROJECTOR_TYPE* pType,
-                       float* pPosX, float* pPosY, float* pPosZ,
-					   float* pDirX, float* pDirY, float* pDirZ,
-					   float* pUpX, float* pUpY, float *pUpZ,
-                       float* pScaleU, float* pScaleV, float* pScaleW,
-                       float* pProjectionAngle ) const;
+        void GetValue( PROJECTOR_TYPE* OutType,
+			FVector3f* OutPos,
+			FVector3f* OutDir,
+			FVector3f* OutUp,
+			FVector3f* OutScaleU,
+			float* OutProjectionAngle ) const;
 
 		//! Set the value to be returned by this node
-        void SetValue( PROJECTOR_TYPE type,
-                       float posX, float posY, float posZ,
-					   float dirX, float dirY, float dirZ,
-					   float upX, float upY, float upZ,
-                       float scaleU, float scaleV, float scaleW,
-                       float projectionAngle );
+        void SetValue( PROJECTOR_TYPE type, FVector3f pos, FVector3f dir, FVector3f up, FVector3f scale, float projectionAngle);
 
 		//-----------------------------------------------------------------------------------------
 		// Interface pattern
@@ -175,11 +170,11 @@ namespace mu
 
 		//! Set the default value of the parameter.
         void SetDefaultValue( PROJECTOR_TYPE type,
-                              float posX, float posY, float posZ,
-							  float dirX, float dirY, float dirZ,
-							  float upX, float upY, float upZ,
-                              float scaleU, float scaleV, float scaleW,
-                              float projectionAngle );
+			FVector3f pos,
+			FVector3f dir,
+			FVector3f up,
+			FVector3f scale,
+			float projectionAngle );
 
         //! Set the number of ranges (dimensions) for this parameter.
         //! By default a parameter has 0 ranges, meaning it only has one value.
