[mutable] Fix bugs with "combine" image operation and some combination of input formats.
#rnx

diff --git a/Engine/Plugins/Experimental/Mutable/Source/MutableRuntime/Private/MuR/OpImageLayer.h b/Engine/Plugins/Experimental/Mutable/Source/MutableRuntime/Private/MuR/OpImageLayer.h
--- a/Engine/Plugins/Experimental/Mutable/Source/MutableRuntime/Private/MuR/OpImageLayer.h
+++ b/Engine/Plugins/Experimental/Mutable/Source/MutableRuntime/Private/MuR/OpImageLayer.h
@@ -2173,7 +2173,7 @@ namespace mu
 			const TArrayView<uint8> ResultBatchView =
 					ResultImage->DataStorage.GetBatchLODRange(BatchId, BatchNumElems, NC, 0, NumLODs);
 
-			const int32 NumElems = BaseBatchView.Num();
+			const int32 NumElems = BaseBatchView.Num() / NC;
 
 			check(NumElems == BlendBatchView.Num() / NC);
 			check(NumElems == MaskBatchView.Num() / 1);
@@ -2310,11 +2310,15 @@ namespace mu
 
 		const EImageFormat BaseFormat = BaseImage->GetFormat();
 
+		Ptr<Image> TempMaskImage;
 		if (MaskImage->GetFormat() != EImageFormat::IF_L_UBYTE)
 		{
-			checkf(false, TEXT("Unsupported mask format."));
+			UE_LOG(LogMutableCore, Log, TEXT("Image layer format not supported. A generic one will be used. "));
 
-			BufferLayerCombine<BLEND_FUNC, 1>(ResultImage, BaseImage, BlendedImage, bOnlyFirstLOD);
+			FImageOperator ImOp = FImageOperator::GetDefault(nullptr);
+			constexpr int32 Quality = 4;
+			TempMaskImage = ImOp.ImagePixelFormat( Quality, MaskImage, EImageFormat::IF_L_UBYTE );
+			MaskImage = TempMaskImage.get();
 		}
 
 		if (BaseFormat == EImageFormat::IF_L_UBYTE)
@@ -2332,7 +2336,13 @@ namespace mu
 		}
 		else
 		{
-			checkf(false, TEXT("Unsupported format."));
+			UE_LOG(LogMutableCore, Log, TEXT("Image layer format not supported. A generic one will be used. "));
+
+			FImageOperator ImOp = FImageOperator::GetDefault(nullptr);
+			constexpr int32 Quality = 4;
+			Ptr<Image> TempBaseImage = ImOp.ImagePixelFormat(Quality, BaseImage, EImageFormat::IF_RGBA_UBYTE);
+			Ptr<Image> TempBlededImage = ImOp.ImagePixelFormat(Quality, BlendedImage, EImageFormat::IF_RGBA_UBYTE);
+			BufferLayerCombine<BLEND_FUNC_MASKED, 4>(ResultImage, TempBaseImage.get(), MaskImage, TempBlededImage.get(), bOnlyFirstLOD);
 		}
 	}
 
