ExtensionData::Hash() value wasn't including the Origin in it's hash, as the return value of HashCombine() was ignored

diff --git a/Engine/Plugins/Experimental/Mutable/Source/MutableRuntime/Private/MuR/ExtensionData.cpp b/Engine/Plugins/Experimental/Mutable/Source/MutableRuntime/Private/MuR/ExtensionData.cpp
--- a/Engine/Plugins/Experimental/Mutable/Source/MutableRuntime/Private/MuR/ExtensionData.cpp
+++ b/Engine/Plugins/Experimental/Mutable/Source/MutableRuntime/Private/MuR/ExtensionData.cpp
@@ -33,7 +33,7 @@ ExtensionDataPtr ExtensionData::StaticUnserialise(InputArchive& Archive)
 uint32 ExtensionData::Hash() const
 {
 	uint32 Result = ::GetTypeHash(Index);
-	HashCombine(Result, ::GetTypeHash((uint8)Origin));
+	Result = HashCombine(Result, ::GetTypeHash((uint8)Origin));
 
 	return Result;
 }
