[Mutable] Fix clothing PhysicsAsset metadata for meshes ingested form a table.

#rnx

diff --git a/Engine/Plugins/Experimental/Mutable/Source/CustomizableObjectEditor/Private/MuCOE/GenerateMutableSource/GenerateMutableSourceTable.cpp b/Engine/Plugins/Experimental/Mutable/Source/CustomizableObjectEditor/Private/MuCOE/GenerateMutableSource/GenerateMutableSourceTable.cpp
--- a/Engine/Plugins/Experimental/Mutable/Source/CustomizableObjectEditor/Private/MuCOE/GenerateMutableSource/GenerateMutableSourceTable.cpp
+++ b/Engine/Plugins/Experimental/Mutable/Source/CustomizableObjectEditor/Private/MuCOE/GenerateMutableSource/GenerateMutableSourceTable.cpp
@@ -23,6 +23,7 @@
 #include "Rendering/SkeletalMeshLODModel.h"
 #include "Rendering/SkeletalMeshModel.h"
 #include "MuCOE/CustomizableObjectVersionBridge.h"
+#include "ClothingAsset.h"
 
 #define LOCTEXT_NAMESPACE "CustomizableObjectEditor"
 
@@ -298,6 +299,29 @@ bool FillTableColumn(const UCustomizableObjectNodeTable* TableNode,	mu::TablePtr
 					AddTagToMutableMeshUnique(*MutableMesh, PhysicsAssetTag);
 				}
 
+				if (GenerationContext.Options.bClothingEnabled)
+				{
+					UClothingAssetBase* ClothingAssetBase = SkeletalMesh->GetSectionClothingAsset(LODIndex, SectionIndex);	
+					UClothingAssetCommon* ClothingAssetCommon = Cast<UClothingAssetCommon>(ClothingAssetBase);
+
+					if (ClothingAssetCommon && ClothingAssetCommon->PhysicsAsset)
+					{	
+						int32 AssetIndex = GenerationContext.ContributingClothingAssetsData.IndexOfByPredicate( 
+						[Guid = ClothingAssetBase->GetAssetGuid()](const FCustomizableObjectClothingAssetData& A)
+						{
+							return A.OriginalAssetGuid == Guid;
+						});
+
+						check(AssetIndex != INDEX_NONE);
+
+						GenerationContext.AddParticipatingObject(*ClothingAssetCommon->PhysicsAsset);
+						
+						const int32 PhysicsAssetIndex = GenerationContext.PhysicsAssets.AddUnique(ClothingAssetCommon->PhysicsAsset);
+						FString ClothPhysicsAssetTag = FString::Printf(TEXT("__ClothPA:%d_%d"), AssetIndex, PhysicsAssetIndex);
+						AddTagToMutableMeshUnique(*MutableMesh, ClothPhysicsAssetTag);
+					}
+				}
+
 				if (!AnimBPAssetTag.IsEmpty())
 				{
 					AddTagToMutableMeshUnique(*MutableMesh, AnimBPAssetTag);
