[Mutable] Clean up CodeRunner WaitIssued step.

#rnx

diff --git a/Engine/Plugins/Experimental/Mutable/Source/MutableRuntime/Private/MuR/CodeRunnerTasked.cpp b/Engine/Plugins/Experimental/Mutable/Source/MutableRuntime/Private/MuR/CodeRunnerTasked.cpp
--- a/Engine/Plugins/Experimental/Mutable/Source/MutableRuntime/Private/MuR/CodeRunnerTasked.cpp
+++ b/Engine/Plugins/Experimental/Mutable/Source/MutableRuntime/Private/MuR/CodeRunnerTasked.cpp
@@ -293,11 +293,13 @@ namespace mu
 
     void CodeRunner::Run(TUniquePtr<FProfileContext>&& ProfileContext, bool bForceInlineExecution)
     {
-		MUTABLE_CPUPROFILER_SCOPE(TEXT("CodeRunner_Run"));
+		MUTABLE_CPUPROFILER_SCOPE(CodeRunner_Run);
 
 		check(!RunnerCompletionEvent.IsCompleted());
 
-		const double StartSeconds = FPlatformTime::Seconds();
+		// TODO: Move MaxAllowedTime somewhere else more accessible, maybe a cvar.
+		const FTimespan MaxAllowedTime = FTimespan::FromMilliseconds(2.0); 
+		const FTimespan TimeOut = FTimespan::FromSeconds(FPlatformTime::Seconds()) + MaxAllowedTime;
 
         while(!OpenTasks.IsEmpty() || !ClosedTasks.IsEmpty() || !IssuedTasks.IsEmpty())
         {
@@ -517,38 +519,40 @@ namespace mu
 			// If at this point there is no open op and we haven't finished, we need to wait for an issued op to complete.
 			if (OpenTasks.IsEmpty() && !IssuedTasks.IsEmpty())
 			{
-				MUTABLE_CPUPROFILER_SCOPE(CodeRunner_WaitIssued);
-
-				for (int32 IssuedIndex = 0; IssuedIndex < IssuedTasks.Num(); ++IssuedIndex)
+				if (!bForceInlineExecution)
 				{
-					if (IssuedTasks[IssuedIndex]->Event.IsValid())
+					TArray<UE::Tasks::FTask, TInlineAllocator<8>> IssuedTasksCompletionEvents;
+					IssuedTasksCompletionEvents.Reserve(IssuedTasks.Num());
+
+					for (TSharedPtr<FIssuedTask>& IssuedTask : IssuedTasks)
 					{
-						if (!bForceInlineExecution)
+						if (IssuedTask->Event.IsValid())
 						{
-							TArray<UE::Tasks::FTask, TInlineAllocator<8>> IssuedTasksCompletionEvents;
-							IssuedTasksCompletionEvents.Reserve(IssuedTasks.Num());
-
-							for (TSharedPtr<FIssuedTask>& IssuedTask : IssuedTasks)
-							{
-								IssuedTasksCompletionEvents.Add(IssuedTask->Event);
-							}
+							IssuedTasksCompletionEvents.Add(IssuedTask->Event);
+						}
+					}
 
-							m_pSystem->WorkingMemoryManager.InvalidateRunnerThread();
+					m_pSystem->WorkingMemoryManager.InvalidateRunnerThread();
 
-							UE::Tasks::Launch(TEXT("CodeRunnerFromIssuedTasksTask"),
-								[Runner = AsShared(), ProfileContext = MoveTemp(ProfileContext)]() mutable
-								{
-									Runner->m_pSystem->WorkingMemoryManager.ResetRunnerThread();
+					UE::Tasks::Launch(TEXT("CodeRunnerFromIssuedTasksTask"),
+						[Runner = AsShared(), ProfileContext = MoveTemp(ProfileContext)]() mutable
+						{
+							Runner->m_pSystem->WorkingMemoryManager.ResetRunnerThread();
 
-									constexpr bool bForceInlineExecution = false;
-									Runner->Run(MoveTemp(ProfileContext), bForceInlineExecution);
-								},
-								UE::Tasks::Prerequisites(UE::Tasks::Any(IssuedTasksCompletionEvents)),
-								UE::Tasks::ETaskPriority::Inherit);
-							
-							return;
-						}
-						else
+							constexpr bool bForceInlineExecution = false;
+							Runner->Run(MoveTemp(ProfileContext), bForceInlineExecution);
+						},
+						UE::Tasks::Prerequisites(UE::Tasks::Any(IssuedTasksCompletionEvents)),
+						UE::Tasks::ETaskPriority::Inherit);
+					
+					return;
+				}	
+				else
+				{
+					MUTABLE_CPUPROFILER_SCOPE(CodeRunner_WaitIssued);
+					for (int32 IssuedIndex = 0; IssuedIndex < IssuedTasks.Num(); ++IssuedIndex)
+					{
+						if (IssuedTasks[IssuedIndex]->Event.IsValid())
 						{
 							if (CVarTaskGraphBusyWait->GetBool())
 							{
@@ -567,10 +571,7 @@ namespace mu
 
 			if (!bForceInlineExecution)
 			{
-				// TODO: Move the timeout somewhere else more accessible, maybe a cvar.
-				const FTimespan Timeout = FTimespan::FromMilliseconds(2.0);
-				const FTimespan Duration = FTimespan::FromSeconds(FPlatformTime::Seconds() - StartSeconds);
-				if (Duration > Timeout)
+				if (FTimespan::FromSeconds(FPlatformTime::Seconds()) > TimeOut)
 				{
 					m_pSystem->WorkingMemoryManager.InvalidateRunnerThread();
 
@@ -1359,7 +1360,7 @@ namespace mu
 
 		check(StartLevel >= 0);
 
-		FMipmapGenerationSettings Settings{};
+		FMipmapGenerationSettings Settings{Args.FilterType, Args.AddressMode};
 		FImageOperator ImOp = FImageOperator::GetDefault(ImagePixelFormatFunc);
 		ImOp.ImageMipmap(Scratch, ImageCompressionQuality, Result.get(), Result.get(), StartLevel, Result->GetLODCount(), Settings);
 	}
