[Mutable] Split Task_Mutable_GetMeshes into multiple tasks.

#rnx

diff --git a/Engine/Plugins/Experimental/Mutable/Source/CustomizableObject/Internal/MuCO/CustomizableObjectSystemPrivate.h b/Engine/Plugins/Experimental/Mutable/Source/CustomizableObject/Internal/MuCO/CustomizableObjectSystemPrivate.h
--- a/Engine/Plugins/Experimental/Mutable/Source/CustomizableObject/Internal/MuCO/CustomizableObjectSystemPrivate.h
+++ b/Engine/Plugins/Experimental/Mutable/Source/CustomizableObject/Internal/MuCO/CustomizableObjectSystemPrivate.h
@@ -485,7 +485,8 @@ private:
 public:
 	/** Instance parameters at the time of the operation request. */
 	mu::ParametersPtr Parameters; 
-	
+	mu::Ptr<mu::System> MutableSystem;
+
 	bool bOnlyUpdateIfNotGenerated = false;
 	bool bIgnoreCloseDist = false;
 	bool bForceHighPriority = false;
diff --git a/Engine/Plugins/Experimental/Mutable/Source/CustomizableObject/Private/MuCO/BusyWaits_Deprecated.cpp b/Engine/Plugins/Experimental/Mutable/Source/CustomizableObject/Private/MuCO/BusyWaits_Deprecated.cpp
--- a/Engine/Plugins/Experimental/Mutable/Source/CustomizableObject/Private/MuCO/BusyWaits_Deprecated.cpp
+++ b/Engine/Plugins/Experimental/Mutable/Source/CustomizableObject/Private/MuCO/BusyWaits_Deprecated.cpp
@@ -10,7 +10,7 @@
 #include "MuR/MutableTrace.h"
 
 
-namespace impl_deprecated
+namespace CustomizableObjectSystem::ImplDeprecated
 {
 	void Subtask_Mutable_GetImages(const TSharedRef<FUpdateContextPrivate>& OperationData)
 	{
@@ -135,7 +135,7 @@ namespace impl_deprecated
 	}
 
 
-	void Task_Mutable_Update_GetImages(const TSharedRef<FUpdateContextPrivate>& OperationData)
+	void Task_Mutable_GetImages(const TSharedRef<FUpdateContextPrivate>& OperationData)
 	{
 		MUTABLE_CPUPROFILER_SCOPE(Task_Mutable_GetImages)
 		FMutableScopeTimer Timer(OperationData->TaskGetImagesTime);
@@ -156,7 +156,7 @@ namespace impl_deprecated
 	}
 
 
-	void Subtask_Mutable_BeginUpdate_GetMesh(const TSharedRef<FUpdateContextPrivate>& OperationData, TSharedPtr<mu::Model> Model)
+	void Subtask_Mutable_BeginUpdate_GetMesh(const TSharedRef<FUpdateContextPrivate>& OperationData)
 	{
 		MUTABLE_CPUPROFILER_SCOPE(Subtask_Mutable_BeginUpdate_GetMesh)
 
@@ -396,17 +396,16 @@ namespace impl_deprecated
 	}
 
 
-	void Task_Mutable_Update_GetMesh(const TSharedRef<FUpdateContextPrivate>& OperationData,
-		const TSharedPtr<mu::Model>& Model)
+	void Task_Mutable_GetMeshes(const TSharedRef<FUpdateContextPrivate>& OperationData)
 	{
-		MUTABLE_CPUPROFILER_SCOPE(Task_Mutable_Update_GetMesh)
+		MUTABLE_CPUPROFILER_SCOPE(Task_Mutable_GetMesh)
 		FMutableScopeTimer Timer(OperationData->TaskGetMeshTime);
 			
 #if WITH_EDITOR
 		const uint32 StartCycles = FPlatformTime::Cycles();
 #endif
 
-		Subtask_Mutable_BeginUpdate_GetMesh(OperationData, Model);
+		Subtask_Mutable_BeginUpdate_GetMesh(OperationData);
 
 		// TODO: Not strictly mutable: move to another worker thread task to free mutable access?
 		impl::Subtask_Mutable_PrepareSkeletonData(OperationData);
diff --git a/Engine/Plugins/Experimental/Mutable/Source/CustomizableObject/Private/MuCO/BusyWaits_Deprecated.h b/Engine/Plugins/Experimental/Mutable/Source/CustomizableObject/Private/MuCO/BusyWaits_Deprecated.h
--- a/Engine/Plugins/Experimental/Mutable/Source/CustomizableObject/Private/MuCO/BusyWaits_Deprecated.h
+++ b/Engine/Plugins/Experimental/Mutable/Source/CustomizableObject/Private/MuCO/BusyWaits_Deprecated.h
@@ -4,6 +4,8 @@
 // THIS FILE SHOULD ONLY CONTAIN CODE BEING DEPRECATED BY UE-207754
 // ---------------------------------------------------------------
 
+#pragma once
+
 #include "Templates/SharedPointer.h"
 
 class FUpdateContextPrivate;
@@ -14,22 +16,22 @@ namespace mu
 }
 
 
-namespace impl_deprecated
+namespace CustomizableObjectSystem::ImplDeprecated
 {
 	// This runs in the mutable thread.
 	void Subtask_Mutable_GetImages(const TSharedRef<FUpdateContextPrivate>& OperationData);
 
 	// This runs in a worker thread.
-	void Task_Mutable_Update_GetImages(const TSharedRef<FUpdateContextPrivate>& OperationData);
+	void Task_Mutable_GetImages(const TSharedRef<FUpdateContextPrivate>& OperationData);
 
 	// This runs in the mutable thread.
 	void Subtask_Mutable_BeginUpdate_GetMesh(const TSharedRef<FUpdateContextPrivate>& OperationData, TSharedPtr<mu::Model> Model);
 
 	// This runs in a worker thread.
-	void Task_Mutable_Update_GetMesh(const TSharedRef<FUpdateContextPrivate>& OperationData, const TSharedPtr<mu::Model>& Model);
-
+	void Task_Mutable_GetMeshes(const TSharedRef<FUpdateContextPrivate>& OperationData);
 }
 
+
 namespace CustomizableObjectMipDataProvider::ImplDeprecated
 {
 	// This runs in the mutable thread.
diff --git a/Engine/Plugins/Experimental/Mutable/Source/CustomizableObject/Private/MuCO/CustomizableObjectSystem.cpp b/Engine/Plugins/Experimental/Mutable/Source/CustomizableObject/Private/MuCO/CustomizableObjectSystem.cpp
--- a/Engine/Plugins/Experimental/Mutable/Source/CustomizableObject/Private/MuCO/CustomizableObjectSystem.cpp
+++ b/Engine/Plugins/Experimental/Mutable/Source/CustomizableObject/Private/MuCO/CustomizableObjectSystem.cpp
@@ -56,7 +56,8 @@ class UAnimInstance;
 DECLARE_CYCLE_STAT(TEXT("MutablePendingRelease Time"), STAT_MutablePendingRelease, STATGROUP_Game);
 DECLARE_CYCLE_STAT(TEXT("MutableTask"), STAT_MutableTask, STATGROUP_Game);
 
-#define UE_MUTABLE_UPDATE_REGION		TEXT("Mutable Update")
+#define UE_MUTABLE_UPDATE_REGION TEXT("Mutable Update")
+#define UE_TASK_MUTABLE_GETMESHES_REGION TEXT("Task_Mutable_GetMeshes")
 
 
 UCustomizableObjectSystem* UCustomizableObjectSystemPrivate::SSystem = nullptr;
@@ -118,8 +119,10 @@ TAutoConsoleVariable<bool> CVarEnableNewSplitMutableTask(
 	TEXT("Enables or disables the then new split GetImages and GetMesh tasks that remove BusyWaits."),
 	ECVF_Scalability);
 
+
 int32 UCustomizableObjectSystemPrivate::SkeletalMeshMinLodQualityLevel = -1;
 
+
 static void CVarMutableSinkFunction()
 {
 	if (UCustomizableObjectSystem::IsCreated())
@@ -145,6 +148,8 @@ FUpdateContextPrivate::FUpdateContextPrivate(UCustomizableObjectInstance& InInst
 	Parameters = Descriptor.GetParameters();
 	NumComponents = InInstance.GetCustomizableObject()->GetComponentCount();
 	FirstLODAvailable = InInstance.GetCustomizableObject()->GetPrivate()->GetMinLODIndex();
+	MutableSystem = UCustomizableObjectSystem::GetInstance()->GetPrivate()->MutableSystem;	
+	check(MutableSystem);
 	
 	InInstance.GetCustomizableObject()->GetPrivate()->ApplyStateForcedValuesToParameters(CapturedDescriptor.GetState(), Parameters.get());
 
@@ -1547,7 +1552,38 @@ void UCustomizableObjectSystemPrivate::UpdateMemoryLimit()
 // Naming: Task_<thread>_<description>
 namespace impl
 {
+	struct FGetImageData
+	{
+		int32 ImageIndex;
+		mu::FResourceID ImageID;
+	};
+	
+
+	/** Process the next Image. If there are no more Images, go to the end of the task. */
+	void Task_Mutable_GetMeshes_GetImage_Loop(
+		const TSharedRef<FUpdateContextPrivate>& OperationData,
+		double StartTime,
+		uint32 StartCycles,
+		const TSharedRef<TArray<FGetImageData>>& GetImagesData,
+		int32 GetImageIndex);
+	
 
+	struct FGetMeshData
+	{
+		int32 ComponentIndex;
+		mu::FResourceID MeshID;
+	};
+
+
+	/** Process the next Mesh. If there are no more Meshes, go to the process Images loop. */
+	void Task_Mutable_GetMeshes_GetMesh_Loop(
+		const TSharedRef<FUpdateContextPrivate>& OperationData,
+		double StartTime,
+		uint32 StartCycles,
+		const TSharedRef<TArray<FGetMeshData>>& GetMeshesData,
+		int32 GetMeshIndex);
+
+	
 	void Subtask_Mutable_UpdateParameterRelevancy(const TSharedRef<FUpdateContextPrivate>& OperationData)
 	{
 		MUTABLE_CPUPROFILER_SCOPE(Subtask_Mutable_UpdateParameterRelevancy)
@@ -1560,7 +1596,6 @@ namespace impl
 		// This must run in the mutable thread.
 		check(UCustomizableObjectSystem::GetInstance() != nullptr);
 		check(UCustomizableObjectSystem::GetInstance()->GetPrivate() != nullptr);
-		const mu::Ptr<mu::System> MutableSystem = UCustomizableObjectSystem::GetInstance()->GetPrivate()->MutableSystem;
 
 		// Update the parameter relevancy.
 		{
@@ -1570,7 +1605,7 @@ namespace impl
 
 			TArray<bool> Relevant;
 			Relevant.SetNumZeroed(NumParameters);
-			MutableSystem->GetParameterRelevancy(OperationData->InstanceID, OperationData->Parameters, Relevant.GetData());
+			OperationData->MutableSystem->GetParameterRelevancy(OperationData->InstanceID, OperationData->Parameters, Relevant.GetData());
 
 			for (int32 ParamIndex = 0; ParamIndex < NumParameters; ++ParamIndex)
 			{
@@ -1749,11 +1784,415 @@ namespace impl
 	}
 
 
+	/** End of the GetMeshes tasks. */
+	void Task_Mutable_GetMeshes_End(
+		const TSharedRef<FUpdateContextPrivate>& OperationData,
+		double StartTime,
+		uint32 StartCycles)
+	{
+		MUTABLE_CPUPROFILER_SCOPE(Task_Mutable_GetMeshes_End)
+
+		// TODO: Not strictly mutable: move to another worker thread task to free mutable access?
+		Subtask_Mutable_PrepareSkeletonData(OperationData);
+
+		if (OperationData->GetCapturedDescriptor().GetBuildParameterRelevancy())
+		{
+			Subtask_Mutable_UpdateParameterRelevancy(OperationData);
+		}
+		else
+		{
+			OperationData->RelevantParametersInProgress.Reset();
+		}
+
+#if WITH_EDITOR
+		const uint32 EndCycles = FPlatformTime::Cycles();
+		OperationData->MutableRuntimeCycles = EndCycles - StartCycles;
+#endif
+		OperationData->TaskGetMeshTime = FPlatformTime::Seconds() - StartTime;
+
+		TRACE_END_REGION(UE_TASK_MUTABLE_GETMESHES_REGION);
+	}
+
+
+	/** TaskGraph task after GetImage has completed. */
+	void Task_Mutable_GetMeshes_GetImage_Post(
+		const TSharedRef<FUpdateContextPrivate>& OperationData,
+		double StartTime,
+		uint32 StartCycles,
+		const TSharedRef<TArray<FGetImageData>>& GetImagesData,
+		int32 GetImageIndex,
+		UE::Tasks::TTask<mu::Ptr<const mu::Image>> GetImageTask)
+	{
+		MUTABLE_CPUPROFILER_SCOPE(Task_Mutable_GetMeshes_GetImage_Post)
+
+		const UCustomizableObjectInstance* Instance = OperationData->Instance.Get();
+		const UCustomizableObject* CustomizableObject = Instance->GetCustomizableObject();
+		const FModelResources& ModelResources = CustomizableObject->GetPrivate()->GetModelResources();
+
+		const int32 ImageIndex = (*GetImagesData)[GetImageIndex].ImageIndex;
+
+		FInstanceUpdateData::FImage& Image = OperationData->InstanceUpdateData.Images[ImageIndex];
+
+		Image.Image = GetImageTask.GetResult();
+		check(Image.Image->IsReference());
+
+		const uint32 ReferenceID = Image.Image->GetReferencedTexture();
+
+		if (ModelResources.PassThroughTextures.IsValidIndex(ReferenceID))
+		{
+			const TSoftObjectPtr<UTexture> Ref = ModelResources.PassThroughTextures[ReferenceID];
+			Instance->GetPrivate()->PassThroughTexturesToLoad.Add(Ref);
+		}
+		else
+		{
+			// internal error.
+			UE_LOG(LogMutable, Error, TEXT("Referenced image [%d] was not stored in the resource array."), ReferenceID);
+		}
+			
+		Task_Mutable_GetMeshes_GetImage_Loop(OperationData, StartTime, StartCycles, GetImagesData, ++GetImageIndex);
+	}
 
 	
+	/** See declaration. */
+	void Task_Mutable_GetMeshes_GetImage_Loop(
+		const TSharedRef<FUpdateContextPrivate>& OperationData,
+		double StartTime,
+		uint32 StartCycles,
+		const TSharedRef<TArray<FGetImageData>>& GetImagesData,
+		int32 GetImageIndex)
+	{
+		MUTABLE_CPUPROFILER_SCOPE(Task_Mutable_GetMesh_GetImages_Loop)
+
+		if (GetImageIndex >= GetImagesData->Num()) 
+		{
+			Task_Mutable_GetMeshes_End(OperationData, StartTime, StartCycles);
+			return;
+		}
+
+		const FGetImageData& ImageData = (*GetImagesData)[GetImageIndex];
+		
+		UE::Tasks::TTask<mu::Ptr<const mu::Image>> GetImageTask = OperationData->MutableSystem->GetImage(OperationData->InstanceID, ImageData.ImageID, 0, 0);
+
+		UE::Tasks::AddNested(UE::Tasks::Launch(TEXT("Task_Mutable_GetMeshes_GetImage_Post"), [=]()
+		{
+			Task_Mutable_GetMeshes_GetImage_Post(OperationData, StartTime, StartCycles, GetImagesData, GetImageIndex, GetImageTask);
+		},
+		GetImageTask));
+	}
+
+
+	/** Gather all GetImages that have to be called. */
+	void Task_Mutable_GetMeshes_GetImages(
+		const TSharedRef<FUpdateContextPrivate>& OperationData,
+		double StartTime,
+		uint32 StartCycles)
+	{
+		MUTABLE_CPUPROFILER_SCOPE(Task_Mutable_GetMeshes_GetImages)
+		
+		const UCustomizableObjectInstance* Instance = OperationData->Instance.Get();
+		const UCustomizableObject* CustomizableObject = Instance->GetCustomizableObject();
+		const FModelResources& ModelResources = CustomizableObject->GetPrivate()->GetModelResources();
+
+		const mu::Instance* MutableInstance = OperationData->MutableInstance;
+
+		TArray<int32> SurfacesSharedId;
+
+		const TSharedRef<TArray<FGetImageData>> GetImagesData = MakeShared<TArray<FGetImageData>>();
+		
+		for (int32 MutableLODIndex = 0; MutableLODIndex < OperationData->NumLODsAvailable; ++MutableLODIndex)
+		{
+			// Skip LODs outside the range we want to generate
+			if (MutableLODIndex < OperationData->GetMinLOD())
+			{
+				continue;
+			}
+
+			const FInstanceUpdateData::FLOD& LOD = OperationData->InstanceUpdateData.LODs[MutableLODIndex];
+
+			for (int32 ComponentIndex = 0; ComponentIndex < LOD.ComponentCount; ++ComponentIndex)
+			{
+				FInstanceUpdateData::FComponent& Component = OperationData->InstanceUpdateData.Components[LOD.FirstComponent + ComponentIndex];
+				Component.FirstSurface = OperationData->InstanceUpdateData.Surfaces.Num();
+				Component.SurfaceCount = 0;
+
+				if (!Component.Mesh)
+				{
+					continue;
+				}
+				
+				// Materials and images
+				const int32 SurfaceCount = Component.Mesh->GetSurfaceCount();
+				for (int32 MeshSurfaceIndex = 0; MeshSurfaceIndex < SurfaceCount; ++MeshSurfaceIndex)
+				{
+					const uint32 SurfaceId = Component.Mesh->GetSurfaceId(MeshSurfaceIndex);
+					const int32 InstanceSurfaceIndex = MutableInstance->FindSurfaceById(MutableLODIndex, ComponentIndex, SurfaceId);
+					check(Component.Mesh->GetVertexCount() > 0 || InstanceSurfaceIndex >= 0);
+
+					if (InstanceSurfaceIndex >= 0)
+					{
+						int32 BaseSurfaceIndex = InstanceSurfaceIndex;
+						int32 BaseLODIndex = MutableLODIndex;
+
+						OperationData->InstanceUpdateData.Surfaces.Push({});
+						FInstanceUpdateData::FSurface& Surface = OperationData->InstanceUpdateData.Surfaces.Last();
+						++Component.SurfaceCount;
+
+						// Now Surface.MaterialIndex is decoded from a parameter at the end of this if()
+						Surface.SurfaceId = SurfaceId;
+
+						const int32 SharedSurfaceId = MutableInstance->GetSharedSurfaceId(MutableLODIndex, ComponentIndex, InstanceSurfaceIndex);
+						const int32 SharedSurfaceIndex = SurfacesSharedId.Find(SharedSurfaceId);
+
+						SurfacesSharedId.Add(SharedSurfaceId);
+
+						if (SharedSurfaceId != INDEX_NONE)
+						{
+							if (SharedSurfaceIndex >= 0)
+							{
+								Surface = OperationData->InstanceUpdateData.Surfaces[SharedSurfaceIndex];
+								continue;
+							}
+
+							// Find the first LOD where this surface can be found
+							MutableInstance->FindBaseSurfaceBySharedId(ComponentIndex, SharedSurfaceId, BaseSurfaceIndex, BaseLODIndex);
+
+							Surface.SurfaceId = MutableInstance->GetSurfaceId(BaseLODIndex, ComponentIndex, BaseSurfaceIndex);
+						}
+
+						// Vectors
+						Surface.FirstVector = OperationData->InstanceUpdateData.Vectors.Num();
+						Surface.VectorCount = MutableInstance->GetVectorCount(BaseLODIndex, ComponentIndex, InstanceSurfaceIndex);
+						for (int32 VectorIndex = 0; VectorIndex < Surface.VectorCount; ++VectorIndex)
+						{
+							MUTABLE_CPUPROFILER_SCOPE(GetVector);
+							OperationData->InstanceUpdateData.Vectors.Push({});
+							FInstanceUpdateData::FVector& Vector = OperationData->InstanceUpdateData.Vectors.Last();
+							Vector.Name = MutableInstance->GetVectorName(BaseLODIndex, ComponentIndex, InstanceSurfaceIndex, VectorIndex);
+							Vector.Vector = MutableInstance->GetVector(BaseLODIndex, ComponentIndex, InstanceSurfaceIndex, VectorIndex);
+						}
+
+						// Scalars
+						Surface.FirstScalar = OperationData->InstanceUpdateData.Scalars.Num();
+						Surface.ScalarCount = MutableInstance->GetScalarCount(BaseLODIndex, ComponentIndex, InstanceSurfaceIndex);
+						for (int32 ScalarIndex = 0; ScalarIndex < Surface.ScalarCount; ++ScalarIndex)
+						{
+							MUTABLE_CPUPROFILER_SCOPE(GetScalar)
+
+							const FName ScalarName = MutableInstance->GetScalarName(BaseLODIndex, ComponentIndex, InstanceSurfaceIndex, ScalarIndex);
+							const float ScalarValue = MutableInstance->GetScalar(BaseLODIndex, ComponentIndex, InstanceSurfaceIndex, ScalarIndex);
+							
+							FString EncodingMaterialIdString = "__MutableMaterialId";
+							
+							// Decoding Material Switch from Mutable parameter name
+							if (ScalarName.ToString().Equals(EncodingMaterialIdString))
+							{
+								Surface.MaterialIndex = static_cast<uint32>(ScalarValue);
+							
+								// This parameter is not needed in the final material instance
+								Surface.ScalarCount -= 1;
+							}
+							else
+							{
+								OperationData->InstanceUpdateData.Scalars.Push({ ScalarName, ScalarValue });
+							}
+						}
+
+						// Images
+						Surface.FirstImage = OperationData->InstanceUpdateData.Images.Num();
+						Surface.ImageCount = MutableInstance->GetImageCount(BaseLODIndex, ComponentIndex, InstanceSurfaceIndex);
+						for (int32 ImageIndex = 0; ImageIndex < Surface.ImageCount; ++ImageIndex)
+						{
+							MUTABLE_CPUPROFILER_SCOPE(GetImageId);
+
+							const int32 UpdateDataImageIndex = OperationData->InstanceUpdateData.Images.AddDefaulted();
+							FInstanceUpdateData::FImage& Image = OperationData->InstanceUpdateData.Images.Last();
+							Image.Name = MutableInstance->GetImageName(BaseLODIndex, ComponentIndex, InstanceSurfaceIndex, ImageIndex);
+							Image.ImageID = MutableInstance->GetImageId(BaseLODIndex, ComponentIndex, BaseSurfaceIndex, ImageIndex);
+							Image.FullImageSizeX = 0;
+							Image.FullImageSizeY = 0;
+							Image.BaseLOD = BaseLODIndex;
+							Image.BaseMip = 0;
+
+							FString KeyName = Image.Name.ToString();
+							int32 ImageKey = FCString::Atoi(*KeyName);
+
+							if (ImageKey >= 0 && ImageKey < ModelResources.ImageProperties.Num())
+							{
+								const FMutableModelImageProperties& Props = ModelResources.ImageProperties[ImageKey];
+
+								if (Props.IsPassThrough)
+								{
+									Image.bIsPassThrough = true;
+
+									// Since it's known it's a pass-through texture there is no need to cache or convert it so we can generate it here already.
+									GetImagesData->Add({ UpdateDataImageIndex, Image.ImageID });
+								}
+							}
+							else
+							{
+								// This means the compiled model (maybe coming from derived data) has images that the asset doesn't know about.
+								UE_LOG(LogMutable, Error, TEXT("CustomizableObject derived data out of sync with asset for [%s]. Try recompiling it."), *CustomizableObject->GetName());
+							}
+						}
+					}
+				}
+			}	
+		}
+
+		Task_Mutable_GetMeshes_GetImage_Loop(OperationData, StartTime, StartCycles, GetImagesData, 0);
+	}
+
+
+	/** TaskGraph task after GetMesh has completed. */
+	void Task_MutableGetMeshes_GetMesh_Post(
+		const TSharedRef<FUpdateContextPrivate>& OperationData,
+		double StartTime,
+		uint32 StartCycles,
+		const TSharedRef<TArray<FGetMeshData>>& GetMeshesData,
+		int32 GetMeshIndex,
+		UE::Tasks::TTask<mu::Ptr<const mu::Mesh>> GetMeshTask)
+	{
+		MUTABLE_CPUPROFILER_SCOPE(Task_MutableGetMeshes_GetMesh_Post)
+
+		const int32 ComponentIndex = (*GetMeshesData)[GetMeshIndex].ComponentIndex;
+		FInstanceUpdateData::FComponent& Component = OperationData->InstanceUpdateData.Components[ComponentIndex];
+
+		Component.Mesh = GetMeshTask.GetResult();
+			
+		Task_Mutable_GetMeshes_GetMesh_Loop(OperationData, StartTime, StartCycles, GetMeshesData, ++GetMeshIndex);
+	}
 	
 
+	/** See declaration. */
+	void Task_Mutable_GetMeshes_GetMesh_Loop(
+		const TSharedRef<FUpdateContextPrivate>& OperationData,
+		double StartTime,
+		uint32 StartCycles,
+		const TSharedRef<TArray<FGetMeshData>>& GetMeshesData,
+		int32 GetMeshIndex)
+	{
+		MUTABLE_CPUPROFILER_SCOPE(Task_Mutable_GetMeshes_GetMesh_Loop)
 
+		if (GetMeshIndex >= GetMeshesData->Num())
+		{
+			Task_Mutable_GetMeshes_GetImages(OperationData, StartTime, StartCycles);
+			return;
+		}
+		
+		const FGetMeshData& MeshData = (*GetMeshesData)[GetMeshIndex];
+		UE::Tasks::TTask<mu::Ptr<const mu::Mesh>> GetMeshTask = OperationData->MutableSystem->GetMesh(OperationData->InstanceID, MeshData.MeshID);
+		
+		UE::Tasks::AddNested(UE::Tasks::Launch(TEXT("Task_MutableGetMeshes_GetMesh_Post"), [=]()
+		{
+			Task_MutableGetMeshes_GetMesh_Post(OperationData, StartTime, StartCycles, GetMeshesData, GetMeshIndex, GetMeshTask);
+		},
+		GetMeshTask));
+	}
+
+
+	namespace impl_new
+	{
+		/** Start of the GetMeshes tasks.
+		  * Gathers all GetMeshes that has to be called. */
+		void Task_Mutable_GetMeshes(const TSharedRef<FUpdateContextPrivate>& OperationData)
+		{
+			MUTABLE_CPUPROFILER_SCOPE(Task_Mutable_GetMeshes)
+			TRACE_BEGIN_REGION(UE_TASK_MUTABLE_GETMESHES_REGION);
+
+			const double StartTime = FPlatformTime::Seconds();
+			const uint32 StartCycles = FPlatformTime::Cycles();
+
+			check(OperationData->Parameters);
+			OperationData->InstanceUpdateData.Clear();
+
+			check(UCustomizableObjectSystem::GetInstance() != nullptr);
+			check(UCustomizableObjectSystem::GetInstance()->GetPrivate() != nullptr);
+
+			UCustomizableInstancePrivate* CustomizableObjectInstancePrivateData = OperationData->Instance->GetPrivate();
+
+			CustomizableObjectInstancePrivateData->PassThroughTexturesToLoad.Empty();
+
+			if (OperationData->PixelFormatOverride)
+			{
+				OperationData->MutableSystem->SetImagePixelConversionOverride( OperationData->PixelFormatOverride );
+			}
+
+			if (!OperationData->bUseMeshCache)
+			{
+				CreateMutableInstance(OperationData);
+				FixLODs(OperationData);
+			}
+
+			// Main instance generation step
+			// LOD mask, set to all ones to build  all LODs
+			const mu::Instance* Instance = OperationData->MutableInstance; // TODO GMTFuture remove
+			if (!Instance)
+			{
+				UE_LOG(LogMutable, Warning, TEXT("An Instace update has failed."));
+				Task_Mutable_GetMeshes_End(OperationData, StartTime, StartCycles);
+				return;
+			}
+
+			const TArray<uint16>& RequestedLODs = OperationData->GetRequestedLODs();
+
+			const TSharedRef<TArray<FGetMeshData>> GetMeshesData = MakeShared<TArray<FGetMeshData>>();
+		
+			OperationData->InstanceUpdateData.LODs.SetNum(OperationData->NumLODsAvailable);
+			for (int32 MutableLODIndex = 0; MutableLODIndex < OperationData->NumLODsAvailable; ++MutableLODIndex)
+			{
+				// Skip LODs outside the range we want to generate
+				if (MutableLODIndex < OperationData->GetMinLOD())
+				{
+					continue;
+				}
+
+				FInstanceUpdateData::FLOD& LOD = OperationData->InstanceUpdateData.LODs[MutableLODIndex];
+				LOD.FirstComponent = OperationData->InstanceUpdateData.Components.Num();
+				LOD.ComponentCount = Instance->GetComponentCount(MutableLODIndex);
+
+				for (int32 ComponentIndex = 0; ComponentIndex < LOD.ComponentCount; ++ComponentIndex)
+				{
+					const int32 UpdateDataComponentIndex = OperationData->InstanceUpdateData.Components.AddDefaulted();
+					FInstanceUpdateData::FComponent& Component = OperationData->InstanceUpdateData.Components.Last();
+					Component.Id = Instance->GetComponentId(MutableLODIndex, ComponentIndex);
+
+					const bool bGenerateLOD = RequestedLODs.IsValidIndex(Component.Id) ? RequestedLODs[Component.Id] <= MutableLODIndex : true;
+
+					// Mesh
+					if (Instance->GetMeshCount(MutableLODIndex, ComponentIndex) > 0)
+					{
+						MUTABLE_CPUPROFILER_SCOPE(GetMesh);
+
+						Component.MeshID = Instance->GetMeshId(MutableLODIndex, ComponentIndex, 0);
+
+						if (bGenerateLOD)
+						{
+							Component.bGenerated = true;
+
+							GetMeshesData->Add({ UpdateDataComponentIndex, Component.MeshID});
+						}
+					}
+				}
+			}
+
+			Task_Mutable_GetMeshes_GetMesh_Loop(OperationData, StartTime, StartCycles, GetMeshesData, 0);
+		}
+	}
+
+
+	void Task_Mutable_GetMeshes(const TSharedRef<FUpdateContextPrivate>& OperationData)
+	{
+		if (CVarEnableNewSplitMutableTask->GetBool())
+		{
+			impl_new::Task_Mutable_GetMeshes(OperationData);
+		}
+		else
+		{
+			CustomizableObjectSystem::ImplDeprecated::Task_Mutable_GetMeshes(OperationData);
+		}
+	}
+
+	
 	// This runs in a worker thread.
 	void Task_Mutable_ReleaseInstance(const TSharedRef<FUpdateContextPrivate>& OperationData, mu::Ptr<mu::System> MutableSystem)
 	{
@@ -2089,7 +2528,7 @@ namespace impl
 				TEXT("Task_Mutable_GetImages"),
 				[OperationData]()
 				{
-					impl_deprecated::Task_Mutable_Update_GetImages(OperationData);
+					CustomizableObjectSystem::ImplDeprecated::Task_Mutable_GetImages(OperationData);
 				});
 		}
 
@@ -2155,14 +2594,11 @@ namespace impl
 			}
 		}
 
-		// Task inputs
-		TSharedPtr<mu::Model> Model = CustomizableObject->GetPrivate()->GetModel();
-	
 		UE::Tasks::FTask Dependency = SystemPrivate->MutableTaskGraph.AddMutableThreadTask(
-			TEXT("Task_Mutable_Update_GetMesh"),
-			[Operation, Model]()
+			TEXT("Task_Mutable_GetMeshes"),
+			[Operation]()
 			{
-				impl_deprecated::Task_Mutable_Update_GetMesh(Operation, Model);
+				Task_Mutable_GetMeshes(Operation);
 			});
 
 		SystemPrivate->AddGameThreadTask(
@@ -2430,14 +2866,11 @@ namespace impl
 		}
 		else
 		{
-			// Task inputs
-			TSharedPtr<mu::Model> Model = CustomizableObject->GetPrivate()->GetModel();
-
 			Mutable_GetMeshTask = SystemPrivateData->MutableTaskGraph.AddMutableThreadTask(
-				TEXT("Task_Mutable_Update_GetMesh"),
-				[Operation, Model]()
+				TEXT("Task_Mutable_GetMeshes"),
+				[Operation]()
 				{
-					impl_deprecated::Task_Mutable_Update_GetMesh(Operation, Model);
+					Task_Mutable_GetMeshes(Operation);
 				});
 
 			// Task: Lock cache
