[Mutable] Calculate bones to remove when regenerating the imported model.

#rnx

diff --git a/Engine/Plugins/Experimental/Mutable/Source/CustomizableObject/Private/MuCO/CustomizableObjectInstance.cpp b/Engine/Plugins/Experimental/Mutable/Source/CustomizableObject/Private/MuCO/CustomizableObjectInstance.cpp
--- a/Engine/Plugins/Experimental/Mutable/Source/CustomizableObject/Private/MuCO/CustomizableObjectInstance.cpp
+++ b/Engine/Plugins/Experimental/Mutable/Source/CustomizableObject/Private/MuCO/CustomizableObjectInstance.cpp
@@ -6787,6 +6787,30 @@ TSet<UAssetUserData*> UCustomizableObjectInstance::GetMergedAssetUserData(int32
 
 #if WITH_EDITORONLY_DATA
 
+void CalculateBonesToRemove(const FSkeletalMeshLODRenderData& LODResource, const FReferenceSkeleton& RefSkeleton, TArray<FBoneReference>& OutBonesToRemove)
+{
+	const int32 NumBones = RefSkeleton.GetNum();
+	OutBonesToRemove.Empty(NumBones);
+
+	TArray<bool> RemovedBones;
+	RemovedBones.Init(true, NumBones);
+
+	for (int32 BoneIndex = 0; BoneIndex < NumBones; ++BoneIndex)
+	{
+		if (LODResource.RequiredBones.Find((uint16)BoneIndex) != INDEX_NONE)
+		{
+			RemovedBones[BoneIndex] = false;
+			continue;
+		}
+
+		const int32 ParentIndex = RefSkeleton.GetParentIndex(BoneIndex);
+		if (!RemovedBones.IsValidIndex(ParentIndex) || !RemovedBones[ParentIndex])
+		{
+			OutBonesToRemove.Add(RefSkeleton.GetBoneName(BoneIndex));
+		}
+	}
+}
+
 void UCustomizableInstancePrivate::RegenerateImportedModels()
 {
 	MUTABLE_CPUPROFILER_SCOPE(RegenerateImportedModels);
@@ -6798,8 +6822,8 @@ void UCustomizableInstancePrivate::RegenerateImportedModels()
 			continue;
 		}
 
-		FSkeletalMeshRenderData* SkelResource = SkeletalMesh->GetResourceForRendering();
-		if (!SkelResource || SkelResource->IsInitialized())
+		FSkeletalMeshRenderData* RenderData = SkeletalMesh->GetResourceForRendering();
+		if (!RenderData || RenderData->IsInitialized())
 		{
 			continue;
 		}
@@ -6844,34 +6868,35 @@ void UCustomizableInstancePrivate::RegenerateImportedModels()
 		ImportedModel->LODModels.Empty();
 
 		int32 OriginalIndex = 0;
-		for (int32 LODIndex = 0; LODIndex < SkelResource->LODRenderData.Num(); ++LODIndex)
+		for (int32 LODIndex = 0; LODIndex < RenderData->LODRenderData.Num(); ++LODIndex)
 		{
 			ImportedModel->LODModels.Add(new FSkeletalMeshLODModel());
+			FSkeletalMeshLODModel& LODModel = ImportedModel->LODModels[LODIndex];
 
-			FSkeletalMeshLODRenderData& LODModel = SkelResource->LODRenderData[LODIndex];
+			FSkeletalMeshLODRenderData& LODRenderData = RenderData->LODRenderData[LODIndex];
 			int32 CurrentSectionInitialVertex = 0;
 
-			ImportedModel->LODModels[LODIndex].ActiveBoneIndices = LODModel.ActiveBoneIndices;
-			ImportedModel->LODModels[LODIndex].NumTexCoords = LODModel.GetNumTexCoords();
-			ImportedModel->LODModels[LODIndex].RequiredBones = LODModel.RequiredBones;
-			ImportedModel->LODModels[LODIndex].NumVertices = LODModel.GetNumVertices();
+			LODModel.ActiveBoneIndices = LODRenderData.ActiveBoneIndices;
+			LODModel.NumTexCoords = LODRenderData.GetNumTexCoords();
+			LODModel.RequiredBones = LODRenderData.RequiredBones;
+			LODModel.NumVertices = LODRenderData.GetNumVertices();
 
 			// Indices
-			if (LODModel.MultiSizeIndexContainer.IsIndexBufferValid())
+			if (LODRenderData.MultiSizeIndexContainer.IsIndexBufferValid())
 			{
-				const int32 NumIndices = LODModel.MultiSizeIndexContainer.GetIndexBuffer()->Num();
-				ImportedModel->LODModels[LODIndex].IndexBuffer.SetNum(NumIndices);
+				const int32 NumIndices = LODRenderData.MultiSizeIndexContainer.GetIndexBuffer()->Num();
+				LODModel.IndexBuffer.SetNum(NumIndices);
 				for (int32 Index = 0; Index < NumIndices; ++Index)
 				{
-					ImportedModel->LODModels[LODIndex].IndexBuffer[Index] = LODModel.MultiSizeIndexContainer.GetIndexBuffer()->Get(Index);
+					LODModel.IndexBuffer[Index] = LODRenderData.MultiSizeIndexContainer.GetIndexBuffer()->Get(Index);
 				}
 			}
 
-			ImportedModel->LODModels[LODIndex].Sections.SetNum(LODModel.RenderSections.Num());
+			LODModel.Sections.SetNum(LODRenderData.RenderSections.Num());
 
-			for (int SectionIndex = 0; SectionIndex < LODModel.RenderSections.Num(); ++SectionIndex)
+			for (int SectionIndex = 0; SectionIndex < LODRenderData.RenderSections.Num(); ++SectionIndex)
 			{
-				const FSkelMeshRenderSection& RenderSection = LODModel.RenderSections[SectionIndex];
+				const FSkelMeshRenderSection& RenderSection = LODRenderData.RenderSections[SectionIndex];
 				FSkelMeshSection& ImportedSection = ImportedModel->LODModels[LODIndex].Sections[SectionIndex];
 
 				ImportedSection.CorrespondClothAssetIndex = RenderSection.CorrespondClothAssetIndex;
@@ -6887,30 +6912,30 @@ void UCustomizableInstancePrivate::RegenerateImportedModels()
 				ImportedSection.NumVertices = RenderSection.NumVertices;
 				ImportedSection.SoftVertices.Empty(RenderSection.NumVertices);
 				ImportedSection.SoftVertices.AddUninitialized(RenderSection.NumVertices);
-				ImportedSection.bUse16BitBoneIndex = LODModel.DoesVertexBufferUse16BitBoneIndex();
+				ImportedSection.bUse16BitBoneIndex = LODRenderData.DoesVertexBufferUse16BitBoneIndex();
 
 				for (uint32 i = 0; i < RenderSection.NumVertices; ++i)
 				{
-					const FPositionVertex* PosPtr = static_cast<const FPositionVertex*>(LODModel.StaticVertexBuffers.PositionVertexBuffer.GetVertexData());
+					const FPositionVertex* PosPtr = static_cast<const FPositionVertex*>(LODRenderData.StaticVertexBuffers.PositionVertexBuffer.GetVertexData());
 					PosPtr += (CurrentSectionInitialVertex + i);
 
-					check(!LODModel.StaticVertexBuffers.StaticMeshVertexBuffer.GetUseHighPrecisionTangentBasis());
-					const FPackedNormal* TangentPtr = static_cast<const FPackedNormal*>(LODModel.StaticVertexBuffers.StaticMeshVertexBuffer.GetTangentData());
+					check(!LODRenderData.StaticVertexBuffers.StaticMeshVertexBuffer.GetUseHighPrecisionTangentBasis());
+					const FPackedNormal* TangentPtr = static_cast<const FPackedNormal*>(LODRenderData.StaticVertexBuffers.StaticMeshVertexBuffer.GetTangentData());
 					TangentPtr += ((CurrentSectionInitialVertex + i) * 2);
 
-					check(LODModel.StaticVertexBuffers.StaticMeshVertexBuffer.GetUseFullPrecisionUVs());
+					check(LODRenderData.StaticVertexBuffers.StaticMeshVertexBuffer.GetUseFullPrecisionUVs());
 
 					using UVsVectorType = typename TDecay<decltype(DeclVal<FSoftSkinVertex>().UVs[0])>::Type;
 
-					const UVsVectorType* TexCoordPosPtr = static_cast<const UVsVectorType*>(LODModel.StaticVertexBuffers.StaticMeshVertexBuffer.GetTexCoordData());
-					const uint32 NumTexCoords = LODModel.StaticVertexBuffers.StaticMeshVertexBuffer.GetNumTexCoords();
+					const UVsVectorType* TexCoordPosPtr = static_cast<const UVsVectorType*>(LODRenderData.StaticVertexBuffers.StaticMeshVertexBuffer.GetTexCoordData());
+					const uint32 NumTexCoords = LODRenderData.StaticVertexBuffers.StaticMeshVertexBuffer.GetNumTexCoords();
 					TexCoordPosPtr += ((CurrentSectionInitialVertex + i) * NumTexCoords);
 
 					FSoftSkinVertex& Vertex = ImportedSection.SoftVertices[i];
 					for (int32 j = 0; j < RenderSection.MaxBoneInfluences; ++j)
 					{
-						Vertex.InfluenceBones[j] = LODModel.SkinWeightVertexBuffer.GetBoneIndex(CurrentSectionInitialVertex + i, j);
-						Vertex.InfluenceWeights[j] = LODModel.SkinWeightVertexBuffer.GetBoneWeight(CurrentSectionInitialVertex + i, j);
+						Vertex.InfluenceBones[j] = LODRenderData.SkinWeightVertexBuffer.GetBoneIndex(CurrentSectionInitialVertex + i, j);
+						Vertex.InfluenceWeights[j] = LODRenderData.SkinWeightVertexBuffer.GetBoneWeight(CurrentSectionInitialVertex + i, j);
 					}
 
 					for (int32 j = RenderSection.MaxBoneInfluences; j < MAX_TOTAL_INFLUENCES; ++j)
@@ -6943,6 +6968,9 @@ void UCustomizableInstancePrivate::RegenerateImportedModels()
 				ImportedSection.BaseVertexIndex = RenderSection.BaseVertexIndex;
 				ImportedSection.BoneMap = RenderSection.BoneMap;
 
+				// Add bones to remove
+				CalculateBonesToRemove(LODRenderData, SkeletalMesh->GetRefSkeleton(), SkeletalMesh->GetLODInfo(LODIndex)->BonesToRemove);
+
 				const TArray<int32>& LODMaterialMap = SkeletalMesh->GetLODInfoArray()[LODIndex].LODMaterialMap;
 
 				if (LODMaterialMap.IsValidIndex(RenderSection.MaterialIndex))
@@ -6979,7 +7007,7 @@ void UCustomizableInstancePrivate::RegenerateImportedModels()
 				ImportedSection.MaxBoneInfluences = RenderSection.MaxBoneInfluences;
 				ImportedSection.OriginalDataSectionIndex = OriginalIndex++;
 
-				FSkelMeshSourceSectionUserData& SectionUserData = ImportedModel->LODModels[LODIndex].UserSectionsData.FindOrAdd(ImportedSection.OriginalDataSectionIndex);
+				FSkelMeshSourceSectionUserData& SectionUserData = LODModel.UserSectionsData.FindOrAdd(ImportedSection.OriginalDataSectionIndex);
 				SectionUserData.bCastShadow = RenderSection.bCastShadow;
 				SectionUserData.bDisabled = RenderSection.bDisabled;
 
@@ -6987,7 +7015,7 @@ void UCustomizableInstancePrivate::RegenerateImportedModels()
 				SectionUserData.ClothingData.AssetGuid = RenderSection.ClothingData.AssetGuid;
 				SectionUserData.ClothingData.AssetLodIndex = RenderSection.ClothingData.AssetLodIndex;
 				
-				ImportedModel->LODModels[LODIndex].SyncronizeUserSectionsDataArray();
+				LODModel.SyncronizeUserSectionsDataArray();
 
 				// DDC keys
 				const USkeletalMeshLODSettings* LODSettings = SkeletalMesh->GetLODSettings();
@@ -6997,7 +7025,7 @@ void UCustomizableInstancePrivate::RegenerateImportedModels()
 				FSkeletalMeshLODInfo* LODInfo = SkeletalMesh->GetLODInfo(LODIndex);
 				LODInfo->BuildGUID = LODInfo->ComputeDeriveDataCacheKey(SkeletalMeshLODGroupSettings);
 
-				ImportedModel->LODModels[LODIndex].BuildStringID = ImportedModel->LODModels[LODIndex].GetLODModelDeriveDataKey();
+				LODModel.BuildStringID = LODModel.GetLODModelDeriveDataKey();
 			}
 		}
 	}
