[Mutable] Fix compilation issue in CodeRunnerTasked.cpp

#rnx

diff --git a/Engine/Plugins/Experimental/Mutable/Source/MutableRuntime/Private/MuR/CodeRunnerTasked.cpp b/Engine/Plugins/Experimental/Mutable/Source/MutableRuntime/Private/MuR/CodeRunnerTasked.cpp
--- a/Engine/Plugins/Experimental/Mutable/Source/MutableRuntime/Private/MuR/CodeRunnerTasked.cpp
+++ b/Engine/Plugins/Experimental/Mutable/Source/MutableRuntime/Private/MuR/CodeRunnerTasked.cpp
@@ -1360,7 +1360,7 @@ namespace mu
 
 		check(StartLevel >= 0);
 
-		FMipmapGenerationSettings Settings{Args.FilterType, Args.AddressMode};
+		FMipmapGenerationSettings Settings{};
 		FImageOperator ImOp = FImageOperator::GetDefault(ImagePixelFormatFunc);
 		ImOp.ImageMipmap(Scratch, ImageCompressionQuality, Result.get(), Result.get(), StartLevel, Result->GetLODCount(), Settings);
 	}
