[Mutable] Split Task_Mutable_GetImages into multiple tasks.

#rnx

diff --git a/Engine/Plugins/Experimental/Mutable/Source/CustomizableObject/Internal/MuCO/CustomizableObjectSystemPrivate.h b/Engine/Plugins/Experimental/Mutable/Source/CustomizableObject/Internal/MuCO/CustomizableObjectSystemPrivate.h
--- a/Engine/Plugins/Experimental/Mutable/Source/CustomizableObject/Internal/MuCO/CustomizableObjectSystemPrivate.h
+++ b/Engine/Plugins/Experimental/Mutable/Source/CustomizableObject/Internal/MuCO/CustomizableObjectSystemPrivate.h
@@ -548,10 +548,8 @@ public:
 	int64 UpdateEndPeakBytes = 0;
 	int64 UpdateEndRealPeakBytes = 0;
 	
-#if WITH_EDITOR
 	/** Used for profiling in the editor. */
 	uint32 MutableRuntimeCycles = 0;
-#endif
 
 	/** Hard references to objects. Avoids GC to collect them. */
 	TArray<TObjectPtr<const UObject>> Objects;
diff --git a/Engine/Plugins/Experimental/Mutable/Source/CustomizableObject/Private/MuCO/CustomizableObjectSystem.cpp b/Engine/Plugins/Experimental/Mutable/Source/CustomizableObject/Private/MuCO/CustomizableObjectSystem.cpp
--- a/Engine/Plugins/Experimental/Mutable/Source/CustomizableObject/Private/MuCO/CustomizableObjectSystem.cpp
+++ b/Engine/Plugins/Experimental/Mutable/Source/CustomizableObject/Private/MuCO/CustomizableObjectSystem.cpp
@@ -58,6 +58,7 @@ DECLARE_CYCLE_STAT(TEXT("MutableTask"), STAT_MutableTask, STATGROUP_Game);
 
 #define UE_MUTABLE_UPDATE_REGION TEXT("Mutable Update")
 #define UE_TASK_MUTABLE_GETMESHES_REGION TEXT("Task_Mutable_GetMeshes")
+#define UE_TASK_MUTABLE_GETIMAGES_REGION TEXT("Task_Mutable_GetImages")
 
 
 UCustomizableObjectSystem* UCustomizableObjectSystemPrivate::SSystem = nullptr;
@@ -1583,6 +1584,26 @@ namespace impl
 		const TSharedRef<TArray<FGetMeshData>>& GetMeshesData,
 		int32 GetMeshIndex);
 
+
+	/** Call GetImage.
+	  * Once GetImage is called, the task must end. Following code will be in a subsequent TaskGraph task. */
+	void Task_Mutable_GetImages_GetImage(
+		const TSharedRef<FUpdateContextPrivate>& OperationData,
+		double StartTime,
+		uint32 StartCycles,
+		const TSharedPtr<TArray<mu::FResourceID>>& ImagesInThisInstance,
+		int32 ImageIndex,
+		UE::Tasks::TTask<mu::FImageDesc> GetImageDescTask);
+
+	
+	/** Process the next Image. If there are no more Images, go to the end of the task. */
+	void Task_Mutable_GetImages_Loop(
+		const TSharedRef<FUpdateContextPrivate>& OperationData,
+		double StartTime,
+		uint32 StartCycles,
+		const TSharedPtr<TArray<mu::FResourceID>>& ImagesInThisInstance,
+		int32 ImageIndex);
+
 	
 	void Subtask_Mutable_UpdateParameterRelevancy(const TSharedRef<FUpdateContextPrivate>& OperationData)
 	{
@@ -2193,6 +2214,250 @@ namespace impl
 	}
 
 	
+		/** End of the GetImages tasks. */
+	void Task_Mutable_GetImages_End(const TSharedRef<FUpdateContextPrivate>& OperationData, double StartCycles, double StartTime)
+	{
+		MUTABLE_CPUPROFILER_SCOPE(Task_Mutable_GetImages_End)
+		
+		// TODO: Not strictly mutable: move to another worker thread task to free mutable access?
+		Subtask_Mutable_PrepareTextures(OperationData);
+
+		const uint32 EndCycles = FPlatformTime::Cycles();
+		OperationData->MutableRuntimeCycles += EndCycles - StartCycles;
+		OperationData->TaskGetImagesTime = FPlatformTime::Seconds() - StartTime;
+
+		TRACE_END_REGION(UE_TASK_MUTABLE_GETIMAGES_REGION);
+	}
+
+
+	/** Call GetImageDesc.
+	  * Once GetImageDesc is called, the task must end. Following code will be in a subsequent TaskGraph task. */
+	void Task_Mutable_GetImages_GetImageDesc(
+		const TSharedRef<FUpdateContextPrivate>& OperationData,
+		double StartTime,
+		uint32 StartCycles,
+		const TSharedPtr<TArray<mu::FResourceID>>& ImagesInThisInstance,
+		int32 ImageIndex)
+	{
+		MUTABLE_CPUPROFILER_SCOPE(Task_Mutable_GetImages_GetImageDesc)
+
+		FInstanceUpdateData::FImage& Image = OperationData->InstanceUpdateData.Images[ImageIndex];
+		
+		// This should only be done when using progressive images, since GetImageDesc does some actual processing.
+		UE::Tasks::TTask<mu::FImageDesc> GetImageDescTask = OperationData->MutableSystem->GetImageDesc(OperationData->InstanceID, Image.ImageID);
+		
+		UE::Tasks::AddNested(UE::Tasks::Launch(TEXT("Task_Mutable_GetImages_GetImage"), [=]()
+		{
+			Task_Mutable_GetImages_GetImage(OperationData, StartTime, StartCycles, ImagesInThisInstance, ImageIndex, GetImageDescTask);
+		},
+		GetImageDescTask));
+	}
+
+
+	/** TaskGraph task after GetImage has completed. */
+	void Task_Mutable_GetImages_GetImage_Post(
+		const TSharedRef<FUpdateContextPrivate>& OperationData,
+    	double StartTime,
+    	uint32 StartCycles,
+    	const TSharedPtr<TArray<mu::FResourceID>>& ImagesInThisInstance,
+    	int32 ImageIndex,
+    	UE::Tasks::TTask<mu::Ptr<const mu::Image>> GetImageTask,
+    	int32 MipSizeX,
+		int32 MipSizeY,
+		int32 FullLODContent,
+		int32 MipsToSkip)
+	{
+		MUTABLE_CPUPROFILER_SCOPE(Task_Mutable_GetImages_GetImage_Post)
+
+		FInstanceUpdateData::FImage& Image = OperationData->InstanceUpdateData.Images[ImageIndex];
+
+		Image.Image = GetImageTask.GetResult();
+		
+		check(Image.Image);
+
+		// We should have generated exactly this size.
+		const bool bSizeMissmatch = Image.Image->GetSizeX() != MipSizeX || Image.Image->GetSizeY() != MipSizeY;
+		if (bSizeMissmatch)
+		{
+			// Generate a correctly-sized but empty image instead, to avoid crashes.
+			UE_LOG(LogMutable, Warning, TEXT("Mutable generated a wrongly-sized image %llu."), Image.ImageID);
+			Image.Image = new mu::Image(MipSizeX, MipSizeY, FullLODContent - MipsToSkip, Image.Image->GetFormat(), mu::EInitializationType::Black);
+		}
+
+		// We need one mip or the complete chain. Otherwise there was a bug.
+		const int32 FullMipCount = Image.Image->GetMipmapCount(Image.Image->GetSizeX(), Image.Image->GetSizeY());
+		const int32 RealMipCount = Image.Image->GetLODCount();
+
+		bool bForceMipchain = 
+			// Did we fail to generate the entire mipchain (if we have mips at all)?
+			(RealMipCount != 1) && (RealMipCount != FullMipCount);
+
+		if (bForceMipchain)
+		{
+			MUTABLE_CPUPROFILER_SCOPE(GetImage_MipFix);
+
+			UE_LOG(LogMutable, Warning, TEXT("Mutable generated an incomplete mip chain for image %llu."), Image.ImageID);
+
+			// Force the right number of mips. The missing data will be black.
+			const mu::Ptr<mu::Image> NewImage = new mu::Image(Image.Image->GetSizeX(), Image.Image->GetSizeY(), FullMipCount, Image.Image->GetFormat(), mu::EInitializationType::Black);
+			check(NewImage);	
+			// Formats with BytesPerBlock == 0 will not allocate memory. This type of images are not expected here.
+			check(!NewImage->DataStorage.IsEmpty());
+
+			for (int32 L = 0; L < RealMipCount; ++L)
+			{
+				TArrayView<uint8> DestView = NewImage->DataStorage.GetLOD(L);
+				TArrayView<const uint8> SrcView = Image.Image->DataStorage.GetLOD(L);
+
+				check(DestView.Num() == SrcView.Num());
+				FMemory::Memcpy(DestView.GetData(), SrcView.GetData(), DestView.Num());
+			}
+			Image.Image = NewImage;
+		}
+
+		ImagesInThisInstance->Add(Image.ImageID);
+
+		Task_Mutable_GetImages_Loop(OperationData, StartTime, StartCycles, ImagesInThisInstance, ++ImageIndex);
+	}
+
+
+	/** See declaration. */
+	void Task_Mutable_GetImages_GetImage(
+		const TSharedRef<FUpdateContextPrivate>& OperationData,
+		double StartTime,
+		uint32 StartCycles,
+		const TSharedPtr<TArray<mu::FResourceID>>& ImagesInThisInstance,
+		int32 ImageIndex,
+		UE::Tasks::TTask<mu::FImageDesc> GetImageDescTask)
+	{
+		MUTABLE_CPUPROFILER_SCOPE(Task_Mutable_GetImages_GetImage)
+
+		const mu::FImageDesc& ImageDesc = GetImageDescTask.GetResult();
+
+		FInstanceUpdateData::FImage& Image = OperationData->InstanceUpdateData.Images[ImageIndex];
+
+		const UCustomizableObjectSystemPrivate* CustomizableObjectSystemPrivateData = UCustomizableObjectSystem::GetInstanceChecked()->GetPrivate();
+		
+		{
+			const uint16 MaxTextureSizeToGenerate = static_cast<uint16>(CustomizableObjectSystemPrivateData->MaxTextureSizeToGenerate);
+			const uint16 MaxSize = FMath::Max(ImageDesc.m_size[0], ImageDesc.m_size[1]);
+			uint16 Reduction = 1;
+
+			if (MaxTextureSizeToGenerate > 0 && MaxSize > MaxTextureSizeToGenerate)
+			{
+				// Find the reduction factor, and the BaseMip of the texture.
+				const uint32 NextPowerOfTwo = FMath::RoundUpToPowerOfTwo(FMath::DivideAndRoundUp(MaxSize, MaxTextureSizeToGenerate));
+				Reduction = FMath::Max(NextPowerOfTwo, 2U); // At least divide the texture by a factor of two
+				Image.BaseMip = FMath::FloorLog2(Reduction);
+			}
+
+			Image.FullImageSizeX = ImageDesc.m_size[0] / Reduction;
+			Image.FullImageSizeY = ImageDesc.m_size[1] / Reduction;
+		}
+
+		const bool bCached = ImagesInThisInstance->Contains(Image.ImageID) || // See if it is cached from this same instance (can happen with LODs)
+			(CVarReuseImagesBetweenInstances.GetValueOnAnyThread() && CustomizableObjectSystemPrivateData->ProtectedObjectCachedImages.Contains(Image.ImageID)); // See if it is cached from another instance
+
+		if (bCached)
+		{
+			UE_LOG(LogMutable, VeryVerbose, TEXT("Texture resource with id [%llu] is cached."), Image.ImageID);
+
+			Task_Mutable_GetImages_Loop(OperationData, StartTime, StartCycles, ImagesInThisInstance, ++ImageIndex);
+			return;
+		}
+		
+		const int32 MaxSize = FMath::Max(Image.FullImageSizeX, Image.FullImageSizeY);
+		const int32 FullLODCount = FMath::CeilLogTwo(MaxSize) + 1;
+		const int32 MinMipsInImage = FMath::Min(FullLODCount, UTexture::GetStaticMinTextureResidentMipCount());
+		const int32 MaxMipsToSkip = FullLODCount - MinMipsInImage;
+		int32 MipsToSkip = FMath::Min(MaxMipsToSkip, OperationData->MipsToSkip);
+
+		if (!FMath::IsPowerOfTwo(Image.FullImageSizeX) || !FMath::IsPowerOfTwo(Image.FullImageSizeY))
+		{
+			// It doesn't make sense to skip mips as non-power-of-two size textures cannot be streamed anyway
+			MipsToSkip = 0;
+		}
+
+		const int32 MipSizeX = FMath::Max(Image.FullImageSizeX >> MipsToSkip, 1);
+		const int32 MipSizeY = FMath::Max(Image.FullImageSizeY >> MipsToSkip, 1);
+		if (MipsToSkip > 0 && CustomizableObjectSystemPrivateData->EnableSkipGenerateResidentMips != 0 && OperationData->LowPriorityTextures.Find(Image.Name.ToString()) != INDEX_NONE)
+		{
+			mu::Ptr<const mu::Image> NewImage = new mu::Image(MipSizeX, MipSizeY, FullLODCount - MipsToSkip, ImageDesc.m_format, mu::EInitializationType::Black);
+
+			UE::Tasks::TTask<mu::Ptr<const mu::Image>> DummyTask = UE::Tasks::MakeCompletedTask<mu::Ptr<const mu::Image>>(NewImage);
+			Task_Mutable_GetImages_GetImage_Post(OperationData, StartTime, StartCycles, ImagesInThisInstance, ImageIndex, DummyTask, MipSizeX, MipSizeY, FullLODCount, MipsToSkip);
+		}
+		else
+		{
+			const UE::Tasks::TTask<mu::Ptr<const mu::Image>> GetImageTask = OperationData->MutableSystem->GetImage(OperationData->InstanceID, Image.ImageID, Image.BaseMip + MipsToSkip, Image.BaseLOD);
+			
+			UE::Tasks::AddNested(UE::Tasks::Launch(TEXT("Task_Mutable_GetImages_GetImage_Post"), [=]()
+			{
+				Task_Mutable_GetImages_GetImage_Post(OperationData, StartTime, StartCycles, ImagesInThisInstance, ImageIndex, GetImageTask, MipSizeX, MipSizeY, FullLODCount, MipsToSkip);
+			},
+			GetImageTask));
+		}
+	}
+
+
+	/** See declaration. */
+	void Task_Mutable_GetImages_Loop(
+		const TSharedRef<FUpdateContextPrivate>& OperationData,
+		double StartTime,
+		uint32 StartCycles,
+		const TSharedPtr<TArray<mu::FResourceID>>& ImagesInThisInstance,
+		int32 ImageIndex)
+	{
+		MUTABLE_CPUPROFILER_SCOPE(Task_Mutable_GetImages_Loop)
+
+		// Process next image. Some images are skipped
+		for (; ImageIndex < OperationData->InstanceUpdateData.Images.Num(); ++ImageIndex)
+		{
+			const FInstanceUpdateData::FImage& Image = OperationData->InstanceUpdateData.Images[ImageIndex];
+			if (!Image.bIsPassThrough)
+			{
+				Task_Mutable_GetImages_GetImageDesc(OperationData, StartTime, StartCycles, ImagesInThisInstance, ImageIndex);
+				return;
+			}
+		}
+
+		// If not image needs to be processed, go to end directly
+		Task_Mutable_GetImages_End(OperationData, StartTime, StartCycles);
+	}
+
+
+	namespace impl_new
+	{
+		// This runs in a worker thread.
+		void Task_Mutable_GetImages(const TSharedRef<FUpdateContextPrivate>& OperationData)
+		{
+			MUTABLE_CPUPROFILER_SCOPE(Task_Mutable_GetImages)
+			TRACE_BEGIN_REGION(UE_TASK_MUTABLE_GETIMAGES_REGION);
+
+			const double StartTime = FPlatformTime::Seconds();		
+			const uint32 StartCycles = FPlatformTime::Cycles();
+
+			const TSharedPtr<TArray<mu::FResourceID>> ImagesInThisInstance = MakeShared<TArray<mu::FResourceID>>();
+			Task_Mutable_GetImages_Loop(OperationData, StartTime, StartCycles, ImagesInThisInstance, 0);
+		}
+		
+	}
+
+
+	/** Start of the GetImages tasks. */
+	void Task_Mutable_GetImages(const TSharedRef<FUpdateContextPrivate>& OperationData)
+	{
+		if (CVarEnableNewSplitMutableTask->GetBool())
+		{
+			impl_new::Task_Mutable_GetImages(OperationData);
+		}
+		else
+		{
+			CustomizableObjectSystem::ImplDeprecated::Task_Mutable_GetImages(OperationData);						
+		}
+	}
+
+
 	// This runs in a worker thread.
 	void Task_Mutable_ReleaseInstance(const TSharedRef<FUpdateContextPrivate>& OperationData, mu::Ptr<mu::System> MutableSystem)
 	{
