[Mutable] Improved Graph Search.

- Better us of the reflection system.
- Find pin names.

#rnx

diff --git a/Engine/Plugins/Experimental/Mutable/Source/CustomizableObjectEditor/Private/MuCOE/CustomizableObjectEditor.cpp b/Engine/Plugins/Experimental/Mutable/Source/CustomizableObjectEditor/Private/MuCOE/CustomizableObjectEditor.cpp
--- a/Engine/Plugins/Experimental/Mutable/Source/CustomizableObjectEditor/Private/MuCOE/CustomizableObjectEditor.cpp
+++ b/Engine/Plugins/Experimental/Mutable/Source/CustomizableObjectEditor/Private/MuCOE/CustomizableObjectEditor.cpp
@@ -346,7 +346,7 @@ FText FCustomizableObjectEditor::GetBaseToolkitName() const
 }
 
 
-void FCustomizableObjectEditor::SelectNode(const UCustomizableObjectNode* Node)
+void FCustomizableObjectEditor::SelectNode(const UEdGraphNode* Node)
 {
 	GraphEditor->JumpToNode(Node);
 }
@@ -2064,450 +2064,170 @@ TSharedPtr<SCustomizableObjectEditorAdvancedPreviewSettings> FCustomizableObject
 }
 
 
-void FCustomizableObjectEditor::OnEnterText(const FText & NewText, ETextCommit::Type TextType)
+void FCustomizableObjectEditor::FindProperty(const FProperty* Property, const void* InContainer, const FString& FindString, const UObject& Context, bool& bFound)
 {
-	bool found = false;
+	if (!Property || !InContainer)
+	{
+		return;
+	}
 
-	if (TextType == ETextCommit::OnEnter)
+	const FString PropertyName = Property->GetDisplayNameText().ToString();
+	if (PropertyName.Contains(FindString))
+	{
+		LogSearchResult(Context, "Property Name", bFound, *PropertyName);
+		bFound = true;
+	}
+	
+	for (int32 Index = 0; Index < Property->ArrayDim; ++Index)
 	{
-		if (!GraphEditor.IsValid() || !GraphEditor->GetCurrentGraph())
+		const uint8* ValuePtr = Property->ContainerPtrToValuePtr<uint8>(InContainer, Index);
+
+		if (const FStrProperty* StringProperty = CastField<FStrProperty>(Property))
 		{
-			return;
+			const FString* StringResult = StringProperty->GetPropertyValuePtr(ValuePtr);
+			if (StringResult->Contains(FindString))
+			{
+				LogSearchResult(Context, "Property Value", bFound, *StringResult);
+				bFound = true;
+			}
 		}
-
-		UEdGraph* Graph = GraphEditor->GetCurrentGraph();
-		
-		for (int i = 0; i < Graph->Nodes.Num(); ++i)
+		else if (const FEnumProperty* EnumProperty = CastField<FEnumProperty>(Property))
 		{
-			UObject* Node = Graph->Nodes[i];
+			const UEnum* EnumResult = EnumProperty->GetEnum();
 
-			// Ensure we are working with a Customizable object Node
-			UCustomizableObjectNode* CustomizableObjectNode = Cast<UCustomizableObjectNode>(Node);
-			if (!CustomizableObjectNode)
+			const FString StringResult = EnumResult->GetDisplayNameTextByIndex(*ValuePtr).ToString();
+			if (StringResult.Contains(FindString))
 			{
-				continue;
+				LogSearchResult(Context, "Property Value", bFound, StringResult);
+				bFound = true;
 			}
-
-			//Find Node Coincidence 
-			FString NodeName = Graph->Nodes[i]->GetNodeTitle(ENodeTitleType::ListView).ToString();
-			int32 cPos = NodeName.Find("\n");
+		}
+		else if (const FSoftObjectProperty* SoftObjectProperty = CastField<FSoftObjectProperty>(Property))
+		{
+			const FString ObjectPath = SoftObjectProperty->GetPropertyValuePtr(ValuePtr)->ToString();
+			if (ObjectPath.Contains(FindString))
+			{
+				LogSearchResult(Context, "Property Value", bFound, ObjectPath);
+				bFound = true;
+			}
+		}
+		else if (const FObjectPropertyBase* ObjectProperty = CastField<FObjectPropertyBase>(Property))
+		{
+			if (const UObject* ObjectValue = ObjectProperty->GetObjectPropertyValue(ValuePtr))
+			{
+				const FString Name = ObjectValue->GetName();
 			
-			if (cPos >= 0)
+				if (ObjectValue->GetName().Contains(FindString))
+				{
+					LogSearchResult(Context, "Property Value", bFound, Name);
+					bFound = true;
+				}				
+			}
+		}
+		else if (const FStructProperty* StructProperty = CastField<FStructProperty>(Property))
+		{
+			for (TFieldIterator<FProperty> It(StructProperty->Struct); It; ++It)
 			{
-				NodeName.RemoveAt(cPos);
-				NodeName.InsertAt(cPos, " ");
+				FindProperty(*It, ValuePtr, FindString, Context, bFound);
 			}
-
-			if (NodeName.Contains(NewText.ToString(), ESearchCase::IgnoreCase))
+		}
+		else if (const FArrayProperty* ArrayProperty = CastField<FArrayProperty>(Property))
+		{
+			FScriptArrayHelper ArrayHelper(ArrayProperty, ValuePtr);
+			for (int32 ValueIdx = 0; ValueIdx < ArrayHelper.Num(); ++ValueIdx)
 			{
-				LogSearchResult(CustomizableObjectNode, "Node", found, NodeName);
-				found = true;
+				FindProperty(ArrayProperty->Inner, ArrayHelper.GetRawPtr(ValueIdx), FindString, Context, bFound);
 			}
-
-			//Find Variable method 1 (using contains of string)
-			FString VariableName = NewText.ToString();
-			VariableName.RemoveSpacesInline();
-
-			for (TFieldIterator<FProperty> It(Node->GetClass()); It; ++It)
+		}
+		else if (const FSetProperty* SetProperty = CastField<FSetProperty>(Property))
+		{
+			FScriptSetHelper SetHelper(SetProperty, ValuePtr);
+			for (FScriptSetHelper::FIterator SetIt = SetHelper.CreateIterator(); SetIt; ++SetIt)
 			{
-				FProperty* Property = *It;
-				FString ResultName = Property->GetName();
-				
-				//Find Variable name
-				if (ResultName.Contains(VariableName))
-				{
-					LogSearchResult(CustomizableObjectNode, "Variable", found, ResultName);
-					found = true;
-				}
-
-				//Find Variable String Content
-				if (FStrProperty *StringProperty = CastField<FStrProperty>(Property))
-				{
-					FString* StringResult = Property->ContainerPtrToValuePtr<FString>(Node);
+				FindProperty(SetProperty->ElementProp, SetHelper.GetElementPtr(SetIt), FindString, Context, bFound);
+			}
+		}
+		else if (const FMapProperty* MapProperty = CastField<FMapProperty>(Property))
+		{
+			FScriptMapHelper MapHelper(MapProperty, ValuePtr);
+			for (FScriptMapHelper::FIterator MapIt = MapHelper.CreateIterator(); MapIt; ++MapIt)
+			{
+				const uint8* MapValuePtr = MapHelper.GetPairPtr(MapIt);
+				FindProperty(MapProperty->KeyProp, MapValuePtr, FindString, Context, bFound);
+				FindProperty(MapProperty->ValueProp, MapValuePtr, FindString, Context, bFound);
+			}
+		}
+	}
+}
 
-					if (StringResult->Contains(NewText.ToString()))
-					{
-						LogSearchResult(CustomizableObjectNode, "Value", found, *StringResult);
-						found = true;
-					}
-				}
 
-				//Find Variable Enum Content (Hard coded)
-				//TODO: improve using full Reflection
-				if (FEnumProperty *EnumProperty = CastField<FEnumProperty>(Property))
-				{
-					UEnum* EnumResult = EnumProperty->GetEnum();
-					FString EnumType = EnumResult->CppType;
-					FString StringResult;
+void FCustomizableObjectEditor::OnEnterText(const FText& NewText, ETextCommit::Type TextType)
+{
+	if (TextType != ETextCommit::OnEnter)
+	{
+		return;
+	}
+	
+	if (!GraphEditor)
+	{
+		return;
+	}
 
-					FFieldClass* SourceObjectClass = Property->GetClass();
+	const UEdGraph* Graph = GraphEditor->GetCurrentGraph();
+	if (!Graph)
+	{
+		return;
+	}
 
-					if (EnumType == "ECustomizableObjectGroupType")
-					{
-						ECustomizableObjectGroupType* value = EnumProperty->ContainerPtrToValuePtr<ECustomizableObjectGroupType>(Node);
-						StringResult = EnumResult->GetDisplayNameTextByIndex((int32)*value).ToString();
-					}
-					else if (EnumType == "ENodeEnabledState")
-					{
-						ENodeEnabledState* value = EnumProperty->ContainerPtrToValuePtr<ENodeEnabledState>(Node);
-						StringResult = EnumResult->GetDisplayNameTextByIndex((int32)*value).ToString();
-					}
-					else if (EnumType == "ECustomizableObjectAutomaticLODStrategy")
-					{
-						ECustomizableObjectAutomaticLODStrategy* value = EnumProperty->ContainerPtrToValuePtr<ECustomizableObjectAutomaticLODStrategy>(Node);
-						StringResult = EnumResult->GetDisplayNameTextByIndex((int32)*value).ToString();
-					}
-					else if (EnumType == "ECustomizableObjectProjectorType")
-					{
-						ECustomizableObjectProjectorType* value = EnumProperty->ContainerPtrToValuePtr<ECustomizableObjectProjectorType>(Node);
-						StringResult = EnumResult->GetDisplayNameTextByIndex((int32)*value).ToString();
-					}
-					else if (EnumType == "EColorArithmeticOperation")
-					{
-						EColorArithmeticOperation* value = EnumProperty->ContainerPtrToValuePtr<EColorArithmeticOperation>(Node);
-						StringResult = EnumResult->GetDisplayNameTextByIndex((int32)*value).ToString();
-					}
-					else if (EnumType == "ECustomizableObjectNodeMaterialVariationType")
-					{
-						ECustomizableObjectNodeMaterialVariationType* value = EnumProperty->ContainerPtrToValuePtr<ECustomizableObjectNodeMaterialVariationType>(Node);
-						StringResult = EnumResult->GetDisplayNameTextByIndex((int32)*value).ToString();
-					}
+	bool bFound = false;
 
-					if (StringResult.Contains(NewText.ToString()))
-					{
-						LogSearchResult(CustomizableObjectNode, "Value", found, StringResult);
-						found = true;
-					}
-				}
+	const FString FindString = NewText.ToString();
 
-				//Find Variable Array Content (Hard coded)
-				//TODO: improve using full Reflection
-				if (FArrayProperty* ArrayProperty = CastField<FArrayProperty>(Property))
-				{
-					FProperty* AProperty = ArrayProperty->Inner;
-					FString ArrayType = AProperty->GetCPPType();
+	for (TObjectPtr<UEdGraphNode> Node : Graph->Nodes)
+	{
+		if (!Node)
+		{
+			continue;
+		}
 
-					if (ArrayType == "FGroupProjectorParameterImage")
-					{
-						TArray<FGroupProjectorParameterImage>  Array =
-							*ArrayProperty->ContainerPtrToValuePtr<TArray<FGroupProjectorParameterImage>>(Node);
-
-						//Class vars
-						FString OptionName = "OptionName";
-
-						for (int a = 0; a < Array.Num(); ++a)
-						{
-							if (OptionName.Contains(VariableName))
-							{
-								LogSearchResult(CustomizableObjectNode, "Variable", found, OptionName);
-								found = true;
-							}
-
-							FString OptionNameContent = Array[a].OptionName;
-
-							if (OptionNameContent.Contains(NewText.ToString()))
-							{
-								LogSearchResult(CustomizableObjectNode, "Value", found, Array[a].OptionName);
-								found = true;
-							}
-						}
-					}
-					else if (ArrayType == "FGroupProjectorParameterPose")
-					{
-						TArray<FGroupProjectorParameterPose>  Array =
-							*ArrayProperty->ContainerPtrToValuePtr<TArray<FGroupProjectorParameterPose>>(Node);
-
-						//Class vars
-						FString PoseName = "PoseName";
-
-						for (int a = 0; a < Array.Num(); ++a)
-						{
-							if (PoseName.Contains(VariableName))
-							{
-								LogSearchResult(CustomizableObjectNode, "Variable", found, PoseName);
-								found = true;
-							}
-
-							FString OptionNameContent = Array[a].PoseName;
-
-							if (OptionNameContent.Contains(NewText.ToString()))
-							{
-								LogSearchResult(CustomizableObjectNode, "Value", found, Array[a].PoseName);
-								found = true;
-							}
-						}
-					}
-					else if (ArrayType == "FCustomizableObjectNodeFloatDescription")
-					{
-						TArray<FCustomizableObjectNodeFloatDescription>  Array =
-							*ArrayProperty->ContainerPtrToValuePtr<TArray<FCustomizableObjectNodeFloatDescription>>(Node);
-						//Class vars
-						FString Name = "Name";
-
-						for (int a = 0; a < Array.Num(); ++a)
-						{
-							if (Name.Contains(VariableName))
-							{
-								LogSearchResult(CustomizableObjectNode, "Variable", found, Name);
-								found = true;
-							}
-
-							FString OptionNameContent = Array[a].Name;
-
-							if (OptionNameContent.Contains(NewText.ToString()))
-							{
-								LogSearchResult(CustomizableObjectNode, "Value", found, Array[a].Name);
-								found = true;
-							}
-						}
-					}
-					else if (ArrayType == "FCustomizableObjectNodeEnumValue")
-					{
-						TArray<FCustomizableObjectNodeEnumValue>  Array =
-							*ArrayProperty->ContainerPtrToValuePtr<TArray<FCustomizableObjectNodeEnumValue>>(Node);
-						//Class vars
-						FString Name = "Name";
-
-						for (int a = 0; a < Array.Num(); ++a)
-						{
-							if (Name.Contains(VariableName))
-							{
-								LogSearchResult(CustomizableObjectNode, "Variable", found, Name);
-								found = true;
-							}
-
-							FString OptionNameContent = Array[a].Name;
-
-							if (OptionNameContent.Contains(NewText.ToString()))
-							{
-								LogSearchResult(CustomizableObjectNode, "Value", found, Array[a].Name);
-								found = true;
-							}
-						}
-					}
-					else if (ArrayType == "FCustomizableObjectNodeExtendMaterialImage")
-					{
-						TArray<FCustomizableObjectNodeExtendMaterialImage>  Array =
-							*ArrayProperty->ContainerPtrToValuePtr<TArray<FCustomizableObjectNodeExtendMaterialImage>>(Node);
-						//Class vars
-						FString Name = "Name";
-
-						for (int a = 0; a < Array.Num(); ++a)
-						{
-							if (Name.Contains(VariableName))
-							{
-								LogSearchResult(CustomizableObjectNode, "Variable", found, Name);
-								found = true;
-							}
-
-							FString OptionNameContent = Array[a].Name;
-
-							if (OptionNameContent.Contains(NewText.ToString()))
-							{
-								LogSearchResult(CustomizableObjectNode, "Value", found, Array[a].Name);
-								found = true;
-							}
-						}
-					}
-					else if (ArrayType == "FCustomizableObjectNodeMaterialImage")
-					{
-						TArray<FCustomizableObjectNodeMaterialImage>  Array =
-							*ArrayProperty->ContainerPtrToValuePtr<TArray<FCustomizableObjectNodeMaterialImage>>(Node);
-						//Class vars
-						FString Name = "Name";
-
-						for (int a = 0; a < Array.Num(); ++a)
-						{
-							if (Name.Contains(VariableName))
-							{
-								LogSearchResult(CustomizableObjectNode, "Variable", found, Name);
-								found = true;
-							}
-
-							FString OptionNameContent = Array[a].Name;
-
-							if (OptionNameContent.Contains(NewText.ToString()))
-							{
-								LogSearchResult(CustomizableObjectNode, "Value", found, Array[a].Name);
-								found = true;
-							}
-						}
-					}
-					else if (ArrayType == "FCustomizableObjectNodeMaterialVector")
-					{
-						TArray<FCustomizableObjectNodeMaterialVector>  Array =
-							*ArrayProperty->ContainerPtrToValuePtr<TArray<FCustomizableObjectNodeMaterialVector>>(Node);
-						//Class vars
-						FString Name = "Name";
-
-						for (int a = 0; a < Array.Num(); ++a)
-						{
-							if (Name.Contains(VariableName))
-							{
-								LogSearchResult(CustomizableObjectNode, "Variable", found, Name);
-								found = true;
-							}
-
-							FString NameContent = Array[a].Name;
-
-							if (NameContent.Contains(NewText.ToString()))
-							{
-								LogSearchResult(CustomizableObjectNode, "Value", found, Array[a].Name);
-								found = true;
-							}
-						}
-					}
-					else if (ArrayType == "FCustomizableObjectNodeMaterialScalar")
-					{
-						TArray<FCustomizableObjectNodeMaterialScalar>  Array =
-							*ArrayProperty->ContainerPtrToValuePtr<TArray<FCustomizableObjectNodeMaterialScalar>>(Node);
-						//Class vars
-						FString Name = "Name";
-
-						for (int a = 0; a < Array.Num(); ++a)
-						{
-							if (Name.Contains(VariableName))
-							{
-								LogSearchResult(CustomizableObjectNode, "Variable", found, Name);
-								found = true;
-							}
-
-							FString NameContent = Array[a].Name;
-
-							if (NameContent.Contains(NewText.ToString()))
-							{
-								LogSearchResult(CustomizableObjectNode, "Value", found, Array[a].Name);
-								found = true;
-							}
-						}
-					}
-					else if (ArrayType == "FCustomizableObjectMaterialVariation")
-					{
-						TArray<FCustomizableObjectMaterialVariation>  Array =
-							*ArrayProperty->ContainerPtrToValuePtr<TArray<FCustomizableObjectMaterialVariation>>(Node);
-						//Class vars
-						FString Tag = "Tag";
-
-						for (int a = 0; a < Array.Num(); ++a)
-						{
-							if (Tag.Contains(VariableName))
-							{
-								LogSearchResult(CustomizableObjectNode, "Variable", found, Tag);
-								found = true;
-							}
-
-							FString OptionNameContent = Array[a].Tag;
-
-							if (OptionNameContent.Contains(NewText.ToString()))
-							{
-								LogSearchResult(CustomizableObjectNode, "Value", found, Array[a].Tag);
-								found = true;
-							}
-						}
-					}
-					else if (ArrayType == "FString")
-					{
-						TArray<FString>  Array =
-							*ArrayProperty->ContainerPtrToValuePtr<TArray<FString>>(Node);
-
-						for (int a = 0; a < Array.Num(); ++a)
-						{
-							if (Array[a].Contains(NewText.ToString()))
-							{
-								LogSearchResult(CustomizableObjectNode, "Value", found, Array[a]);
-								found = true;
-							}
-						}
+		// Node names are not in the reflection system
+		const FString NodeName = Node->GetNodeTitle(ENodeTitleType::FullTitle).ToString().Replace(TEXT("\n"), TEXT(" "));
+		if (NodeName.Contains(NewText.ToString(), ESearchCase::IgnoreCase))
+		{
+			LogSearchResult(*Node, "Node", bFound, NodeName);
+			bFound = true;
+		}
 
-					}
-					else if (ArrayType == "FCustomizableObjectState")
-					{
-						TArray<FCustomizableObjectState>  Array =
-							*ArrayProperty->ContainerPtrToValuePtr<TArray<FCustomizableObjectState>>(Node);
-						//Class vars
-						FString Name = "Name";
-						FString RuntimeParameters = "RuntimeParameters";
-
-						for (int a = 0; a < Array.Num(); ++a)
-						{
-							if (Name.Contains(VariableName))
-							{
-								LogSearchResult(CustomizableObjectNode, "Variable", found, Name);
-								found = true;
-							}
-
-							if (Array[a].Name.Contains(NewText.ToString()))
-							{
-								LogSearchResult(CustomizableObjectNode, "Value", found, Array[a].Name);
-								found = true;
-							}
-
-							if (RuntimeParameters.Contains(VariableName))
-							{
-								LogSearchResult(CustomizableObjectNode, "Variable", found, Name);
-								found = true;
-							}
-
-							for (int s = 0; s < Array[a].RuntimeParameters.Num(); ++s)
-							{
-								if (Array[a].RuntimeParameters[s].Contains(NewText.ToString()))
-								{
-									LogSearchResult(CustomizableObjectNode, "Value", found, Array[a].RuntimeParameters[s]);
-									found = true;
-								}
-							}
-						}
-					}
-					else if (ArrayType == "FCustomizableObjectNodeStaticMeshLOD")
-					{
-						TArray<FCustomizableObjectNodeStaticMeshLOD>  Array =
-							*ArrayProperty->ContainerPtrToValuePtr<TArray<FCustomizableObjectNodeStaticMeshLOD>>(Node);
-
-						FString LODs = "LODs";
-
-						for (int a = 0; a < Array.Num(); ++a)
-						{
-							if (LODs.Contains(VariableName))
-							{
-								LogSearchResult(CustomizableObjectNode, "Variable", found, LODs);
-								found = true;
-							}
-
-							for (int s = 0; s < Array[a].Materials.Num(); ++s)
-							{
-								FString Name = "Name";
-
-								if (Name.Contains(VariableName))
-								{
-									LogSearchResult(CustomizableObjectNode, "Variable", found, Name);
-									found = true;
-								}
-
-								if (Array[a].Materials[s].Name.Contains(NewText.ToString()))
-								{
-									LogSearchResult(CustomizableObjectNode, "Value", found, Array[a].Materials[s].Name);
-									found = true;
-								}
-							}
-						}
-					}
-				}				
+		// Pins are not in the reflection system
+		for (const UEdGraphPin* Pin : Node->GetAllPins())
+		{
+			const FString PinFriendlyName = Pin->PinFriendlyName.ToString();
+			if (PinFriendlyName.Contains(FindString))
+			{
+				LogSearchResult(*Node, "Pin", bFound, PinFriendlyName);
+				bFound = true;
 			}
 		}
 
-		const FText Text = found ?
-			LOCTEXT("SearchCompleted", "Search completed") :
-			FText::FromString("No Results for: " + NewText.ToString());
-
-		FCustomizableObjectEditorLogger::CreateLog(Text)
-			.Category(ELoggerCategory::GraphSearch)
-			.CustomNotification()
-			.Log();
+		// Find anything marked as a UPROPERTY
+		for (TFieldIterator<FProperty> It(Node->GetClass()); It; ++It)
+		{
+			FindProperty(*It, Node, FindString, *Node, bFound);
+		}
 	}
+
+	const FText Text = bFound ?
+		LOCTEXT("SearchCompleted", "Search completed") :
+		FText::FromString("No Results for: " + NewText.ToString());
+
+	FCustomizableObjectEditorLogger::CreateLog(Text)
+		.Category(ELoggerCategory::GraphSearch)
+		.CustomNotification()
+		.Log();
 }
 
 
-void FCustomizableObjectEditor::LogSearchResult(UCustomizableObjectNode* Node, FString Type, bool bIsFirst, FString Result) const
+void FCustomizableObjectEditor::LogSearchResult(const UObject& Context, const FString& Type, bool bIsFirst, const FString& Result) const
 {
 	if (!bIsFirst)
 	{
@@ -2517,7 +2237,7 @@ void FCustomizableObjectEditor::LogSearchResult(UCustomizableObjectNode* Node, F
 	}
 	
 	FCustomizableObjectEditorLogger::CreateLog(FText::FromString(Type + ": " + Result))
-	.Context(*Node)
+	.Context(Context)
 	.BaseObject()
 	.Notification(false)
 	.Log();
diff --git a/Engine/Plugins/Experimental/Mutable/Source/CustomizableObjectEditor/Private/MuCOE/CustomizableObjectEditor.h b/Engine/Plugins/Experimental/Mutable/Source/CustomizableObjectEditor/Private/MuCOE/CustomizableObjectEditor.h
--- a/Engine/Plugins/Experimental/Mutable/Source/CustomizableObjectEditor/Private/MuCOE/CustomizableObjectEditor.h
+++ b/Engine/Plugins/Experimental/Mutable/Source/CustomizableObjectEditor/Private/MuCOE/CustomizableObjectEditor.h
@@ -197,7 +197,7 @@ public:
 	virtual UCustomizableObjectInstance* GetPreviewInstance() override;
 	virtual bool CanPasteNodes() const override;
 	virtual void PasteNodesHere(const FVector2D& Location) override;
-	virtual void SelectNode(const UCustomizableObjectNode* Node) override;
+	virtual void SelectNode(const UEdGraphNode* Node) override;
 	virtual void ReconstructAllChildNodes(UCustomizableObjectNode& StartNode, const UClass& NodeType) override;
 	virtual UProjectorParameter* GetProjectorParameter() override;
 	virtual UCustomSettings* GetCustomSettings() override;
@@ -351,12 +351,12 @@ private:
 	void OnEnterText(const FText& NewText, ETextCommit::Type TextType);
 
 	/** Logs the search results of the search
-	 * @param Node The Customizable Object Node we have found to be related with the searched string.
+	 * @param Context The UObject we have found to be related with the searched string.
 	 * @param Type The type of relation with the searched word. It is a node, a value or maybe a variable?
 	 * @param bIsFirst Is this the first time we encountered something during our search?
 	 * @param Result The string containing the search word we are looking for in Node
 	 */
-	void LogSearchResult(UCustomizableObjectNode* Node, FString Type, bool bIsFirst, FString Result) const;
+	void LogSearchResult(const UObject& Context, const FString& Type, bool bIsFirst, const FString& Result) const;
 
 	/** Open the Texture Analyzer tab */
 	void OpenTextureAnalyzerTab();
@@ -367,6 +367,14 @@ private:
 	/** Creates the necessary components for the preview of the CO instance */
 	void CreatePreviewComponents();
 
+	/** Recursively find any property that its name or value contains the given string.
+	  * @param Property Root property.
+	  * @param Container Root property container (address of the property value).
+	  * @param FindString String to find for.
+	  * @param Context UObject Context where this string has been found.
+	  * @param bFound Mark as true if any property has been found. */
+	void FindProperty(const FProperty* Property, const void* Container, const FString& FindString, const UObject& Context, bool& bFound);
+	
 public:
 	void OnCustomizableObjectStatusChanged(FCustomizableObjectStatus::EState PreviousState, FCustomizableObjectStatus::EState CurrentState);
 
diff --git a/Engine/Plugins/Experimental/Mutable/Source/CustomizableObjectEditor/Private/MuCOE/CustomizableObjectEditorLogger.cpp b/Engine/Plugins/Experimental/Mutable/Source/CustomizableObjectEditor/Private/MuCOE/CustomizableObjectEditorLogger.cpp
--- a/Engine/Plugins/Experimental/Mutable/Source/CustomizableObjectEditor/Private/MuCOE/CustomizableObjectEditorLogger.cpp
+++ b/Engine/Plugins/Experimental/Mutable/Source/CustomizableObjectEditor/Private/MuCOE/CustomizableObjectEditorLogger.cpp
@@ -10,6 +10,7 @@
 #include "MuCOE/CustomizableObjectGraph.h"
 #include "MuCOE/ICustomizableObjectEditor.h"
 #include "MuCOE/Nodes/CustomizableObjectNode.h"
+#include "Toolkits/ToolkitManager.h"
 #include "Widgets/Notifications/SNotificationList.h"
 
 #define LOCTEXT_NAMESPACE "CustomizableObject"
@@ -24,7 +25,7 @@ class FCustomizableObjectToken : public IMessageToken
 {
 public:
 	/** Factory method, tokens can only be constructed as shared refs */
-	CUSTOMIZABLEOBJECTEDITOR_API static TSharedRef<IMessageToken> Create(const UCustomizableObjectNode* Node)
+	CUSTOMIZABLEOBJECTEDITOR_API static TSharedRef<IMessageToken> Create(const UEdGraphNode* Node)
 	{
 		return MakeShareable(new FCustomizableObjectToken(Node));
 	}
@@ -36,46 +37,14 @@ public:
 	}
 
 	// Own interface
-	CUSTOMIZABLEOBJECTEDITOR_API const UCustomizableObjectNode* GetNode() const
+	CUSTOMIZABLEOBJECTEDITOR_API const UEdGraphNode* GetNode() const
 	{
-		return Cast<UCustomizableObjectNode>(NodeBeingReferenced.Get());
+		return Cast<UEdGraphNode>(NodeBeingReferenced.Get());
 	}
 
-	static void OnMessageLogLinkActivated(const class TSharedRef<IMessageToken>& Token)
-	{
-		// Just an object link
-		if (Token->GetType() == EMessageToken::Object)
-		{
-			const TSharedRef<FUObjectToken> UObjectToken = StaticCastSharedRef<FUObjectToken>(Token);
-			if (UObjectToken->GetObject().IsValid())
-			{
-				UObject* Object = UObjectToken->GetObject().Get();
-				GEditor->GetEditorSubsystem<UAssetEditorSubsystem>()->OpenEditorForAsset(Object);
-			}
-		}
-
-		// Fake type, but no possibility for custom types
-		else if (Token->GetType() == EMessageToken::Text)
-		{
-			const TSharedRef<FCustomizableObjectToken> UObjectToken = StaticCastSharedRef<FCustomizableObjectToken>(Token);
-			if (UObjectToken->GetNode())
-			{
-				const UCustomizableObjectNode* Node = UObjectToken->GetNode();
-				UObject* Object = Node->GetCustomizableObjectGraph()->GetOuter();
-
-				// Make sure the editor exists for this asset
-				GEditor->GetEditorSubsystem<UAssetEditorSubsystem>()->OpenEditorForAsset(Object);
-
-				// Find it
-				const TSharedPtr<ICustomizableObjectEditor> Editor = Node->GetGraphEditor();
-				Editor->SelectNode(Node);
-			}
-		}
-	}
-	
 private:
 	/** Private constructor */
-	explicit FCustomizableObjectToken(const UCustomizableObjectNode* Node) : NodeBeingReferenced(Node)
+	explicit FCustomizableObjectToken(const UEdGraphNode* Node) : NodeBeingReferenced(Node)
 	{
 		const FString String = Node ? *Node->GetName() : TEXT("<None>");
 		CachedText = FText::FromString(String);
@@ -86,7 +55,7 @@ private:
 };
 
 
-void OnMessageLogLinkActivated(const class TSharedRef<IMessageToken>& Token)
+void OnMessageLogLinkActivated(const TSharedRef<IMessageToken>& Token)
 {
 	// Just an object link
 	if (Token->GetType() == EMessageToken::Object)
@@ -105,15 +74,20 @@ void OnMessageLogLinkActivated(const class TSharedRef<IMessageToken>& Token)
 		const TSharedRef<FCustomizableObjectToken> UObjectToken = StaticCastSharedRef<FCustomizableObjectToken>(Token);
 		if (UObjectToken->GetNode())
 		{
-			const UCustomizableObjectNode* Node = UObjectToken->GetNode();
-			UObject* Object = Node->GetCustomizableObjectGraph()->GetOuter();
+			const UEdGraphNode* Node = UObjectToken->GetNode();
+			UObject* Object = Node->GetOuter()->GetOuter();
 
 			// Make sure the editor exists for this asset
 			GEditor->GetEditorSubsystem<UAssetEditorSubsystem>()->OpenEditorForAsset(Object);
 
 			// Find it
-			const TSharedPtr<ICustomizableObjectEditor> Editor = Node->GetGraphEditor();
-			Editor->SelectNode(Node);
+			if (const UCustomizableObject* CustomizableObject = Cast<UCustomizableObject>(Object))
+			{
+				if (const TSharedPtr<IToolkit> FoundAssetEditor = FToolkitManager::Get().FindEditorForAsset(CustomizableObject))
+				{
+					StaticCastSharedPtr<ICustomizableObjectEditor>(FoundAssetEditor)->SelectNode(Node);
+				}
+			}
 		}
 	}
 }
@@ -242,12 +216,12 @@ void FCustomizableObjectEditorLogger::Log(FLogParameters& LogParameters)
 
 	for (const UObject* Context : LogParameters.ParamContext)
 	{
-		if (const UCustomizableObjectNode* Node = Cast<const UCustomizableObjectNode>(Context))
+		if (const UEdGraphNode* Node = Cast<const UEdGraphNode>(Context))
 		{
 			if (LogParameters.bParamBaseObject)
 			{
 				Message->AddToken(FTextToken::Create(FText::FromString(TEXT(" "))));
-				const UObject* Asset = Node->GetCustomizableObjectGraph()->GetOuter();
+				const UObject* Asset = Node->GetOuter()->GetOuter();
 				Message->AddToken(FUObjectToken::Create(Asset)->OnMessageTokenActivated(FOnMessageTokenActivated::CreateStatic(&OnMessageLogLinkActivated)));
 			}
 
diff --git a/Engine/Plugins/Experimental/Mutable/Source/CustomizableObjectEditor/Private/MuCOE/Nodes/CustomizableObjectNode.cpp b/Engine/Plugins/Experimental/Mutable/Source/CustomizableObjectEditor/Private/MuCOE/Nodes/CustomizableObjectNode.cpp
--- a/Engine/Plugins/Experimental/Mutable/Source/CustomizableObjectEditor/Private/MuCOE/Nodes/CustomizableObjectNode.cpp
+++ b/Engine/Plugins/Experimental/Mutable/Source/CustomizableObjectEditor/Private/MuCOE/Nodes/CustomizableObjectNode.cpp
@@ -36,6 +36,8 @@ bool UCustomizableObjectNode::IsSingleOutputNode() const
 UEdGraphPin* UCustomizableObjectNode::CustomCreatePin(EEdGraphPinDirection Direction, const FName& Type, const FName& Name, bool bIsArray)
 {
 	UEdGraphPin* Pin = CreatePin(Direction, Type, Name);
+
+	Pin->PinFriendlyName = FText::FromName(Name);
 	if (bIsArray)
 	{
 		Pin->PinType.ContainerType = EPinContainerType::Array;
diff --git a/Engine/Plugins/Experimental/Mutable/Source/CustomizableObjectEditor/Public/MuCOE/ICustomizableObjectEditor.h b/Engine/Plugins/Experimental/Mutable/Source/CustomizableObjectEditor/Public/MuCOE/ICustomizableObjectEditor.h
--- a/Engine/Plugins/Experimental/Mutable/Source/CustomizableObjectEditor/Public/MuCOE/ICustomizableObjectEditor.h
+++ b/Engine/Plugins/Experimental/Mutable/Source/CustomizableObjectEditor/Public/MuCOE/ICustomizableObjectEditor.h
@@ -32,7 +32,7 @@ public:
 	virtual void UpdateObjectProperties() {}
 
 	// Make sure the graph editor selects a specific node
-	virtual void SelectNode(const UCustomizableObjectNode* Node) {}
+	virtual void SelectNode(const UEdGraphNode* Node) {}
 
 	virtual void ReconstructNode(UEdGraphNode* Node) {}
 
