[Mutable] Adapt the mesh generation process to create a CustomizableObjectSkeletalMesh.
- Avoid enqueuing automatic LOD updates when mesh LOD streaming is enabled.
- Add helper functions to init the streamable mesh with dummy data.
- Remove double initialization of SkinWeightVertexBuffers.

#rnx

diff --git a/Engine/Plugins/Experimental/Mutable/Source/CustomizableObject/Classes/MuCO/CustomizableObject.h b/Engine/Plugins/Experimental/Mutable/Source/CustomizableObject/Classes/MuCO/CustomizableObject.h
--- a/Engine/Plugins/Experimental/Mutable/Source/CustomizableObject/Classes/MuCO/CustomizableObject.h
+++ b/Engine/Plugins/Experimental/Mutable/Source/CustomizableObject/Classes/MuCO/CustomizableObject.h
@@ -544,6 +544,11 @@ public:
 	 * USkeletalMeshes are only reused between the same CO. */
 	UPROPERTY(EditAnywhere, Category = CustomizableObject)
 	bool bEnableMeshCache = false;
+
+	/** If true, Mesh LODs will be streamed on demand. It requires streaming of SkeletalMeshes and Mutable.StreamMeshLODsEnabled to be enabled.
+	 *  Does not support Clothing, Morphs, and alternative SkinWeightProfiles yet. */
+	UPROPERTY(EditAnywhere, Category = CustomizableObject)
+	bool bEnableMeshStreaming = false;
 	
 #if WITH_EDITORONLY_DATA
 	// Compilation options to use in editor and for packaging for this object.
diff --git a/Engine/Plugins/Experimental/Mutable/Source/CustomizableObject/Internal/MuCO/CustomizableObjectInstancePrivate.h b/Engine/Plugins/Experimental/Mutable/Source/CustomizableObject/Internal/MuCO/CustomizableObjectInstancePrivate.h
--- a/Engine/Plugins/Experimental/Mutable/Source/CustomizableObject/Internal/MuCO/CustomizableObjectInstancePrivate.h
+++ b/Engine/Plugins/Experimental/Mutable/Source/CustomizableObject/Internal/MuCO/CustomizableObjectInstancePrivate.h
@@ -168,8 +168,6 @@ class CUSTOMIZABLEOBJECT_API UCustomizableInstancePrivate : public UObject
 public:
 	GENERATED_BODY()
 
-	UCustomizableInstancePrivate();
-
 	/** The generated skeletal meshes for this Instance, one for each component */
 	UPROPERTY(Transient, VisibleAnywhere, Category = NoCategory)
 	TArray<TObjectPtr<USkeletalMesh>> SkeletalMeshes;
@@ -308,17 +306,17 @@ public:
 	UCustomizableObjectInstance* GetPublic() const;
 
 	// If any components are using this instance, they will store the min of their distances to the player here every frame for LOD purposes
-	float MinSquareDistFromComponentToPlayer;
-	float LastMinSquareDistFromComponentToPlayer; // The same as the previous dist for last frame
+	float MinSquareDistFromComponentToPlayer = FLT_MAX;
+	float LastMinSquareDistFromComponentToPlayer = FLT_MAX; // The same as the previous dist for last frame
 												
 	// This is the LODs that the Customizable Object has
-	int32 NumLODsAvailable;
+	int32 NumLODsAvailable = INT32_MAX;
 
 	// First SkeletalMesh LOD we can generate on the running platform
-	uint8 FirstLODAvailable;
+	uint8 FirstLODAvailable = 0;
 
-	// Maximum number of SkeletalMesh LODs to stream
-	uint8 NumMaxLODsToStream;
+	// First Non-streamable SkeletalMesh LOD
+	uint8 FirstResidentLOD = MAX_MESH_LOD_COUNT;
 	
 	UPROPERTY(Transient)
 	TArray<FCustomizableInstanceComponentData> ComponentsData;
diff --git a/Engine/Plugins/Experimental/Mutable/Source/CustomizableObject/Internal/MuCO/CustomizableObjectSystemPrivate.h b/Engine/Plugins/Experimental/Mutable/Source/CustomizableObject/Internal/MuCO/CustomizableObjectSystemPrivate.h
--- a/Engine/Plugins/Experimental/Mutable/Source/CustomizableObject/Internal/MuCO/CustomizableObjectSystemPrivate.h
+++ b/Engine/Plugins/Experimental/Mutable/Source/CustomizableObject/Internal/MuCO/CustomizableObjectSystemPrivate.h
@@ -504,6 +504,9 @@ public:
 	bool bReuseInstanceTextures = false;
 	bool bUseMeshCache = false;
 	
+	/** Whether the mesh to generate should support Mesh LOD streaming or not. */
+	bool bStreamMeshLODs = false;
+
 	/** This option comes from the operation request. It is used to reduce the number of mipmaps that mutable must generate for images.  */
 	int32 MipsToSkip = 0;
 
@@ -513,6 +516,7 @@ public:
 	int32 NumComponents = 0;
 	int32 NumLODsAvailable = 0;
 	int32 FirstLODAvailable = 0;
+	uint8 FirstResidentLOD = 0;
 
 	TMap<uint32, FTexturePlatformData*> ImageToPlatformDataMap;
 
diff --git a/Engine/Plugins/Experimental/Mutable/Source/CustomizableObject/Internal/MuCO/UnrealConversionUtils.h b/Engine/Plugins/Experimental/Mutable/Source/CustomizableObject/Internal/MuCO/UnrealConversionUtils.h
--- a/Engine/Plugins/Experimental/Mutable/Source/CustomizableObject/Internal/MuCO/UnrealConversionUtils.h
+++ b/Engine/Plugins/Experimental/Mutable/Source/CustomizableObject/Internal/MuCO/UnrealConversionUtils.h
@@ -44,6 +44,16 @@ namespace UnrealConversionUtils
 		const int32 InFirstBoneMapIndex);
 
 
+	/** Initializes the LODResource's VertexBuffers with dummy data to prepare it for streaming.
+	 * @param LODResource - LODRenderData to update
+	 * @param InMutableMesh - Mutable mesh to be used as reference for the section data update on the Skeletal Mesh
+	 * @param bAllowCPUAccess - Keeps this LODs data on the CPU so it can be used for things such as sampling in FX.
+	 */
+	CUSTOMIZABLEOBJECT_API void InitVertexBuffersWithDummyData(
+		FSkeletalMeshLODRenderData& LODResource,
+		const mu::Ptr<const mu::Mesh> InMutableMesh,
+		const bool bAllowCPUAccess);
+
 	/** Performs a copy of the data found on the vertex buffers on the mutable mesh to the buffers of the skeletal mesh
 	 * @param LODResource - LODRenderData to update
 	 * @param InMutableMesh - Mutable mesh to be used as reference for the section data update on the Skeletal Mesh
@@ -55,6 +65,15 @@ namespace UnrealConversionUtils
 		const bool bAllowCPUAccess);
 
 	
+	/**
+	 * Initializes the LODResource's IndexBuffers with dummy data to prepare it for streaming.
+	 * @param LODResource - LODRenderData to be updated.
+	 * @param InMutableMesh - The mutable mesh whose index buffers you want to work with
+	 */
+	CUSTOMIZABLEOBJECT_API void InitIndexBuffersWithDummyData(
+		FSkeletalMeshLODRenderData& LODResource,
+		const mu::Ptr<const mu::Mesh> InMutableMesh);
+
 	/**
 	 *Performs a copy of the data found on the index buffers on the mutable mesh to the buffers of the skeletal mesh
 	 * @param LODResource - LODRenderData to be updated.
diff --git a/Engine/Plugins/Experimental/Mutable/Source/CustomizableObject/Private/MuCO/CustomizableObjectInstance.cpp b/Engine/Plugins/Experimental/Mutable/Source/CustomizableObject/Private/MuCO/CustomizableObjectInstance.cpp
--- a/Engine/Plugins/Experimental/Mutable/Source/CustomizableObject/Private/MuCO/CustomizableObjectInstance.cpp
+++ b/Engine/Plugins/Experimental/Mutable/Source/CustomizableObject/Private/MuCO/CustomizableObjectInstance.cpp
@@ -21,6 +21,7 @@
 #include "Tasks/Task.h"
 
 #include "MuCO/CustomizableObjectSystemPrivate.h"
+#include "MuCO/CustomizableObjectSkeletalMesh.h"
 #include "MuCO/CustomizableInstanceLODManagement.h"
 #include "MuCO/CustomizableObjectInstancePrivate.h"
 #include "MuCO/CustomizableObjectExtension.h"
@@ -223,7 +224,8 @@ void UCustomizableInstancePrivate::InitCustomizableObjectData(const UCustomizabl
 	// Init LOD Data
 	NumLODsAvailable = InCustomizableObject->GetNumLODs();
 	FirstLODAvailable = InCustomizableObject->LODSettings.FirstLODAvailable;
-	NumMaxLODsToStream = InCustomizableObject->LODSettings.bLODStreamingEnabled ? InCustomizableObject->LODSettings.NumLODsToStream : 0;
+	FirstResidentLOD = InCustomizableObject->LODSettings.bLODStreamingEnabled ? InCustomizableObject->LODSettings.NumLODsToStream : 0;
+	FirstResidentLOD = FMath::Clamp(FirstResidentLOD, FirstLODAvailable, NumLODsAvailable);
 
 	// Init Component Data
 	FCustomizableInstanceComponentData TemplateComponentData;
@@ -1820,7 +1822,14 @@ bool UCustomizableInstancePrivate::UpdateSkeletalMesh_PostBeginUpdate0(UCustomiz
 		// Create and initialize the SkeletalMesh for this component
 		MUTABLE_CPUPROFILER_SCOPE(ConstructMesh);
 
-		SkeletalMeshes[Component.Id] = NewObject<USkeletalMesh>(GetTransientPackage(), NAME_None, RF_Transient);
+		if (OperationData->bStreamMeshLODs)
+		{
+			SkeletalMeshes[Component.Id] = UCustomizableObjectSkeletalMesh::CreateSkeletalMesh(OperationData, *Public, Component.Id);
+		}
+		else
+		{
+			SkeletalMeshes[Component.Id] = NewObject<USkeletalMesh>(GetTransientPackage(), NAME_None, RF_Transient);
+		}
 
 		USkeletalMesh* SkeletalMesh = SkeletalMeshes[Component.Id];
 		check(SkeletalMesh);
@@ -2036,8 +2045,14 @@ bool UCustomizableInstancePrivate::UpdateSkeletalMesh_PostBeginUpdate0(UCustomiz
 				BuildOrCopyMorphTargetsData(OperationData, SkeletalMesh, OldSkeletalMesh, Public, ComponentIndex);
 				BuildOrCopyClothingData(OperationData, SkeletalMesh, OldSkeletalMesh, Public, ComponentIndex);
 
-				ensure(SkeletalMesh->GetResourceForRendering()->LODRenderData.Num() > 0);
+				FSkeletalMeshRenderData* RenderData = SkeletalMesh->GetResourceForRendering();
+				ensure(RenderData && RenderData->LODRenderData.Num() > 0);
 				ensure(SkeletalMesh->GetLODInfoArray().Num() > 0);
+
+				for (FSkeletalMeshLODRenderData& LODResource : RenderData->LODRenderData)
+				{
+					UnrealConversionUtils::UpdateSkeletalMeshLODRenderDataBuffersSize(LODResource);
+				}
 			}
 
 			if (OperationData->bUseMeshCache)
@@ -2057,18 +2072,6 @@ bool UCustomizableInstancePrivate::UpdateSkeletalMesh_PostBeginUpdate0(UCustomiz
 }
 
 
-UCustomizableInstancePrivate::UCustomizableInstancePrivate()
-{
-	MinSquareDistFromComponentToPlayer = FLT_MAX;
-	LastMinSquareDistFromComponentToPlayer = FLT_MAX;
-	
-	NumLODsAvailable = INT32_MAX;
-	FirstLODAvailable = 0;
-
-	NumMaxLODsToStream = MAX_MESH_LOD_COUNT;
-}
-
-
 UCustomizableObjectInstance* UCustomizableObjectInstance::Clone()
 {
 	MUTABLE_CPUPROFILER_SCOPE(UCustomizableObjectInstance::Clone);
@@ -3168,8 +3171,7 @@ void UCustomizableInstancePrivate::InitSkeletalMeshData(const TSharedRef<FUpdate
 
 	check(SkeletalMesh);
 
-	// Mutable skeletal meshes are generated dynamically in-game and cannot be streamed from disk
-	SkeletalMesh->NeverStream = 1;
+	SkeletalMesh->NeverStream = !OperationData->bStreamMeshLODs;
 
 	SkeletalMesh->SetImportedBounds(RefSkeletalMeshData.Bounds);
 	SkeletalMesh->SetPostProcessAnimBlueprint(RefSkeletalMeshData.PostProcessAnimInst.Get());
@@ -3223,12 +3225,20 @@ void UCustomizableInstancePrivate::InitSkeletalMeshData(const TSharedRef<FUpdate
 		SkeletalMesh->AllocateResourceForRendering();
 
 		FSkeletalMeshRenderData* RenderData = SkeletalMesh->GetResourceForRendering();
+		RenderData->NumInlinedLODs = OperationData->NumLODsAvailable - OperationData->FirstResidentLOD;
+		RenderData->NumNonOptionalLODs = OperationData->NumLODsAvailable - OperationData->FirstLODAvailable;
+		RenderData->CurrentFirstLODIdx = OperationData->FirstResidentLOD;
+		RenderData->LODBiasModifier = 0;
+
 		for (int32 LODIndex = 0; LODIndex < OperationData->NumLODsAvailable; ++LODIndex)
 		{
 			RenderData->LODRenderData.Add(new FSkeletalMeshLODRenderData());
+			
+			FSkeletalMeshLODRenderData& LODRenderData = RenderData->LODRenderData[LODIndex];
+			LODRenderData.bIsLODOptional = LODIndex < OperationData->FirstLODAvailable;
+			LODRenderData.bStreamedDataInlined = LODIndex >= OperationData->FirstResidentLOD;
 
 			const FMutableRefLODData& LODData = RefSkeletalMeshData.LODData[LODIndex];
-
 			FSkeletalMeshLODInfo& LODInfo = SkeletalMesh->AddLODInfo();
 			LODInfo.ScreenSize = LODData.LODInfo.ScreenSize;
 			LODInfo.LODHysteresis = LODData.LODInfo.LODHysteresis;
@@ -4687,17 +4697,91 @@ bool UCustomizableInstancePrivate::BuildOrCopyRenderData(const TSharedRef<FUpdat
 
 		FSkeletalMeshLODRenderData& LODResource = RenderData->LODRenderData[LODIndex];
 
+		// Set active and required bones
+		LODResource.ActiveBoneIndices.Append(Component.ActiveBones);
+		LODResource.RequiredBones.Append(Component.ActiveBones);
+
+		// Set RenderSections
 		UnrealConversionUtils::SetupRenderSections(
 			LODResource,
 			Component.Mesh,
 			OperationData->InstanceUpdateData.BoneMaps,
 			Component.FirstBoneMap);
 
-		UnrealConversionUtils::CopyMutableVertexBuffers(
-			LODResource,
-			Component.Mesh,
-			SkeletalMesh->GetLODInfo(LODIndex)->bAllowCPUAccess);
+		if (LODResource.bStreamedDataInlined) // Non-streamable LOD
+		{
+			// Copy Vertices
+			UnrealConversionUtils::CopyMutableVertexBuffers(
+				LODResource,
+				Component.Mesh,
+				SkeletalMesh->GetLODInfo(LODIndex)->bAllowCPUAccess);
+
+			// Copy indices.
+			if (!UnrealConversionUtils::CopyMutableIndexBuffers(LODResource, Component.Mesh))
+			{
+				// End with failure
+				return false;
+			}
+
+			// Copy SkinWeightProfiles
+			if (!ModelResources.SkinWeightProfilesInfo.IsEmpty())
+			{
+				bool bHasSkinWeightProfiles = false;
+
+				const mu::FMeshBufferSet& MutableMeshVertexBuffers = Component.Mesh->GetVertexBuffers();
+
+				const int32 SkinWeightProfilesCount = ModelResources.SkinWeightProfilesInfo.Num();
+				for (int32 ProfileIndex = 0; ProfileIndex < SkinWeightProfilesCount; ++ProfileIndex)
+				{
+					const int32 ProfileSemanticsIndex = ProfileIndex + 10;
+					int32 BoneIndicesBufferIndex, BoneIndicesBufferChannelIndex;
+					MutableMeshVertexBuffers.FindChannel(mu::MBS_BONEINDICES, ProfileSemanticsIndex, &BoneIndicesBufferIndex, &BoneIndicesBufferChannelIndex);
+
+					int32 BoneWeightsBufferIndex, BoneWeightsBufferChannelIndex;
+					MutableMeshVertexBuffers.FindChannel(mu::MBS_BONEWEIGHTS, ProfileSemanticsIndex, &BoneWeightsBufferIndex, &BoneWeightsBufferChannelIndex);
+
+					if (BoneIndicesBufferIndex < 0 || BoneIndicesBufferIndex != BoneWeightsBufferIndex)
+					{
+						continue;
+					}
+
+					if (!bHasSkinWeightProfiles)
+					{
+						LODResource.SkinWeightProfilesData.Init(&LODResource.SkinWeightVertexBuffer);
+						bHasSkinWeightProfiles = true;
+					}
+
+					const FMutableSkinWeightProfileInfo& Profile = ModelResources.SkinWeightProfilesInfo[ProfileIndex];
 
+					const FSkinWeightProfileInfo* ExistingProfile = SkeletalMesh->GetSkinWeightProfiles().FindByPredicate(
+						[&Profile](const FSkinWeightProfileInfo& P) { return P.Name == Profile.Name; });
+
+					if (!ExistingProfile)
+					{
+						SkeletalMesh->AddSkinWeightProfile({ Profile.Name, Profile.DefaultProfile, Profile.DefaultProfileFromLODIndex });
+					}
+
+					UnrealConversionUtils::CopyMutableSkinWeightProfilesBuffers(
+						LODResource,
+						Profile.Name,
+						MutableMeshVertexBuffers,
+						BoneIndicesBufferIndex);
+				}
+			}
+		}
+		else // Streamable LOD. 
+		{
+			// Init VertexBuffers for streaming
+			UnrealConversionUtils::InitVertexBuffersWithDummyData(
+				LODResource,
+				Component.Mesh,
+				SkeletalMesh->GetLODInfo(LODIndex)->bAllowCPUAccess);
+
+			// Init IndexBuffers for streaming
+			UnrealConversionUtils::InitIndexBuffersWithDummyData(LODResource, Component.Mesh);
+
+			// SkinWeightProfilesInfo Not supported yet
+		}
 
 		if (LODResource.StaticVertexBuffers.ColorVertexBuffer.GetNumVertices())
 		{
@@ -4720,72 +4804,6 @@ bool UCustomizableInstancePrivate::BuildOrCopyRenderData(const TSharedRef<FUpdat
 			FSlateNotificationManager::Get().AddNotification(Info);
 #endif
 		}
-
-		// Update active and required bones
-		LODResource.ActiveBoneIndices.Append(Component.ActiveBones);
-		LODResource.RequiredBones.Append(Component.ActiveBones);
-
-		if (!ModelResources.SkinWeightProfilesInfo.IsEmpty())
-		{
-			bool bHasSkinWeightProfiles = false;
-
-			const mu::FMeshBufferSet& MutableMeshVertexBuffers = Component.Mesh->GetVertexBuffers();
-			
-			const int32 SkinWeightProfilesCount = ModelResources.SkinWeightProfilesInfo.Num();
-			for (int32 ProfileIndex = 0; ProfileIndex < SkinWeightProfilesCount; ++ProfileIndex)
-			{
-				const int32 ProfileSemanticsIndex = ProfileIndex + 10;
-				int32 BoneIndicesBufferIndex, BoneIndicesBufferChannelIndex;
-				MutableMeshVertexBuffers.FindChannel(mu::MBS_BONEINDICES, ProfileSemanticsIndex, &BoneIndicesBufferIndex, &BoneIndicesBufferChannelIndex);
-
-				int32 BoneWeightsBufferIndex, BoneWeightsBufferChannelIndex;
-				MutableMeshVertexBuffers.FindChannel(mu::MBS_BONEWEIGHTS, ProfileSemanticsIndex, &BoneWeightsBufferIndex, &BoneWeightsBufferChannelIndex);
-
-				if (BoneIndicesBufferIndex < 0 || BoneIndicesBufferIndex != BoneWeightsBufferIndex)
-				{
-					continue;
-				}
-
-				if (!bHasSkinWeightProfiles)
-				{
-					LODResource.SkinWeightProfilesData.Init(&LODResource.SkinWeightVertexBuffer);
-					bHasSkinWeightProfiles = true;
-				}
-
-				const FMutableSkinWeightProfileInfo& Profile = ModelResources.SkinWeightProfilesInfo[ProfileIndex];
-
-				const FSkinWeightProfileInfo* ExistingProfile = SkeletalMesh->GetSkinWeightProfiles().FindByPredicate(
-					[&Profile](const FSkinWeightProfileInfo& P) { return P.Name == Profile.Name; });
-
-				if (!ExistingProfile)
-				{
-					SkeletalMesh->AddSkinWeightProfile({ Profile.Name, Profile.DefaultProfile, Profile.DefaultProfileFromLODIndex });
-				}
-
-				UnrealConversionUtils::CopyMutableSkinWeightProfilesBuffers(
-					LODResource,
-					Profile.Name,
-					MutableMeshVertexBuffers,
-					BoneIndicesBufferIndex);
-			}
-		}
-
-		// Copy indices.
-		if (!UnrealConversionUtils::CopyMutableIndexBuffers(LODResource, Component.Mesh))
-		{
-			// End with failure
-			return false;
-		}
-
-		// Update LOD and streaming data
-		const FMutableRefLODRenderData& RefLODRenderData = ModelResources.ReferenceSkeletalMeshesData[Component.Id].LODData[LODIndex].RenderData;
-		LODResource.bIsLODOptional = RefLODRenderData.bIsLODOptional;
-		LODResource.bStreamedDataInlined = RefLODRenderData.bStreamedDataInlined;
-
-		// WARNING! BufferSize must be > 0 or all texture mips in this LOD will be requested at once. 
-		// USkeletalMesh::IsMaterialUsed checks this size to see if a material is being used. If it fails the textures used by it won't be included 
-		// in the map of textures to stream.
-		LODResource.BuffersSize = 1;
 	}
 
 	// Copy LODRenderData from the FirstGeneratedLOD to the LODs below
@@ -6375,7 +6393,7 @@ void UCustomizableObjectInstance::SetRequestedLODs(int32 InMinLOD, int32 , const
 	bool bUpdateRequestedLODs = false;
 	if (UCustomizableObjectSystem::GetInstance()->IsOnlyGenerateRequestedLODsEnabled())
 	{
-		const uint16 FirstNonStreamedLODIndex = FMath::Clamp(PrivateData->NumMaxLODsToStream, 0, MaxLODIdx);
+		const uint16 FirstNonStreamedLODIndex = FMath::Clamp(PrivateData->FirstResidentLOD, 0, MaxLODIdx);
 
 		const int32 ComponentCount = GetNumComponents();
 		if (ComponentCount != MutableUpdateCandidate.RequestedLODLevels.Num())
diff --git a/Engine/Plugins/Experimental/Mutable/Source/CustomizableObject/Private/MuCO/CustomizableObjectSystem.cpp b/Engine/Plugins/Experimental/Mutable/Source/CustomizableObject/Private/MuCO/CustomizableObjectSystem.cpp
--- a/Engine/Plugins/Experimental/Mutable/Source/CustomizableObject/Private/MuCO/CustomizableObjectSystem.cpp
+++ b/Engine/Plugins/Experimental/Mutable/Source/CustomizableObject/Private/MuCO/CustomizableObjectSystem.cpp
@@ -149,6 +149,8 @@ FUpdateContextPrivate::FUpdateContextPrivate(UCustomizableObjectInstance& InInst
 	Parameters = Descriptor.GetParameters();
 	NumComponents = InInstance.GetCustomizableObject()->GetComponentCount();
 	FirstLODAvailable = InInstance.GetCustomizableObject()->GetPrivate()->GetMinLODIndex();
+	FirstResidentLOD = InInstance.GetPrivate()->FirstResidentLOD;
+
 	MutableSystem = UCustomizableObjectSystem::GetInstance()->GetPrivate()->MutableSystem;	
 	check(MutableSystem);
 	
@@ -759,6 +761,22 @@ FMutableResourceCache& UCustomizableObjectSystemPrivate::GetObjectCache(const UC
 	return ModelResourcesCache.Last();
 }
 
+bool bForceStreamMeshLODs = false;
+
+static FAutoConsoleVariableRef CVarMutableForceStreamMeshLODs(
+	TEXT("Mutable.ForceStreamMeshLODs"),
+	bForceStreamMeshLODs,
+	TEXT("true/false - If true, and bStreamMeshLODs is enabled, all COs will stream mesh LODs. "),
+	ECVF_Default);
+
+
+bool bStreamMeshLODs = true;
+
+static FAutoConsoleVariableRef CVarMutableStreamMeshLODsEnabled(
+	TEXT("Mutable.StreamMeshLODsEnabled"),
+	bStreamMeshLODs,
+	TEXT("true/false - If true, enable generated meshes to stream mesh LODs. "),
+	ECVF_Default);
 
 int32 UCustomizableObjectSystemPrivate::EnableMutableProgressiveMipStreaming = 1;
 
@@ -1688,24 +1706,35 @@ namespace impl
 	void FixLODs(const TSharedRef<FUpdateContextPrivate>& Operation)
 	{
 		Operation->NumLODsAvailable = Operation->MutableInstance->GetLODCount();
+		
+		int32 CurrentMinLOD = Operation->bStreamMeshLODs ? 0 : Operation->GetMinLOD();
+		CurrentMinLOD = FMath::Clamp(CurrentMinLOD, Operation->FirstLODAvailable, Operation->NumLODsAvailable - 1);
+		Operation->SetMinLOD(CurrentMinLOD);
 
-		int32 CurrentMinLOD = FMath::Max(Operation->GetMinLOD(), Operation->FirstLODAvailable);
-
-		if (CurrentMinLOD >= Operation->NumLODsAvailable)
+		if (Operation->bStreamMeshLODs)
 		{
-			CurrentMinLOD = Operation->NumLODsAvailable - 1;
+			Operation->FirstResidentLOD = FMath::Clamp(Operation->FirstResidentLOD, Operation->FirstLODAvailable, Operation->NumLODsAvailable - 1);
 		}
-
-		Operation->SetMinLOD(CurrentMinLOD);
-
+		else 
+		{
+			Operation->FirstResidentLOD = Operation->FirstLODAvailable;
+		}
+		
 		// Initialize RequestedLODs to zero if not set
 		TArray<uint16> RequestedLODs = Operation->GetRequestedLODs();
 		RequestedLODs.SetNumZeroed(Operation->NumComponents);
 
 		for (int32 ComponentIndex = 0; ComponentIndex < Operation->NumComponents; ++ComponentIndex)
 		{
-			// Clamp value to the valid LOD range.
-			RequestedLODs[ComponentIndex] = FMath::Min(RequestedLODs[ComponentIndex], (uint16)(Operation->NumLODsAvailable - 1));
+			if (Operation->bStreamMeshLODs)
+			{
+				RequestedLODs[ComponentIndex] = CurrentMinLOD;
+			}
+			else
+			{
+				// Clamp value to the valid LOD range.
+				RequestedLODs[ComponentIndex] = FMath::Min(RequestedLODs[ComponentIndex], (uint16)(Operation->NumLODsAvailable - 1));
+			}
 		}
 
 		Operation->SetRequestedLODs(RequestedLODs);
@@ -3054,6 +3083,9 @@ namespace impl
 		//-------------------------------------------------------------
 		Operation->InstanceID = Operation->bLiveUpdateMode ? CandidateInstancePrivateData->LiveUpdateModeInstanceID : 0;
 		Operation->bUseMeshCache = CustomizableObject->bEnableMeshCache && !Operation->bLiveUpdateMode && CVarEnableMeshCache.GetValueOnGameThread();
+
+		const bool bStreamingEnabled = (CustomizableObject->bEnableMeshStreaming || bForceStreamMeshLODs) && bStreamMeshLODs;
+		Operation->bStreamMeshLODs = bStreamingEnabled && IStreamingManager::Get().IsRenderAssetStreamingEnabled(EStreamableRenderAssetType::SkeletalMesh);
 #if WITH_EDITOR
 		Operation->PixelFormatOverride = SystemPrivateData->ImageFormatOverrideFunc;
 #endif
diff --git a/Engine/Plugins/Experimental/Mutable/Source/CustomizableObject/Private/MuCO/UnrealConversionUtils.cpp b/Engine/Plugins/Experimental/Mutable/Source/CustomizableObject/Private/MuCO/UnrealConversionUtils.cpp
--- a/Engine/Plugins/Experimental/Mutable/Source/CustomizableObject/Private/MuCO/UnrealConversionUtils.cpp
+++ b/Engine/Plugins/Experimental/Mutable/Source/CustomizableObject/Private/MuCO/UnrealConversionUtils.cpp
@@ -192,6 +192,98 @@ namespace UnrealConversionUtils
 	}
 
 
+	void InitVertexBuffersWithDummyData(
+		FSkeletalMeshLODRenderData& LODResource,
+		const mu::Ptr<const mu::Mesh> InMutableMesh,
+		const bool bAllowCPUAccess)
+	{
+		MUTABLE_CPUPROFILER_SCOPE(InitVertexBuffersWithDummyData);
+
+		const mu::FMeshBufferSet& MutableMeshVertexBuffers = InMutableMesh->GetVertexBuffers();
+		check(MutableMeshVertexBuffers.GetElementCount() > 0);
+
+		const bool bUseFullPrecisionUVs = true;
+
+		const int32 NumVertices = 1;
+		const int32 NumTexCoords = MutableMeshVertexBuffers.GetBufferChannelCount(MUTABLE_VERTEXBUFFER_TEXCOORDS);
+
+		FStaticMeshVertexBuffers_InitWithMutableData(
+			LODResource.StaticVertexBuffers,
+			NumVertices,
+			NumTexCoords,
+			bUseFullPrecisionUVs,
+			bAllowCPUAccess,
+			MutableMeshVertexBuffers.GetBufferData(MUTABLE_VERTEXBUFFER_POSITION),
+			MutableMeshVertexBuffers.GetBufferData(MUTABLE_VERTEXBUFFER_TANGENT),
+			MutableMeshVertexBuffers.GetBufferData(MUTABLE_VERTEXBUFFER_TEXCOORDS)
+		);
+
+		mu::MESH_BUFFER_FORMAT BoneIndexFormat = mu::MBF_NONE;
+		int32 NumBoneInfluences = 0;
+		int32 BoneIndexBuffer = -1;
+		int32 BoneIndexChannel = -1;
+		MutableMeshVertexBuffers.FindChannel(mu::MBS_BONEINDICES, 0, &BoneIndexBuffer, &BoneIndexChannel);
+		if (BoneIndexBuffer >= 0 || BoneIndexChannel >= 0)
+		{
+			MutableMeshVertexBuffers.GetChannel(BoneIndexBuffer, BoneIndexChannel,
+				nullptr, nullptr, &BoneIndexFormat, &NumBoneInfluences, nullptr);
+		}
+
+		mu::MESH_BUFFER_FORMAT BoneWeightFormat = mu::MBF_NONE;
+		int32 BoneWeightBuffer = -1;
+		int32 BoneWeightChannel = -1;
+		MutableMeshVertexBuffers.FindChannel(mu::MBS_BONEWEIGHTS, 0, &BoneWeightBuffer, &BoneWeightChannel);
+		if (BoneWeightBuffer >= 0 || BoneWeightChannel >= 0)
+		{
+			MutableMeshVertexBuffers.GetChannel(BoneWeightBuffer, BoneWeightChannel,
+				nullptr, nullptr, &BoneWeightFormat, nullptr, nullptr);
+		}
+
+		if (BoneIndexFormat == mu::MBF_UINT16)
+		{
+			LODResource.SkinWeightVertexBuffer.SetUse16BitBoneIndex(true);
+		}
+
+		if (BoneWeightFormat == mu::MBF_NUINT16)
+		{
+			LODResource.SkinWeightVertexBuffer.SetUse16BitBoneWeight(true);
+		}
+
+		// Init skin weight buffer
+		FSkinWeightVertexBuffer_InitWithMutableData(
+			LODResource.SkinWeightVertexBuffer,
+			NumVertices,
+			NumBoneInfluences * NumVertices,
+			NumBoneInfluences,
+			bAllowCPUAccess,
+			MutableMeshVertexBuffers.GetBufferData(BoneIndexBuffer),
+			MutableMeshVertexBuffers.GetElementSize(BoneIndexBuffer)
+		);
+
+		// Optional buffers
+		for (int32 Buffer = MUTABLE_VERTEXBUFFER_TEXCOORDS + 1; Buffer < MutableMeshVertexBuffers.GetBufferCount(); ++Buffer)
+		{
+			if (MutableMeshVertexBuffers.GetBufferChannelCount(Buffer) > 0)
+			{
+				mu::MESH_BUFFER_SEMANTIC Semantic;
+				mu::MESH_BUFFER_FORMAT Format;
+				int32 SemanticIndex;
+				int32 ComponentCount;
+				int32 Offset;
+				MutableMeshVertexBuffers.GetChannel(Buffer, 0, &Semantic, &SemanticIndex, &Format, &ComponentCount, &Offset);
+
+				// colour buffer?
+				if (Semantic == mu::MBS_COLOUR)
+				{
+					const void* DataPtr = MutableMeshVertexBuffers.GetBufferData(Buffer);
+					FColorVertexBuffers_InitWithMutableData(LODResource.StaticVertexBuffers, NumVertices, DataPtr);
+					check(LODResource.StaticVertexBuffers.ColorVertexBuffer.GetStride() == MutableMeshVertexBuffers.GetElementSize(Buffer));
+				}
+			}
+		}
+	}
+
+
 	void CopyMutableVertexBuffers(
 		FSkeletalMeshLODRenderData& LODResource,
 		const mu::MeshPtrConst MutableMesh,
@@ -274,26 +366,8 @@ namespace UnrealConversionUtils
 				int32 Offset;
 				MutableMeshVertexBuffers.GetChannel(Buffer, 0, &Semantic, &SemanticIndex, &Format, &ComponentCount, &Offset);
 
-				if (Semantic == mu::MBS_BONEINDICES)
-				{
-					const int32 BonesPerVertex = ComponentCount;
-					const int32 NumBones = BonesPerVertex * NumVertices;
-
-					check(FGPUBaseSkinVertexFactory::UseUnlimitedBoneInfluences(BonesPerVertex) ||
-						!LODResource.SkinWeightVertexBuffer.GetVariableBonesPerVertex());
-					FSkinWeightVertexBuffer_InitWithMutableData(
-						LODResource.SkinWeightVertexBuffer,
-						NumVertices,
-						NumBones,
-						NumBoneInfluences,
-						bAllowCPUAccess,
-						MutableMeshVertexBuffers.GetBufferData(Buffer),
-						MutableMeshVertexBuffers.GetBufferDataSize(Buffer)
-					);
-				}
-
 				// colour buffer?
-				else if (Semantic == mu::MBS_COLOUR)
+				if (Semantic == mu::MBS_COLOUR)
 				{
 					const void* DataPtr = MutableMeshVertexBuffers.GetBufferData(Buffer);
 					FColorVertexBuffers_InitWithMutableData(LODResource.StaticVertexBuffers, NumVertices, DataPtr);
@@ -303,6 +377,20 @@ namespace UnrealConversionUtils
 		}
 	}
 
+		
+	void InitIndexBuffersWithDummyData(FSkeletalMeshLODRenderData& LODResource, const mu::Ptr<const mu::Mesh> InMutableMesh)
+	{
+		MUTABLE_CPUPROFILER_SCOPE(InitIndexBuffersWithDummyData);
+
+		check(InMutableMesh->GetIndexBuffers().GetElementCount() > 0);
+		
+		const int32 NumIndices = 3;
+		const int32 ElementSize = InMutableMesh->GetIndexBuffers().GetElementSize(0);
+
+		LODResource.MultiSizeIndexContainer.CreateIndexBuffer(ElementSize);
+		LODResource.MultiSizeIndexContainer.GetIndexBuffer()->Insert(0, NumIndices);
+	}
+
 
 	bool CopyMutableIndexBuffers(
 		FSkeletalMeshLODRenderData& LODResource,
@@ -563,7 +651,7 @@ namespace UnrealConversionUtils
 	}
 
 
-	CUSTOMIZABLEOBJECT_API void UpdateSkeletalMeshLODRenderDataBuffersSize(FSkeletalMeshLODRenderData& LODResource)
+	void UpdateSkeletalMeshLODRenderDataBuffersSize(FSkeletalMeshLODRenderData& LODResource)
 	{
 		LODResource.BuffersSize = 0;
 		
