[mutable] Check if loading or not when saving the benchmarking data
- Reported in the benchmarking file and logs if the level has ran BeginPlay to know if the level was loaded at the time of the mesh update.
- Added state of the instnace descriptor for each of the image updates

diff --git a/Engine/Plugins/Experimental/Mutable/Source/CustomizableObject/Internal/MuCO/CustomizableObjectMipDataProvider.h b/Engine/Plugins/Experimental/Mutable/Source/CustomizableObject/Internal/MuCO/CustomizableObjectMipDataProvider.h
--- a/Engine/Plugins/Experimental/Mutable/Source/CustomizableObject/Internal/MuCO/CustomizableObjectMipDataProvider.h
+++ b/Engine/Plugins/Experimental/Mutable/Source/CustomizableObject/Internal/MuCO/CustomizableObjectMipDataProvider.h
@@ -49,6 +49,8 @@ public:
 	int32 GetState() const;
 
 	const TArray<mu::Ptr<const mu::Image>>& GetImageParameterValues() const;
+
+	FString CapturedDescriptor;
 	
 private:
 	FString CustomizableObjectPathName;
diff --git a/Engine/Plugins/Experimental/Mutable/Source/CustomizableObject/Internal/MuCO/CustomizableObjectSystemPrivate.h b/Engine/Plugins/Experimental/Mutable/Source/CustomizableObject/Internal/MuCO/CustomizableObjectSystemPrivate.h
--- a/Engine/Plugins/Experimental/Mutable/Source/CustomizableObject/Internal/MuCO/CustomizableObjectSystemPrivate.h
+++ b/Engine/Plugins/Experimental/Mutable/Source/CustomizableObject/Internal/MuCO/CustomizableObjectSystemPrivate.h
@@ -533,6 +533,7 @@ public:
 	TArray<bool> MeshChanged;
 	
 	bool UpdateStarted = false;
+	bool bLevelBegunPlay = false;
 
 	// Update stats
 	double StartQueueTime = 0.0;
diff --git a/Engine/Plugins/Experimental/Mutable/Source/CustomizableObject/Internal/MuCO/LogBenchmarkUtil.h b/Engine/Plugins/Experimental/Mutable/Source/CustomizableObject/Internal/MuCO/LogBenchmarkUtil.h
--- a/Engine/Plugins/Experimental/Mutable/Source/CustomizableObject/Internal/MuCO/LogBenchmarkUtil.h
+++ b/Engine/Plugins/Experimental/Mutable/Source/CustomizableObject/Internal/MuCO/LogBenchmarkUtil.h
@@ -9,6 +9,7 @@
 
 class UCustomizableObjectInstance;
 class UTexture2D;
+struct FCustomizableObjectInstanceDescriptor;
 
 
 /** Stat which automatically gets updated to Insights when modified. */
@@ -87,7 +88,7 @@ public:
 	/** Gathers update stats when it has finished. */
 	void FinishUpdateMesh(const TSharedRef<FUpdateContextPrivate>& Context);
 
-	void FinishUpdateImage(const FString& CustomizableObjectPathName, const FString& InstancePathName, double TaskUpdateImageTime, const int64 TaskUpdateImageMemoryPeak, const int64 TaskUpdateImageRealMemoryPeak) const;
+	void FinishUpdateImage(const FString& CustomizableObjectPathName, const FString& InstancePathName, const FString& InstanceDescriptor, double TaskUpdateImageTime, const int64 TaskUpdateImageMemoryPeak, const int64 TaskUpdateImageRealMemoryPeak) const;
 	
 	DECLARE_BENCHMARK_STAT(NumAllocatedTextures, uint32);
 	DECLARE_BENCHMARK_STAT(TextureGPUSize, uint64);
diff --git a/Engine/Plugins/Experimental/Mutable/Source/CustomizableObject/Private/MuCO/BusyWaits_Deprecated.cpp b/Engine/Plugins/Experimental/Mutable/Source/CustomizableObject/Private/MuCO/BusyWaits_Deprecated.cpp
--- a/Engine/Plugins/Experimental/Mutable/Source/CustomizableObject/Private/MuCO/BusyWaits_Deprecated.cpp
+++ b/Engine/Plugins/Experimental/Mutable/Source/CustomizableObject/Private/MuCO/BusyWaits_Deprecated.cpp
@@ -554,9 +554,11 @@ namespace CustomizableObjectMipDataProvider::ImplDeprecated
 
 			const FString& CustomizableObjectPathName = OperationData->UpdateContext->GetCustomizableObjectPathName();
 			const FString& InstancePathName = OperationData->UpdateContext->GetInstancePathName();
-			
+
+			const FString& Descriptor = OperationData->UpdateContext->CapturedDescriptor;
+
 			FFunctionGraphTask::CreateAndDispatchWhenReady(
-			[CustomizableObjectPathName, InstancePathName, Time, PeakMemory, RealMemoryPeak]()
+			[CustomizableObjectPathName, InstancePathName, Time, PeakMemory, RealMemoryPeak, Descriptor]()
 			{
 				if (!UCustomizableObjectSystem::IsCreated()) // We are shutting down
 				{
@@ -569,7 +571,7 @@ namespace CustomizableObjectMipDataProvider::ImplDeprecated
 					return;
 				}
 
-				System->GetPrivate()->LogBenchmarkUtil.FinishUpdateImage(CustomizableObjectPathName, InstancePathName, Time, PeakMemory, RealMemoryPeak);
+				System->GetPrivate()->LogBenchmarkUtil.FinishUpdateImage(CustomizableObjectPathName, InstancePathName, Descriptor, Time, PeakMemory, RealMemoryPeak);
 			},
 			TStatId{},
 			nullptr,
diff --git a/Engine/Plugins/Experimental/Mutable/Source/CustomizableObject/Private/MuCO/CustomizableObjectInstance.cpp b/Engine/Plugins/Experimental/Mutable/Source/CustomizableObject/Private/MuCO/CustomizableObjectInstance.cpp
--- a/Engine/Plugins/Experimental/Mutable/Source/CustomizableObject/Private/MuCO/CustomizableObjectInstance.cpp
+++ b/Engine/Plugins/Experimental/Mutable/Source/CustomizableObject/Private/MuCO/CustomizableObjectInstance.cpp
@@ -5591,6 +5591,12 @@ void UCustomizableInstancePrivate::BuildMaterials(const TSharedRef<FUpdateContex
 		CustomizableObject->GetPrivate()->GetModel(),
 		OperationData->Parameters,
 	    OperationData->GetCapturedDescriptor().GetState());
+
+	// Cache the descriptor as a string if we want to later report it using our benchmark utility. 
+	if (CVarEnableBenchmark.GetValueOnAnyThread())
+	{
+		UpdateContext->CapturedDescriptor = OperationData->GetCapturedDescriptor().ToString();
+	}
 	
 	const bool bReuseTextures = OperationData->bReuseInstanceTextures;
 
diff --git a/Engine/Plugins/Experimental/Mutable/Source/CustomizableObject/Private/MuCO/CustomizableObjectMipDataProvider.cpp b/Engine/Plugins/Experimental/Mutable/Source/CustomizableObject/Private/MuCO/CustomizableObjectMipDataProvider.cpp
--- a/Engine/Plugins/Experimental/Mutable/Source/CustomizableObject/Private/MuCO/CustomizableObjectMipDataProvider.cpp
+++ b/Engine/Plugins/Experimental/Mutable/Source/CustomizableObject/Private/MuCO/CustomizableObjectMipDataProvider.cpp
@@ -281,9 +281,11 @@ namespace Impl
 
 						const FString& CustomizableObjectPathName = OperationData->UpdateContext->GetCustomizableObjectPathName();
 						const FString& InstancePathName = OperationData->UpdateContext->GetInstancePathName();
-					
+
+						const FString& Descriptor = OperationData->UpdateContext->CapturedDescriptor;
+
 						ExecuteOnGameThread(UE_SOURCE_LOCATION, 
-						[CustomizableObjectPathName, InstancePathName, Time, PeakMemory, RealMemoryPeak]()
+						[CustomizableObjectPathName, InstancePathName, Time, PeakMemory, RealMemoryPeak, Descriptor ]
 						{
 							if (!UCustomizableObjectSystem::IsCreated()) // We are shutting down
 							{
@@ -296,7 +298,7 @@ namespace Impl
 								return;
 							}
 
-							System->GetPrivate()->LogBenchmarkUtil.FinishUpdateImage(CustomizableObjectPathName, InstancePathName, Time, PeakMemory, RealMemoryPeak);
+							System->GetPrivate()->LogBenchmarkUtil.FinishUpdateImage(CustomizableObjectPathName, InstancePathName, Descriptor, Time, PeakMemory, RealMemoryPeak);
 						});
 					}
 
@@ -336,7 +338,7 @@ int32 FMutableTextureMipDataProvider::GetMips(const FTextureUpdateContext& Conte
 	check(!Texture->NeverStream);
 	const TIndirectArray<FTexture2DMipMap>& OwnerMips = Texture->GetPlatformMips();
 
-	int32 NumMips = OwnerMips.Num();
+	const int32 NumMips = OwnerMips.Num();
 	check(ImageRef.ImageID > 0);
 
 	// Maximum value to skip, will be minimized by the first mip level requested
diff --git a/Engine/Plugins/Experimental/Mutable/Source/CustomizableObject/Private/MuCO/CustomizableObjectSystem.cpp b/Engine/Plugins/Experimental/Mutable/Source/CustomizableObject/Private/MuCO/CustomizableObjectSystem.cpp
--- a/Engine/Plugins/Experimental/Mutable/Source/CustomizableObject/Private/MuCO/CustomizableObjectSystem.cpp
+++ b/Engine/Plugins/Experimental/Mutable/Source/CustomizableObject/Private/MuCO/CustomizableObjectSystem.cpp
@@ -2946,6 +2946,12 @@ namespace impl
 	{
 		MUTABLE_CPUPROFILER_SCOPE(Task_Game_StartUpdate)
 
+		// Check if a level has been loaded
+		if (CVarEnableBenchmark.GetValueOnAnyThread() && GWorld)
+		{
+			Operation->bLevelBegunPlay = GWorld->GetBegunPlay();
+		}
+
 		Operation->StartUpdateTime = FPlatformTime::Seconds();
 
 		UCustomizableObjectSystem::GetInstance()->GetPrivate()->MutableTaskGraph.AllowLaunchingMutableTaskLowPriority(false, false);
diff --git a/Engine/Plugins/Experimental/Mutable/Source/CustomizableObject/Private/MuCO/LogBenchmarkUtil.cpp b/Engine/Plugins/Experimental/Mutable/Source/CustomizableObject/Private/MuCO/LogBenchmarkUtil.cpp
--- a/Engine/Plugins/Experimental/Mutable/Source/CustomizableObject/Private/MuCO/LogBenchmarkUtil.cpp
+++ b/Engine/Plugins/Experimental/Mutable/Source/CustomizableObject/Private/MuCO/LogBenchmarkUtil.cpp
@@ -56,7 +56,9 @@ TSharedPtr<FArchive> CreateFile()
 	check(Archive);
 
 	const FString HeaderRow = TEXT(
-		"ID_CO;ID_COI;ID_UpdateType;ID_Descriptor;ID_UpdateResult;Time_Queue;Time_Update;Time_TaskGetMesh;Time_TaskLockCache;"
+		"ID_CO;ID_COI;ID_UpdateType;ID_Descriptor;ID_UpdateResult;"											// Basic identifying data
+		"Context_LevelBegunPlay;"																	// The context of the update
+		"Time_Queue;Time_Update;Time_TaskGetMesh;Time_TaskLockCache;"										// Time and memory data ...
 		"Time_TaskGetImages;Time_TaskConvertResources;Time_TaskCallbacks;Memory_Update;Memory_Update_Real;"
 		"Time_TaskUpdateImage;Memory_TaskUpdateImage;Memory_TaskUpdateImage_Real");
 
@@ -323,16 +325,19 @@ void FLogBenchmarkUtil::FinishUpdateMesh(const TSharedRef<FUpdateContextPrivate>
 	const double Time_TaskConvertResources = Context->TaskConvertResourcesTime * 1000;
 	const double Time_TaskCallbacks =  Context->TaskCallbacksTime * 1000;
 
+	// Context data to know in what situation did the update happen
+	const FString Context_LevelBegunPlay =  Context->bLevelBegunPlay ? TEXT("true") : TEXT("false");
+
 	const double Memory_UpdateEndPeakMB = (Context->UpdateEndPeakBytes / 1024.0) / 1024.0;
 	const double Memory_UpdateEndRealPeakMB = (Context->UpdateEndRealPeakBytes / 1024.0) / 1024.0;
 	
-	const FString UpdateString = FString::Printf(TEXT("%s;%s;%s;%s;%s;%f;%f;%f;%f;%f;%f;%f;%f;%f"), *ID_CO, *ID_COI, *ID_UpdateType, *ID_Descriptor, *ID_UpdateResult, Time_Queue, Time_Update, Time_TaskGetMesh, Time_TaskLockCache, Time_TaskGetImages, Time_TaskConvertResources, Time_TaskCallbacks, Memory_UpdateEndPeakMB, Memory_UpdateEndRealPeakMB);
+	const FString UpdateString = FString::Printf(TEXT("%s;%s;%s;%s;%s;%s;%f;%f;%f;%f;%f;%f;%f;%f;%f"), *ID_CO, *ID_COI, *ID_UpdateType, *ID_Descriptor, *ID_UpdateResult, *Context_LevelBegunPlay, Time_Queue, Time_Update, Time_TaskGetMesh, Time_TaskLockCache, Time_TaskGetImages, Time_TaskConvertResources, Time_TaskCallbacks, Memory_UpdateEndPeakMB, Memory_UpdateEndRealPeakMB);
 	LogBenchmarkUtil::Write(*Archive, UpdateString);
 	Archive->Flush();
 }
 
 
-void FLogBenchmarkUtil::FinishUpdateImage(const FString& CustomizableObjectPathName, const FString& InstancePathName, const double TaskUpdateImageTime, const int64 TaskUpdateImageMemoryPeak, const int64 TaskUpdateImageRealMemoryPeak) const
+void FLogBenchmarkUtil::FinishUpdateImage(const FString& CustomizableObjectPathName, const FString& InstancePathName, const FString& InstanceDescriptor, const double TaskUpdateImageTime, const int64 TaskUpdateImageMemoryPeak, const int64 TaskUpdateImageRealMemoryPeak) const
 {
 	check(IsInGameThread());
 
@@ -343,13 +348,13 @@ void FLogBenchmarkUtil::FinishUpdateImage(const FString& CustomizableObjectPathN
 
 	const FString& ID_CO = CustomizableObjectPathName;
 	const FString& ID_COI = InstancePathName;
+	const FString ID_Descriptor = InstanceDescriptor;
 	const FString ID_UpdateType = TEXT("Image");
 	const double Time_TaskUpdateImage = TaskUpdateImageTime * 1000;
-
 	const double Memory_TaskUpdateImagePeakMB = (TaskUpdateImageMemoryPeak / 1024.0) / 1024.0;
 	const double Memory_TaskUpdateImageRealPeakMB = (TaskUpdateImageRealMemoryPeak / 1024.0) / 1024.0;
 
-	const FString UpdateString = FString::Printf(TEXT("%s;%s;%s;;;;;;;;;;;;%f;%f;%f"), *ID_CO, *ID_COI, *ID_UpdateType, Time_TaskUpdateImage, Memory_TaskUpdateImagePeakMB,Memory_TaskUpdateImageRealPeakMB);
+	const FString UpdateString = FString::Printf(TEXT("%s;%s;%s;%s;;;;;;;;;;;;%f;%f;%f"), *ID_CO, *ID_COI, *ID_UpdateType, *ID_Descriptor, Time_TaskUpdateImage, Memory_TaskUpdateImagePeakMB,Memory_TaskUpdateImageRealPeakMB);
 	LogBenchmarkUtil::Write(*Archive, UpdateString);
 	Archive->Flush();	
 }
