[Mutable] Fix UV not showing the correct UV when Materials had the same name.

#rnx

diff --git a/Engine/Plugins/Experimental/Mutable/Source/CustomizableObjectEditor/Private/MuCOE/CustomizableObjectEditor.cpp b/Engine/Plugins/Experimental/Mutable/Source/CustomizableObjectEditor/Private/MuCOE/CustomizableObjectEditor.cpp
--- a/Engine/Plugins/Experimental/Mutable/Source/CustomizableObjectEditor/Private/MuCOE/CustomizableObjectEditor.cpp
+++ b/Engine/Plugins/Experimental/Mutable/Source/CustomizableObjectEditor/Private/MuCOE/CustomizableObjectEditor.cpp
@@ -1543,7 +1543,7 @@ void FCustomizableObjectEditor::OnObjectModified(UObject* Object)
 void FCustomizableObjectEditor::CompileObject()
 {
 	// Resetting viewport parameters
-	Viewport->SetDrawDefaultUVMaterial(true);
+	Viewport->SetDrawDefaultUVMaterial();
 
 	UE_LOG(LogMutable, Verbose, TEXT("PROFILE: -----------------------------------------------------------"));
 	UE_LOG(LogMutable, Verbose, TEXT("PROFILE: [ %16.8f ] FCustomizableObjectEditor::CompileObject start."), FPlatformTime::Seconds());
@@ -2275,7 +2275,7 @@ void FCustomizableObjectEditor::OnUpdatePreviewInstance()
 	}
 	
 	Viewport->SetPreviewComponents(PreviewSkeletalMeshComponents);
-	Viewport->SetDrawDefaultUVMaterial(false);
+	Viewport->SetDrawDefaultUVMaterial();
 	ViewportClient->Invalidate();
 	ViewportClient->ReSetAnimation();
 	ViewportClient->SetReferenceMeshMissingWarningMessage(false);
diff --git a/Engine/Plugins/Experimental/Mutable/Source/CustomizableObjectEditor/Private/MuCOE/CustomizableObjectEditorViewportClient.cpp b/Engine/Plugins/Experimental/Mutable/Source/CustomizableObjectEditor/Private/MuCOE/CustomizableObjectEditorViewportClient.cpp
--- a/Engine/Plugins/Experimental/Mutable/Source/CustomizableObjectEditor/Private/MuCOE/CustomizableObjectEditorViewportClient.cpp
+++ b/Engine/Plugins/Experimental/Mutable/Source/CustomizableObjectEditor/Private/MuCOE/CustomizableObjectEditorViewportClient.cpp
@@ -81,12 +81,6 @@ FCustomizableObjectEditorViewportClient::FCustomizableObjectEditorViewportClient
 
 	bShowBones = false;
 
-	MaterialToDrawInUVs = 0;
-	MaterialToDrawInUVsLOD = 0;
-	MaterialToDrawInUVsIndex = 0;
-	UVChannelToDrawInUVs = 0;
-	MaterialToDrawInUVsComponent = 0;
-
 	bReferenceMeshMissingWarningMessageVisible = false;
 
 	DrawHelper.bDrawPivot = false;
@@ -462,7 +456,7 @@ void FCustomizableObjectEditorViewportClient::Draw(FViewport* InViewport, FCanva
 	if(bDrawUVs)
 	{
 		constexpr int32 YPos = 24;
-		DrawUVs(InViewport, Canvas, YPos, MaterialToDrawInUVs);
+		DrawUVs(InViewport, Canvas, YPos);
 	}
 
 	if (bReferenceMeshMissingWarningMessageVisible)
@@ -491,16 +485,12 @@ namespace
 	}
 }
 
-void FCustomizableObjectEditorViewportClient::DrawUVs(FViewport* InViewport, FCanvas* InCanvas, int32 InTextYPos, const FString& MaterialName )
+void FCustomizableObjectEditorViewportClient::DrawUVs(FViewport* InViewport, FCanvas* InCanvas, int32 InTextYPos)
 {
-	//use the overriden LOD level
-	// TODO
-	const uint32 LODLevel = MaterialToDrawInUVsLOD; //FMath::Clamp(StaticMeshComponent->ForcedLodModel - 1, 0, StaticMesh->RenderData->LODResources.Num() - 1);
-
-	// TODO
-	int32 UVChannel = UVChannelToDrawInUVs; //StaticMeshEditorPtr.Pin()->GetCurrentUVChannel();
-
-	const uint32 ComponentIndex = MaterialToDrawInUVsComponent;
+	const uint32 ComponentIndex = UVDrawComponentIndex;
+	const uint32 LODLevel = UVDrawLODIndex; 	// TODO use the overriden LOD level
+	const int32 SectionIndex = UVDrawSectionIndex;
+	const int32 UVChannel = UVDrawUVIndex;
 
 	//draw a string showing what UV channel and LOD is being displayed
 	InCanvas->DrawShadowedString( 
@@ -526,9 +516,7 @@ void FCustomizableObjectEditorViewportClient::DrawUVs(FViewport* InViewport, FCa
 		BoxOrigin + FVector2D( BoxSize, BoxSize ),	// bottomright
 		BoxOrigin + FVector2D( 0, BoxSize ),		// bottomleft
 	};
-
-	const FVector Color(1.0f, 1.0f, 1.0f);
-
+	
 	//draw texture border
 	FLinearColor BorderColor = FLinearColor::White;
 	FBatchedElements* BatchedElements = InCanvas->GetBatchedElements(FCanvas::ET_Line);
@@ -543,11 +531,14 @@ void FCustomizableObjectEditorViewportClient::DrawUVs(FViewport* InViewport, FCa
 	BatchedElements->AddLine( FVector( Box[ 2 ], 0.0f ), FVector( Box[ 3 ], 0.0f ), BorderColor, HitProxyId );
 	BatchedElements->AddLine( FVector( Box[ 3 ], 0.0f ), FVector( Box[ 0 ], 0.0f ), BorderColor, HitProxyId );
 
-	if ( StaticMeshComponent.IsValid() && StaticMeshComponent->GetStaticMesh() && StaticMeshComponent->GetStaticMesh()->GetRenderData() )
+	if (StaticMeshComponent.IsValid() &&
+		StaticMeshComponent->GetStaticMesh() &&
+		StaticMeshComponent->GetStaticMesh()->GetRenderData() &&
+		StaticMeshComponent->GetStaticMesh()->GetRenderData()->LODResources.IsValidIndex(LODLevel))
 	{
 		FStaticMeshLODResources* RenderData = &StaticMeshComponent->GetStaticMesh()->GetRenderData()->LODResources[LODLevel];
 		
-		if( RenderData && ( ( uint32 )UVChannel < RenderData->VertexBuffers.StaticMeshVertexBuffer.GetNumTexCoords() ) )
+		if (RenderData && UVChannel >= 0 && UVChannel < static_cast<int32>(RenderData->VertexBuffers.StaticMeshVertexBuffer.GetNumTexCoords()))
 		{
 			//draw triangles
 			FIndexArrayView Indices = RenderData->IndexBuffer.GetArrayView();
@@ -565,11 +556,9 @@ void FCustomizableObjectEditorViewportClient::DrawUVs(FViewport* InViewport, FCa
 				// UVs, we'll draw the line segment in red
 				
 				// If we are supporting a version lower than LWC get the right real type. 
-				using Vector2DRealType = TDecay<decltype( DeclVal<FVector2D>().X )>::Type; 
-
-				constexpr Vector2DRealType Epsilon = static_cast<Vector2DRealType>(1e-4);
-				constexpr Vector2DRealType One     = static_cast<Vector2DRealType>(1);
-				constexpr Vector2DRealType Zero    = static_cast<Vector2DRealType>(0);
+				using Vector2DRealType = TDecay<decltype( DeclVal<FVector2D>().X )>::Type;
+				
+				constexpr Vector2DRealType Zero = static_cast<Vector2DRealType>(0);
 
 				UV1 = ClampUVRange(UV1.X, UV1.Y) * UVBoxScale + UVBoxOrigin;
 				UV2 = ClampUVRange(UV2.X, UV2.Y) * UVBoxScale + UVBoxOrigin;
@@ -581,57 +570,31 @@ void FCustomizableObjectEditorViewportClient::DrawUVs(FViewport* InViewport, FCa
 			}
 		}
 	}
-
 	else if (SkeletalMeshComponents.Num())
 	{
-		int32 CurrentComponentIndex = 0;
-
-		for (TWeakObjectPtr<USkeletalMeshComponent> SkeletalMeshComponent : SkeletalMeshComponents)
+		if (SkeletalMeshComponents.IsValidIndex(ComponentIndex))
 		{
-			if (!SkeletalMeshComponent.IsValid() || !UE_MUTABLE_GETSKINNEDASSET(SkeletalMeshComponent) || CurrentComponentIndex != ComponentIndex)
+			TWeakObjectPtr<USkeletalMeshComponent> SkeletalMeshComponent = SkeletalMeshComponents[ComponentIndex];
+
+			if (!SkeletalMeshComponent.IsValid() || !UE_MUTABLE_GETSKINNEDASSET(SkeletalMeshComponent))
 			{
-				CurrentComponentIndex++;
-				continue;
+				return;
 			}
 
-			bool bFoundMaterial = false;
-
 			const FSkeletalMeshRenderData* MeshRes = UE_MUTABLE_GETSKINNEDASSET(SkeletalMeshComponent)->GetResourceForRendering();
-			if (UVChannel < (int32)MeshRes->LODRenderData[LODLevel].GetNumTexCoords())
+			if (!MeshRes->LODRenderData.IsValidIndex(LODLevel))
+			{
+				return;
+			}
+			
+			if (UVChannel >= 0 && UVChannel < static_cast<int32>(MeshRes->LODRenderData[LODLevel].GetNumTexCoords()))
 			{
 				// Find material index from name
 				const FSkeletalMeshLODRenderData& lodModel = MeshRes->LODRenderData[LODLevel];
-				int MaterialIndex = 0;
-				int MaterialIndexCount = 0;
-				{
-					for (int s = 0; s < lodModel.RenderSections.Num(); ++s)
-					{
-						int SectionMaterial = lodModel.RenderSections[s].MaterialIndex;
-						UMaterialInterface* Material = SkeletalMeshComponent->GetMaterial(SectionMaterial);
-
-						if (!Material)
-						{
-							continue;
-						}
-
-						const UMaterial* BaseMaterial = Material->GetBaseMaterial();
-						if (BaseMaterial && BaseMaterial->GetName() == MaterialName)
-						{
-							MaterialIndex = s;
-
-							if (MaterialIndexCount == MaterialToDrawInUVsIndex)
-							{
-								bFoundMaterial = true;
-								break;
-							}
-							MaterialIndexCount++;
-						}
-					}
-				}
 
-				if (!bFoundMaterial)
+				if (!lodModel.RenderSections.IsValidIndex(SectionIndex))
 				{
-					continue;
+					return;
 				}
 
 				const FStaticMeshVertexBuffer& Vertices = lodModel.StaticVertexBuffers.StaticMeshVertexBuffer;
@@ -639,8 +602,8 @@ void FCustomizableObjectEditorViewportClient::DrawUVs(FViewport* InViewport, FCa
 				TArray<uint32> Indices;
 				lodModel.MultiSizeIndexContainer.GetIndexBuffer(Indices);
 
-				uint32 NumTriangles = lodModel.RenderSections[MaterialIndex].NumTriangles;
-				int IndexIndex = lodModel.RenderSections[MaterialIndex].BaseIndex;
+				uint32 NumTriangles = lodModel.RenderSections[SectionIndex].NumTriangles;
+				int IndexIndex = lodModel.RenderSections[SectionIndex].BaseIndex;
 
 				BatchedElements->AddReserveLines(NumTriangles * 3);
 
@@ -658,8 +621,6 @@ void FCustomizableObjectEditorViewportClient::DrawUVs(FViewport* InViewport, FCa
 					// If we are supporting a version lower than LWC get the right real type. 
 					using Vector2DRealType = TDecay<decltype(DeclVal<FVector2D>().X)>::Type;
 
-					constexpr Vector2DRealType Epsilon = static_cast<Vector2DRealType>(1e-4);
-					constexpr Vector2DRealType One = static_cast<Vector2DRealType>(1);
 					constexpr Vector2DRealType Zero = static_cast<Vector2DRealType>(0);
 
 					UV1 = ClampUVRange(UV1.X, UV1.Y) * UVBoxScale + UVBoxOrigin;
@@ -671,11 +632,6 @@ void FCustomizableObjectEditorViewportClient::DrawUVs(FViewport* InViewport, FCa
 					BatchedElements->AddLine( FVector(UV3, Zero), FVector(UV1, Zero), BorderColor, HitProxyId );
 				}
 			}
-
-			if (bFoundMaterial)
-			{
-				break;
-			}
 		}
 	}
 }
@@ -919,38 +875,17 @@ void FCustomizableObjectEditorViewportClient::SetDrawUVOverlay()
 }
 
 
-void FCustomizableObjectEditorViewportClient::SetDrawUVOverlayMaterial(const FString& MaterialName, const FString& UVChannel)
+void FCustomizableObjectEditorViewportClient::SetDrawUV(const int32 ComponentIndex, const int32 LODIndex, const int32 SectionIndex, const int32 UVIndex)
 {
-	// Get LOD Index
-	FString NameWithLOD, ComponentString;
-	MaterialName.Split(FString("_Component_"), &NameWithLOD, &ComponentString);
-	check(ComponentString.IsNumeric());
-	MaterialToDrawInUVsComponent = FCString::Atoi(*ComponentString);
-
-	FString Name, LODIndex;
-	bool bSplit = NameWithLOD.Split(FString(" LOD_"), &Name, &LODIndex);
-	check(bSplit && LODIndex.IsNumeric());
-
-	MaterialToDrawInUVsLOD = FCString::Atoi(*LODIndex);
-	UVChannelToDrawInUVs = FCString::Atoi(*UVChannel);
-
-	// Get Material Index, added if the name of the material already exists within the skeletal mesh.
-	FString DuplicatedMaterialIndex;
-	bSplit = Name.Split(FString("__"), &MaterialToDrawInUVs, &DuplicatedMaterialIndex);
-
-	if (bSplit && DuplicatedMaterialIndex.IsNumeric())
-	{
-		MaterialToDrawInUVsIndex = FCString::Atoi(*DuplicatedMaterialIndex);
-	}
-	else
-	{
-		MaterialToDrawInUVs = Name;
-		MaterialToDrawInUVsIndex = 0;
-	}
+	UVDrawComponentIndex = ComponentIndex;
+	UVDrawLODIndex = LODIndex;
+	UVDrawSectionIndex = SectionIndex;
+	UVDrawUVIndex = UVIndex;
 
 	Invalidate();
 }
 
+
 bool FCustomizableObjectEditorViewportClient::IsSetDrawUVOverlayChecked() const
 {
 	return bDrawUVs;
diff --git a/Engine/Plugins/Experimental/Mutable/Source/CustomizableObjectEditor/Private/MuCOE/CustomizableObjectInstanceEditor.cpp b/Engine/Plugins/Experimental/Mutable/Source/CustomizableObjectEditor/Private/MuCOE/CustomizableObjectInstanceEditor.cpp
--- a/Engine/Plugins/Experimental/Mutable/Source/CustomizableObjectEditor/Private/MuCOE/CustomizableObjectInstanceEditor.cpp
+++ b/Engine/Plugins/Experimental/Mutable/Source/CustomizableObjectEditor/Private/MuCOE/CustomizableObjectInstanceEditor.cpp
@@ -676,7 +676,7 @@ void FCustomizableObjectInstanceEditor::OnUpdatePreviewInstance()
 	}
 
 	// Sets which UV material to draw
-	Viewport->SetDrawDefaultUVMaterial(false);
+	Viewport->SetDrawDefaultUVMaterial();
 
 	if (TextureAnalyzer.IsValid())
 	{
diff --git a/Engine/Plugins/Experimental/Mutable/Source/CustomizableObjectEditor/Private/MuCOE/SCustomizableObjectEditorViewport.cpp b/Engine/Plugins/Experimental/Mutable/Source/CustomizableObjectEditor/Private/MuCOE/SCustomizableObjectEditorViewport.cpp
--- a/Engine/Plugins/Experimental/Mutable/Source/CustomizableObjectEditor/Private/MuCOE/SCustomizableObjectEditorViewport.cpp
+++ b/Engine/Plugins/Experimental/Mutable/Source/CustomizableObjectEditor/Private/MuCOE/SCustomizableObjectEditorViewport.cpp
@@ -523,76 +523,26 @@ bool SCustomizableObjectEditorViewportTabBody::IsCameraModeActive(int Value)
 }
 
 
-void SCustomizableObjectEditorViewportTabBody::SetDrawDefaultUVMaterial(bool bIsCompilation)
+void SCustomizableObjectEditorViewportTabBody::SetDrawDefaultUVMaterial()
 {
-	GenerateUVMaterialOptions();
-	bool bMaterialFound = false;
-
-	if (!ArrayUVMaterialOptionString.IsEmpty())
-	{
-		if (!bIsCompilation && SelectedUVMaterial.IsValid())
-		{
-			// We check if the selected Material still exists after the update
-			FString MaterialName = *SelectedUVMaterial;
-
-			for (int32 MaterialIndex = 0; MaterialIndex < ArrayUVMaterialOptionString.Num(); ++MaterialIndex)
-			{
-				if (MaterialName == *ArrayUVMaterialOptionString[MaterialIndex])
-				{
-					bMaterialFound = true;
-					break;
-				}
-			}
-		}
-
-		// If the CO is recompiled or couldn't find the last used materia, draw the first material
-		if (bIsCompilation || !bMaterialFound)
-		{
-			LevelViewportClient->SetDrawUVOverlayMaterial(*(ArrayUVMaterialOptionString[0]), "0");
-			SelectedUVMaterial = ArrayUVMaterialOptionString[0];
-		}
-	}
-	else
-	{
-		SelectedUVMaterial = nullptr;
-	}
-
+	GenerateUVSectionOptions();
 	GenerateUVChannelOptions();
-
-	if (!ArrayUVChannelOptionString.IsEmpty())
+	
+	if (!SelectedUVSection ||
+		!SelectedUVChannel)
 	{
-		bool bUVChannelFound = false;
-
-		if (!bIsCompilation && SelectedUVChannel.IsValid())
-		{
-			// We check if the selected Material still exists after the update
-			FString UVChannel = *SelectedUVChannel;
-
-			for (int32 MaterialIndex = 0; MaterialIndex < ArrayUVChannelOptionString.Num(); ++MaterialIndex)
-			{
-				if (UVChannel == *ArrayUVChannelOptionString[MaterialIndex])
-				{
-					bUVChannelFound = true;
-					break;
-				}
-			}
-		}
-
-		// If the CO is recompiled or couldn't find the last used materia, draw the UV channel 0
-		if (bIsCompilation || !bMaterialFound || !bUVChannelFound)
-		{
-			SelectedUVChannel = ArrayUVChannelOptionString[0];
-		}
+		LevelViewportClient->SetDrawUV(-1, -1, -1, -1);
 	}
 	else
 	{
-		SelectedUVChannel = nullptr;
-	}
+		const int32 SectionOptionIndex = UVSectionOptionString.IndexOfByKey(SelectedUVSection);
+		check(SectionOptionIndex != INDEX_NONE);
+		const FSection& Section = UVSectionOption[SectionOptionIndex];
 
-	if (UVMaterialOptionCombo.IsValid() && UVChannelOptionCombo.IsValid())
-	{
-		UVMaterialOptionCombo->SetSelectedItem(SelectedUVMaterial);
-		UVChannelOptionCombo->SetSelectedItem(SelectedUVChannel);
+		const int32 UVIndex = UVChannelOptionString.IndexOfByKey(SelectedUVChannel);
+		check(UVIndex != INDEX_NONE);
+	
+		LevelViewportClient->SetDrawUV(Section.ComponentIndex, Section.LODIndex, Section.SectionIndex, UVIndex);
 	}
 }
 
@@ -673,85 +623,22 @@ TSharedRef<SWidget> SCustomizableObjectEditorViewportTabBody::GenerateUVMaterial
 	MenuBuilder.BeginSection("ShowUV");
 	{
 		// Generating an array with all the options of the combobox 
-		GenerateUVMaterialOptions();
-
-		// Sorting array
-        ArrayUVMaterialOptionString.Sort(::CompareNames);
-
-		// Setting initial selected option
-		if (ArrayUVMaterialOptionString.Num())
-		{
-			bool bFound = false;
-
-			if (SelectedUVMaterial.IsValid())
-			{
-				for (int32 i = 0; i < ArrayUVMaterialOptionString.Num(); ++i)
-				{
-					if (*SelectedUVMaterial == *ArrayUVMaterialOptionString[i])
-					{
-						bFound = true;
-						SelectedUVMaterial = ArrayUVMaterialOptionString[i];
-						break;
-					}
-				}
-			}
-
-			if (!bFound)
-			{
-				SelectedUVMaterial = ArrayUVMaterialOptionString[0];
-			}
-		}
-		else
-		{
-			SelectedUVMaterial = nullptr;
-		}
-	
-		UVMaterialOptionCombo = SNew(STextComboBox)
-			.OptionsSource(&ArrayUVMaterialOptionString)
-			.InitiallySelectedItem(SelectedUVMaterial)
-			.OnSelectionChanged(this, &SCustomizableObjectEditorViewportTabBody::OnMaterialChanged);
+		GenerateUVSectionOptions();
+		
+		UVSectionOptionCombo = SNew(STextComboBox)
+			.OptionsSource(&UVSectionOptionString)
+			.InitiallySelectedItem(SelectedUVSection)
+			.OnSelectionChanged(this, &SCustomizableObjectEditorViewportTabBody::OnSectionChanged);
 
 		// Generating an array with all the options of the combobox 
 		GenerateUVChannelOptions();
-
-		// Setting initial selected option
-		if (ArrayUVChannelOptionString.Num())
-		{
-			if (SelectedUVChannel.IsValid())
-			{
-				bool bFound = false;
-
-				for (int32 i = 0; i < ArrayUVChannelOptionString.Num(); ++i)
-				{
-					if (*SelectedUVChannel == *ArrayUVChannelOptionString[i])
-					{
-						bFound = true;
-						SelectedUVChannel = ArrayUVChannelOptionString[i];
-						break;
-					}
-				}
-
-				if (!bFound)
-				{
-					SelectedUVChannel = ArrayUVChannelOptionString[0];
-				}
-			}
-			else
-			{
-				SelectedUVChannel = ArrayUVChannelOptionString[0];
-			}
-		}
-		else
-		{
-			SelectedUVChannel = nullptr;
-		}
-
+		
 		UVChannelOptionCombo = SNew(STextComboBox)
-			.OptionsSource(&ArrayUVChannelOptionString)
+			.OptionsSource(&UVChannelOptionString)
 			.InitiallySelectedItem(SelectedUVChannel)
 			.OnSelectionChanged(this, &SCustomizableObjectEditorViewportTabBody::OnUVChannelChanged);
 		
-		MenuBuilder.AddWidget(UVMaterialOptionCombo.ToSharedRef(), FText::FromString(TEXT("Material")));
+		MenuBuilder.AddWidget(UVSectionOptionCombo.ToSharedRef(), FText::FromString(TEXT("Section")));
 		MenuBuilder.AddWidget(UVChannelOptionCombo.ToSharedRef(), FText::FromString(TEXT("UV Channel")));
 	}
 
@@ -777,20 +664,24 @@ TSharedRef<SWidget> SCustomizableObjectEditorViewportTabBody::ShowStateTestData(
 }
 
 
-bool SCustomizableObjectEditorViewportTabBody::IsMaterialsComboEnabled() const
-{
-	return true;
-	//return LevelViewportClient.IsValid() && LevelViewportClient->IsSetDrawUVOverlayChecked();
-}
-
-
-void SCustomizableObjectEditorViewportTabBody::GenerateUVMaterialOptions()
+void SCustomizableObjectEditorViewportTabBody::GenerateUVSectionOptions()
 {
-	ArrayUVMaterialOptionString.Empty();
+	ON_SCOPE_EXIT
+	{
+		if (UVSectionOptionCombo)
+		{
+			UVSectionOptionCombo->RefreshOptions();
+			UVSectionOptionCombo->SetSelectedItem(SelectedUVSection);
+		}
+	};
+		
+	UVSectionOptionString.Empty();
+	UVSectionOption.Empty();
 	
-	int32 ComponentIndex = 0;
-	for (UDebugSkelMeshComponent* PreviewSkeletalMeshComponent : PreviewSkeletalMeshComponents)
+	for (int32 ComponentIndex = 0; ComponentIndex < PreviewSkeletalMeshComponents.Num(); ++ComponentIndex)
 	{
+		const UDebugSkelMeshComponent* PreviewSkeletalMeshComponent = PreviewSkeletalMeshComponents[ComponentIndex];
+		
 		if (PreviewSkeletalMeshComponent != nullptr && UE_MUTABLE_GETSKINNEDASSET(PreviewSkeletalMeshComponent) != nullptr && UE_MUTABLE_GETSKINNEDASSET(PreviewSkeletalMeshComponent)->GetResourceForRendering() != nullptr)
 		{
 			const TArray<UMaterialInterface*> Materials = PreviewSkeletalMeshComponent->GetMaterials();
@@ -800,109 +691,147 @@ void SCustomizableObjectEditorViewportTabBody::GenerateUVMaterialOptions()
 				for (int32 SectionIndex = 0; SectionIndex < MeshRes->LODRenderData[LODIndex].RenderSections.Num(); ++SectionIndex)
 				{
 					const FSkelMeshRenderSection& Section = MeshRes->LODRenderData[LODIndex].RenderSections[SectionIndex];
-					
-					if (!Materials.IsValidIndex(Section.MaterialIndex))
-					{
-						continue;
-					}
 
-					if (!Materials[Section.MaterialIndex])
-					{
-						continue;
-					}
-
-					FString BaseMaterialName;
+					FString BaseMaterialName = FString::Printf(TEXT("Section %i"), SectionIndex);
 
-					const UMaterial* BaseMaterial = Materials[Section.MaterialIndex]->GetBaseMaterial();
-					if (BaseMaterial)
+					if (Materials.IsValidIndex(Section.MaterialIndex))
 					{
-						BaseMaterialName = BaseMaterial->GetName();
+						if (UMaterialInterface* MaterialInterface = Materials[Section.MaterialIndex])
+						{
+							if (const UMaterial* BaseMaterial = MaterialInterface->GetBaseMaterial())
+							{
+								BaseMaterialName += " - " + BaseMaterial->GetName();
+							}
+						}
 					}
 
-					BaseMaterialName += FString::Printf(TEXT(" LOD_%d_Component_%d"), LODIndex, ComponentIndex);
-
-					ArrayUVMaterialOptionString.Add(MakeShareable(new FString(BaseMaterialName)));
+					UVSectionOptionString.Add(MakeShared<FString>(BaseMaterialName));
+					UVSectionOption.Emplace(ComponentIndex, LODIndex, SectionIndex);
 				}
 			}
 		}
+	}
 
-		ComponentIndex++;
+	if (SelectedUVSection)
+	{
+		if (const TSharedPtr<FString>* Result = UVSectionOptionString.FindByPredicate([this](const TSharedPtr<FString>& Other){ return *SelectedUVSection == *Other; }))
+		{
+			SelectedUVSection = *Result;
+			return;
+		}	
 	}
+	
+	SelectedUVSection = UVSectionOptionString.IsEmpty() ? nullptr : UVSectionOptionString[0];
 }
 
 
-void SCustomizableObjectEditorViewportTabBody::OnMaterialChanged(TSharedPtr<FString> Selected, ESelectInfo::Type SelectInfo)
+void SCustomizableObjectEditorViewportTabBody::OnSectionChanged(TSharedPtr<FString> Selected, ESelectInfo::Type SelectInfo)
 {
-	if (Selected.IsValid() && SelectInfo == ESelectInfo::OnMouseClick)
-	{
-		SelectedUVMaterial = Selected;
+	SelectedUVSection = Selected;
+	
+	// We need to update options for the new section
+	GenerateUVChannelOptions();
 
-		//We need to update options for the new LOD
-		GenerateUVChannelOptions();
+	// Reset the UVChannel selection
+	SelectedUVChannel = UVChannelOptionString.IsEmpty() ? nullptr : UVChannelOptionString[0];
+	UVChannelOptionCombo->SetSelectedItem(SelectedUVChannel);
+	
+	if (!LevelViewportClient)
+	{
+		return;
+	}
 
-		// Resets the value to the first element of the array
-		if (!ArrayUVChannelOptionString.IsEmpty())
-		{
-			SelectedUVChannel = ArrayUVChannelOptionString[0];
-			UVChannelOptionCombo->SetSelectedItem(SelectedUVChannel);
-		}
+	if (SelectedUVSection)
+	{
+		const int32 SectionOptionIndex = UVSectionOptionString.IndexOfByKey(SelectedUVSection);
+		check(SectionOptionIndex != INDEX_NONE);
+		const FSection& Section = UVSectionOption[SectionOptionIndex];
 
-		if (LevelViewportClient.IsValid() && SelectedUVChannel.IsValid())
-		{
-			LevelViewportClient->SetDrawUVOverlayMaterial(*Selected, *SelectedUVChannel);
-		}
+		const int32 UVIndex = UVChannelOptionString.IndexOfByKey(SelectedUVChannel);
+		check(UVIndex != INDEX_NONE);
+	
+		LevelViewportClient->SetDrawUV(Section.ComponentIndex, Section.LODIndex, Section.SectionIndex, UVIndex);
+	}
+	else
+	{
+		LevelViewportClient->SetDrawUV(-1, -1, -1, -1);
 	}
 }
 
 
 void SCustomizableObjectEditorViewportTabBody::GenerateUVChannelOptions()
 {
-	ArrayUVChannelOptionString.Empty();
-
-	if (!SelectedUVMaterial)
+	ON_SCOPE_EXIT
+	{
+		if (UVChannelOptionCombo)
+		{
+			UVChannelOptionCombo->RefreshOptions();
+			UVChannelOptionCombo->SetSelectedItem(SelectedUVChannel);
+		}
+	};
+	
+	UVChannelOptionString.Empty();
+	
+	if (!SelectedUVSection)
 	{
+		SelectedUVChannel = nullptr;
 		return;
 	}
+	
+	const int32 Index = UVSectionOptionString.IndexOfByKey(SelectedUVSection);
+	check(Index != INDEX_NONE);
+	const FSection& Section = UVSectionOption[Index];
 
-	FString SelectedMat = *SelectedUVMaterial;
-
-	// From String to values
-	FString NameWithLOD, ComponentString;
-	SelectedMat.Split(FString("_Component_"), &NameWithLOD, &ComponentString);
-	check(ComponentString.IsNumeric());
-	int32 ComponentIndex = FCString::Atoi(*ComponentString);
-	check(PreviewSkeletalMeshComponents.IsValidIndex(ComponentIndex));
-
-	FString Name, LODString;
-	bool bSplit = NameWithLOD.Split(FString(" LOD_"), &Name, &LODString);
-	int32 LODIndex = FCString::Atoi(*LODString);
-
-	UDebugSkelMeshComponent* PreviewSkeletalMeshComponent = PreviewSkeletalMeshComponents[ComponentIndex];
+	const UDebugSkelMeshComponent* PreviewSkeletalMeshComponent = PreviewSkeletalMeshComponents[Section.ComponentIndex];
 
 	if (PreviewSkeletalMeshComponent != nullptr && UE_MUTABLE_GETSKINNEDASSET(PreviewSkeletalMeshComponent) != nullptr
 		&& UE_MUTABLE_GETSKINNEDASSET(PreviewSkeletalMeshComponent)->GetResourceForRendering() != nullptr)
 	{
 		const FSkeletalMeshRenderData* MeshRes = UE_MUTABLE_GETSKINNEDASSET(PreviewSkeletalMeshComponent)->GetResourceForRendering();
 		
-		int32 UVChannels = MeshRes->LODRenderData[LODIndex].GetNumTexCoords();
-		
+		const int32 UVChannels = MeshRes->LODRenderData[Section.LODIndex].GetNumTexCoords();
 		for (int32 UVChan = 0; UVChan < UVChannels; ++UVChan)
 		{
-			ArrayUVChannelOptionString.Add(MakeShareable(new FString(FString::FromInt(UVChan))));
+			UVChannelOptionString.Add(MakeShareable(new FString(FString::FromInt(UVChan))));
 		}
 	}
+
+	if (SelectedUVChannel)
+	{
+		if (const TSharedPtr<FString>* Result = UVChannelOptionString.FindByPredicate([this](const TSharedPtr<FString>& Other) { return *SelectedUVChannel == *Other; }))
+		{
+			SelectedUVChannel = *Result;
+			return;
+		}
+	}
+
+	SelectedUVChannel = UVChannelOptionString.IsEmpty() ? nullptr : UVChannelOptionString[0];
 }
 
 
 void SCustomizableObjectEditorViewportTabBody::OnUVChannelChanged(TSharedPtr<FString> Selected, ESelectInfo::Type SelectInfo)
 {
-	if (Selected.IsValid() && SelectInfo == ESelectInfo::OnMouseClick)
+	SelectedUVChannel = Selected;
+
+	if (!LevelViewportClient)
 	{
-		SelectedUVChannel = Selected;
-		if (LevelViewportClient.IsValid() && SelectedUVMaterial.IsValid())
-		{
-			LevelViewportClient->SetDrawUVOverlayMaterial(*SelectedUVMaterial, *Selected);
-		}
+		return;
+	}
+
+	if (SelectedUVChannel)
+	{
+		const int32 SectionOptionIndex = UVSectionOptionString.IndexOfByKey(SelectedUVSection);
+		check(SectionOptionIndex != INDEX_NONE);
+		const FSection& Section = UVSectionOption[SectionOptionIndex];
+
+		const int32 UVIndex = UVChannelOptionString.IndexOfByKey(SelectedUVChannel);
+		check(UVIndex != INDEX_NONE);
+	
+		LevelViewportClient->SetDrawUV(Section.ComponentIndex, Section.LODIndex, Section.SectionIndex, UVIndex);
+	}
+	else
+	{
+		LevelViewportClient->SetDrawUV(-1, -1, -1, -1);
 	}
 }
 
diff --git a/Engine/Plugins/Experimental/Mutable/Source/CustomizableObjectEditor/Private/MuCOE/SCustomizableObjectEditorViewport.h b/Engine/Plugins/Experimental/Mutable/Source/CustomizableObjectEditor/Private/MuCOE/SCustomizableObjectEditorViewport.h
--- a/Engine/Plugins/Experimental/Mutable/Source/CustomizableObjectEditor/Private/MuCOE/SCustomizableObjectEditorViewport.h
+++ b/Engine/Plugins/Experimental/Mutable/Source/CustomizableObjectEditor/Private/MuCOE/SCustomizableObjectEditorViewport.h
@@ -217,7 +217,7 @@ public:
 
 	/** Sets the first material of the instance as the default section to draw in the UVs Overlay (when there is no material selected in the combobox)
 		Param bIsCompilation: true if this function is called during a CO compilation, false when is called from an instance update	*/
-	void SetDrawDefaultUVMaterial(bool bIsCompilation);
+	void SetDrawDefaultUVMaterial();
 
 	/** Sets the bones visibility */
 	void SetShowBones();
@@ -254,26 +254,35 @@ private:
 	// The currently selected view mode.
 	EViewModeIndex CurrentViewMode = VMI_Lit;
 
-	// Selection for the material to filter when showing UVs
-	TArray<TSharedPtr<FString>> ArrayUVMaterialOptionString;
+	/** Section identifier. */
+	struct FSection
+	{
+		int32 ComponentIndex = -1;
+		int32 LODIndex = -1;
+		int32 SectionIndex = -1;
+	};
+	
+	// Selection for the section to filter when showing UVs
+	TArray<TSharedPtr<FString>> UVSectionOptionString;
+	TArray<FSection> UVSectionOption;
 
-	// Combo box for UV material selection
-	TSharedPtr< STextComboBox> UVMaterialOptionCombo;
+	// Combo box for UV section selection
+	TSharedPtr<STextComboBox> UVSectionOptionCombo;
 
 	// Selection for the channel to filter when showing UVs
-	TArray<TSharedPtr<FString>> ArrayUVChannelOptionString;
+	TArray<TSharedPtr<FString>> UVChannelOptionString;
 
 	// Combo box for UV channel selection
-	TSharedPtr< STextComboBox> UVChannelOptionCombo;
+	TSharedPtr<STextComboBox> UVChannelOptionCombo;
 
 	// Generates the Combobox options for each UV Channel
-	void GenerateUVMaterialOptions();
+	void GenerateUVSectionOptions();
 
 	// Generates the Combobox options for each UV Channel
 	void GenerateUVChannelOptions();
 
 	// Selected option of the UV Material ComboBox
-	TSharedPtr<FString> SelectedUVMaterial;
+	TSharedPtr<FString> SelectedUVSection;
 	
 	// Selected option of the UV Channel ComboBox
 	TSharedPtr<FString> SelectedUVChannel;
@@ -282,8 +291,7 @@ private:
 	TArray< TSharedPtr< FString > > MaterialNames;
 
 	// Slate callbacks to manage the material combo
-	bool IsMaterialsComboEnabled() const;
-	void OnMaterialChanged(TSharedPtr<FString> Selected, ESelectInfo::Type SelectInfo);
+	void OnSectionChanged(TSharedPtr<FString> Selected, ESelectInfo::Type SelectInfo);
 	void OnUVChannelChanged(TSharedPtr<FString> Selected, ESelectInfo::Type SelectInfo);
 
 	/** Current LOD selection*/
diff --git a/Engine/Plugins/Experimental/Mutable/Source/CustomizableObjectEditor/Public/MuCOE/CustomizableObjectEditorViewportClient.h b/Engine/Plugins/Experimental/Mutable/Source/CustomizableObjectEditor/Public/MuCOE/CustomizableObjectEditorViewportClient.h
--- a/Engine/Plugins/Experimental/Mutable/Source/CustomizableObjectEditor/Public/MuCOE/CustomizableObjectEditorViewportClient.h
+++ b/Engine/Plugins/Experimental/Mutable/Source/CustomizableObjectEditor/Public/MuCOE/CustomizableObjectEditorViewportClient.h
@@ -112,7 +112,7 @@ public:
 	 *	@param	InCanvas					The canvas to draw to
 	 *	@param	InTextYPos					The Y-position to begin drawing the UV-Overlay at (due to text displayed above it).
 	 */
-	void DrawUVs(FViewport* InViewport, FCanvas* InCanvas, int32 InTextYPos, const FString& MaterialName);
+	void DrawUVs(FViewport* InViewport, FCanvas* InCanvas, int32 InTextYPos);
 
 	/**
 	 * Bake the instance currently present in the editor. Internally will schedule the update of the instance before baking it's resources
@@ -138,8 +138,8 @@ public:
 	/** Callback for checking the UV overlay show flag. */
 	bool IsSetDrawUVOverlayChecked() const;
 
-	// Set the material index whose UVs will be drawn.
-	void SetDrawUVOverlayMaterial(const FString& MaterialName, const FString& UVChannel);
+	/** Specify which UV to draw. -1 values will not draw anything. */
+	void SetDrawUV(const int32 ComponentIndex, const int32 LODIndex, const int32 SectionIndex, const int32 UVIndex);
 
 	/** Callback for toggling the grid show flag. */
 	void UpdateShowGridFromButton();
@@ -306,12 +306,11 @@ private:
 	bool bDrawUVs;
 	bool bDrawSky;
 
-	// Material index to draw in the uv preview.
-	FString MaterialToDrawInUVs;
-	int32 MaterialToDrawInUVsLOD;
-	int32 MaterialToDrawInUVsIndex;
-	int32 UVChannelToDrawInUVs;
-	int32 MaterialToDrawInUVsComponent;
+	// UV to draw
+	int32 UVDrawComponentIndex = 0;
+	int32 UVDrawSectionIndex = 0;
+	int32 UVDrawLODIndex = 0;
+	int32 UVDrawUVIndex = 0;
 
 	bool bReferenceMeshMissingWarningMessageVisible;
 
