[Backout] - CL32394772
Original CL Desc
-----------------------------------------------------------------
[Mutable] Fix UV not showing the correct UV when Materials had the same name.

#rnx

diff --git a/Engine/Plugins/Experimental/Mutable/Source/CustomizableObjectEditor/Private/MuCOE/CustomizableObjectEditor.cpp b/Engine/Plugins/Experimental/Mutable/Source/CustomizableObjectEditor/Private/MuCOE/CustomizableObjectEditor.cpp
--- a/Engine/Plugins/Experimental/Mutable/Source/CustomizableObjectEditor/Private/MuCOE/CustomizableObjectEditor.cpp
+++ b/Engine/Plugins/Experimental/Mutable/Source/CustomizableObjectEditor/Private/MuCOE/CustomizableObjectEditor.cpp
@@ -1543,7 +1543,7 @@ void FCustomizableObjectEditor::OnObjectModified(UObject* Object)
 void FCustomizableObjectEditor::CompileObject()
 {
 	// Resetting viewport parameters
-	Viewport->SetDrawDefaultUVMaterial();
+	Viewport->SetDrawDefaultUVMaterial(true);
 
 	UE_LOG(LogMutable, Verbose, TEXT("PROFILE: -----------------------------------------------------------"));
 	UE_LOG(LogMutable, Verbose, TEXT("PROFILE: [ %16.8f ] FCustomizableObjectEditor::CompileObject start."), FPlatformTime::Seconds());
@@ -2275,7 +2275,7 @@ void FCustomizableObjectEditor::OnUpdatePreviewInstance()
 	}
 	
 	Viewport->SetPreviewComponents(PreviewSkeletalMeshComponents);
-	Viewport->SetDrawDefaultUVMaterial();
+	Viewport->SetDrawDefaultUVMaterial(false);
 	ViewportClient->Invalidate();
 	ViewportClient->ReSetAnimation();
 	ViewportClient->SetReferenceMeshMissingWarningMessage(false);
diff --git a/Engine/Plugins/Experimental/Mutable/Source/CustomizableObjectEditor/Private/MuCOE/CustomizableObjectEditorViewportClient.cpp b/Engine/Plugins/Experimental/Mutable/Source/CustomizableObjectEditor/Private/MuCOE/CustomizableObjectEditorViewportClient.cpp
--- a/Engine/Plugins/Experimental/Mutable/Source/CustomizableObjectEditor/Private/MuCOE/CustomizableObjectEditorViewportClient.cpp
+++ b/Engine/Plugins/Experimental/Mutable/Source/CustomizableObjectEditor/Private/MuCOE/CustomizableObjectEditorViewportClient.cpp
@@ -81,6 +81,12 @@ FCustomizableObjectEditorViewportClient::FCustomizableObjectEditorViewportClient
 
 	bShowBones = false;
 
+	MaterialToDrawInUVs = 0;
+	MaterialToDrawInUVsLOD = 0;
+	MaterialToDrawInUVsIndex = 0;
+	UVChannelToDrawInUVs = 0;
+	MaterialToDrawInUVsComponent = 0;
+
 	bReferenceMeshMissingWarningMessageVisible = false;
 
 	DrawHelper.bDrawPivot = false;
@@ -456,7 +462,7 @@ void FCustomizableObjectEditorViewportClient::Draw(FViewport* InViewport, FCanva
 	if(bDrawUVs)
 	{
 		constexpr int32 YPos = 24;
-		DrawUVs(InViewport, Canvas, YPos);
+		DrawUVs(InViewport, Canvas, YPos, MaterialToDrawInUVs);
 	}
 
 	if (bReferenceMeshMissingWarningMessageVisible)
@@ -485,12 +491,16 @@ namespace
 	}
 }
 
-void FCustomizableObjectEditorViewportClient::DrawUVs(FViewport* InViewport, FCanvas* InCanvas, int32 InTextYPos)
+void FCustomizableObjectEditorViewportClient::DrawUVs(FViewport* InViewport, FCanvas* InCanvas, int32 InTextYPos, const FString& MaterialName )
 {
-	const uint32 ComponentIndex = UVDrawComponentIndex;
-	const uint32 LODLevel = UVDrawLODIndex; 	// TODO use the overriden LOD level
-	const int32 SectionIndex = UVDrawSectionIndex;
-	const int32 UVChannel = UVDrawUVIndex;
+	//use the overriden LOD level
+	// TODO
+	const uint32 LODLevel = MaterialToDrawInUVsLOD; //FMath::Clamp(StaticMeshComponent->ForcedLodModel - 1, 0, StaticMesh->RenderData->LODResources.Num() - 1);
+
+	// TODO
+	int32 UVChannel = UVChannelToDrawInUVs; //StaticMeshEditorPtr.Pin()->GetCurrentUVChannel();
+
+	const uint32 ComponentIndex = MaterialToDrawInUVsComponent;
 
 	//draw a string showing what UV channel and LOD is being displayed
 	InCanvas->DrawShadowedString( 
@@ -516,7 +526,9 @@ void FCustomizableObjectEditorViewportClient::DrawUVs(FViewport* InViewport, FCa
 		BoxOrigin + FVector2D( BoxSize, BoxSize ),	// bottomright
 		BoxOrigin + FVector2D( 0, BoxSize ),		// bottomleft
 	};
-	
+
+	const FVector Color(1.0f, 1.0f, 1.0f);
+
 	//draw texture border
 	FLinearColor BorderColor = FLinearColor::White;
 	FBatchedElements* BatchedElements = InCanvas->GetBatchedElements(FCanvas::ET_Line);
@@ -531,14 +543,11 @@ void FCustomizableObjectEditorViewportClient::DrawUVs(FViewport* InViewport, FCa
 	BatchedElements->AddLine( FVector( Box[ 2 ], 0.0f ), FVector( Box[ 3 ], 0.0f ), BorderColor, HitProxyId );
 	BatchedElements->AddLine( FVector( Box[ 3 ], 0.0f ), FVector( Box[ 0 ], 0.0f ), BorderColor, HitProxyId );
 
-	if (StaticMeshComponent.IsValid() &&
-		StaticMeshComponent->GetStaticMesh() &&
-		StaticMeshComponent->GetStaticMesh()->GetRenderData() &&
-		StaticMeshComponent->GetStaticMesh()->GetRenderData()->LODResources.IsValidIndex(LODLevel))
+	if ( StaticMeshComponent.IsValid() && StaticMeshComponent->GetStaticMesh() && StaticMeshComponent->GetStaticMesh()->GetRenderData() )
 	{
 		FStaticMeshLODResources* RenderData = &StaticMeshComponent->GetStaticMesh()->GetRenderData()->LODResources[LODLevel];
 		
-		if (RenderData && UVChannel >= 0 && UVChannel < static_cast<int32>(RenderData->VertexBuffers.StaticMeshVertexBuffer.GetNumTexCoords()))
+		if( RenderData && ( ( uint32 )UVChannel < RenderData->VertexBuffers.StaticMeshVertexBuffer.GetNumTexCoords() ) )
 		{
 			//draw triangles
 			FIndexArrayView Indices = RenderData->IndexBuffer.GetArrayView();
@@ -556,9 +565,11 @@ void FCustomizableObjectEditorViewportClient::DrawUVs(FViewport* InViewport, FCa
 				// UVs, we'll draw the line segment in red
 				
 				// If we are supporting a version lower than LWC get the right real type. 
-				using Vector2DRealType = TDecay<decltype( DeclVal<FVector2D>().X )>::Type;
-				
-				constexpr Vector2DRealType Zero = static_cast<Vector2DRealType>(0);
+				using Vector2DRealType = TDecay<decltype( DeclVal<FVector2D>().X )>::Type; 
+
+				constexpr Vector2DRealType Epsilon = static_cast<Vector2DRealType>(1e-4);
+				constexpr Vector2DRealType One     = static_cast<Vector2DRealType>(1);
+				constexpr Vector2DRealType Zero    = static_cast<Vector2DRealType>(0);
 
 				UV1 = ClampUVRange(UV1.X, UV1.Y) * UVBoxScale + UVBoxOrigin;
 				UV2 = ClampUVRange(UV2.X, UV2.Y) * UVBoxScale + UVBoxOrigin;
@@ -570,31 +581,57 @@ void FCustomizableObjectEditorViewportClient::DrawUVs(FViewport* InViewport, FCa
 			}
 		}
 	}
+
 	else if (SkeletalMeshComponents.Num())
 	{
-		if (SkeletalMeshComponents.IsValidIndex(ComponentIndex))
-		{
-			TWeakObjectPtr<USkeletalMeshComponent> SkeletalMeshComponent = SkeletalMeshComponents[ComponentIndex];
+		int32 CurrentComponentIndex = 0;
 
-			if (!SkeletalMeshComponent.IsValid() || !UE_MUTABLE_GETSKINNEDASSET(SkeletalMeshComponent))
+		for (TWeakObjectPtr<USkeletalMeshComponent> SkeletalMeshComponent : SkeletalMeshComponents)
+		{
+			if (!SkeletalMeshComponent.IsValid() || !UE_MUTABLE_GETSKINNEDASSET(SkeletalMeshComponent) || CurrentComponentIndex != ComponentIndex)
 			{
-				return;
+				CurrentComponentIndex++;
+				continue;
 			}
 
+			bool bFoundMaterial = false;
+
 			const FSkeletalMeshRenderData* MeshRes = UE_MUTABLE_GETSKINNEDASSET(SkeletalMeshComponent)->GetResourceForRendering();
-			if (!MeshRes->LODRenderData.IsValidIndex(LODLevel))
-			{
-				return;
-			}
-			
-			if (UVChannel >= 0 && UVChannel < static_cast<int32>(MeshRes->LODRenderData[LODLevel].GetNumTexCoords()))
+			if (UVChannel < (int32)MeshRes->LODRenderData[LODLevel].GetNumTexCoords())
 			{
 				// Find material index from name
 				const FSkeletalMeshLODRenderData& lodModel = MeshRes->LODRenderData[LODLevel];
+				int MaterialIndex = 0;
+				int MaterialIndexCount = 0;
+				{
+					for (int s = 0; s < lodModel.RenderSections.Num(); ++s)
+					{
+						int SectionMaterial = lodModel.RenderSections[s].MaterialIndex;
+						UMaterialInterface* Material = SkeletalMeshComponent->GetMaterial(SectionMaterial);
+
+						if (!Material)
+						{
+							continue;
+						}
+
+						const UMaterial* BaseMaterial = Material->GetBaseMaterial();
+						if (BaseMaterial && BaseMaterial->GetName() == MaterialName)
+						{
+							MaterialIndex = s;
+
+							if (MaterialIndexCount == MaterialToDrawInUVsIndex)
+							{
+								bFoundMaterial = true;
+								break;
+							}
+							MaterialIndexCount++;
+						}
+					}
+				}
 
-				if (!lodModel.RenderSections.IsValidIndex(SectionIndex))
+				if (!bFoundMaterial)
 				{
-					return;
+					continue;
 				}
 
 				const FStaticMeshVertexBuffer& Vertices = lodModel.StaticVertexBuffers.StaticMeshVertexBuffer;
@@ -602,8 +639,8 @@ void FCustomizableObjectEditorViewportClient::DrawUVs(FViewport* InViewport, FCa
 				TArray<uint32> Indices;
 				lodModel.MultiSizeIndexContainer.GetIndexBuffer(Indices);
 
-				uint32 NumTriangles = lodModel.RenderSections[SectionIndex].NumTriangles;
-				int IndexIndex = lodModel.RenderSections[SectionIndex].BaseIndex;
+				uint32 NumTriangles = lodModel.RenderSections[MaterialIndex].NumTriangles;
+				int IndexIndex = lodModel.RenderSections[MaterialIndex].BaseIndex;
 
 				BatchedElements->AddReserveLines(NumTriangles * 3);
 
@@ -621,6 +658,8 @@ void FCustomizableObjectEditorViewportClient::DrawUVs(FViewport* InViewport, FCa
 					// If we are supporting a version lower than LWC get the right real type. 
 					using Vector2DRealType = TDecay<decltype(DeclVal<FVector2D>().X)>::Type;
 
+					constexpr Vector2DRealType Epsilon = static_cast<Vector2DRealType>(1e-4);
+					constexpr Vector2DRealType One = static_cast<Vector2DRealType>(1);
 					constexpr Vector2DRealType Zero = static_cast<Vector2DRealType>(0);
 
 					UV1 = ClampUVRange(UV1.X, UV1.Y) * UVBoxScale + UVBoxOrigin;
@@ -632,6 +671,11 @@ void FCustomizableObjectEditorViewportClient::DrawUVs(FViewport* InViewport, FCa
 					BatchedElements->AddLine( FVector(UV3, Zero), FVector(UV1, Zero), BorderColor, HitProxyId );
 				}
 			}
+
+			if (bFoundMaterial)
+			{
+				break;
+			}
 		}
 	}
 }
@@ -875,17 +919,38 @@ void FCustomizableObjectEditorViewportClient::SetDrawUVOverlay()
 }
 
 
-void FCustomizableObjectEditorViewportClient::SetDrawUV(const int32 ComponentIndex, const int32 LODIndex, const int32 SectionIndex, const int32 UVIndex)
+void FCustomizableObjectEditorViewportClient::SetDrawUVOverlayMaterial(const FString& MaterialName, const FString& UVChannel)
 {
-	UVDrawComponentIndex = ComponentIndex;
-	UVDrawLODIndex = LODIndex;
-	UVDrawSectionIndex = SectionIndex;
-	UVDrawUVIndex = UVIndex;
+	// Get LOD Index
+	FString NameWithLOD, ComponentString;
+	MaterialName.Split(FString("_Component_"), &NameWithLOD, &ComponentString);
+	check(ComponentString.IsNumeric());
+	MaterialToDrawInUVsComponent = FCString::Atoi(*ComponentString);
+
+	FString Name, LODIndex;
+	bool bSplit = NameWithLOD.Split(FString(" LOD_"), &Name, &LODIndex);
+	check(bSplit && LODIndex.IsNumeric());
+
+	MaterialToDrawInUVsLOD = FCString::Atoi(*LODIndex);
+	UVChannelToDrawInUVs = FCString::Atoi(*UVChannel);
+
+	// Get Material Index, added if the name of the material already exists within the skeletal mesh.
+	FString DuplicatedMaterialIndex;
+	bSplit = Name.Split(FString("__"), &MaterialToDrawInUVs, &DuplicatedMaterialIndex);
+
+	if (bSplit && DuplicatedMaterialIndex.IsNumeric())
+	{
+		MaterialToDrawInUVsIndex = FCString::Atoi(*DuplicatedMaterialIndex);
+	}
+	else
+	{
+		MaterialToDrawInUVs = Name;
+		MaterialToDrawInUVsIndex = 0;
+	}
 
 	Invalidate();
 }
 
-
 bool FCustomizableObjectEditorViewportClient::IsSetDrawUVOverlayChecked() const
 {
 	return bDrawUVs;
diff --git a/Engine/Plugins/Experimental/Mutable/Source/CustomizableObjectEditor/Private/MuCOE/CustomizableObjectInstanceEditor.cpp b/Engine/Plugins/Experimental/Mutable/Source/CustomizableObjectEditor/Private/MuCOE/CustomizableObjectInstanceEditor.cpp
--- a/Engine/Plugins/Experimental/Mutable/Source/CustomizableObjectEditor/Private/MuCOE/CustomizableObjectInstanceEditor.cpp
+++ b/Engine/Plugins/Experimental/Mutable/Source/CustomizableObjectEditor/Private/MuCOE/CustomizableObjectInstanceEditor.cpp
@@ -676,7 +676,7 @@ void FCustomizableObjectInstanceEditor::OnUpdatePreviewInstance()
 	}
 
 	// Sets which UV material to draw
-	Viewport->SetDrawDefaultUVMaterial();
+	Viewport->SetDrawDefaultUVMaterial(false);
 
 	if (TextureAnalyzer.IsValid())
 	{
diff --git a/Engine/Plugins/Experimental/Mutable/Source/CustomizableObjectEditor/Private/MuCOE/SCustomizableObjectEditorViewport.cpp b/Engine/Plugins/Experimental/Mutable/Source/CustomizableObjectEditor/Private/MuCOE/SCustomizableObjectEditorViewport.cpp
--- a/Engine/Plugins/Experimental/Mutable/Source/CustomizableObjectEditor/Private/MuCOE/SCustomizableObjectEditorViewport.cpp
+++ b/Engine/Plugins/Experimental/Mutable/Source/CustomizableObjectEditor/Private/MuCOE/SCustomizableObjectEditorViewport.cpp
@@ -523,26 +523,76 @@ bool SCustomizableObjectEditorViewportTabBody::IsCameraModeActive(int Value)
 }
 
 
-void SCustomizableObjectEditorViewportTabBody::SetDrawDefaultUVMaterial()
+void SCustomizableObjectEditorViewportTabBody::SetDrawDefaultUVMaterial(bool bIsCompilation)
 {
-	GenerateUVSectionOptions();
+	GenerateUVMaterialOptions();
+	bool bMaterialFound = false;
+
+	if (!ArrayUVMaterialOptionString.IsEmpty())
+	{
+		if (!bIsCompilation && SelectedUVMaterial.IsValid())
+		{
+			// We check if the selected Material still exists after the update
+			FString MaterialName = *SelectedUVMaterial;
+
+			for (int32 MaterialIndex = 0; MaterialIndex < ArrayUVMaterialOptionString.Num(); ++MaterialIndex)
+			{
+				if (MaterialName == *ArrayUVMaterialOptionString[MaterialIndex])
+				{
+					bMaterialFound = true;
+					break;
+				}
+			}
+		}
+
+		// If the CO is recompiled or couldn't find the last used materia, draw the first material
+		if (bIsCompilation || !bMaterialFound)
+		{
+			LevelViewportClient->SetDrawUVOverlayMaterial(*(ArrayUVMaterialOptionString[0]), "0");
+			SelectedUVMaterial = ArrayUVMaterialOptionString[0];
+		}
+	}
+	else
+	{
+		SelectedUVMaterial = nullptr;
+	}
+
 	GenerateUVChannelOptions();
-	
-	if (!SelectedUVSection ||
-		!SelectedUVChannel)
+
+	if (!ArrayUVChannelOptionString.IsEmpty())
 	{
-		LevelViewportClient->SetDrawUV(-1, -1, -1, -1);
+		bool bUVChannelFound = false;
+
+		if (!bIsCompilation && SelectedUVChannel.IsValid())
+		{
+			// We check if the selected Material still exists after the update
+			FString UVChannel = *SelectedUVChannel;
+
+			for (int32 MaterialIndex = 0; MaterialIndex < ArrayUVChannelOptionString.Num(); ++MaterialIndex)
+			{
+				if (UVChannel == *ArrayUVChannelOptionString[MaterialIndex])
+				{
+					bUVChannelFound = true;
+					break;
+				}
+			}
+		}
+
+		// If the CO is recompiled or couldn't find the last used materia, draw the UV channel 0
+		if (bIsCompilation || !bMaterialFound || !bUVChannelFound)
+		{
+			SelectedUVChannel = ArrayUVChannelOptionString[0];
+		}
 	}
 	else
 	{
-		const int32 SectionOptionIndex = UVSectionOptionString.IndexOfByKey(SelectedUVSection);
-		check(SectionOptionIndex != INDEX_NONE);
-		const FSection& Section = UVSectionOption[SectionOptionIndex];
+		SelectedUVChannel = nullptr;
+	}
 
-		const int32 UVIndex = UVChannelOptionString.IndexOfByKey(SelectedUVChannel);
-		check(UVIndex != INDEX_NONE);
-	
-		LevelViewportClient->SetDrawUV(Section.ComponentIndex, Section.LODIndex, Section.SectionIndex, UVIndex);
+	if (UVMaterialOptionCombo.IsValid() && UVChannelOptionCombo.IsValid())
+	{
+		UVMaterialOptionCombo->SetSelectedItem(SelectedUVMaterial);
+		UVChannelOptionCombo->SetSelectedItem(SelectedUVChannel);
 	}
 }
 
@@ -623,22 +673,85 @@ TSharedRef<SWidget> SCustomizableObjectEditorViewportTabBody::GenerateUVMaterial
 	MenuBuilder.BeginSection("ShowUV");
 	{
 		// Generating an array with all the options of the combobox 
-		GenerateUVSectionOptions();
-		
-		UVSectionOptionCombo = SNew(STextComboBox)
-			.OptionsSource(&UVSectionOptionString)
-			.InitiallySelectedItem(SelectedUVSection)
-			.OnSelectionChanged(this, &SCustomizableObjectEditorViewportTabBody::OnSectionChanged);
+		GenerateUVMaterialOptions();
+
+		// Sorting array
+        ArrayUVMaterialOptionString.Sort(::CompareNames);
+
+		// Setting initial selected option
+		if (ArrayUVMaterialOptionString.Num())
+		{
+			bool bFound = false;
+
+			if (SelectedUVMaterial.IsValid())
+			{
+				for (int32 i = 0; i < ArrayUVMaterialOptionString.Num(); ++i)
+				{
+					if (*SelectedUVMaterial == *ArrayUVMaterialOptionString[i])
+					{
+						bFound = true;
+						SelectedUVMaterial = ArrayUVMaterialOptionString[i];
+						break;
+					}
+				}
+			}
+
+			if (!bFound)
+			{
+				SelectedUVMaterial = ArrayUVMaterialOptionString[0];
+			}
+		}
+		else
+		{
+			SelectedUVMaterial = nullptr;
+		}
+	
+		UVMaterialOptionCombo = SNew(STextComboBox)
+			.OptionsSource(&ArrayUVMaterialOptionString)
+			.InitiallySelectedItem(SelectedUVMaterial)
+			.OnSelectionChanged(this, &SCustomizableObjectEditorViewportTabBody::OnMaterialChanged);
 
 		// Generating an array with all the options of the combobox 
 		GenerateUVChannelOptions();
-		
+
+		// Setting initial selected option
+		if (ArrayUVChannelOptionString.Num())
+		{
+			if (SelectedUVChannel.IsValid())
+			{
+				bool bFound = false;
+
+				for (int32 i = 0; i < ArrayUVChannelOptionString.Num(); ++i)
+				{
+					if (*SelectedUVChannel == *ArrayUVChannelOptionString[i])
+					{
+						bFound = true;
+						SelectedUVChannel = ArrayUVChannelOptionString[i];
+						break;
+					}
+				}
+
+				if (!bFound)
+				{
+					SelectedUVChannel = ArrayUVChannelOptionString[0];
+				}
+			}
+			else
+			{
+				SelectedUVChannel = ArrayUVChannelOptionString[0];
+			}
+		}
+		else
+		{
+			SelectedUVChannel = nullptr;
+		}
+
 		UVChannelOptionCombo = SNew(STextComboBox)
-			.OptionsSource(&UVChannelOptionString)
+			.OptionsSource(&ArrayUVChannelOptionString)
 			.InitiallySelectedItem(SelectedUVChannel)
 			.OnSelectionChanged(this, &SCustomizableObjectEditorViewportTabBody::OnUVChannelChanged);
 		
-		MenuBuilder.AddWidget(UVSectionOptionCombo.ToSharedRef(), FText::FromString(TEXT("Section")));
+		MenuBuilder.AddWidget(UVMaterialOptionCombo.ToSharedRef(), FText::FromString(TEXT("Material")));
 		MenuBuilder.AddWidget(UVChannelOptionCombo.ToSharedRef(), FText::FromString(TEXT("UV Channel")));
 	}
 
@@ -664,24 +777,20 @@ TSharedRef<SWidget> SCustomizableObjectEditorViewportTabBody::ShowStateTestData(
 }
 
 
-void SCustomizableObjectEditorViewportTabBody::GenerateUVSectionOptions()
+bool SCustomizableObjectEditorViewportTabBody::IsMaterialsComboEnabled() const
 {
-	ON_SCOPE_EXIT
-	{
-		if (UVSectionOptionCombo)
-		{
-			UVSectionOptionCombo->RefreshOptions();
-			UVSectionOptionCombo->SetSelectedItem(SelectedUVSection);
-		}
-	};
-		
-	UVSectionOptionString.Empty();
-	UVSectionOption.Empty();
+	return true;
+	//return LevelViewportClient.IsValid() && LevelViewportClient->IsSetDrawUVOverlayChecked();
+}
+
+
+void SCustomizableObjectEditorViewportTabBody::GenerateUVMaterialOptions()
+{
+	ArrayUVMaterialOptionString.Empty();
 	
-	for (int32 ComponentIndex = 0; ComponentIndex < PreviewSkeletalMeshComponents.Num(); ++ComponentIndex)
+	int32 ComponentIndex = 0;
+	for (UDebugSkelMeshComponent* PreviewSkeletalMeshComponent : PreviewSkeletalMeshComponents)
 	{
-		const UDebugSkelMeshComponent* PreviewSkeletalMeshComponent = PreviewSkeletalMeshComponents[ComponentIndex];
-		
 		if (PreviewSkeletalMeshComponent != nullptr && UE_MUTABLE_GETSKINNEDASSET(PreviewSkeletalMeshComponent) != nullptr && UE_MUTABLE_GETSKINNEDASSET(PreviewSkeletalMeshComponent)->GetResourceForRendering() != nullptr)
 		{
 			const TArray<UMaterialInterface*> Materials = PreviewSkeletalMeshComponent->GetMaterials();
@@ -691,147 +800,109 @@ void SCustomizableObjectEditorViewportTabBody::GenerateUVSectionOptions()
 				for (int32 SectionIndex = 0; SectionIndex < MeshRes->LODRenderData[LODIndex].RenderSections.Num(); ++SectionIndex)
 				{
 					const FSkelMeshRenderSection& Section = MeshRes->LODRenderData[LODIndex].RenderSections[SectionIndex];
+					
+					if (!Materials.IsValidIndex(Section.MaterialIndex))
+					{
+						continue;
+					}
 
-					FString BaseMaterialName = FString::Printf(TEXT("Section %i"), SectionIndex);
+					if (!Materials[Section.MaterialIndex])
+					{
+						continue;
+					}
+
+					FString BaseMaterialName;
 
-					if (Materials.IsValidIndex(Section.MaterialIndex))
+					const UMaterial* BaseMaterial = Materials[Section.MaterialIndex]->GetBaseMaterial();
+					if (BaseMaterial)
 					{
-						if (UMaterialInterface* MaterialInterface = Materials[Section.MaterialIndex])
-						{
-							if (const UMaterial* BaseMaterial = MaterialInterface->GetBaseMaterial())
-							{
-								BaseMaterialName += " - " + BaseMaterial->GetName();
-							}
-						}
+						BaseMaterialName = BaseMaterial->GetName();
 					}
 
-					UVSectionOptionString.Add(MakeShared<FString>(BaseMaterialName));
-					UVSectionOption.Emplace(ComponentIndex, LODIndex, SectionIndex);
+					BaseMaterialName += FString::Printf(TEXT(" LOD_%d_Component_%d"), LODIndex, ComponentIndex);
+
+					ArrayUVMaterialOptionString.Add(MakeShareable(new FString(BaseMaterialName)));
 				}
 			}
 		}
-	}
 
-	if (SelectedUVSection)
-	{
-		if (const TSharedPtr<FString>* Result = UVSectionOptionString.FindByPredicate([this](const TSharedPtr<FString>& Other){ return *SelectedUVSection == *Other; }))
-		{
-			SelectedUVSection = *Result;
-			return;
-		}	
+		ComponentIndex++;
 	}
-	
-	SelectedUVSection = UVSectionOptionString.IsEmpty() ? nullptr : UVSectionOptionString[0];
 }
 
 
-void SCustomizableObjectEditorViewportTabBody::OnSectionChanged(TSharedPtr<FString> Selected, ESelectInfo::Type SelectInfo)
+void SCustomizableObjectEditorViewportTabBody::OnMaterialChanged(TSharedPtr<FString> Selected, ESelectInfo::Type SelectInfo)
 {
-	SelectedUVSection = Selected;
-	
-	// We need to update options for the new section
-	GenerateUVChannelOptions();
-
-	// Reset the UVChannel selection
-	SelectedUVChannel = UVChannelOptionString.IsEmpty() ? nullptr : UVChannelOptionString[0];
-	UVChannelOptionCombo->SetSelectedItem(SelectedUVChannel);
-	
-	if (!LevelViewportClient)
+	if (Selected.IsValid() && SelectInfo == ESelectInfo::OnMouseClick)
 	{
-		return;
-	}
+		SelectedUVMaterial = Selected;
 
-	if (SelectedUVSection)
-	{
-		const int32 SectionOptionIndex = UVSectionOptionString.IndexOfByKey(SelectedUVSection);
-		check(SectionOptionIndex != INDEX_NONE);
-		const FSection& Section = UVSectionOption[SectionOptionIndex];
+		//We need to update options for the new LOD
+		GenerateUVChannelOptions();
 
-		const int32 UVIndex = UVChannelOptionString.IndexOfByKey(SelectedUVChannel);
-		check(UVIndex != INDEX_NONE);
-	
-		LevelViewportClient->SetDrawUV(Section.ComponentIndex, Section.LODIndex, Section.SectionIndex, UVIndex);
-	}
-	else
-	{
-		LevelViewportClient->SetDrawUV(-1, -1, -1, -1);
+		// Resets the value to the first element of the array
+		if (!ArrayUVChannelOptionString.IsEmpty())
+		{
+			SelectedUVChannel = ArrayUVChannelOptionString[0];
+			UVChannelOptionCombo->SetSelectedItem(SelectedUVChannel);
+		}
+
+		if (LevelViewportClient.IsValid() && SelectedUVChannel.IsValid())
+		{
+			LevelViewportClient->SetDrawUVOverlayMaterial(*Selected, *SelectedUVChannel);
+		}
 	}
 }
 
 
 void SCustomizableObjectEditorViewportTabBody::GenerateUVChannelOptions()
 {
-	ON_SCOPE_EXIT
-	{
-		if (UVChannelOptionCombo)
-		{
-			UVChannelOptionCombo->RefreshOptions();
-			UVChannelOptionCombo->SetSelectedItem(SelectedUVChannel);
-		}
-	};
-	
-	UVChannelOptionString.Empty();
-	
-	if (!SelectedUVSection)
+	ArrayUVChannelOptionString.Empty();
+
+	if (!SelectedUVMaterial)
 	{
-		SelectedUVChannel = nullptr;
 		return;
 	}
-	
-	const int32 Index = UVSectionOptionString.IndexOfByKey(SelectedUVSection);
-	check(Index != INDEX_NONE);
-	const FSection& Section = UVSectionOption[Index];
 
-	const UDebugSkelMeshComponent* PreviewSkeletalMeshComponent = PreviewSkeletalMeshComponents[Section.ComponentIndex];
+	FString SelectedMat = *SelectedUVMaterial;
+
+	// From String to values
+	FString NameWithLOD, ComponentString;
+	SelectedMat.Split(FString("_Component_"), &NameWithLOD, &ComponentString);
+	check(ComponentString.IsNumeric());
+	int32 ComponentIndex = FCString::Atoi(*ComponentString);
+	check(PreviewSkeletalMeshComponents.IsValidIndex(ComponentIndex));
+
+	FString Name, LODString;
+	bool bSplit = NameWithLOD.Split(FString(" LOD_"), &Name, &LODString);
+	int32 LODIndex = FCString::Atoi(*LODString);
+
+	UDebugSkelMeshComponent* PreviewSkeletalMeshComponent = PreviewSkeletalMeshComponents[ComponentIndex];
 
 	if (PreviewSkeletalMeshComponent != nullptr && UE_MUTABLE_GETSKINNEDASSET(PreviewSkeletalMeshComponent) != nullptr
 		&& UE_MUTABLE_GETSKINNEDASSET(PreviewSkeletalMeshComponent)->GetResourceForRendering() != nullptr)
 	{
 		const FSkeletalMeshRenderData* MeshRes = UE_MUTABLE_GETSKINNEDASSET(PreviewSkeletalMeshComponent)->GetResourceForRendering();
 		
-		const int32 UVChannels = MeshRes->LODRenderData[Section.LODIndex].GetNumTexCoords();
+		int32 UVChannels = MeshRes->LODRenderData[LODIndex].GetNumTexCoords();
+		
 		for (int32 UVChan = 0; UVChan < UVChannels; ++UVChan)
 		{
-			UVChannelOptionString.Add(MakeShareable(new FString(FString::FromInt(UVChan))));
+			ArrayUVChannelOptionString.Add(MakeShareable(new FString(FString::FromInt(UVChan))));
 		}
 	}
-
-	if (SelectedUVChannel)
-	{
-		if (const TSharedPtr<FString>* Result = UVChannelOptionString.FindByPredicate([this](const TSharedPtr<FString>& Other) { return *SelectedUVChannel == *Other; }))
-		{
-			SelectedUVChannel = *Result;
-			return;
-		}
-	}
-
-	SelectedUVChannel = UVChannelOptionString.IsEmpty() ? nullptr : UVChannelOptionString[0];
 }
 
 
 void SCustomizableObjectEditorViewportTabBody::OnUVChannelChanged(TSharedPtr<FString> Selected, ESelectInfo::Type SelectInfo)
 {
-	SelectedUVChannel = Selected;
-
-	if (!LevelViewportClient)
+	if (Selected.IsValid() && SelectInfo == ESelectInfo::OnMouseClick)
 	{
-		return;
-	}
-
-	if (SelectedUVChannel)
-	{
-		const int32 SectionOptionIndex = UVSectionOptionString.IndexOfByKey(SelectedUVSection);
-		check(SectionOptionIndex != INDEX_NONE);
-		const FSection& Section = UVSectionOption[SectionOptionIndex];
-
-		const int32 UVIndex = UVChannelOptionString.IndexOfByKey(SelectedUVChannel);
-		check(UVIndex != INDEX_NONE);
-	
-		LevelViewportClient->SetDrawUV(Section.ComponentIndex, Section.LODIndex, Section.SectionIndex, UVIndex);
-	}
-	else
-	{
-		LevelViewportClient->SetDrawUV(-1, -1, -1, -1);
+		SelectedUVChannel = Selected;
+		if (LevelViewportClient.IsValid() && SelectedUVMaterial.IsValid())
+		{
+			LevelViewportClient->SetDrawUVOverlayMaterial(*SelectedUVMaterial, *Selected);
+		}
 	}
 }
 
diff --git a/Engine/Plugins/Experimental/Mutable/Source/CustomizableObjectEditor/Private/MuCOE/SCustomizableObjectEditorViewport.h b/Engine/Plugins/Experimental/Mutable/Source/CustomizableObjectEditor/Private/MuCOE/SCustomizableObjectEditorViewport.h
--- a/Engine/Plugins/Experimental/Mutable/Source/CustomizableObjectEditor/Private/MuCOE/SCustomizableObjectEditorViewport.h
+++ b/Engine/Plugins/Experimental/Mutable/Source/CustomizableObjectEditor/Private/MuCOE/SCustomizableObjectEditorViewport.h
@@ -217,7 +217,7 @@ public:
 
 	/** Sets the first material of the instance as the default section to draw in the UVs Overlay (when there is no material selected in the combobox)
 		Param bIsCompilation: true if this function is called during a CO compilation, false when is called from an instance update	*/
-	void SetDrawDefaultUVMaterial();
+	void SetDrawDefaultUVMaterial(bool bIsCompilation);
 
 	/** Sets the bones visibility */
 	void SetShowBones();
@@ -254,35 +254,26 @@ private:
 	// The currently selected view mode.
 	EViewModeIndex CurrentViewMode = VMI_Lit;
 
-	/** Section identifier. */
-	struct FSection
-	{
-		int32 ComponentIndex = -1;
-		int32 LODIndex = -1;
-		int32 SectionIndex = -1;
-	};
-	
-	// Selection for the section to filter when showing UVs
-	TArray<TSharedPtr<FString>> UVSectionOptionString;
-	TArray<FSection> UVSectionOption;
+	// Selection for the material to filter when showing UVs
+	TArray<TSharedPtr<FString>> ArrayUVMaterialOptionString;
 
-	// Combo box for UV section selection
-	TSharedPtr<STextComboBox> UVSectionOptionCombo;
+	// Combo box for UV material selection
+	TSharedPtr< STextComboBox> UVMaterialOptionCombo;
 
 	// Selection for the channel to filter when showing UVs
-	TArray<TSharedPtr<FString>> UVChannelOptionString;
+	TArray<TSharedPtr<FString>> ArrayUVChannelOptionString;
 
 	// Combo box for UV channel selection
-	TSharedPtr<STextComboBox> UVChannelOptionCombo;
+	TSharedPtr< STextComboBox> UVChannelOptionCombo;
 
 	// Generates the Combobox options for each UV Channel
-	void GenerateUVSectionOptions();
+	void GenerateUVMaterialOptions();
 
 	// Generates the Combobox options for each UV Channel
 	void GenerateUVChannelOptions();
 
 	// Selected option of the UV Material ComboBox
-	TSharedPtr<FString> SelectedUVSection;
+	TSharedPtr<FString> SelectedUVMaterial;
 	
 	// Selected option of the UV Channel ComboBox
 	TSharedPtr<FString> SelectedUVChannel;
@@ -291,7 +282,8 @@ private:
 	TArray< TSharedPtr< FString > > MaterialNames;
 
 	// Slate callbacks to manage the material combo
-	void OnSectionChanged(TSharedPtr<FString> Selected, ESelectInfo::Type SelectInfo);
+	bool IsMaterialsComboEnabled() const;
+	void OnMaterialChanged(TSharedPtr<FString> Selected, ESelectInfo::Type SelectInfo);
 	void OnUVChannelChanged(TSharedPtr<FString> Selected, ESelectInfo::Type SelectInfo);
 
 	/** Current LOD selection*/
diff --git a/Engine/Plugins/Experimental/Mutable/Source/CustomizableObjectEditor/Public/MuCOE/CustomizableObjectEditorViewportClient.h b/Engine/Plugins/Experimental/Mutable/Source/CustomizableObjectEditor/Public/MuCOE/CustomizableObjectEditorViewportClient.h
--- a/Engine/Plugins/Experimental/Mutable/Source/CustomizableObjectEditor/Public/MuCOE/CustomizableObjectEditorViewportClient.h
+++ b/Engine/Plugins/Experimental/Mutable/Source/CustomizableObjectEditor/Public/MuCOE/CustomizableObjectEditorViewportClient.h
@@ -112,7 +112,7 @@ public:
 	 *	@param	InCanvas					The canvas to draw to
 	 *	@param	InTextYPos					The Y-position to begin drawing the UV-Overlay at (due to text displayed above it).
 	 */
-	void DrawUVs(FViewport* InViewport, FCanvas* InCanvas, int32 InTextYPos);
+	void DrawUVs(FViewport* InViewport, FCanvas* InCanvas, int32 InTextYPos, const FString& MaterialName);
 
 	/**
 	 * Bake the instance currently present in the editor. Internally will schedule the update of the instance before baking it's resources
@@ -138,8 +138,8 @@ public:
 	/** Callback for checking the UV overlay show flag. */
 	bool IsSetDrawUVOverlayChecked() const;
 
-	/** Specify which UV to draw. -1 values will not draw anything. */
-	void SetDrawUV(const int32 ComponentIndex, const int32 LODIndex, const int32 SectionIndex, const int32 UVIndex);
+	// Set the material index whose UVs will be drawn.
+	void SetDrawUVOverlayMaterial(const FString& MaterialName, const FString& UVChannel);
 
 	/** Callback for toggling the grid show flag. */
 	void UpdateShowGridFromButton();
@@ -306,11 +306,12 @@ private:
 	bool bDrawUVs;
 	bool bDrawSky;
 
-	// UV to draw
-	int32 UVDrawComponentIndex = 0;
-	int32 UVDrawSectionIndex = 0;
-	int32 UVDrawLODIndex = 0;
-	int32 UVDrawUVIndex = 0;
+	// Material index to draw in the uv preview.
+	FString MaterialToDrawInUVs;
+	int32 MaterialToDrawInUVsLOD;
+	int32 MaterialToDrawInUVsIndex;
+	int32 UVChannelToDrawInUVs;
+	int32 MaterialToDrawInUVsComponent;
 
 	bool bReferenceMeshMissingWarningMessageVisible;
 
