[Mutable] Fix mip generation when using block compression memory optimizations.

#rnx

diff --git a/Engine/Plugins/Experimental/Mutable/Source/MutableRuntime/Private/MuR/OpImageMipmap.cpp b/Engine/Plugins/Experimental/Mutable/Source/MutableRuntime/Private/MuR/OpImageMipmap.cpp
--- a/Engine/Plugins/Experimental/Mutable/Source/MutableRuntime/Private/MuR/OpImageMipmap.cpp
+++ b/Engine/Plugins/Experimental/Mutable/Source/MutableRuntime/Private/MuR/OpImageMipmap.cpp
@@ -415,7 +415,7 @@ namespace OpImageMipmapInternal
 			for (int32 L = DestLODBegin; L < DestLODEnd; ++L)
 			{
 				uint8* DestData = DestImage->GetLODData(L);
-				GenerateNextMipmapUint8SimpleAverage<PIXEL_SIZE>(SrcData, DestData, SourceSize);
+				GenerateNextMipmapUint8Unfiltered<PIXEL_SIZE>(SrcData, DestData, SourceSize);
 				
 				SrcData = DestData;
 				SourceSize = FIntVector2(
@@ -466,25 +466,26 @@ namespace OpImageMipmapInternal
 		{
 		case EMipmapFilterType::MFT_SimpleAverage:
 		{
-			const uint8* SourceData = SourceImage->GetLODData(SrcLOD);
+			const uint8* SrcData = SourceImage->GetLODData(SrcLOD);
 			uint8* DestData = DestImage->GetLODData(DestLODBegin);
 
 			GenerateNextMipBlockCompressed<PixelSize, EMipmapFilterType::MFT_SimpleAverage>(
-					SourceData, DestData, SourceSize, SrcFormat, DestFormat);
+					SrcData, DestData, SourceSize, SrcFormat, DestFormat);
 
+			SrcData = DestData;
 			FIntVector2 CurrentMipSize = FIntVector2(
 					FMath::DivideAndRoundUp(SourceSize[0], 2),
 					FMath::DivideAndRoundUp(SourceSize[1], 2));
-
+		
 			if (CurrentMipSize.X > 1 || CurrentMipSize.Y > 1)
 			{
-				for (int32 L = DestLODBegin + 1; L < DestLODEnd - 1; ++L)
+				for (int32 L = DestLODBegin + 1; L < DestLODEnd; ++L)
 				{
-					const uint8* CurrentMipData = DestData;
-					uint8* NextMipData = DestImage->GetLODData(L);
+					DestData = DestImage->GetLODData(L);
 
-					GenerateNextMipmapUint8SimpleAverage<PixelSize>(CurrentMipData, NextMipData, CurrentMipSize);
+					GenerateNextMipmapUint8SimpleAverage<PixelSize>(SrcData, DestData, CurrentMipSize);
 
+					SrcData = DestData;
 					CurrentMipSize = FIntVector2(
 							FMath::DivideAndRoundUp(CurrentMipSize.X, 2),
 							FMath::DivideAndRoundUp(CurrentMipSize.Y, 2));
@@ -495,25 +496,26 @@ namespace OpImageMipmapInternal
 		}
 		case EMipmapFilterType::MFT_Unfiltered:
 		{
-			const uint8* SourceData = SourceImage->GetLODData(SrcLOD);
+			const uint8* SrcData = SourceImage->GetLODData(SrcLOD);
 			uint8* DestData = DestImage->GetLODData(DestLODBegin);
 			
 			GenerateNextMipBlockCompressed<PixelSize, EMipmapFilterType::MFT_Unfiltered>(
-					SourceData, DestData, SourceSize, SrcFormat, DestFormat);
+					SrcData, DestData, SourceSize, SrcFormat, DestFormat);
 
+			SrcData = DestData;
 			FIntVector2 CurrentMipSize = FIntVector2(
 					FMath::DivideAndRoundUp(SourceSize.X, 2), 
 					FMath::DivideAndRoundUp(SourceSize.Y, 2));
 
 			if (CurrentMipSize.X > 1 || CurrentMipSize.Y > 1)
 			{
-				for (int32 L = DestLODBegin + 1; L < DestLODEnd - 1; ++L)
+				for (int32 L = DestLODBegin + 1; L < DestLODEnd; ++L)
 				{
-					const uint8* CurrentMipData = DestData;
-					uint8* NextMipData = DestImage->GetLODData(L);
+					DestData = DestImage->GetLODData(L);
 
-					GenerateNextMipmapUint8SimpleAverage<PixelSize>(CurrentMipData, NextMipData, CurrentMipSize);
+					GenerateNextMipmapUint8Unfiltered<PixelSize>(SrcData, DestData, CurrentMipSize);
 
+					SrcData = DestData;
 					CurrentMipSize = FIntVector2(
 							FMath::DivideAndRoundUp(CurrentMipSize.X, 2),
 							FMath::DivideAndRoundUp(CurrentMipSize.Y, 2));
diff --git a/Engine/Plugins/Experimental/Mutable/Source/MutableRuntime/Private/MuR/OpImagePixelFormat.cpp b/Engine/Plugins/Experimental/Mutable/Source/MutableRuntime/Private/MuR/OpImagePixelFormat.cpp
--- a/Engine/Plugins/Experimental/Mutable/Source/MutableRuntime/Private/MuR/OpImagePixelFormat.cpp
+++ b/Engine/Plugins/Experimental/Mutable/Source/MutableRuntime/Private/MuR/OpImagePixelFormat.cpp
@@ -125,6 +125,8 @@ namespace ImagePixelFormatInternal
 			const int32 ResultLOD = ResultLODBegin + L;
 			FIntVector2 MipSize = Result->CalculateMipSize(ResultLOD);
 			
+			check(Base->CalculateMipSize(BaseLOD) == MipSize);
+			
 			DecFunc(MipSize.X, MipSize.Y, Base->GetLODData(BaseLOD), Result->GetLODData(ResultLOD));
 		}
 	}
@@ -142,7 +144,9 @@ namespace ImagePixelFormatInternal
 			const int32 BaseLOD = BaseLODBegin + L;
 			const int32 ResultLOD = ResultLODBegin + L;
 			FIntVector2 MipSize = Result->CalculateMipSize(ResultLOD);
-			
+		
+			check(Base->CalculateMipSize(BaseLOD) == MipSize);
+
 			CompFunc(MipSize.X, MipSize.Y, Base->GetLODData(BaseLOD), Result->GetLODData(ResultLOD), Quality);
 		}
 	}
