[mutable] Fix performance problem with optimization of the image swizzle operation.
#rnx

diff --git a/Engine/Plugins/Experimental/Mutable/Source/MutableTools/Internal/MuT/AST.h b/Engine/Plugins/Experimental/Mutable/Source/MutableTools/Internal/MuT/AST.h
--- a/Engine/Plugins/Experimental/Mutable/Source/MutableTools/Internal/MuT/AST.h
+++ b/Engine/Plugins/Experimental/Mutable/Source/MutableTools/Internal/MuT/AST.h
@@ -335,6 +335,24 @@ namespace mu
 		Ptr<ASTOp> Visit(Ptr<ASTOp> at, const class ASTOpImagePixelFormat* currentFormatOp);
 	};
 
+	//---------------------------------------------------------------------------------------------
+	//!
+	//---------------------------------------------------------------------------------------------
+	class Sink_ImageSwizzleAST
+	{
+	public:
+
+		Ptr<ASTOp> Apply(const class ASTOpImageSwizzle* Root);
+
+	protected:
+
+		const class ASTOpImageSwizzle* Root = nullptr;
+		Ptr<ASTOp> InitialSource;
+		TMap<FSinkerOldToNewKey, Ptr<ASTOp>> OldToNew;
+
+		Ptr<ASTOp> Visit(Ptr<ASTOp> at, const class ASTOpImageSwizzle* CurrentSwizzleOp);
+	};
+
 
 	//---------------------------------------------------------------------------------------------
 	//!
@@ -381,6 +399,7 @@ namespace mu
 	{
 		Sink_ImageCropAST ImageCropSinker;
 		Sink_ImagePixelFormatAST ImagePixelFormatSinker;
+		Sink_ImageSwizzleAST ImageSwizzleSinker;
 		Sink_ImageMipmapAST ImageMipmapSinker;
 		Sink_MeshFormatAST MeshFormatSinker;
 	};
diff --git a/Engine/Plugins/Experimental/Mutable/Source/MutableTools/Private/MuT/ASTOpImageCrop.cpp b/Engine/Plugins/Experimental/Mutable/Source/MutableTools/Private/MuT/ASTOpImageCrop.cpp
--- a/Engine/Plugins/Experimental/Mutable/Source/MutableTools/Private/MuT/ASTOpImageCrop.cpp
+++ b/Engine/Plugins/Experimental/Mutable/Source/MutableTools/Private/MuT/ASTOpImageCrop.cpp
@@ -220,7 +220,7 @@ namespace mu
 		check(InRoot->GetOpType() == OP_TYPE::IM_CROP);
 
 		m_root = InRoot;
-		OldToNew.Empty();
+		OldToNew.Reset();
 
 		m_initialSource = InRoot->Source.child();
 		Ptr<ASTOp> newSource = Visit(m_initialSource, InRoot);
diff --git a/Engine/Plugins/Experimental/Mutable/Source/MutableTools/Private/MuT/ASTOpImagePixelFormat.cpp b/Engine/Plugins/Experimental/Mutable/Source/MutableTools/Private/MuT/ASTOpImagePixelFormat.cpp
--- a/Engine/Plugins/Experimental/Mutable/Source/MutableTools/Private/MuT/ASTOpImagePixelFormat.cpp
+++ b/Engine/Plugins/Experimental/Mutable/Source/MutableTools/Private/MuT/ASTOpImagePixelFormat.cpp
@@ -329,7 +329,7 @@ namespace mu
 	mu::Ptr<ASTOp> Sink_ImagePixelFormatAST::Apply(const ASTOpImagePixelFormat* root)
 	{
 		m_root = root;
-		OldToNew.Empty();
+		OldToNew.Reset();
 
 		check(root->GetOpType() == OP_TYPE::IM_PIXELFORMAT);
 
diff --git a/Engine/Plugins/Experimental/Mutable/Source/MutableTools/Private/MuT/ASTOpImageSwizzle.cpp b/Engine/Plugins/Experimental/Mutable/Source/MutableTools/Private/MuT/ASTOpImageSwizzle.cpp
--- a/Engine/Plugins/Experimental/Mutable/Source/MutableTools/Private/MuT/ASTOpImageSwizzle.cpp
+++ b/Engine/Plugins/Experimental/Mutable/Source/MutableTools/Private/MuT/ASTOpImageSwizzle.cpp
@@ -20,6 +20,7 @@
 #include "MuT/ASTOpImageLayerColor.h"
 #include "MuT/ASTOpImageRasterMesh.h"
 #include "MuT/ASTOpImageTransform.h"
+#include "MuT/ASTOpImageMipmap.h"
 #include "MuT/ASTOpSwitch.h"
 #include "MuT/StreamsPrivate.h"
 
@@ -229,6 +230,11 @@ namespace mu
 			}
 		}
 
+		if (!channelSourceAt)
+		{
+			return at;
+		}
+
 		// If we are not changing channel order, just remove the swizzle and adjust the format.
 		bool bSameChannelOrder = true;
 		int32 NumChannelsInFormat = GetImageFormatData(Format).Channels;
@@ -239,11 +245,6 @@ namespace mu
 				bSameChannelOrder = false;
 			}
 		}
-
-		if (!channelSourceAt)
-		{
-			return at;
-		}
 		
 		// If all channels are the same, and in the same order, and the source format is the same that we are
 		// setting in the swizzle, then the swizzle won't do anything.
@@ -258,217 +259,12 @@ namespace mu
 
 		OP_TYPE sourceType = channelSourceAt->GetOpType();
 
-		if (bAllChannelsAreTheSame && channelSourceAt)
+		if (bAllChannelsAreTheSame)
 		{
-			// The instruction can be sunk
-			switch (sourceType)
-			{
-
-			case OP_TYPE::IM_SWITCH:
-			{
-				// Move the swizzle down all the paths
-				Ptr<ASTOpSwitch> nop = mu::Clone<ASTOpSwitch>(channelSourceAt);
-
-				if (nop->def)
-				{
-					Ptr<ASTOpImageSwizzle> defOp = mu::Clone<ASTOpImageSwizzle>(this);
-					ReplaceAllSources(defOp, nop->def.child());
-					nop->def = defOp;
-				}
-
-				for (int32 v = 0; v < nop->cases.Num(); ++v)
-				{
-					if (nop->cases[v].branch)
-					{
-						Ptr<ASTOpImageSwizzle> bOp = mu::Clone<ASTOpImageSwizzle>(this);
-						ReplaceAllSources(bOp, nop->cases[v].branch.child());
-						nop->cases[v].branch = bOp;
-					}
-				}
-
-				at = nop;
-				break;
-			}
-
-			case OP_TYPE::IM_CONDITIONAL:
-			{
-				// We move the swizzle down the two paths
-				Ptr<ASTOpConditional> nop = mu::Clone<ASTOpConditional>(channelSourceAt);
-
-				Ptr<ASTOpImageSwizzle> aOp = mu::Clone<ASTOpImageSwizzle>(this);
-				ReplaceAllSources(aOp, nop->yes.child());
-				nop->yes = aOp;
-
-				Ptr<ASTOpImageSwizzle> bOp = mu::Clone<ASTOpImageSwizzle>(this);
-				ReplaceAllSources(bOp, nop->no.child());
-				nop->no = bOp;
-
-				at = nop;
-				break;
-			}
-
-			case OP_TYPE::IM_LAYER:
-			{
-				// We move the swizzle down the two paths
-				Ptr<ASTOpImageLayer> nop = mu::Clone<ASTOpImageLayer>(channelSourceAt);
-
-				if (nop->Flags == 0)
-				{
-					Ptr<ASTOpImageSwizzle> aOp = mu::Clone<ASTOpImageSwizzle>(this);
-					ReplaceAllSources(aOp, nop->base.child());
-					nop->base = aOp;
-
-					Ptr<ASTOpImageSwizzle> bOp = mu::Clone<ASTOpImageSwizzle>(this);
-					ReplaceAllSources(bOp, nop->blend.child());
-					nop->blend = bOp;
-
-					at = nop;
-				}
-				break;
-			}
-
-			case OP_TYPE::IM_LAYERCOLOUR:
-			{
-				// We move the swizzle down the base path
-				Ptr<ASTOpImageLayerColor> nop = mu::Clone<ASTOpImageLayerColor>(channelSourceAt);
-
-				Ptr<ASTOpImageSwizzle> NewSwizzle = mu::Clone<ASTOpImageSwizzle>(this);
-				ReplaceAllSources(NewSwizzle, nop->base.child());
-				nop->base = NewSwizzle;
-
-				// We need to swizzle the colour too
-				Ptr<ASTOpFixed> cOp = new ASTOpFixed;
-				cOp->op.type = OP_TYPE::CO_SWIZZLE;
-				for (int i = 0; i < MUTABLE_OP_MAX_SWIZZLE_CHANNELS; ++i)
-				{
-					cOp->SetChild(cOp->op.args.ColourSwizzle.sources[i], nop->color);
-					cOp->op.args.ColourSwizzle.sourceChannels[i] = SourceChannels[i];
-				}
-				nop->color = cOp;
-
-				at = nop;
-				break;
-			}
-
-			case OP_TYPE::IM_DISPLACE:
-			{
-				Ptr<ASTOpFixed> NewDisplace = mu::Clone<ASTOpFixed>(channelSourceAt);
-				Ptr<ASTOpImageSwizzle> NewSwizzle = mu::Clone<ASTOpImageSwizzle>(this);
-				ReplaceAllSources(NewSwizzle, NewDisplace->children[NewDisplace->op.args.ImageDisplace.source].child());
-				NewDisplace->SetChild(NewDisplace->op.args.ImageDisplace.source, NewSwizzle);
-				at = NewDisplace;
-				break;
-			}
-
-			case OP_TYPE::IM_INVERT:
-			{
-				Ptr<ASTOpFixed> NewInvert = mu::Clone<ASTOpFixed>(channelSourceAt);
-				Ptr<ASTOpImageSwizzle> NewSwizzle = mu::Clone<ASTOpImageSwizzle>(this);
-				ReplaceAllSources(NewSwizzle, NewInvert->children[NewInvert->op.args.ImageInvert.base].child());
-				NewInvert->SetChild(NewInvert->op.args.ImageInvert.base, NewSwizzle);
-				at = NewInvert;
-				break;
-			}
-
-			case OP_TYPE::IM_RASTERMESH:
-			{
-				Ptr<ASTOpImageRasterMesh> NewRaster = mu::Clone<ASTOpImageRasterMesh>(channelSourceAt);
-				Ptr<ASTOpImageSwizzle> NewSwizzle = mu::Clone<ASTOpImageSwizzle>(this);
-				ReplaceAllSources(NewSwizzle, NewRaster->image.child());
-				NewRaster->image = NewSwizzle;
-
-				// If we are swapping rgb and alphas, we need to correct some flags
-				{
-					// We should only find these two cases
-					if (SourceChannels[0] == 3 || SourceChannels[1] == 3 || SourceChannels[2] == 3)
-					{
-						NewRaster->bIsRGBFadingEnabled = NewRaster->bIsAlphaFadingEnabled;
-					}
-					else if (Sources[3] && SourceChannels[3] < 3 )
-					{
-						NewRaster->bIsAlphaFadingEnabled = NewRaster->bIsRGBFadingEnabled;
-					}
-				}
-
-				at = NewRaster;
-				break;
-			}
-
-			case OP_TYPE::IM_TRANSFORM:
-			{
-				Ptr<ASTOpImageTransform> NewTransform = mu::Clone<ASTOpImageTransform>(channelSourceAt);
-				Ptr<ASTOpImageSwizzle> NewSwizzle = mu::Clone<ASTOpImageSwizzle>(this);
-				ReplaceAllSources(NewSwizzle, NewTransform->Base.child());
-				NewTransform->Base = NewSwizzle;
-				at = NewTransform;
-				break;
-			}
-
-			//case OP_TYPE::IM_RESIZE:
-			//{
-			//	Ptr<ASTOpFixed> NewResize = mu::Clone<ASTOpFixed>(channelSourceAt);
-			//	Ptr<ASTOpImageSwizzle> NewSwizzle = mu::Clone<ASTOpImageSwizzle>(this);
-			//	ReplaceAllSources(NewSwizzle, NewResize->children[NewResize->op.args.ImageResize.source].child());
-			//	NewResize->SetChild(NewResize->op.args.ImageResize.source, NewSwizzle);
-			//	at = NewResize;
-			//	break;
-			//}
-
-			// This is not valid: binarize always forces format L8
-			//case OP_TYPE::IM_BINARISE:
-			//{
-			//	Ptr<ASTOpFixed> NewBinarise = mu::Clone<ASTOpFixed>(channelSourceAt);
-			//	Ptr<ASTOpImageSwizzle> NewSwizzle = mu::Clone<ASTOpImageSwizzle>(this);
-			//	ReplaceAllSources(NewSwizzle, NewBinarise->children[NewBinarise->op.args.ImageBinarise.base].child());
-			//	NewBinarise->SetChild(NewBinarise->op.args.ImageBinarise.base, NewSwizzle);
-			//	at = NewBinarise;
-			//	break;
-			//}
-
-			case OP_TYPE::IM_PIXELFORMAT:
-			{
-				if (bSameChannelOrder)
-				{
-					Ptr<ASTOpImagePixelFormat> NewFormat = mu::Clone<ASTOpImagePixelFormat>(channelSourceAt);
-					NewFormat->Format = Format;
-					at = NewFormat;
-				}
-				break;
-			}
-
-			case OP_TYPE::IM_COMPOSE:
-			{
-				// Move the swizzle down the base and the block, but not the mask.
-				Ptr<ASTOpImageCompose> nop = mu::Clone<ASTOpImageCompose>(channelSourceAt);
-
-				Ptr<ASTOpImageSwizzle> aOp = mu::Clone<ASTOpImageSwizzle>(this);
-				ReplaceAllSources(aOp, nop->Base.child());
-				nop->Base = aOp;
-
-				Ptr<ASTOpImageSwizzle> bOp = mu::Clone<ASTOpImageSwizzle>(this);
-				ReplaceAllSources(bOp, nop->BlockImage.child());
-				nop->BlockImage = bOp;
-
-				at = nop;
-				break;
-			}
-
-			case OP_TYPE::IM_BLANKLAYOUT:
-			{
-				// We can remove the swizzle entirely.
-				// It is not 100% equivalent, because blank layouts are initialized with 0,0,0,1 so the result could be
-				// different, but those pixels shouldn't be used anyway.
-				at = channelSourceAt;
-				break;
-			}
-
-			default:
-				bAllChannelsAreTheSame = false;
-				break;
-			}
+			at = context.ImageSwizzleSinker.Apply(this);
 		}
 
-		if (!bAllChannelsAreTheSame && bAllChannelsAreTheSameType)
+		if (!at && bAllChannelsAreTheSameType)
 		{
 			// Maybe we can still sink the instruction in some cases
 
@@ -1283,4 +1079,274 @@ namespace mu
 		return pRes;
 	}
 
+
+
+
+	//---------------------------------------------------------------------------------------------
+	//---------------------------------------------------------------------------------------------
+	//---------------------------------------------------------------------------------------------
+	mu::Ptr<ASTOp> Sink_ImageSwizzleAST::Apply(const ASTOpImageSwizzle* InRoot)
+	{
+		Root = InRoot;
+		OldToNew.Reset();
+
+		check(Root->GetOpType() == OP_TYPE::IM_SWIZZLE);
+
+		// This sinker only works assuming all swizzle channels come from the same image operation.
+		bool bAllChannelsAreTheSame = true;
+		Ptr<ASTOp> Source;
+		for (int c = 0; c < MUTABLE_OP_MAX_SWIZZLE_CHANNELS; ++c)
+		{
+			Ptr<ASTOp> Candidate = Root->Sources[c].child();
+			if (Candidate)
+			{
+				if (!Source)
+				{
+					Source = Candidate;
+				}
+				else
+				{
+					bAllChannelsAreTheSame = bAllChannelsAreTheSame && (Source == Candidate);
+				}
+			}
+		}
+
+		if (!bAllChannelsAreTheSame || !Source)
+		{
+			return nullptr;
+		}
+
+		InitialSource = Source;
+		mu::Ptr<ASTOp> NewSource = Visit(InitialSource, Root);
+
+		Root = nullptr;
+
+		// If there is any change, it is the new root.
+		if (NewSource != InitialSource)
+		{
+			return NewSource;
+		}
+
+		return nullptr;
+	}
+
+
+	//---------------------------------------------------------------------------------------------
+	mu::Ptr<ASTOp> Sink_ImageSwizzleAST::Visit(mu::Ptr<ASTOp> at, const ASTOpImageSwizzle* CurrentSwizzleOp)
+	{
+		if (!at) return nullptr;
+
+		// Already visited?
+		const Ptr<ASTOp>* Cached = OldToNew.Find({ at,CurrentSwizzleOp });
+		if (Cached)
+		{
+			return *Cached;
+		}
+
+		mu::Ptr<ASTOp> newAt = at;
+		switch (at->GetOpType())
+		{
+
+		case OP_TYPE::IM_CONDITIONAL:
+		{
+			// We move the op down the two paths
+			auto newOp = mu::Clone<ASTOpConditional>(at);
+			newOp->yes = Visit(newOp->yes.child(), CurrentSwizzleOp);
+			newOp->no = Visit(newOp->no.child(), CurrentSwizzleOp);
+			newAt = newOp;
+			break;
+		}
+
+		case OP_TYPE::IM_SWITCH:
+		{
+			// We move the op down all the paths
+			Ptr<ASTOpSwitch> newOp = mu::Clone<ASTOpSwitch>(at);
+			newOp->def = Visit(newOp->def.child(), CurrentSwizzleOp);
+			for (auto& c : newOp->cases)
+			{
+				c.branch = Visit(c.branch.child(), CurrentSwizzleOp);
+			}
+			newAt = newOp;
+			break;
+		}
+
+		case OP_TYPE::IM_COMPOSE:
+		{
+			Ptr<ASTOpImageCompose> newOp = mu::Clone<ASTOpImageCompose>(at);
+			newOp->Base = Visit(newOp->Base.child(), CurrentSwizzleOp);
+			newOp->BlockImage = Visit(newOp->BlockImage.child(), CurrentSwizzleOp);
+			newAt = newOp;
+			break;
+		}
+
+		case OP_TYPE::IM_PATCH:
+		{
+			Ptr<ASTOpImagePatch> newOp = mu::Clone<ASTOpImagePatch>(at);
+			newOp->base = Visit(newOp->base.child(), CurrentSwizzleOp);
+			newOp->patch = Visit(newOp->patch.child(), CurrentSwizzleOp);
+			newAt = newOp;
+			break;
+		}
+
+		case OP_TYPE::IM_MIPMAP:
+		{
+			Ptr<ASTOpImageMipmap> newOp = mu::Clone<ASTOpImageMipmap>(at);
+			newOp->Source = Visit(newOp->Source.child(), CurrentSwizzleOp);
+			newAt = newOp;
+			break;
+		}
+
+		case OP_TYPE::IM_INTERPOLATE:
+		{
+			Ptr<ASTOpFixed> newOp = mu::Clone<ASTOpFixed>(at);
+
+			for (int32 v = 0; v < MUTABLE_OP_MAX_INTERPOLATE_COUNT; ++v)
+			{
+				Ptr<ASTOp> child = newOp->children[newOp->op.args.ImageInterpolate.targets[v]].child();
+				Ptr<ASTOp> bOp = Visit(child, CurrentSwizzleOp);
+				newOp->SetChild(newOp->op.args.ImageInterpolate.targets[v], bOp);
+			}
+
+			newAt = newOp;
+			break;
+		}
+
+		case OP_TYPE::IM_LAYER:
+		{
+			Ptr<ASTOpImageLayer> nop = mu::Clone<ASTOpImageLayer>(at);
+			nop->base = Visit(nop->base.child(), CurrentSwizzleOp);
+			nop->blend = Visit(nop->blend.child(), CurrentSwizzleOp);
+			newAt = nop;
+			break;
+		}
+
+		case OP_TYPE::IM_LAYERCOLOUR:
+		{
+			Ptr<ASTOpImageLayerColor> nop = mu::Clone<ASTOpImageLayerColor>(at);
+			nop->base = Visit(nop->base.child(), CurrentSwizzleOp);
+
+			// We need to swizzle the colour too
+			Ptr<ASTOpFixed> NewColorOp = new ASTOpFixed;
+			NewColorOp->op.type = OP_TYPE::CO_SWIZZLE;
+			for (int i = 0; i < MUTABLE_OP_MAX_SWIZZLE_CHANNELS; ++i)
+			{
+				NewColorOp->SetChild(NewColorOp->op.args.ColourSwizzle.sources[i], nop->color);
+				NewColorOp->op.args.ColourSwizzle.sourceChannels[i] = CurrentSwizzleOp->SourceChannels[i];
+			}
+			nop->color = NewColorOp;
+
+			newAt = nop;
+			break;
+		}
+
+		case OP_TYPE::IM_DISPLACE:
+		{
+			Ptr<ASTOpFixed> NewOp = mu::Clone<ASTOpFixed>(at);
+
+			Ptr<ASTOp> Child = NewOp->children[NewOp->op.args.ImageDisplace.source].child();
+			Ptr<ASTOp> NewSource = Visit(Child, CurrentSwizzleOp);
+			NewOp->SetChild(NewOp->op.args.ImageDisplace.source, NewSource);
+
+			newAt = NewOp;
+			break;
+		}
+
+		case OP_TYPE::IM_INVERT:
+		{
+			Ptr<ASTOpFixed> NewOp = mu::Clone<ASTOpFixed>(at);
+
+			Ptr<ASTOp> Child = NewOp->children[NewOp->op.args.ImageInvert.base].child();
+			Ptr<ASTOp> NewSource = Visit(Child, CurrentSwizzleOp);
+			NewOp->SetChild(NewOp->op.args.ImageInvert.base, NewSource);
+
+			newAt = NewOp;
+			break;
+		}
+
+		case OP_TYPE::IM_RASTERMESH:
+		{
+			Ptr<ASTOpImageRasterMesh> NewOp = mu::Clone<ASTOpImageRasterMesh>(at);
+
+			NewOp->image = Visit(NewOp->image.child(), CurrentSwizzleOp);
+
+			// If we are swapping rgb and alphas, we need to correct some flags
+			{
+				// We should only find these two cases
+				if (CurrentSwizzleOp->SourceChannels[0] == 3 || CurrentSwizzleOp->SourceChannels[1] == 3 || CurrentSwizzleOp->SourceChannels[2] == 3)
+				{
+					NewOp->bIsRGBFadingEnabled = NewOp->bIsAlphaFadingEnabled;
+				}
+				else if (CurrentSwizzleOp->Sources[3] && CurrentSwizzleOp->SourceChannels[3] < 3)
+				{
+					NewOp->bIsAlphaFadingEnabled = NewOp->bIsRGBFadingEnabled;
+				}
+			}
+
+			newAt = NewOp;
+			break;
+		}
+
+		case OP_TYPE::IM_TRANSFORM:
+		{
+			Ptr<ASTOpImageTransform> nop = mu::Clone<ASTOpImageTransform>(at);
+			nop->Base = Visit(nop->Base.child(), CurrentSwizzleOp);
+			newAt = nop;
+			break;
+		}
+
+		case OP_TYPE::IM_PIXELFORMAT:
+		{
+			// If we are not changing channel order, just remove the swizzle and adjust the format.
+			bool bSameChannelOrder = true;
+			int32 NumChannelsInFormat = GetImageFormatData(CurrentSwizzleOp->Format).Channels;
+			for (int32 c = 0; c < NumChannelsInFormat; ++c)
+			{
+				if (CurrentSwizzleOp->Sources[c] && CurrentSwizzleOp->SourceChannels[c] != c)
+				{
+					bSameChannelOrder = false;
+				}
+			}
+
+			if (bSameChannelOrder)
+			{
+				Ptr<ASTOpImagePixelFormat> NewOp = mu::Clone<ASTOpImagePixelFormat>(at);
+				NewOp->Format = CurrentSwizzleOp->Format;
+				newAt = NewOp;
+			}
+			break;
+		}
+
+		case OP_TYPE::IM_BLANKLAYOUT:
+		{
+			Ptr<ASTOpFixed> NewOp = mu::Clone<ASTOpFixed>(at);
+
+			// We can remove the swizzle entirely.
+			// It is not 100% equivalent, because blank layouts are initialized with 0,0,0,1 so the result could be
+			// different, but those pixels shouldn't be used anyway.
+
+			newAt = NewOp;
+			break;
+		}
+
+		default:
+			break;
+		}
+
+		// end on tree branch, replace with format
+		if (at == newAt && at != InitialSource)
+		{
+			mu::Ptr<ASTOpImageSwizzle> NewOp = mu::Clone<ASTOpImageSwizzle>(CurrentSwizzleOp);
+			check(NewOp->GetOpType() == OP_TYPE::IM_SWIZZLE);
+
+			ReplaceAllSources(NewOp, at);
+
+			newAt = NewOp;
+		}
+
+		OldToNew.Add({ at, CurrentSwizzleOp }, newAt);
+
+		return newAt;
+	}
+
+
 }
diff --git a/Engine/Plugins/Experimental/Mutable/Source/MutableTools/Private/MuT/ASTOpMeshFormat.cpp b/Engine/Plugins/Experimental/Mutable/Source/MutableTools/Private/MuT/ASTOpMeshFormat.cpp
--- a/Engine/Plugins/Experimental/Mutable/Source/MutableTools/Private/MuT/ASTOpMeshFormat.cpp
+++ b/Engine/Plugins/Experimental/Mutable/Source/MutableTools/Private/MuT/ASTOpMeshFormat.cpp
@@ -112,7 +112,7 @@ mu::Ptr<ASTOp> Sink_MeshFormatAST::Apply(const ASTOpMeshFormat* root)
 {
 	m_root = root;
 
-	OldToNew.Empty();
+	OldToNew.Reset();
 
 	m_initialSource = m_root->Source.child();
 	mu::Ptr<ASTOp> newSource = Visit(m_initialSource, m_root);
