[Mutable] Move compilation-related data to the FModelResource struct.
Change type for NumLODs, FirstLODAvailable, etc.

#rnx

diff --git a/Engine/Plugins/Experimental/Mutable/Source/CustomizableObject/Classes/MuCO/CustomizableObject.h b/Engine/Plugins/Experimental/Mutable/Source/CustomizableObject/Classes/MuCO/CustomizableObject.h
--- a/Engine/Plugins/Experimental/Mutable/Source/CustomizableObject/Classes/MuCO/CustomizableObject.h
+++ b/Engine/Plugins/Experimental/Mutable/Source/CustomizableObject/Classes/MuCO/CustomizableObject.h
@@ -364,23 +364,6 @@ struct FMutableLODSettings
 	FPerPlatformInt NumMaxStreamedLODs = MAX_MESH_LOD_COUNT;
 
 #endif
-
-	/** Number of LODs that the object root has, which can be higher than the reference mesh LOD count. It's set at the end of the compilation process */
-	UPROPERTY()
-	int32 NumLODsInRoot = 0;
-
-	/** First LOD available, some platforms may remove lower LODs when cooking, this MinLOD represents the first LOD we can generate */
-	UPROPERTY()
-	int32 FirstLODAvailable = 0;
-
-	/** Whether we should stream LODs for on the running platform */
-	UPROPERTY()
-	bool bLODStreamingEnabled = false;
-
-	/** If bEnableLODStreaming is true, maximum number of LODs to stream */
-	UPROPERTY()
-	uint32 NumLODsToStream = 0;
-
 };
 
 
@@ -605,10 +588,6 @@ private:
 	TArray<FProfileParameterDat> InstancePropertiesProfiles;
 #endif // WITH_EDITORONLY_DATA
 
-	/** Amount of components in this CO. Set at the end of the model compilation process. */
-	UPROPERTY()
-	int32 NumMeshComponentsInRoot = 0; // TODO UE-205600, move to FModelResources
-
 public:
 	/** Get the number of components this Customizable Object has. */
 	UFUNCTION(BlueprintCallable, Category = CustomizableObject)
@@ -784,10 +763,6 @@ private:
 	UPROPERTY(EditAnywhere, Category = CustomizableObject)
 	TArray<FName> LowPriorityTextures;
 
-	/** Map of Hash to Streaming blocks, used to stream a block of data representing a resource from the BulkData */
-	UPROPERTY()	
-	TMap<uint64, FMutableStreamableBlock> HashToStreamableBlock; // TODO UE-205600, move to FModelResources
-
 	// Customizable Object Population data start ------------------------------------------------------
 	/** Array to store the selected Population Class tags for this Customizable Object */
 	UPROPERTY()
@@ -813,8 +788,10 @@ private:
 #endif
 	
 public:
+#if WITH_EDITORONLY_DATA
 	FPostCompileDelegate& GetPostCompileDelegate() const;
-	
+#endif
+
 	/** Create a new instance of this object. The instance parameters will be initialized with the object default values. */
 	UFUNCTION(BlueprintCallable, Category = CustomizableObject)
 	UCustomizableObjectInstance* CreateInstance();
diff --git a/Engine/Plugins/Experimental/Mutable/Source/CustomizableObject/Internal/MuCO/CustomizableObjectPrivate.h b/Engine/Plugins/Experimental/Mutable/Source/CustomizableObject/Internal/MuCO/CustomizableObjectPrivate.h
--- a/Engine/Plugins/Experimental/Mutable/Source/CustomizableObject/Internal/MuCO/CustomizableObjectPrivate.h
+++ b/Engine/Plugins/Experimental/Mutable/Source/CustomizableObject/Internal/MuCO/CustomizableObjectPrivate.h
@@ -583,6 +583,25 @@ struct FModelResources
 	TArray<FMorphTargetVertexData> EditorOnlyMorphTargetReconstructionData;
 #endif
 	
+	/** Map of Hash to Streaming blocks, used to stream a block of data representing a resource from the BulkData */
+	UPROPERTY()
+	TMap<uint64, FMutableStreamableBlock> HashToStreamableBlock;
+
+	/** Max number of components in the compiled Model. */
+	UPROPERTY()
+	uint8 NumComponents = 0; 
+
+	/** Max number of LODs in the compiled Model. */
+	UPROPERTY()
+	uint8 NumLODs = 0;
+
+	/** Max number of LODs to stream. Mutable will always generate at least on LOD. */
+	UPROPERTY()
+	uint8 NumLODsToStream = 0;
+
+	/** First LOD available, some platforms may remove lower LODs when cooking, this MinLOD represents the first LOD we can generate */
+	UPROPERTY()
+	uint8 FirstLODAvailable = 0;
 
 };
 
@@ -728,10 +747,6 @@ public:
 	void OnParticipatingObjectDirty(UPackage* Package, bool);
 #endif
 
-	TMap<uint64, FMutableStreamableBlock>& GetHashToStreamableBlock();
-
-	int32& GetNumMeshComponentsInRoot();
-
 	TArray<FString>& GetCustomizableObjectClassTags();
 	
 	TArray<FString>& GetPopulationClassTags();
@@ -774,6 +789,14 @@ public:
 	/** List of Participating Objects (packages) has been marked as dirty since the last compilation. */
 	TArray<FName> DirtyParticipatingObjects;
 	
+	/** Map to identify what CustomizableObject owns a parameter. Used to display a tooltip when hovering a parameter
+	 * in the Prev. instance panel */
+	UPROPERTY(Transient)
+	TMap<FString, FString> CustomizableObjectPathMap;
+
+	UPROPERTY(Transient)
+	TMap<FString, FCustomizableObjectIdPair> GroupNodeMap;
+
 	/** If the object is compiled, this flag is false unless it was compiled with maximum optimizations. If the object is not compiled, its value is meaningless. */
 	bool bIsCompiledWithoutOptimization = true;
 
@@ -782,6 +805,8 @@ public:
 	
 	ECustomizableObjectCompilationState CompilationState = ECustomizableObjectCompilationState::None;
 	
+	FPostCompileDelegate PostCompileDelegate;
+
 #if WITH_EDITOR
 	/** Map of PlatformName to CachedPlatformData. Only valid while cooking. */
 	TMap<FString, FMutableCachedPlatformData> CachedPlatformsData;
@@ -790,16 +815,6 @@ public:
 
 	FCustomizableObjectStatus Status;
 
-	/** Map to identify what CustomizableObject owns a parameter. Used to display a tooltip when hovering a parameter
-	 * in the Prev. instance panel */
-	UPROPERTY(Transient)
-	TMap<FString, FString> CustomizableObjectPathMap;
-
-	UPROPERTY(Transient)
-	TMap<FString, FCustomizableObjectIdPair> GroupNodeMap;
-
-	FPostCompileDelegate PostCompileDelegate;
-
 	// This is information about the parameters in the model that is generated at model compile time.
 	UPROPERTY(Transient)
 	TArray<FMutableModelParameterProperties> ParameterProperties;
@@ -811,6 +826,6 @@ public:
 	// This is a manual version number for the binary blobs in this asset.
 	// Increasing it invalidates all the previously compiled models.
 	// Warning: If while merging code both versions have changed, take the highest+1.
-	static constexpr int32 CurrentSupportedVersion = 435;
+	static constexpr int32 CurrentSupportedVersion = 436;
 };
 
diff --git a/Engine/Plugins/Experimental/Mutable/Source/CustomizableObject/Internal/MuCO/CustomizableObjectSystemPrivate.h b/Engine/Plugins/Experimental/Mutable/Source/CustomizableObject/Internal/MuCO/CustomizableObjectSystemPrivate.h
--- a/Engine/Plugins/Experimental/Mutable/Source/CustomizableObject/Internal/MuCO/CustomizableObjectSystemPrivate.h
+++ b/Engine/Plugins/Experimental/Mutable/Source/CustomizableObject/Internal/MuCO/CustomizableObjectSystemPrivate.h
@@ -513,9 +513,9 @@ public:
 	mu::Instance::ID InstanceID = 0; // Redundant
 	const mu::Instance* MutableInstance = nullptr;
 
-	int32 NumComponents = 0;
-	int32 NumLODsAvailable = 0;
-	int32 FirstLODAvailable = 0;
+	uint8 NumComponents = 0;
+	uint8 NumLODsAvailable = 0;
+	uint8 FirstLODAvailable = 0;
 	uint8 FirstResidentLOD = 0;
 
 	TMap<uint32, FTexturePlatformData*> ImageToPlatformDataMap;
diff --git a/Engine/Plugins/Experimental/Mutable/Source/CustomizableObject/Private/MuCO/CustomizableObject.cpp b/Engine/Plugins/Experimental/Mutable/Source/CustomizableObject/Private/MuCO/CustomizableObject.cpp
--- a/Engine/Plugins/Experimental/Mutable/Source/CustomizableObject/Private/MuCO/CustomizableObject.cpp
+++ b/Engine/Plugins/Experimental/Mutable/Source/CustomizableObject/Private/MuCO/CustomizableObject.cpp
@@ -375,7 +375,6 @@ void UCustomizableObjectPrivate::ClearCompiledData(bool bIsCooking)
 	ParticipatingObjects.Empty();
 #endif
 
-	GetPublic()->HashToStreamableBlock.Empty();
 	GetPublic()->BulkData = nullptr;
 }
 
@@ -563,19 +562,16 @@ void UCustomizableObjectPrivate::SaveCompiledData(FArchive& MemoryWriter, bool b
 	MemoryWriter << LocalModelResources.RealTimeMorphTargetNames;
 	MemoryWriter << LocalModelResources.RealTimeMorphStreamableBlocks;
 
+	MemoryWriter << LocalModelResources.HashToStreamableBlock;
+
 	MemoryWriter << GetPublic()->ClothMeshToMeshVertData;
 	MemoryWriter << GetPublic()->ContributingClothingAssetsData;
 	MemoryWriter << GetPublic()->ClothSharedConfigsData; 
 
-	MemoryWriter << GetPublic()->HashToStreamableBlock;
-
-	MemoryWriter << GetPublic()->LODSettings.NumLODsInRoot;
-	MemoryWriter << GetPublic()->NumMeshComponentsInRoot;
-
-	MemoryWriter << GetPublic()->LODSettings.FirstLODAvailable;
-
-	MemoryWriter << GetPublic()->LODSettings.NumLODsToStream;
-	MemoryWriter << GetPublic()->LODSettings.bLODStreamingEnabled;
+	MemoryWriter << LocalModelResources.NumComponents;
+	MemoryWriter << LocalModelResources.NumLODs;
+	MemoryWriter << LocalModelResources.NumLODsToStream;
+	MemoryWriter << LocalModelResources.FirstLODAvailable;
 
 	// Editor Only data
 	MemoryWriter << bDisableTextureStreaming;
@@ -687,19 +683,16 @@ void UCustomizableObjectPrivate::LoadCompiledData(FArchive& MemoryReader, const
 		MemoryReader << LocalModelResource.RealTimeMorphTargetNames;
 		MemoryReader << LocalModelResource.RealTimeMorphStreamableBlocks;
 
+		MemoryReader << LocalModelResource.HashToStreamableBlock;
+
 		MemoryReader << GetPublic()->ClothMeshToMeshVertData;
 		MemoryReader << GetPublic()->ContributingClothingAssetsData;
 		MemoryReader << GetPublic()->ClothSharedConfigsData;
 
-		MemoryReader << GetPublic()->HashToStreamableBlock;
-
-		MemoryReader << GetPublic()->LODSettings.NumLODsInRoot;
-		MemoryReader << GetPublic()->NumMeshComponentsInRoot;
-
-		MemoryReader << GetPublic()->LODSettings.FirstLODAvailable;
-
-		MemoryReader << GetPublic()->LODSettings.NumLODsToStream;
-		MemoryReader << GetPublic()->LODSettings.bLODStreamingEnabled;
+		MemoryReader << LocalModelResource.NumComponents;
+		MemoryReader << LocalModelResource.NumLODs;
+		MemoryReader << LocalModelResource.NumLODsToStream;
+		MemoryReader << LocalModelResource.FirstLODAvailable;
 
 		bool bInvalidateModel = false;
 
@@ -1208,11 +1201,12 @@ UCustomizableObject* UCustomizableObjectPrivate::GetPublic() const
 	return Public;
 }
 
-
+#if WITH_EDITORONLY_DATA
 FPostCompileDelegate& UCustomizableObject::GetPostCompileDelegate() const
 {
 	return GetPrivate()->PostCompileDelegate;
 }
+#endif
 
 
 UCustomizableObjectInstance* UCustomizableObject::CreateInstance()
@@ -1231,12 +1225,22 @@ UCustomizableObjectInstance* UCustomizableObject::CreateInstance()
 
 int32 UCustomizableObject::GetNumLODs() const
 {
-	return LODSettings.NumLODsInRoot;
+	if (IsCompiled())
+	{
+		return GetPrivate()->GetModelResources().NumLODs;
+	}
+
+	return 0;
 }
 
 int32 UCustomizableObject::GetComponentCount() const
 {
-	return IsCompiled() ? NumMeshComponentsInRoot : 0;
+	if (IsCompiled())
+	{
+		return GetPrivate()->GetModelResources().NumComponents;
+	}
+
+	return 0;
 }
 
 int32 UCustomizableObject::GetParameterCount() const
@@ -1805,7 +1809,7 @@ int32 UCustomizableObjectPrivate::GetMinLODIndex() const
 		MinLODIdx = GetPublic()->LODSettings.MinLOD.GetValue();
 	}
 
-	return FMath::Max(MinLODIdx, GetPublic()->LODSettings.FirstLODAvailable);
+	return FMath::Max(MinLODIdx, static_cast<int32>(GetModelResources().FirstLODAvailable));
 }
 
 
@@ -1953,17 +1957,6 @@ void UCustomizableObjectPrivate::OnParticipatingObjectDirty(UPackage* Package, b
 }
 #endif
 
-TMap<uint64, FMutableStreamableBlock>& UCustomizableObjectPrivate::GetHashToStreamableBlock()
-{
-	return GetPublic()->HashToStreamableBlock;
-}
-
-
-int32& UCustomizableObjectPrivate::GetNumMeshComponentsInRoot()
-{
-	return GetPublic()->NumMeshComponentsInRoot;
-}
-
 
 TArray<FString>& UCustomizableObjectPrivate::GetCustomizableObjectClassTags()
 {
@@ -2117,6 +2110,8 @@ void UCustomizableObjectBulk::PrepareBulkData(UCustomizableObject* InOuter, cons
 		return;
 	}
 
+	FModelResources& ModelResources = CustomizableObject->GetPrivate()->GetModelResources(true);
+
 	uint64 TargetBulkDataFileBytes = CustomizableObject->CompileOptions.PackagedDataBytesLimit;
 	const uint64 MaxChunkSize = UCustomizableObjectSystem::GetInstance()->GetMaxChunkSizeForPlatform(TargetPlatform);
 	TargetBulkDataFileBytes = FMath::Min(TargetBulkDataFileBytes, MaxChunkSize);
@@ -2144,7 +2139,7 @@ void UCustomizableObjectBulk::PrepareBulkData(UCustomizableObject* InOuter, cons
 		for (int32 BlockIndex = 0; BlockIndex < NumBlocks; ++BlockIndex)
 		{
 			uint32 BlockId = Model->GetRomId(BlockIndex);
-			const FMutableStreamableBlock& StreamableBlock = CustomizableObject->GetPrivate()->GetHashToStreamableBlock()[BlockId];
+			const FMutableStreamableBlock& StreamableBlock = ModelResources.HashToStreamableBlock[BlockId];
 			const uint32 BlockSize = StreamableBlock.Size;
 
 			FBlock CurrentBlock = { EDataType::Model, BlockId, BlockSize, SourceOffset };
@@ -2271,7 +2266,7 @@ void UCustomizableObjectBulk::PrepareBulkData(UCustomizableObject* InOuter, cons
 			{
 				FBlock ThisBlock = CurrentFile.Blocks[FileBlockIndex];
 
-				FMutableStreamableBlock& StreamableBlock = CustomizableObject->GetPrivate()->GetHashToStreamableBlock()[ThisBlock.Id];
+				FMutableStreamableBlock& StreamableBlock = ModelResources.HashToStreamableBlock[ThisBlock.Id];
 				check(StreamableBlock.Size == ThisBlock.Size);
 				StreamableBlock.FileId = FileId;
 				StreamableBlock.Offset = OffsetInFile;
diff --git a/Engine/Plugins/Experimental/Mutable/Source/CustomizableObject/Private/MuCO/CustomizableObjectInstance.cpp b/Engine/Plugins/Experimental/Mutable/Source/CustomizableObject/Private/MuCO/CustomizableObjectInstance.cpp
--- a/Engine/Plugins/Experimental/Mutable/Source/CustomizableObject/Private/MuCO/CustomizableObjectInstance.cpp
+++ b/Engine/Plugins/Experimental/Mutable/Source/CustomizableObject/Private/MuCO/CustomizableObjectInstance.cpp
@@ -222,10 +222,10 @@ void UCustomizableInstancePrivate::InitCustomizableObjectData(const UCustomizabl
 	}
 
 	// Init LOD Data
-	NumLODsAvailable = InCustomizableObject->GetNumLODs();
-	FirstLODAvailable = InCustomizableObject->LODSettings.FirstLODAvailable;
-	FirstResidentLOD = InCustomizableObject->LODSettings.bLODStreamingEnabled ? InCustomizableObject->LODSettings.NumLODsToStream : 0;
-	FirstResidentLOD = FMath::Clamp(FirstResidentLOD, FirstLODAvailable, NumLODsAvailable);
+	const FModelResources& ModelResources = InCustomizableObject->GetPrivate()->GetModelResources();
+	NumLODsAvailable = ModelResources.NumLODs;
+	FirstLODAvailable = ModelResources.FirstLODAvailable;
+	FirstResidentLOD = FMath::Clamp(ModelResources.NumLODsToStream, FirstLODAvailable, NumLODsAvailable);
 
 	// Init Component Data
 	FCustomizableInstanceComponentData TemplateComponentData;
@@ -6356,11 +6356,6 @@ void UCustomizableObjectInstance::SetRequestedLODs(int32 InMinLOD, int32 , const
 		return;
 	}
 
-	if (!GetCustomizableObject()->LODSettings.bLODStreamingEnabled)
-	{
-		return;
-	}
-
 	if (GetPrivate()->SkeletalMeshStatus == ESkeletalMeshStatus::Error)
 	{
 		return;
@@ -6828,7 +6823,7 @@ void CalculateBonesToRemove(const FSkeletalMeshLODRenderData& LODResource, const
 		}
 
 		const int32 ParentIndex = RefSkeleton.GetParentIndex(BoneIndex);
-		if (!RemovedBones.IsValidIndex(ParentIndex) || !RemovedBones[ParentIndex])
+		if (RemovedBones.IsValidIndex(ParentIndex) && !RemovedBones[ParentIndex])
 		{
 			OutBonesToRemove.Add(RefSkeleton.GetBoneName(BoneIndex));
 		}
diff --git a/Engine/Plugins/Experimental/Mutable/Source/CustomizableObject/Private/MuCO/UnrealMutableModelDiskStreamer.cpp b/Engine/Plugins/Experimental/Mutable/Source/CustomizableObject/Private/MuCO/UnrealMutableModelDiskStreamer.cpp
--- a/Engine/Plugins/Experimental/Mutable/Source/CustomizableObject/Private/MuCO/UnrealMutableModelDiskStreamer.cpp
+++ b/Engine/Plugins/Experimental/Mutable/Source/CustomizableObject/Private/MuCO/UnrealMutableModelDiskStreamer.cpp
@@ -102,7 +102,8 @@ bool FUnrealMutableModelBulkReader::PrepareStreamingForObject(UCustomizableObjec
 		NewData.BulkFilePrefix = BulkData->GetBulkFilePrefix();
 #endif
 
-		NewData.StreamableBlocks = CustomizableObject->GetPrivate()->GetHashToStreamableBlock();
+		const FModelResources& ModelResources = CustomizableObject->GetPrivate()->GetModelResources();
+		NewData.StreamableBlocks = ModelResources.HashToStreamableBlock;
 		if (NewData.StreamableBlocks.IsEmpty())
 		{
 			UE_LOG(LogMutable, Warning, TEXT("Streaming: Customizable Object %s has no data to stream."), *CustomizableObject->GetName());
diff --git a/Engine/Plugins/Experimental/Mutable/Source/CustomizableObjectEditor/Private/MuCOE/CustomizableObjectCompiler.cpp b/Engine/Plugins/Experimental/Mutable/Source/CustomizableObjectEditor/Private/MuCOE/CustomizableObjectCompiler.cpp
--- a/Engine/Plugins/Experimental/Mutable/Source/CustomizableObjectEditor/Private/MuCOE/CustomizableObjectCompiler.cpp
+++ b/Engine/Plugins/Experimental/Mutable/Source/CustomizableObjectEditor/Private/MuCOE/CustomizableObjectCompiler.cpp
@@ -1099,14 +1099,10 @@ void FCustomizableObjectCompiler::CompileInternal(UCustomizableObject* Object, c
 		Object->GetPrivate()->CustomizableObjectPathMap = GenerationContext.CustomizableObjectPathMap;
 #endif
 
-		Object->LODSettings.NumLODsInRoot = GenerationContext.NumLODsInRoot;
-		
-		Object->GetPrivate()->GetNumMeshComponentsInRoot() = GenerationContext.NumMeshComponentsInRoot;
-		
-		Object->LODSettings.FirstLODAvailable = GenerationContext.FirstLODAvailable;
-
-		Object->LODSettings.bLODStreamingEnabled = GenerationContext.bEnableLODStreaming;
-		Object->LODSettings.NumLODsToStream = GenerationContext.NumMaxLODsToStream;
+		ModelResources.NumComponents = GenerationContext.NumMeshComponentsInRoot;
+		ModelResources.NumLODs = GenerationContext.NumLODsInRoot;
+		ModelResources.NumLODsToStream = GenerationContext.bEnableLODStreaming ? GenerationContext.NumMaxLODsToStream : 0;
+		ModelResources.FirstLODAvailable = GenerationContext.FirstLODAvailable;
 
 		Object->GetPrivate()->GetStreamedResourceData() = MoveTemp(GenerationContext.StreamedResourceData);
 
@@ -1283,19 +1279,20 @@ void FCustomizableObjectCompiler::FinishCompilation()
 
 	// Generate a map that using the resource id tells the offset and size of the resource inside the bulk data
 	// At this point it is assumed that all data goes into a single file.
-	if (Model)
+	if (Model && !Options.bIsCooking)
 	{
-		uint64 Offset = 0;
-
+		FModelResources& ModelResources = CurrentObject->GetPrivate()->GetModelResources(false);
+		
 		const int32 NumStreamingFiles = Model->GetRomCount();
-		CurrentObject->GetPrivate()->GetHashToStreamableBlock().Empty(NumStreamingFiles);
+		ModelResources.HashToStreamableBlock.Empty(NumStreamingFiles);
 
+		uint64 Offset = 0;
 		for (int32 FileIndex = 0; FileIndex < NumStreamingFiles; ++FileIndex)
 		{
 			const uint32 ResourceId = Model->GetRomId(FileIndex);
 			const uint32 ResourceSize = Model->GetRomSize(FileIndex);
 
-			CurrentObject->GetPrivate()->GetHashToStreamableBlock().Add(ResourceId, FMutableStreamableBlock{0, ResourceSize, Offset });
+			ModelResources.HashToStreamableBlock.Add(ResourceId, FMutableStreamableBlock{0, ResourceSize, Offset });
 			Offset += ResourceSize;
 		}
 	}
diff --git a/Engine/Plugins/Experimental/Mutable/Source/CustomizableObjectEditor/Private/MuCOE/GenerateMutableSource/GenerateMutableSource.cpp b/Engine/Plugins/Experimental/Mutable/Source/CustomizableObjectEditor/Private/MuCOE/GenerateMutableSource/GenerateMutableSource.cpp
--- a/Engine/Plugins/Experimental/Mutable/Source/CustomizableObjectEditor/Private/MuCOE/GenerateMutableSource/GenerateMutableSource.cpp
+++ b/Engine/Plugins/Experimental/Mutable/Source/CustomizableObjectEditor/Private/MuCOE/GenerateMutableSource/GenerateMutableSource.cpp
@@ -845,7 +845,7 @@ mu::NodeObjectPtr GenerateMutableSource(const UEdGraphPin * Pin, FMutableGraphGe
 					GenerationContext.bEnableLODStreaming = GenerationContext.bEnableLODStreaming &&
 						RefSkeletalMesh->GetEnableLODStreaming(GenerationContext.Options.TargetPlatform);
 
-					GenerationContext.NumMaxLODsToStream = FMath::Min(GenerationContext.NumMaxLODsToStream,
+					GenerationContext.NumMaxLODsToStream = FMath::Min(static_cast<int32>(GenerationContext.NumMaxLODsToStream),
 						RefSkeletalMesh->GetMaxNumStreamedLODs(GenerationContext.Options.TargetPlatform));
 				}
 			}
@@ -1745,7 +1745,7 @@ int32 GetMaxTextureSize(const UTexture2D& ReferenceTexture, const UTextureLODSet
 }
 
 
-int32 GetTextureSizeInGame(const UTexture2D& Texture, const UTextureLODSettings& LODSettings, int32 SurfaceLODBias)
+int32 GetTextureSizeInGame(const UTexture2D& Texture, const UTextureLODSettings& LODSettings, uint8 SurfaceLODBias)
 {
 	const int32 SourceSize = (uint32)FMath::Max3(Texture.Source.GetSizeX(), Texture.Source.GetSizeY(), (int64)1);
 	const int32 NumMipsSource = FMath::CeilLogTwo(SourceSize) + 1;
diff --git a/Engine/Plugins/Experimental/Mutable/Source/CustomizableObjectEditor/Private/MuCOE/GenerateMutableSource/GenerateMutableSource.h b/Engine/Plugins/Experimental/Mutable/Source/CustomizableObjectEditor/Private/MuCOE/GenerateMutableSource/GenerateMutableSource.h
--- a/Engine/Plugins/Experimental/Mutable/Source/CustomizableObjectEditor/Private/MuCOE/GenerateMutableSource/GenerateMutableSource.h
+++ b/Engine/Plugins/Experimental/Mutable/Source/CustomizableObjectEditor/Private/MuCOE/GenerateMutableSource/GenerateMutableSource.h
@@ -720,14 +720,14 @@ struct FMutableGraphGenerationContext
 	};
 	TArray< ObjectParent > ComponentNewNode;
 
-	int32 FromLOD = 0; // LOD to append to the CurrentLOD when using AutomaticLODs. 
-	int32 CurrentLOD = 0;
-	int32 NumLODsInRoot = 0;
+	uint8 FromLOD = 0; // LOD to append to the CurrentLOD when using AutomaticLODs. 
+	uint8 CurrentLOD = 0;
 	int32 CurrentMeshComponent = 0;
-	int32 NumMeshComponentsInRoot = 0;
 
-	int32 FirstLODAvailable = MAX_MESH_LOD_COUNT;
-	int32 NumMaxLODsToStream = MAX_MESH_LOD_COUNT;
+	uint8 NumLODsInRoot = 0;
+	uint8 NumMeshComponentsInRoot = 0;
+	uint8 FirstLODAvailable = MAX_MESH_LOD_COUNT;
+	uint8 NumMaxLODsToStream = MAX_MESH_LOD_COUNT;
 
 	bool bEnableLODStreaming = true;
 
@@ -870,7 +870,7 @@ uint32 ComputeLODBiasForTexture(const FMutableGraphGenerationContext& Generation
 int32 GetMaxTextureSize(const UTexture2D& ReferenceTexture, const UTextureLODSettings& LODSettings);
 
 // Max texture size of the texture with per platform MaxTextureSize and LODBias applied.
-int32 GetTextureSizeInGame(const UTexture2D& Texture, const UTextureLODSettings& LODSettings, int32 SurfaceLODBias = 0);
+int32 GetTextureSizeInGame(const UTexture2D& Texture, const UTextureLODSettings& LODSettings, uint8 SurfaceLODBias = 0);
 
 mu::Ptr<mu::Image> GenerateImageConstant( UTexture*, FMutableGraphGenerationContext&, bool bIsReference);
 
diff --git a/Engine/Plugins/Experimental/Mutable/Source/CustomizableObjectEditor/Private/MuCOE/GenerateMutableSource/GenerateMutableSourceImage.cpp b/Engine/Plugins/Experimental/Mutable/Source/CustomizableObjectEditor/Private/MuCOE/GenerateMutableSource/GenerateMutableSourceImage.cpp
--- a/Engine/Plugins/Experimental/Mutable/Source/CustomizableObjectEditor/Private/MuCOE/GenerateMutableSource/GenerateMutableSourceImage.cpp
+++ b/Engine/Plugins/Experimental/Mutable/Source/CustomizableObjectEditor/Private/MuCOE/GenerateMutableSource/GenerateMutableSourceImage.cpp
@@ -840,8 +840,8 @@ mu::NodeImagePtr GenerateMutableSourceImage(const UEdGraphPin* Pin, FMutableGrap
 			GenerationContext.AddParticipatingObject(*TypedNodeTransform->ReferenceTexture);
 
 			const UTextureLODSettings& TextureLODSettings = GenerationContext.Options.TargetPlatform->GetTextureLODSettings();
-			const int32 FirstLODAvailable = GenerationContext.Options.bUseLODAsBias ? GenerationContext.FirstLODAvailable : 0;
-			MaxReferenceTextureSizeInGame = GetTextureSizeInGame(*TypedNodeTransform->ReferenceTexture, TextureLODSettings, FirstLODAvailable);
+			const uint8 SurfaceLODBias = GenerationContext.Options.bUseLODAsBias ? GenerationContext.FirstLODAvailable : 0;
+			MaxReferenceTextureSizeInGame = GetTextureSizeInGame(*TypedNodeTransform->ReferenceTexture, TextureLODSettings, SurfaceLODBias);
 		}
 
 		// Max TextureSize allowed
diff --git a/Engine/Plugins/Experimental/Mutable/Source/CustomizableObjectEditor/Private/MuCOE/GenerateMutableSource/GenerateMutableSourceSurface.cpp b/Engine/Plugins/Experimental/Mutable/Source/CustomizableObjectEditor/Private/MuCOE/GenerateMutableSource/GenerateMutableSourceSurface.cpp
--- a/Engine/Plugins/Experimental/Mutable/Source/CustomizableObjectEditor/Private/MuCOE/GenerateMutableSource/GenerateMutableSourceSurface.cpp
+++ b/Engine/Plugins/Experimental/Mutable/Source/CustomizableObjectEditor/Private/MuCOE/GenerateMutableSource/GenerateMutableSourceSurface.cpp
@@ -636,8 +636,8 @@ mu::NodeSurfacePtr GenerateMutableSourceSurface(const UEdGraphPin * Pin, FMutabl
 						if (ReferenceTexture)
 						{
 							// Apply base LODBias. It will be propagated to most images.
-							const uint32 FirstLODAvailable = GenerationContext.Options.bUseLODAsBias ? GenerationContext.FirstLODAvailable : 0;
-							const uint32 BaseLODBias = ComputeLODBiasForTexture(GenerationContext, *ReferenceTexture) + FirstLODAvailable;
+							const uint32 SurfaceLODBias = GenerationContext.Options.bUseLODAsBias ? GenerationContext.FirstLODAvailable : 0;
+							const uint32 BaseLODBias = ComputeLODBiasForTexture(GenerationContext, *ReferenceTexture) + SurfaceLODBias;
 							mu::NodeImagePtr LastImage = ResizeTextureByNumMips(ImageNode, BaseLODBias);
 
 							mu::NodeImageMipmapPtr MipmapImage = new mu::NodeImageMipmap();
