[mutable] Added an error message in case of missing mesh/texture input connections to some clip modifier nodes.
#rnx
#tests Editor

diff --git a/Engine/Plugins/Experimental/Mutable/Source/CustomizableObjectEditor/Private/MuCOE/GenerateMutableSource/GenerateMutableSourceModifier.cpp b/Engine/Plugins/Experimental/Mutable/Source/CustomizableObjectEditor/Private/MuCOE/GenerateMutableSource/GenerateMutableSourceModifier.cpp
--- a/Engine/Plugins/Experimental/Mutable/Source/CustomizableObjectEditor/Private/MuCOE/GenerateMutableSource/GenerateMutableSourceModifier.cpp
+++ b/Engine/Plugins/Experimental/Mutable/Source/CustomizableObjectEditor/Private/MuCOE/GenerateMutableSource/GenerateMutableSourceModifier.cpp
@@ -101,6 +101,12 @@ mu::NodeModifierPtr GenerateMutableSourceModifier(const UEdGraphPin * Pin, FMuta
 
 			ClipNode->SetBindingMethod(BindingMethod);
 		}
+		else
+		{
+			FText ErrorMsg = LOCTEXT("ClipDeform mesh", "The clip deform node requires an input clip shape.");
+			GenerationContext.Compiler->CompilerLog(ErrorMsg, TypedNodeClipDeform, EMessageSeverity::Error);
+			Result = nullptr;
+		}
 	
 		ClipNode->SetMultipleTagPolicy(TypedNodeClipDeform->MultipleTagPolicy);
 		for (const FString& Tag : TypedNodeClipDeform->Tags)
@@ -162,6 +168,12 @@ mu::NodeModifierPtr GenerateMutableSourceModifier(const UEdGraphPin * Pin, FMuta
 
 			ClipNode->SetClipMesh(ClipMesh.get());
 		}
+		else
+		{
+			FText ErrorMsg = LOCTEXT("Clipping mesh", "The clip mesh with mesh node requires an input clip mesh.");
+			GenerationContext.Compiler->CompilerLog(ErrorMsg, TypedNodeClipMesh, EMessageSeverity::Error);
+			Result = nullptr;
+		}
 
 		ClipNode->SetMultipleTagPolicy(TypedNodeClipMesh->MultipleTagPolicy);
 		for (const FString& Tag : TypedNodeClipMesh->Tags)
@@ -204,6 +216,12 @@ mu::NodeModifierPtr GenerateMutableSourceModifier(const UEdGraphPin * Pin, FMuta
 
 			ClipNode->SetClipMask(ClipMask.get());
 		}
+		else
+		{
+			FText ErrorMsg = LOCTEXT("ClipUVMask mesh", "The clip mesh with UV Mask node requires an input texture mask.");
+			GenerationContext.Compiler->CompilerLog(ErrorMsg, TypedNodeClipUVMask, EMessageSeverity::Error);
+			Result = nullptr;
+		}
 
 		ClipNode->SetLayoutIndex(TypedNodeClipUVMask->UVChannelForMask);
 
