[Mutable] Split Task_Mutable_UpdateMesh into multiple tasks.
Move old task to BusyWaits_Deprecated.

#rnx

diff --git a/Engine/Plugins/Experimental/Mutable/Source/CustomizableObject/Private/MuCO/BusyWaits_Deprecated.cpp b/Engine/Plugins/Experimental/Mutable/Source/CustomizableObject/Private/MuCO/BusyWaits_Deprecated.cpp
--- a/Engine/Plugins/Experimental/Mutable/Source/CustomizableObject/Private/MuCO/BusyWaits_Deprecated.cpp
+++ b/Engine/Plugins/Experimental/Mutable/Source/CustomizableObject/Private/MuCO/BusyWaits_Deprecated.cpp
@@ -5,6 +5,7 @@
 #include "MuCO/CustomizableObjectSystem.h"
 #include "MuCO/CustomizableObjectSystemPrivate.h"
 #include "MuCO/CustomizableObjectMipDataProvider.h"
+#include "MuCO/CustomizableObjectMeshUpdate.h"
 
 #include "MuR/Model.h"
 #include "MuR/MutableTrace.h"
@@ -596,3 +597,49 @@ namespace CustomizableObjectMipDataProvider::ImplDeprecated
 		}
 	}
 }
+
+namespace CustomizableObjectMeshUpdate::ImplDeprecated
+{
+	void Task_Mutable_UpdateMesh(const TSharedPtr<FMutableMeshOperationData>& OperationData)
+	{
+		MUTABLE_CPUPROFILER_SCOPE(Task_Mutable_UpdateMesh);
+
+		mu::SystemPtr System = OperationData->System;
+		const TSharedPtr<mu::Model, ESPMode::ThreadSafe> Model = OperationData->Model;
+
+#if WITH_EDITOR
+		// Recompiling a CO in the editor will invalidate the previously generated Model. Check that it is valid before accessing the streamed data.
+		if (!Model || !Model->IsValid())
+		{
+			return;
+		}
+#endif
+
+		// For now, we are forcing the recreation of mutable-side instances with every update.
+		mu::Instance::ID InstanceID = System->NewInstance(Model);
+		UE_LOG(LogMutable, Verbose, TEXT("Creating Mutable instance with id [%d] for a mesh update"), InstanceID);
+
+		// LOD mask, set to all ones to build all LODs
+		const uint32 LODMask = 0xFFFFFFFF;
+
+		// Main instance generation step
+		const mu::Instance* Instance = System->BeginUpdate(InstanceID, OperationData->Parameters, OperationData->State, LODMask);
+		check(Instance);
+
+		// Generate the required meshes
+		for (int32 Index = OperationData->PendingFirstLODIdx; Index < OperationData->CurrentFirstLODIdx; ++Index)
+		{
+			const mu::FResourceID MeshID = OperationData->MeshIDs[Index];
+			OperationData->Meshes[Index] = System->GetMeshInline(InstanceID, MeshID);
+		}
+
+		// End update
+		System->EndUpdate(InstanceID);
+		System->ReleaseInstance(InstanceID);
+
+		if (CVarClearWorkingMemoryOnUpdateEnd.GetValueOnAnyThread())
+		{
+			System->ClearWorkingMemory();
+		}
+	}
+}
diff --git a/Engine/Plugins/Experimental/Mutable/Source/CustomizableObject/Private/MuCO/BusyWaits_Deprecated.h b/Engine/Plugins/Experimental/Mutable/Source/CustomizableObject/Private/MuCO/BusyWaits_Deprecated.h
--- a/Engine/Plugins/Experimental/Mutable/Source/CustomizableObject/Private/MuCO/BusyWaits_Deprecated.h
+++ b/Engine/Plugins/Experimental/Mutable/Source/CustomizableObject/Private/MuCO/BusyWaits_Deprecated.h
@@ -10,6 +10,8 @@
 
 class FUpdateContextPrivate;
 struct FMutableImageOperationData;
+struct FMutableMeshOperationData;
+
 namespace mu
 {
 	class Model;
@@ -37,3 +39,10 @@ namespace CustomizableObjectMipDataProvider::ImplDeprecated
 	// This runs in the mutable thread.
 	void Task_Mutable_UpdateImage(TSharedPtr<FMutableImageOperationData> OperationData);
 }
+
+
+namespace CustomizableObjectMeshUpdate::ImplDeprecated
+{
+	// This runs in the mutable thread.
+	void Task_Mutable_UpdateMesh(const TSharedPtr<FMutableMeshOperationData>& OperationData);
+}
\ No newline at end of file
diff --git a/Engine/Plugins/Experimental/Mutable/Source/CustomizableObject/Private/MuCO/CustomizableObjectMeshUpdate.cpp b/Engine/Plugins/Experimental/Mutable/Source/CustomizableObject/Private/MuCO/CustomizableObjectMeshUpdate.cpp
--- a/Engine/Plugins/Experimental/Mutable/Source/CustomizableObject/Private/MuCO/CustomizableObjectMeshUpdate.cpp
+++ b/Engine/Plugins/Experimental/Mutable/Source/CustomizableObject/Private/MuCO/CustomizableObjectMeshUpdate.cpp
@@ -19,8 +19,12 @@ CustomizableObjectMeshUpdate.cpp: Helpers to stream in CustomizableObject skelet
 #include "Streaming/RenderAssetUpdate.inl"
 #include "Rendering/SkeletalMeshRenderData.h"
 
+
+#include "BusyWaits_Deprecated.h"
+
 template class TRenderAssetUpdate<FSkelMeshUpdateContext>;
 
+#define UE_MUTABLE_UPDATE_MESH_REGION		TEXT("Task_Mutable_UpdateMesh")
 
 FCustomizableObjectMeshStreamIn::FCustomizableObjectMeshStreamIn(const UCustomizableObjectSkeletalMesh* InMesh, bool bInHighPrio, EThreadType CreateResourcesThread)
 	: FSkeletalMeshStreamIn(InMesh, CreateResourcesThread),
@@ -45,6 +49,18 @@ FCustomizableObjectMeshStreamIn::FCustomizableObjectMeshStreamIn(const UCustomiz
 	PushTask(FContext(InMesh, TT_None), TT_Async, SRA_UPDATE_CALLBACK(DoInitiate), TT_None, nullptr);
 }
 
+void FCustomizableObjectMeshStreamIn::OnUpdateMeshFinished()
+{
+	check(TaskSynchronization.GetValue() > 0)
+
+	// At this point task synchronization would hold the number of pending requests.
+	TaskSynchronization.Decrement();
+
+	// The tick here is intended to schedule the success or cancel callback.
+	// Using TT_None ensure gets which could create a dead lock.
+	Tick(FSkeletalMeshUpdate::TT_None);
+}
+
 void FCustomizableObjectMeshStreamIn::DoInitiate(const FContext& Context)
 {
 	check(Context.CurrentThread == TT_Async);
@@ -92,21 +108,55 @@ void FCustomizableObjectMeshStreamIn::DoCancelMeshUpdate(const FContext& Context
 
 namespace impl
 {
-	void Task_Mutable_UpdateMesh(const TSharedPtr<FMutableMeshOperationData> OperationData)
+	void Task_Mutable_UpdateMesh_End(const TSharedPtr<FMutableMeshOperationData> OperationData, TRefCountPtr<FCustomizableObjectMeshStreamIn>& Task, mu::Instance::ID InstanceID)
 	{
-		MUTABLE_CPUPROFILER_SCOPE(Task_Mutable_UpdateMesh);
+		MUTABLE_CPUPROFILER_SCOPE(Task_Mutable_UpdateMesh_End);
+
+		// End update
+		OperationData->System->EndUpdate(InstanceID);
+		OperationData->System->ReleaseInstance(InstanceID);
+
+		if (CVarClearWorkingMemoryOnUpdateEnd.GetValueOnAnyThread())
+		{
+			OperationData->System->ClearWorkingMemory();
+		}
 
-		// This runs in a worker thread.
-		check(OperationData.IsValid());
-		check(OperationData->System.get());
-		check(OperationData->Model);
-		check(OperationData->Parameters.get());
+		Task->OnUpdateMeshFinished();
+		
+		TRACE_END_REGION(UE_MUTABLE_UPDATE_MESH_REGION);
+	}
+
+	void Task_Mutable_UpdateMesh_Loop(
+		const TSharedPtr<FMutableMeshOperationData> OperationData,
+		TRefCountPtr<FCustomizableObjectMeshStreamIn>& Task,
+		mu::Instance::ID InstanceID,
+		int32 LODIndex)
+	{
+		MUTABLE_CPUPROFILER_SCOPE(Task_Mutable_UpdateMesh_Loop);
 
-		if (!OperationData.IsValid())
+		if (LODIndex == OperationData->CurrentFirstLODIdx)
 		{
+			Task_Mutable_UpdateMesh_End(OperationData, Task, InstanceID);
 			return;
 		}
 
+		const mu::FResourceID MeshID = OperationData->MeshIDs[LODIndex];
+		UE::Tasks::TTask<mu::Ptr<const mu::Mesh>> GetMeshTask = OperationData->System->GetMesh(InstanceID, MeshID);
+
+		UE::Tasks::AddNested(UE::Tasks::Launch(TEXT("Task_MutableGetMeshes_GetMesh_Post"), [=]() mutable
+			{
+				OperationData->Meshes[LODIndex] = GetMeshTask.GetResult();
+
+				Task_Mutable_UpdateMesh_Loop(OperationData, Task, InstanceID, LODIndex + 1);
+			},
+			GetMeshTask));
+	}
+
+	void Task_Mutable_UpdateMesh(const TSharedPtr<FMutableMeshOperationData> OperationData, TRefCountPtr<FCustomizableObjectMeshStreamIn>& Task)
+	{
+		MUTABLE_CPUPROFILER_SCOPE(Task_Mutable_UpdateMesh);
+		TRACE_BEGIN_REGION(UE_MUTABLE_UPDATE_MESH_REGION);
+
 		mu::SystemPtr System = OperationData->System;
 		const TSharedPtr<mu::Model, ESPMode::ThreadSafe> Model = OperationData->Model;
 
@@ -114,6 +164,7 @@ namespace impl
 		// Recompiling a CO in the editor will invalidate the previously generated Model. Check that it is valid before accessing the streamed data.
 		if (!Model || !Model->IsValid())
 		{
+			TRACE_END_REGION(UE_MUTABLE_UPDATE_MESH_REGION);
 			return;
 		}
 #endif
@@ -129,22 +180,9 @@ namespace impl
 		const mu::Instance* Instance = System->BeginUpdate(InstanceID, OperationData->Parameters, OperationData->State, LODMask);
 		check(Instance);
 
-		// Generate the required meshes
-		for (int32 Index = OperationData->PendingFirstLODIdx; Index < OperationData->CurrentFirstLODIdx; ++Index)
-		{
-			const mu::FResourceID MeshID = OperationData->MeshIDs[Index];
-			OperationData->Meshes[Index] = System->GetMeshInline(InstanceID, MeshID);
-		}
-
-		// End update
-		System->EndUpdate(InstanceID);
-		System->ReleaseInstance(InstanceID);
-
-		if (CVarClearWorkingMemoryOnUpdateEnd.GetValueOnAnyThread())
-		{
-			System->ClearWorkingMemory();
-		}
+		Task_Mutable_UpdateMesh_Loop(OperationData, Task, InstanceID, OperationData->PendingFirstLODIdx);
 	}
+	
 } // namespace
 
 
@@ -174,18 +212,22 @@ void FCustomizableObjectMeshStreamIn::RequestMeshUpdate(const FContext& Context)
 	TRefCountPtr<FCustomizableObjectMeshStreamIn> RefThis = this;
 	TSharedPtr<FMutableMeshOperationData> SharedOperationData = OperationData;
 
+
 	MutableTaskId = CustomizableObjectSystem->MutableTaskGraph.AddMutableThreadTaskLowPriority(
 		TEXT("Mutable_MeshUpdate"),
-		[SharedOperationData, RefThis]()
+		[SharedOperationData, RefThis]() mutable
 		{
-			impl::Task_Mutable_UpdateMesh(SharedOperationData);
-
-			// At this point task synchronization would hold the number of pending requests.
-			RefThis->TaskSynchronization.Decrement();
-
-			// The tick here is intended to schedule the success or cancel callback.
-			// Using TT_None ensure gets which could create a dead lock.
-			RefThis->Tick(FSkeletalMeshUpdate::TT_None);
+			if (CVarEnableNewSplitMutableTask.GetValueOnAnyThread())
+			{
+				impl::Task_Mutable_UpdateMesh(SharedOperationData, RefThis);
+			}
+			else
+			{
+				using namespace CustomizableObjectMeshUpdate;
+				ImplDeprecated::Task_Mutable_UpdateMesh(SharedOperationData);
+				
+				RefThis->OnUpdateMeshFinished();
+			}
 		});
 }
 
diff --git a/Engine/Plugins/Experimental/Mutable/Source/CustomizableObject/Private/MuCO/CustomizableObjectMeshUpdate.h b/Engine/Plugins/Experimental/Mutable/Source/CustomizableObject/Private/MuCO/CustomizableObjectMeshUpdate.h
--- a/Engine/Plugins/Experimental/Mutable/Source/CustomizableObject/Private/MuCO/CustomizableObjectMeshUpdate.h
+++ b/Engine/Plugins/Experimental/Mutable/Source/CustomizableObject/Private/MuCO/CustomizableObjectMeshUpdate.h
@@ -53,6 +53,8 @@ class FCustomizableObjectMeshStreamIn : public FSkeletalMeshStreamIn
 public:
 	FCustomizableObjectMeshStreamIn(const UCustomizableObjectSkeletalMesh* InMesh, bool bHighPrio, EThreadType CreateResourcesThread);
 
+	void OnUpdateMeshFinished();
+
 private:
 
 	void DoInitiate(const FContext& Context);
