[mutable] Now users can set which row will be the default value of a table node parameter.

#rnx

diff --git a/Engine/Plugins/Experimental/Mutable/Source/CustomizableObjectEditor/Classes/MuCOE/Nodes/CustomizableObjectNodeTable.h b/Engine/Plugins/Experimental/Mutable/Source/CustomizableObjectEditor/Classes/MuCOE/Nodes/CustomizableObjectNodeTable.h
--- a/Engine/Plugins/Experimental/Mutable/Source/CustomizableObjectEditor/Classes/MuCOE/Nodes/CustomizableObjectNodeTable.h
+++ b/Engine/Plugins/Experimental/Mutable/Source/CustomizableObjectEditor/Classes/MuCOE/Nodes/CustomizableObjectNodeTable.h
@@ -220,6 +220,10 @@ public:
 	UPROPERTY(EditAnywhere, Category = TableProperties)
 	FName VersionColumn;
 
+	/** Name of the row that will be used as default value. */
+	UPROPERTY(EditAnywhere, Category = TableProperties)
+	FName DefaultRowName;
+
 	UPROPERTY(EditAnywhere, Category = UI, meta = (DisplayName = "Parameter UI Metadata"))
 	FMutableParamUIMetadata ParamUIMetadata;
 
diff --git a/Engine/Plugins/Experimental/Mutable/Source/CustomizableObjectEditor/Private/MuCOE/GenerateMutableSource/GenerateMutableSourceColor.cpp b/Engine/Plugins/Experimental/Mutable/Source/CustomizableObjectEditor/Private/MuCOE/GenerateMutableSource/GenerateMutableSourceColor.cpp
--- a/Engine/Plugins/Experimental/Mutable/Source/CustomizableObjectEditor/Private/MuCOE/GenerateMutableSource/GenerateMutableSourceColor.cpp
+++ b/Engine/Plugins/Experimental/Mutable/Source/CustomizableObjectEditor/Private/MuCOE/GenerateMutableSource/GenerateMutableSourceColor.cpp
@@ -325,6 +325,7 @@ mu::NodeColourPtr GenerateMutableSourceColor(const UEdGraphPin* Pin, FMutableGra
 						ColorTableNode->SetColumn(ColumnName);
 						ColorTableNode->SetParameterName(TypedNodeTable->ParameterName);
 						ColorTableNode->SetNoneOption(TypedNodeTable->bAddNoneOption);
+						ColorTableNode->SetDefaultRowName(TypedNodeTable->DefaultRowName.ToString());
 					}
 				}
 				else
diff --git a/Engine/Plugins/Experimental/Mutable/Source/CustomizableObjectEditor/Private/MuCOE/GenerateMutableSource/GenerateMutableSourceFloat.cpp b/Engine/Plugins/Experimental/Mutable/Source/CustomizableObjectEditor/Private/MuCOE/GenerateMutableSource/GenerateMutableSourceFloat.cpp
--- a/Engine/Plugins/Experimental/Mutable/Source/CustomizableObjectEditor/Private/MuCOE/GenerateMutableSource/GenerateMutableSourceFloat.cpp
+++ b/Engine/Plugins/Experimental/Mutable/Source/CustomizableObjectEditor/Private/MuCOE/GenerateMutableSource/GenerateMutableSourceFloat.cpp
@@ -365,6 +365,7 @@ mu::NodeScalarPtr GenerateMutableSourceFloat(const UEdGraphPin* Pin, FMutableGra
 						ScalarTableNode->SetColumn(ColumnName);
 						ScalarTableNode->SetParameterName(TypedNodeTable->ParameterName);
 						ScalarTableNode->SetNoneOption(TypedNodeTable->bAddNoneOption);
+						ScalarTableNode->SetDefaultRowName(TypedNodeTable->DefaultRowName.ToString());
 					}
 				}
 				else
diff --git a/Engine/Plugins/Experimental/Mutable/Source/CustomizableObjectEditor/Private/MuCOE/GenerateMutableSource/GenerateMutableSourceImage.cpp b/Engine/Plugins/Experimental/Mutable/Source/CustomizableObjectEditor/Private/MuCOE/GenerateMutableSource/GenerateMutableSourceImage.cpp
--- a/Engine/Plugins/Experimental/Mutable/Source/CustomizableObjectEditor/Private/MuCOE/GenerateMutableSource/GenerateMutableSourceImage.cpp
+++ b/Engine/Plugins/Experimental/Mutable/Source/CustomizableObjectEditor/Private/MuCOE/GenerateMutableSource/GenerateMutableSourceImage.cpp
@@ -1091,6 +1091,7 @@ mu::NodeImagePtr GenerateMutableSourceImage(const UEdGraphPin* Pin, FMutableGrap
 							ImageTableNode->SetColumn(ColumnName);
 							ImageTableNode->SetParameterName(TypedNodeTable->ParameterName);
 							ImageTableNode->SetNoneOption(TypedNodeTable->bAddNoneOption);
+							ImageTableNode->SetDefaultRowName(TypedNodeTable->DefaultRowName.ToString());
 
 							if (UTexture2D* DefaultTexture2D = Cast<UTexture2D>(DefaultTexture))
 							{
diff --git a/Engine/Plugins/Experimental/Mutable/Source/CustomizableObjectEditor/Private/MuCOE/GenerateMutableSource/GenerateMutableSourceMesh.cpp b/Engine/Plugins/Experimental/Mutable/Source/CustomizableObjectEditor/Private/MuCOE/GenerateMutableSource/GenerateMutableSourceMesh.cpp
--- a/Engine/Plugins/Experimental/Mutable/Source/CustomizableObjectEditor/Private/MuCOE/GenerateMutableSource/GenerateMutableSourceMesh.cpp
+++ b/Engine/Plugins/Experimental/Mutable/Source/CustomizableObjectEditor/Private/MuCOE/GenerateMutableSource/GenerateMutableSourceMesh.cpp
@@ -3735,6 +3735,7 @@ mu::NodeMeshPtr GenerateMutableSourceMesh(const UEdGraphPin* Pin,
 						MeshTableNode->SetColumn(MutableColumnName);
 						MeshTableNode->SetParameterName(TypedNodeTable->ParameterName);
 						MeshTableNode->SetNoneOption(TypedNodeTable->bAddNoneOption);
+						MeshTableNode->SetDefaultRowName(TypedNodeTable->DefaultRowName.ToString());
 
 						if (DefaultSkeletalMesh)
 						{
diff --git a/Engine/Plugins/Experimental/Mutable/Source/CustomizableObjectEditor/Private/MuCOE/Nodes/CustomizableObjectNodeTableDetails.cpp b/Engine/Plugins/Experimental/Mutable/Source/CustomizableObjectEditor/Private/MuCOE/Nodes/CustomizableObjectNodeTableDetails.cpp
--- a/Engine/Plugins/Experimental/Mutable/Source/CustomizableObjectEditor/Private/MuCOE/Nodes/CustomizableObjectNodeTableDetails.cpp
+++ b/Engine/Plugins/Experimental/Mutable/Source/CustomizableObjectEditor/Private/MuCOE/Nodes/CustomizableObjectNodeTableDetails.cpp
@@ -48,7 +48,7 @@ void FCustomizableObjectNodeTableDetails::CustomizeDetails(const TSharedPtr<IDet
 	if (Node.IsValid())
 	{
 		IDetailCategoryBuilder& CustomizableObjectCategory = DetailBuilder->EditCategory("TableProperties");
-		//DetailBuilder->HideProperty("VersionColumn");
+		DetailBuilder->HideProperty("VersionColumn");
 		IDetailCategoryBuilder& UICategory = DetailBuilder->EditCategory("UI");
 		DetailBuilder->HideProperty("ParamUIMetadataColumn");
 		IDetailCategoryBuilder& AnimationCategory = DetailBuilder->EditCategory("AnimationProperties");
diff --git a/Engine/Plugins/Experimental/Mutable/Source/MutableTools/Internal/MuT/NodeMeshTablePrivate.h b/Engine/Plugins/Experimental/Mutable/Source/MutableTools/Internal/MuT/NodeMeshTablePrivate.h
--- a/Engine/Plugins/Experimental/Mutable/Source/MutableTools/Internal/MuT/NodeMeshTablePrivate.h
+++ b/Engine/Plugins/Experimental/Mutable/Source/MutableTools/Internal/MuT/NodeMeshTablePrivate.h
@@ -23,13 +23,14 @@ namespace mu
 		TablePtr Table;
 		FString ColumnName;
 		bool bNoneOption = false;
+		FString DefaultRowName;
 
 		TArray<NodeLayoutPtr> Layouts;
 
 		//!
 		void Serialise( OutputArchive& arch ) const
 		{
-            uint32 ver = 3;
+            uint32 ver = 4;
 			arch << ver;
 
 			arch << ParameterName;
@@ -37,6 +38,7 @@ namespace mu
 			arch << ColumnName;
 			arch << Layouts;
 			arch << bNoneOption;
+			arch << DefaultRowName;
 		}
 
 		//!
@@ -44,7 +46,7 @@ namespace mu
 		{
             uint32 ver;
 			arch >> ver;
-            check(ver>=1 && ver<=3);
+            check(ver>=1 && ver<=4);
 
 			if (ver == 1)
 			{
@@ -80,6 +82,11 @@ namespace mu
 			{
 				bNoneOption = Table->GetPrivate()->bNoneOption_DEPRECATED;
 			}
+
+			if (ver >= 4)
+			{
+				arch >> DefaultRowName;
+			}
 		}
 
 		// NodeMesh::Private interface
diff --git a/Engine/Plugins/Experimental/Mutable/Source/MutableTools/Private/MuT/CodeGenerator.cpp b/Engine/Plugins/Experimental/Mutable/Source/MutableTools/Private/MuT/CodeGenerator.cpp
--- a/Engine/Plugins/Experimental/Mutable/Source/MutableTools/Private/MuT/CodeGenerator.cpp
+++ b/Engine/Plugins/Experimental/Mutable/Source/MutableTools/Private/MuT/CodeGenerator.cpp
@@ -309,17 +309,19 @@ namespace mu
 	}
 
 
-	Ptr<ASTOp> CodeGenerator::GenerateTableVariable(Ptr<const Node> InNode, const FTableCacheKey& CacheKey, bool bAddNoneOption)
+	Ptr<ASTOp> CodeGenerator::GenerateTableVariable(Ptr<const Node> InNode, const FTableCacheKey& CacheKey, bool bAddNoneOption, const FString& DefaultRowName)
 	{
 		Ptr<ASTOp> result;
 
         FParameterDesc param;
         param.m_name = CacheKey.ParameterName;
-        if ( param.m_name.Len()==0 )
+
+        if (param.m_name.Len() == 0)
         {
-            param.m_name = TCHAR_TO_ANSI(*CacheKey.Table->GetName());
+            param.m_name = CacheKey.Table->GetName();
         }
-        param.m_type = PARAMETER_TYPE::T_INT;
+        
+		param.m_type = PARAMETER_TYPE::T_INT;
         param.m_defaultValue.Set<ParamIntType>(0);
 
 		// Add the possible values
@@ -354,16 +356,22 @@ namespace mu
 
 				if (nameCol > -1)
 				{
-					value.m_name = TCHAR_TO_ANSI(*CacheKey.Table->GetPrivate()->Rows[i].Values[nameCol].String);
+					value.m_name = CacheKey.Table->GetPrivate()->Rows[i].Values[nameCol].String;
 				}
 
 				param.m_possibleValues.Add(value);
 
-                // The first row is the default one
-                if (i==0)
-                {
-                    param.m_defaultValue.Set<ParamIntType>(value.m_value);
-                }
+				// Set the first row as the default one (if there is none option)
+				if (i == 0 && !bAddNoneOption)
+				{
+					param.m_defaultValue.Set<ParamIntType>(value.m_value);
+				}
+				
+				// Set the selected row as default (if exists)
+				if (value.m_name == DefaultRowName)
+				{
+					param.m_defaultValue.Set<ParamIntType>(value.m_value);
+				}
             }
         }
 
diff --git a/Engine/Plugins/Experimental/Mutable/Source/MutableTools/Private/MuT/CodeGenerator.h b/Engine/Plugins/Experimental/Mutable/Source/MutableTools/Private/MuT/CodeGenerator.h
--- a/Engine/Plugins/Experimental/Mutable/Source/MutableTools/Private/MuT/CodeGenerator.h
+++ b/Engine/Plugins/Experimental/Mutable/Source/MutableTools/Private/MuT/CodeGenerator.h
@@ -315,7 +315,7 @@ namespace mu
 
         //-----------------------------------------------------------------------------------------
         //!
-        Ptr<ASTOp> GenerateTableVariable(Ptr<const Node>, const FTableCacheKey&, bool bAddNoneOption);
+        Ptr<ASTOp> GenerateTableVariable(Ptr<const Node>, const FTableCacheKey&, bool bAddNoneOption, const FString& DefaultRowName);
 
         //!
         Ptr<ASTOp> GenerateMissingBoolCode(const TCHAR* strWhere, bool value, const void* errorContext );
@@ -743,7 +743,7 @@ namespace mu
         if ( !Variable)
         {
             // Create the table variable expression
-            Variable = GenerateTableVariable( node.m_pNode, CacheKey, node.bNoneOption);
+            Variable = GenerateTableVariable( node.m_pNode, CacheKey, node.bNoneOption, node.DefaultRowName);
 
             GeneratedTables.Add(CacheKey, Variable );
         }
diff --git a/Engine/Plugins/Experimental/Mutable/Source/MutableTools/Private/MuT/NodeColourTable.cpp b/Engine/Plugins/Experimental/Mutable/Source/MutableTools/Private/MuT/NodeColourTable.cpp
--- a/Engine/Plugins/Experimental/Mutable/Source/MutableTools/Private/MuT/NodeColourTable.cpp
+++ b/Engine/Plugins/Experimental/Mutable/Source/MutableTools/Private/MuT/NodeColourTable.cpp
@@ -64,6 +64,12 @@ namespace mu
 		m_pD->bNoneOption = bAddNoneOption;
 	}
 
+	
+	//---------------------------------------------------------------------------------------------
+	void NodeColourTable::SetDefaultRowName(const FString RowName)
+	{
+		m_pD->DefaultRowName = RowName;
+	}
 
 }
 
diff --git a/Engine/Plugins/Experimental/Mutable/Source/MutableTools/Private/MuT/NodeColourTablePrivate.h b/Engine/Plugins/Experimental/Mutable/Source/MutableTools/Private/MuT/NodeColourTablePrivate.h
--- a/Engine/Plugins/Experimental/Mutable/Source/MutableTools/Private/MuT/NodeColourTablePrivate.h
+++ b/Engine/Plugins/Experimental/Mutable/Source/MutableTools/Private/MuT/NodeColourTablePrivate.h
@@ -23,17 +23,19 @@ namespace mu
 		TablePtr Table;
 		FString ColumnName;
 		bool bNoneOption = false;
+		FString DefaultRowName;
 
 		//!
 		void Serialise( OutputArchive& arch ) const
 		{
-            uint32_t ver = 2;
+            uint32_t ver = 3;
 			arch << ver;
 
 			arch << ParameterName;
 			arch << Table;
 			arch << ColumnName;
 			arch << bNoneOption;
+			arch << DefaultRowName;
 		}
 
 		//!
@@ -41,7 +43,7 @@ namespace mu
 		{
             uint32 ver;
 			arch >> ver;
-			check(ver <= 2);
+			check(ver <= 3);
 
 			if (ver == 0)
 			{
@@ -75,6 +77,11 @@ namespace mu
 			{
 				bNoneOption = Table->GetPrivate()->bNoneOption_DEPRECATED;
 			}
+
+			if (ver >= 3)
+			{
+				arch >> DefaultRowName;
+			}
 		}
 
 	};
diff --git a/Engine/Plugins/Experimental/Mutable/Source/MutableTools/Private/MuT/NodeImageTable.cpp b/Engine/Plugins/Experimental/Mutable/Source/MutableTools/Private/MuT/NodeImageTable.cpp
--- a/Engine/Plugins/Experimental/Mutable/Source/MutableTools/Private/MuT/NodeImageTable.cpp
+++ b/Engine/Plugins/Experimental/Mutable/Source/MutableTools/Private/MuT/NodeImageTable.cpp
@@ -84,6 +84,13 @@ namespace mu
 		m_pD->ReferenceImageDesc = ImageDesc;
 	}
 
+
+	//---------------------------------------------------------------------------------------------
+	void NodeImageTable::SetDefaultRowName(const FString& RowName)
+	{
+		m_pD->DefaultRowName = RowName;
+	}
+
 }
 
 
diff --git a/Engine/Plugins/Experimental/Mutable/Source/MutableTools/Private/MuT/NodeImageTablePrivate.h b/Engine/Plugins/Experimental/Mutable/Source/MutableTools/Private/MuT/NodeImageTablePrivate.h
--- a/Engine/Plugins/Experimental/Mutable/Source/MutableTools/Private/MuT/NodeImageTablePrivate.h
+++ b/Engine/Plugins/Experimental/Mutable/Source/MutableTools/Private/MuT/NodeImageTablePrivate.h
@@ -24,11 +24,12 @@ namespace mu
 		uint16 MaxTextureSize = 0;
 		FImageDesc ReferenceImageDesc;
 		bool bNoneOption = false;
+		FString DefaultRowName;
 
 		//!
 		void Serialise( OutputArchive& arch ) const
 		{
-            uint32_t ver = 5;
+            uint32_t ver = 6;
 			arch << ver;
 
 			arch << ParameterName;
@@ -44,6 +45,7 @@ namespace mu
 			}
 
 			arch << bNoneOption;
+			arch << DefaultRowName;
 		}
 
 		//!
@@ -51,7 +53,7 @@ namespace mu
 		{
             uint32_t ver;
 			arch >> ver;
-			check(ver>=1 && ver<= 5);
+			check(ver>=1 && ver<= 6);
 
 			if (ver == 1)
 			{
@@ -101,6 +103,11 @@ namespace mu
 			{
 				arch >> bNoneOption;
 			}
+
+			if (ver >= 6)
+			{
+				arch >> DefaultRowName;
+			}
 		}
 	};
 
diff --git a/Engine/Plugins/Experimental/Mutable/Source/MutableTools/Private/MuT/NodeMeshTable.cpp b/Engine/Plugins/Experimental/Mutable/Source/MutableTools/Private/MuT/NodeMeshTable.cpp
--- a/Engine/Plugins/Experimental/Mutable/Source/MutableTools/Private/MuT/NodeMeshTable.cpp
+++ b/Engine/Plugins/Experimental/Mutable/Source/MutableTools/Private/MuT/NodeMeshTable.cpp
@@ -108,6 +108,13 @@ namespace mu
 		m_pD->bNoneOption = bAddNoneOption;
 	}
 
+
+	//---------------------------------------------------------------------------------------------
+	void NodeMeshTable::SetDefaultRowName(const FString& RowName)
+	{
+		m_pD->DefaultRowName = RowName;
+	}
+
 }
 
 
diff --git a/Engine/Plugins/Experimental/Mutable/Source/MutableTools/Private/MuT/NodeScalarTable.cpp b/Engine/Plugins/Experimental/Mutable/Source/MutableTools/Private/MuT/NodeScalarTable.cpp
--- a/Engine/Plugins/Experimental/Mutable/Source/MutableTools/Private/MuT/NodeScalarTable.cpp
+++ b/Engine/Plugins/Experimental/Mutable/Source/MutableTools/Private/MuT/NodeScalarTable.cpp
@@ -63,6 +63,11 @@ namespace mu
 		m_pD->bNoneOption = bAddNoneOption;
 	}
 
+	//---------------------------------------------------------------------------------------------
+	void NodeScalarTable::SetDefaultRowName(const FString& RowName)
+	{
+		m_pD->DefaultRowName = RowName;
+	}
 }
 
 
diff --git a/Engine/Plugins/Experimental/Mutable/Source/MutableTools/Private/MuT/NodeScalarTablePrivate.h b/Engine/Plugins/Experimental/Mutable/Source/MutableTools/Private/MuT/NodeScalarTablePrivate.h
--- a/Engine/Plugins/Experimental/Mutable/Source/MutableTools/Private/MuT/NodeScalarTablePrivate.h
+++ b/Engine/Plugins/Experimental/Mutable/Source/MutableTools/Private/MuT/NodeScalarTablePrivate.h
@@ -22,17 +22,19 @@ namespace mu
 		TablePtr Table;
 		FString ColumnName;
 		bool bNoneOption = false;
+		FString DefaultRowName;
 
 		//!
 		void Serialise( OutputArchive& arch ) const
 		{
-            uint32 ver = 2;
+            uint32 ver = 3;
 			arch << ver;
 
 			arch << ParameterName;
 			arch << Table;
 			arch << ColumnName;
 			arch << bNoneOption;
+			arch << DefaultRowName;
 		}
 
 		//!
@@ -40,7 +42,7 @@ namespace mu
 		{
             uint32 ver;
 			arch >> ver;
-            check(ver <= 2);
+            check(ver <= 3);
 
 			if (ver==0)
 			{
@@ -74,6 +76,11 @@ namespace mu
 			{
 				bNoneOption = Table->GetPrivate()->bNoneOption_DEPRECATED;
 			}
+
+			if (ver >= 3)
+			{
+				arch >> DefaultRowName;
+			}
 		}
 
 	};
diff --git a/Engine/Plugins/Experimental/Mutable/Source/MutableTools/Public/MuT/NodeColourTable.h b/Engine/Plugins/Experimental/Mutable/Source/MutableTools/Public/MuT/NodeColourTable.h
--- a/Engine/Plugins/Experimental/Mutable/Source/MutableTools/Public/MuT/NodeColourTable.h
+++ b/Engine/Plugins/Experimental/Mutable/Source/MutableTools/Public/MuT/NodeColourTable.h
@@ -61,6 +61,9 @@ namespace mu
 		//!
 		void SetNoneOption(bool bAddNoneOption);
 
+		//! Set the row name to be used as default value
+		void SetDefaultRowName(const FString RowName);
+
 		//-----------------------------------------------------------------------------------------
 		// Interface pattern
 		//-----------------------------------------------------------------------------------------
diff --git a/Engine/Plugins/Experimental/Mutable/Source/MutableTools/Public/MuT/NodeImageTable.h b/Engine/Plugins/Experimental/Mutable/Source/MutableTools/Public/MuT/NodeImageTable.h
--- a/Engine/Plugins/Experimental/Mutable/Source/MutableTools/Public/MuT/NodeImageTable.h
+++ b/Engine/Plugins/Experimental/Mutable/Source/MutableTools/Public/MuT/NodeImageTable.h
@@ -73,6 +73,9 @@ namespace mu
 		//!
 		void SetReferenceImageDescriptor(const FImageDesc& Descriptor);
 
+		//! Set the row name to be used as default value
+		void SetDefaultRowName(const FString& RowName);
+
 		//-----------------------------------------------------------------------------------------
 		// Interface pattern
 		//-----------------------------------------------------------------------------------------
diff --git a/Engine/Plugins/Experimental/Mutable/Source/MutableTools/Public/MuT/NodeMeshTable.h b/Engine/Plugins/Experimental/Mutable/Source/MutableTools/Public/MuT/NodeMeshTable.h
--- a/Engine/Plugins/Experimental/Mutable/Source/MutableTools/Public/MuT/NodeMeshTable.h
+++ b/Engine/Plugins/Experimental/Mutable/Source/MutableTools/Public/MuT/NodeMeshTable.h
@@ -74,6 +74,9 @@ namespace mu
 		//! Adds the "None" option to the parameter that represents this table column
 		void SetNoneOption(bool bAddOption);
 
+		//! Set the row name to be used as default value
+		void SetDefaultRowName(const FString& RowName);
+
 		//-----------------------------------------------------------------------------------------
 		// Interface pattern
 		//-----------------------------------------------------------------------------------------
diff --git a/Engine/Plugins/Experimental/Mutable/Source/MutableTools/Public/MuT/NodeScalarTable.h b/Engine/Plugins/Experimental/Mutable/Source/MutableTools/Public/MuT/NodeScalarTable.h
--- a/Engine/Plugins/Experimental/Mutable/Source/MutableTools/Public/MuT/NodeScalarTable.h
+++ b/Engine/Plugins/Experimental/Mutable/Source/MutableTools/Public/MuT/NodeScalarTable.h
@@ -62,6 +62,9 @@ namespace mu
 		//! Adds the "None" option to the parameter that represents this table column
 		void SetNoneOption(bool bAddOption);
 
+		//! Set the row name to be used as default value
+		void SetDefaultRowName(const FString& RowName);
+
 		//-----------------------------------------------------------------------------------------
 		// Interface pattern
 		//-----------------------------------------------------------------------------------------
