[mutable] Replaced SMutableTextSearchBox with the generic SSearchableComboBox.
#rnx
#tests Editor

diff --git a/Engine/Plugins/Experimental/Mutable/Source/CustomizableObjectEditor/Private/MuCOE/CustomizableInstanceDetails.cpp b/Engine/Plugins/Experimental/Mutable/Source/CustomizableObjectEditor/Private/MuCOE/CustomizableInstanceDetails.cpp
--- a/Engine/Plugins/Experimental/Mutable/Source/CustomizableObjectEditor/Private/MuCOE/CustomizableInstanceDetails.cpp
+++ b/Engine/Plugins/Experimental/Mutable/Source/CustomizableObjectEditor/Private/MuCOE/CustomizableInstanceDetails.cpp
@@ -19,7 +19,6 @@
 #include "MuCO/CustomizableObjectInstance.h"
 #include "MuCO/CustomizableObjectSystem.h"
 #include "MuCOE/SCustomizableInstanceProperties.h"
-#include "MuCOE/SMutableTextSearchBox.h"
 #include "MuCOE/UnrealEditorPortabilityHelpers.h"
 
 #include "SSearchableComboBox.h"
@@ -870,7 +869,7 @@ TSharedRef<SWidget> FCustomizableInstanceDetails::GenerateIntWidget(const int32
 {
 	UCustomizableObject* CustomizableObject = CustomInstance->GetCustomizableObject();
 
-	int numValues = CustomizableObject->GetIntParameterNumOptions(ParamIndexInObject);
+	int32 numValues = CustomizableObject->GetIntParameterNumOptions(ParamIndexInObject);
 	bool bIsParamMultidimensional = CustomInstance->GetCustomizableObject()->IsParameterMultidimensional((ParamIndexInObject));
 
 	if (!bIsParamMultidimensional && numValues)
@@ -878,10 +877,17 @@ TSharedRef<SWidget> FCustomizableInstanceDetails::GenerateIntWidget(const int32
 		FString ToolTipText = FString("None");
 		FString ParamName = CustomizableObject->GetParameterName(ParamIndexInObject);
 		FString Value = CustomInstance->GetIntParameterSelectedOption(ParamName, -1);
-		TArray<FString> OptionNamesAttribute;
-		int ValueIndex = 0;
 
-		for (int i = 0; i < numValues; ++i)
+		TSharedPtr<TArray<TSharedPtr<FString>>>* FoundOptions = IntParameterOptions.Find(ParamIndexInObject);
+		TArray<TSharedPtr<FString>>& OptionNamesAttribute = FoundOptions && FoundOptions->IsValid() ? 
+															*FoundOptions->Get() :
+															*(IntParameterOptions.Add(ParamIndexInObject, MakeShared<TArray<TSharedPtr<FString>>>()).Get());
+
+		OptionNamesAttribute.Empty();
+
+		int32 ValueIndex = 0;
+
+		for (int32 i = 0; i < numValues; ++i)
 		{
 			FString PossibleValue = CustomizableObject->GetIntParameterAvailableOption(ParamIndexInObject, i);
 
@@ -899,23 +905,33 @@ TSharedRef<SWidget> FCustomizableInstanceDetails::GenerateIntWidget(const int32
 				}
 			}
 
-			OptionNamesAttribute.Add(CustomizableObject->GetIntParameterAvailableOption(ParamIndexInObject, i));
+			OptionNamesAttribute.Add(TSharedPtr<FString>(new FString(CustomizableObject->GetIntParameterAvailableOption(ParamIndexInObject, i))));
 		}
 
-		//TODO(Max): UE-207137
-		return SNew(SMutableTextSearchBox)
+		return SNew(SSearchableComboBox)
 			.ToolTipText(FText::FromString(ToolTipText))
-			.PossibleSuggestions(OptionNamesAttribute)
-			.InitialText(FText::FromString(OptionNamesAttribute[ValueIndex]))
-			.MustMatchPossibleSuggestions(TAttribute<bool>(true))
-			.SuggestionListPlacement(EMenuPlacement::MenuPlacement_ComboBox)
-			.OnTextCommitted(this, &FCustomizableInstanceDetails::OnIntParameterComboBoxChanged, ParamName);
+			.OptionsSource(&OptionNamesAttribute)
+			.InitiallySelectedItem(OptionNamesAttribute[ValueIndex])
+			.Method(EPopupMethod::UseCurrentWindow)
+			.OnSelectionChanged(this, &FCustomizableInstanceDetails::OnIntParameterComboBoxChanged, ParamName)
+			.OnGenerateWidget(this, &FCustomizableInstanceDetails::OnGenerateWidgetIntParameter)
+			.Content()
+			[
+				SNew(STextBlock)
+					.Text(FText::FromString(*OptionNamesAttribute[ValueIndex]))
+			];
 	}
 
 	return SNew(STextBlock).Text(LOCTEXT("MultidimensionalINTParameter_Text", "Multidimensional INT Parameter not supported"));
 }
 
 
+TSharedRef<SWidget> FCustomizableInstanceDetails::OnGenerateWidgetIntParameter(TSharedPtr<FString> InItem) const
+{
+	return SNew(STextBlock).Text(FText::FromString(*InItem.Get()));
+}
+
+
 void FCustomizableInstanceDetails::OnIntParameterComboBoxChanged(TSharedPtr<FString> Selection, ESelectInfo::Type SelectInfo, FString ParamName)
 {
 	TArray<FCustomizableObjectIntParameterValue>& IntParameters = CustomInstance->GetIntParameters();
@@ -1399,7 +1415,13 @@ TSharedRef<SWidget> FCustomizableInstanceDetails::GenerateMultidimensionalProjec
 	{
 		const int32 NumPoseValues = CustomizableObject->GetIntParameterNumOptions(PoseSwitchEnumParamIndexInObject);
 
-		TArray<FString> PoseOptionNamesAttribute;
+		TSharedPtr<TArray<TSharedPtr<FString>>>* FoundOptions = ProjectorParameterPoseOptions.Find(PoseSwitchEnumParamIndexInObject);
+		TArray<TSharedPtr<FString>>& PoseOptionNamesAttribute = FoundOptions && FoundOptions->IsValid() ?
+																*FoundOptions->Get()
+																: *(ProjectorParameterPoseOptions.Add(PoseSwitchEnumParamIndexInObject, MakeShared<TArray<TSharedPtr<FString>>>()).Get());
+
+		PoseOptionNamesAttribute.Empty();
+
 		FString PoseValue = CustomInstance->GetIntParameterSelectedOption(PoseSwitchEnumParamName, -1);
 		int32 PoseValueIndex = 0;
 
@@ -1411,7 +1433,7 @@ TSharedRef<SWidget> FCustomizableInstanceDetails::GenerateMultidimensionalProjec
 				PoseValueIndex = j;
 			}
 
-			PoseOptionNamesAttribute.Add(PossibleValue);
+			PoseOptionNamesAttribute.Add(MakeShared<FString>(PossibleValue));
 		}
 
 		ProjectorBox->AddSlot()
@@ -1425,14 +1447,18 @@ TSharedRef<SWidget> FCustomizableInstanceDetails::GenerateMultidimensionalProjec
 			.VAlign(VAlign_Center)
 			.FillWidth(0.45f)
 			[
-				//TODO(Max): UE-207137
-				SNew(SMutableTextSearchBox)
-				.PossibleSuggestions(PoseOptionNamesAttribute)
-				.InitialText(FText::FromString(PoseOptionNamesAttribute[PoseValueIndex]))
-				.MustMatchPossibleSuggestions(TAttribute<bool>(true))
-				.SuggestionListPlacement(EMenuPlacement::MenuPlacement_ComboBox)
-				.OnTextCommitted(this, &FCustomizableInstanceDetails::OnProjectorTextureParameterComboBoxChanged, PoseSwitchEnumParamName, -1)
+				SNew(SSearchableComboBox)
 				.ToolTipText(LOCTEXT("Pose selector tooltip", "Select the skeletal mesh pose used for projection. This does not control the actual visual mesh pose in the viewport (or during gameplay for that matter). It has to be manually set. You can drag&drop a pose onto the preview viewport."))
+				.OptionsSource(&PoseOptionNamesAttribute)
+				.InitiallySelectedItem(PoseOptionNamesAttribute[PoseValueIndex])
+				.Method(EPopupMethod::UseCurrentWindow)
+				.OnSelectionChanged(this, &FCustomizableInstanceDetails::OnProjectorTextureParameterComboBoxChanged, PoseSwitchEnumParamName, -1)
+				.OnGenerateWidget(this, &FCustomizableInstanceDetails::OnGenerateWidgetProjectorParameter)
+				.Content()
+				[
+					SNew(STextBlock)
+						.Text(FText::FromString(*PoseOptionNamesAttribute[PoseValueIndex]))
+				]
 			]
 
 			+ SHorizontalBox::Slot()
@@ -1587,7 +1613,7 @@ TSharedRef<SWidget> FCustomizableInstanceDetails::GenerateMultidimensionalProjec
 		// Avoid filling this arraw with repeated array  options
 		if (RangeIndex == 0)
 		{
-			ProjectorTextureOptions.Add(OptionNamesAttribute);
+			ProjectorTextureOptions.Add(MakeShared<TArray<TSharedPtr<FString>>>(OptionNamesAttribute));
 		}
 
 		if (NumValues > 1)
@@ -1599,14 +1625,14 @@ TSharedRef<SWidget> FCustomizableInstanceDetails::GenerateMultidimensionalProjec
 				SNew(SBox)
 				[
 					SNew(SSearchableComboBox)
-					.OptionsSource(&ProjectorTextureOptions.Last())
-					.InitiallySelectedItem(ProjectorTextureOptions.Last()[ValueIndex])
+					.OptionsSource(ProjectorTextureOptions.Last().Get())
+					.InitiallySelectedItem((*ProjectorTextureOptions.Last())[ValueIndex])
 					.OnGenerateWidget(this, &FCustomizableInstanceDetails::MakeTextureComboEntryWidget)
 					.OnSelectionChanged(this, &FCustomizableInstanceDetails::OnProjectorTextureParameterComboBoxChanged, TextureSwitchEnumParamName, RangeIndex)
 					.Content()
 					[
 						SNew(STextBlock)
-						.Text(FText::FromString(*ProjectorTextureOptions.Last()[ValueIndex]))
+						.Text(FText::FromString(*(*ProjectorTextureOptions.Last())[ValueIndex]))
 					]
 				]
 			];
@@ -1691,6 +1717,12 @@ TSharedRef<SWidget> FCustomizableInstanceDetails::GenerateMultidimensionalProjec
 }
 
 
+TSharedRef<SWidget> FCustomizableInstanceDetails::OnGenerateWidgetProjectorParameter(TSharedPtr<FString> InItem) const
+{
+	return SNew(STextBlock).Text(FText::FromString(*InItem.Get()));
+}
+
+
 FReply FCustomizableInstanceDetails::OnProjectorLayerAdded(FString ParamName) const
 {
 	const int32 NumLayers = CustomInstance->AddValueToProjectorRange(ParamName) + 1;
diff --git a/Engine/Plugins/Experimental/Mutable/Source/CustomizableObjectEditor/Private/MuCOE/CustomizableInstanceDetails.h b/Engine/Plugins/Experimental/Mutable/Source/CustomizableObjectEditor/Private/MuCOE/CustomizableInstanceDetails.h
--- a/Engine/Plugins/Experimental/Mutable/Source/CustomizableObjectEditor/Private/MuCOE/CustomizableInstanceDetails.h
+++ b/Engine/Plugins/Experimental/Mutable/Source/CustomizableObjectEditor/Private/MuCOE/CustomizableInstanceDetails.h
@@ -76,6 +76,7 @@ private:
 	// Int Parameters Functions
 	TSharedRef<SWidget> GenerateIntWidget(const int32 ParamIndexInObject);
 	void OnIntParameterComboBoxChanged(TSharedPtr<FString> Selection, ESelectInfo::Type SelectInfo, const FString ParamName);
+	TSharedRef<SWidget> OnGenerateWidgetIntParameter(TSharedPtr<FString> InItem) const;
 	
 	// Float Parameters Functions
 	TSharedRef<SWidget> GenerateFloatWidget(const int32 ParamIndexInObject);
@@ -112,6 +113,7 @@ private:
 	FReply OnProjectorLayerRemoved(const FString ParamName, const int32 RangeIndex) const;
 	void OnProjectorTextureParameterComboBoxChanged(TSharedPtr<FString> Selection, ESelectInfo::Type SelectInfo, const FString ParamName, int32 RangeIndex) const;
 	TSharedRef<SWidget> MakeTextureComboEntryWidget(TSharedPtr<FString> InItem) const;
+	TSharedRef<SWidget> OnGenerateWidgetProjectorParameter(TSharedPtr<FString> InItem) const;
 
 	// Parameter Functions
 	FReply OnCopyAllParameters();
@@ -155,10 +157,15 @@ private:
 	TMap<int32, bool> ParamHasParent;
 
 	/** Array with all the possible multilayer projector texture options */
-	TArray<TArray<TSharedPtr<FString>>> ProjectorTextureOptions;
+	TArray<TSharedPtr<TArray<TSharedPtr<FString>>>> ProjectorTextureOptions;
+
+	/** Map from ParamIndexInObject to the param's int selector options */
+	TMap<int32, TSharedPtr<TArray<TSharedPtr<FString>>>> IntParameterOptions;
+
+	/** Map from ParamIndexInObject to the projector param pose options  */
+	TMap<int32, TSharedPtr<TArray<TSharedPtr<FString>>>> ProjectorParameterPoseOptions;
 
 	/** True when a slider is being edited*/
 	bool bUpdatingSlider = false;
-
 };
 
diff --git a/Engine/Plugins/Experimental/Mutable/Source/CustomizableObjectEditor/Private/MuCOE/CustomizableObjectEditor.cpp b/Engine/Plugins/Experimental/Mutable/Source/CustomizableObjectEditor/Private/MuCOE/CustomizableObjectEditor.cpp
--- a/Engine/Plugins/Experimental/Mutable/Source/CustomizableObjectEditor/Private/MuCOE/CustomizableObjectEditor.cpp
+++ b/Engine/Plugins/Experimental/Mutable/Source/CustomizableObjectEditor/Private/MuCOE/CustomizableObjectEditor.cpp
@@ -51,7 +51,6 @@
 #include "MuCOE/SCustomizableObjectEditorViewport.h"
 #include "PropertyEditorModule.h"
 #include "ScopedTransaction.h"
-#include "MuCOE/SCustomizableInstanceProperties.h"
 #include "MuCO/CustomizableObjectInstancePrivate.h"
 #include "MuCO/CustomizableObjectPrivate.h"
 #include "UObject/EnumProperty.h"
diff --git a/Engine/Plugins/Experimental/Mutable/Source/CustomizableObjectEditor/Private/MuCOE/SCustomizableInstanceProperties.cpp b/Engine/Plugins/Experimental/Mutable/Source/CustomizableObjectEditor/Private/MuCOE/SCustomizableInstanceProperties.cpp
--- a/Engine/Plugins/Experimental/Mutable/Source/CustomizableObjectEditor/Private/MuCOE/SCustomizableInstanceProperties.cpp
+++ b/Engine/Plugins/Experimental/Mutable/Source/CustomizableObjectEditor/Private/MuCOE/SCustomizableInstanceProperties.cpp
@@ -14,7 +14,7 @@
 #include "MuCO/CustomizableObjectSystem.h"
 #include "MuCOE/CustomizableInstanceDetails.h"
 #include "MuCOE/CustomizableObjectEditorUtilities.h"
-#include "MuCOE/SMutableTextSearchBox.h"
+#include "SSearchableComboBox.h"
 #include "MuCOE/UnrealEditorPortabilityHelpers.h"
 #include "MuCOE/SCustomizableObjectEditorViewport.h"
 #include "Serialization/BufferArchive.h"
@@ -748,10 +748,14 @@ void SCustomizableInstanceProperties::AddParameter(int32 ParamIndexInObject)
 		if (!bIsParamMultidimensional && numValues)
 		{
 			FString ToolTipText = FString("None");
-			//TSharedPtr< TArray< TSharedPtr<FString> > > OptionNames = MakeShareable(new TArray< TSharedPtr<FString> >());
-			//IntOptionNames.Add(OptionNames);
-			
-			TArray<FString> OptionNamesAttribute;
+
+			TSharedPtr<TArray<TSharedPtr<FString>>>* FoundOptions = IntParameterOptions.Find(ParamIndexInObject);
+			TArray<TSharedPtr<FString>>& OptionNamesAttribute = FoundOptions && FoundOptions->IsValid() ?
+																*FoundOptions->Get() :
+																*(IntParameterOptions.Add(ParamIndexInObject, MakeShared<TArray<TSharedPtr<FString>>>()).Get());
+
+			OptionNamesAttribute.Empty();
+
 			FString Value = GetIntParameterValue(ParamName);
 			int ValueIndex = 0;
 			
@@ -773,8 +777,7 @@ void SCustomizableInstanceProperties::AddParameter(int32 ParamIndexInObject)
 					}
 				}
 
-				//OptionNames->Add(MakeShareable(new FString(CustomizableObject->GetIntParameterAvailableOption(ParamIndexInObject, i))));
-				OptionNamesAttribute.Add(CustomizableObject->GetIntParameterAvailableOption(ParamIndexInObject, i));
+				OptionNamesAttribute.Add(MakeShared<FString>(CustomizableObject->GetIntParameterAvailableOption(ParamIndexInObject, i)));
 			}
 
 			//OptionNames->Sort(::CompareNames);
@@ -788,13 +791,18 @@ void SCustomizableInstanceProperties::AddParameter(int32 ParamIndexInObject)
 				SNew(SHorizontalBox)
 				+ SHorizontalBox::Slot().FillWidth(0.9f)
 				[
-					SNew(SMutableTextSearchBox)
+					SNew(SSearchableComboBox)
 					.ToolTipText(FText::FromString(ToolTipText))
-					.PossibleSuggestions(OptionNamesAttribute)
-					.InitialText(FText::FromString(OptionNamesAttribute[ValueIndex]))
-					.MustMatchPossibleSuggestions(TAttribute<bool>(true))
-					.SuggestionListPlacement(EMenuPlacement::MenuPlacement_ComboBox)
-					.OnTextCommitted(this, &SCustomizableInstanceProperties::OnIntParameterComboBoxChanged, ParamName)
+					.OptionsSource(&OptionNamesAttribute)
+					.InitiallySelectedItem(OptionNamesAttribute[ValueIndex])
+					.Method(EPopupMethod::UseCurrentWindow)
+					.OnSelectionChanged(this, &SCustomizableInstanceProperties::OnIntParameterComboBoxChanged, ParamName)
+					.OnGenerateWidget(this, &SCustomizableInstanceProperties::OnGenerateWidgetIntParameter)
+					.Content()
+					[
+						SNew(STextBlock)
+							.Text(FText::FromString(*OptionNamesAttribute[ValueIndex]))
+					]
 				]
 
 				+ SHorizontalBox::Slot().AutoWidth().Padding(2.5f, 0.0f, 0.0f, 0.0f)
@@ -988,7 +996,13 @@ void SCustomizableInstanceProperties::AddParameter(int32 ParamIndexInObject)
 			{
 				const int32 NumPoseValues = CustomizableObject->GetIntParameterNumOptions(PoseSwitchEnumParamIndexInObject);
 
-				TArray<FString> PoseOptionNamesAttribute;
+				TSharedPtr<TArray<TSharedPtr<FString>>>* FoundOptions = ProjectorParameterPoseOptions.Find(PoseSwitchEnumParamIndexInObject);
+				TArray<TSharedPtr<FString>>& PoseOptionNamesAttribute = FoundOptions && FoundOptions->IsValid() ?
+					*FoundOptions->Get()
+					: *(ProjectorParameterPoseOptions.Add(PoseSwitchEnumParamIndexInObject, MakeShared<TArray<TSharedPtr<FString>>>()).Get());
+
+				PoseOptionNamesAttribute.Empty();
+
 				FString PoseValue = GetIntParameterValue(PoseSwitchEnumParamName, -1);
 				int32 PoseValueIndex = 0;
 
@@ -1000,7 +1014,7 @@ void SCustomizableInstanceProperties::AddParameter(int32 ParamIndexInObject)
 						PoseValueIndex = j;
 					}
 
-					PoseOptionNamesAttribute.Add(PossibleValue);
+					PoseOptionNamesAttribute.Add(MakeShared<FString>(PossibleValue));
 				}
 
 				ParameterBox->AddSlot()
@@ -1014,13 +1028,18 @@ void SCustomizableInstanceProperties::AddParameter(int32 ParamIndexInObject)
 							.VAlign(VAlign_Center)
 							.FillWidth(0.45f)
 							[
-								SNew(SMutableTextSearchBox)
-								.PossibleSuggestions(PoseOptionNamesAttribute)
-								.InitialText(FText::FromString(PoseOptionNamesAttribute[PoseValueIndex]))
-								.MustMatchPossibleSuggestions(TAttribute<bool>(true))
-								.SuggestionListPlacement(EMenuPlacement::MenuPlacement_ComboBox)
-								.OnTextCommitted(this, &SCustomizableInstanceProperties::OnProjectorTextureParameterComboBoxChanged, PoseSwitchEnumParamName, -1)
+								SNew(SSearchableComboBox)
 								.ToolTipText(LOCTEXT("Pose selector tooltip", "Select the skeletal mesh pose used for projection. This does not control the actual visual mesh pose in the viewport (or during gameplay for that matter). It has to be manually set. You can drag&drop a pose onto the preview viewport."))
+								.OptionsSource(&PoseOptionNamesAttribute)
+								.InitiallySelectedItem(PoseOptionNamesAttribute[PoseValueIndex])
+								.Method(EPopupMethod::UseCurrentWindow)
+								.OnSelectionChanged(this, &SCustomizableInstanceProperties::OnProjectorTextureParameterComboBoxChanged, PoseSwitchEnumParamName, -1)
+								.OnGenerateWidget(this, &SCustomizableInstanceProperties::OnGenerateWidgetProjectorParameter)
+								.Content()
+								[
+									SNew(STextBlock)
+										.Text(FText::FromString(*PoseOptionNamesAttribute[PoseValueIndex]))
+								]
 							]
 
 						+ SHorizontalBox::Slot()
@@ -1063,7 +1082,7 @@ void SCustomizableInstanceProperties::AddParameter(int32 ParamIndexInObject)
 				FString OpacitySliderParamName = ParamName + FMultilayerProjector::OPACITY_PARAMETER_POSTFIX;
 				FString OpacitySliderParamNameWithRange = OpacitySliderParamName + FString::Printf(TEXT("__%d"), RangeIndex);
 
-				TArray<FString> OptionNamesAttribute;
+				TArray<TSharedPtr<FString>> OptionNamesAttribute;
 				FString Value = GetIntParameterValue(TextureSwitchEnumParamName, RangeIndex);
 				int32 ValueIndex = 0;
 
@@ -1077,7 +1096,13 @@ void SCustomizableInstanceProperties::AddParameter(int32 ParamIndexInObject)
 					{
 						ValueIndex = CandidateIndex;
 					}
-					OptionNamesAttribute.Add(CustomizableObject->GetIntParameterAvailableOption(TextureSwitchEnumParamIndexInObject, CandidateIndex));
+					OptionNamesAttribute.Add(MakeShared<FString>(CustomizableObject->GetIntParameterAvailableOption(TextureSwitchEnumParamIndexInObject, CandidateIndex)));
+				}
+
+				// Avoid filling this arraw with repeated array  options
+				if (RangeIndex == 0)
+				{
+					ProjectorTextureOptions.Add(MakeShared<TArray<TSharedPtr<FString>>>(OptionNamesAttribute));
 				}
 				
 				TSharedPtr<SHorizontalBox> SliderBox;
@@ -1129,12 +1154,16 @@ void SCustomizableInstanceProperties::AddParameter(int32 ParamIndexInObject)
 							SNew(SBox)
 							.MinDesiredWidth(120.f)
 							[
-								SNew(SMutableTextSearchBox)
-								.PossibleSuggestions(OptionNamesAttribute)
-								.InitialText(FText::FromString(OptionNamesAttribute[ValueIndex]))
-								.MustMatchPossibleSuggestions(TAttribute<bool>(true))
-								.SuggestionListPlacement(EMenuPlacement::MenuPlacement_ComboBox)
-								.OnTextCommitted(this, &SCustomizableInstanceProperties::OnProjectorTextureParameterComboBoxChanged, TextureSwitchEnumParamName, RangeIndex)
+								SNew(SSearchableComboBox)
+								.OptionsSource(ProjectorTextureOptions.Last().Get())
+								.InitiallySelectedItem((*ProjectorTextureOptions.Last())[ValueIndex])
+								.OnGenerateWidget(this, &SCustomizableInstanceProperties::OnGenerateWidgetProjectorParameter)
+								.OnSelectionChanged(this, &SCustomizableInstanceProperties::OnProjectorTextureParameterComboBoxChanged, TextureSwitchEnumParamName, RangeIndex)
+								.Content()
+								[
+									SNew(STextBlock)
+										.Text(FText::FromString(*(*ProjectorTextureOptions.Last())[ValueIndex]))
+								]
 							]
 						]
 					]
@@ -1267,6 +1296,12 @@ void SCustomizableInstanceProperties::AddParameter(int32 ParamIndexInObject)
 }
 
 
+TSharedRef<SWidget> SCustomizableInstanceProperties::OnGenerateWidgetIntParameter(TSharedPtr<FString> InItem) const
+{
+	return SNew(STextBlock).Text(FText::FromString(*InItem.Get()));
+}
+
+
 void SCustomizableInstanceProperties::Tick(const FGeometry& AllottedGeometry, const double InCurrentTime, const float InDeltaTime)
 {
 	SCompoundWidget::Tick(AllottedGeometry, InCurrentTime, InDeltaTime);
@@ -1823,6 +1858,12 @@ void SCustomizableInstanceProperties::OnProjectorTextureParameterComboBoxChanged
 }
 
 
+TSharedRef<SWidget> SCustomizableInstanceProperties::OnGenerateWidgetProjectorParameter(TSharedPtr<FString> InItem) const
+{
+	return SNew(STextBlock).Text(FText::FromString(*InItem.Get()));
+}
+
+
 void SCustomizableInstanceProperties::OnProjectorFloatParameterChanged(float Value, int SliderIndex)
 {
 	OnFloatParameterChanged(Value, SliderIndex);
diff --git a/Engine/Plugins/Experimental/Mutable/Source/CustomizableObjectEditor/Private/MuCOE/SMutableParametersWidget.cpp b/Engine/Plugins/Experimental/Mutable/Source/CustomizableObjectEditor/Private/MuCOE/SMutableParametersWidget.cpp
--- a/Engine/Plugins/Experimental/Mutable/Source/CustomizableObjectEditor/Private/MuCOE/SMutableParametersWidget.cpp
+++ b/Engine/Plugins/Experimental/Mutable/Source/CustomizableObjectEditor/Private/MuCOE/SMutableParametersWidget.cpp
@@ -2,7 +2,7 @@
 
 #include "MuCOE/SMutableParametersWidget.h"
 
-#include "MuCOE/SMutableTextSearchBox.h"
+#include "SSearchableComboBox.h"
 #include "MuCOE/UnrealEditorPortabilityHelpers.h"
 #include "MuR/ModelPrivate.h"
 #include "MuR/ParametersPrivate.h"
@@ -228,6 +228,7 @@ void SMutableParametersWidget::ScheduleUpdateIfRequired(const int32& InParameter
 	}
 }
 
+
 void SMutableParametersWidget::GenerateAndAttachParameterSlate(const int32 ParamIndex, TSharedPtr<SHorizontalBox> ParameterHorizontalBox,  mu::RangeIndexPtrConst RangeIndex)
 {
 	switch (MutableParameters->GetType(ParamIndex) )
@@ -278,32 +279,50 @@ void SMutableParametersWidget::GenerateAndAttachParameterSlate(const int32 Param
 	case mu::PARAMETER_TYPE::T_INT:
 	{
 		// If we have a list of options, add a combo box too
-		TSharedPtr<SMutableTextSearchBox> ParamComboBox;
+		TSharedPtr<SSearchableComboBox> ParamComboBox;
 		const int32 ValueCount = MutableParameters->GetIntPossibleValueCount(ParamIndex);
 		if (ValueCount > 0)
 		{
 			const FString ToolTipText = FString("None");
-			TArray<FString> OptionNamesAttribute;
+
+			TSharedPtr<TArray<TSharedPtr<FString>>>* FoundOptions = IntParameterOptions.Find(ParamIndex);
+			TArray<TSharedPtr<FString>>& OptionNamesAttribute = FoundOptions && FoundOptions->IsValid() ?
+																*FoundOptions->Get() :
+																*(IntParameterOptions.Add(ParamIndex, MakeShared<TArray<TSharedPtr<FString>>>()).Get());
+
+			OptionNamesAttribute.Empty();
+
 			const int32 Value = MutableParameters->GetIntValue( ParamIndex , RangeIndex);
 			const int32 ValueIndex = MutableParameters->GetIntValueIndex( ParamIndex, Value );
 			for (int32 i = 0; i < ValueCount; ++i)
 			{
 				const FString& ValueText = MutableParameters->GetIntPossibleValueName( ParamIndex, i );
-				OptionNamesAttribute.Add( ValueText );				
+				OptionNamesAttribute.Add(MakeShared<FString>(ValueText));				
 			}
 
 			ParameterHorizontalBox->AddSlot()
 				.Padding(4.0f)
 				.VAlign(VAlign_Center)
 				[
-					SAssignNew(ParamComboBox, SMutableTextSearchBox)
-					.Visibility(this, &SMutableParametersWidget::GetParameterVisibility, ParamIndex)
+					SAssignNew(ParamComboBox, SSearchableComboBox)
+					.SearchVisibility(this, &SMutableParametersWidget::GetParameterVisibility, ParamIndex)
 					.ToolTipText(FText::FromString(ToolTipText))
-					.PossibleSuggestions(OptionNamesAttribute)
-					.InitialText(FText::FromString(OptionNamesAttribute[ValueIndex]))
-					.MustMatchPossibleSuggestions(TAttribute<bool>(true))
-					.SuggestionListPlacement(EMenuPlacement::MenuPlacement_ComboBox)
-					.OnTextCommitted(this, &SMutableParametersWidget::OnIntParameterTextChanged, ParamIndex,RangeIndex)
+					.OptionsSource(&OptionNamesAttribute)
+					.InitiallySelectedItem(OptionNamesAttribute[ValueIndex])
+					.Method(EPopupMethod::UseCurrentWindow)
+					.OnSelectionChanged(this, &SMutableParametersWidget::OnIntParameterTextChanged, ParamIndex, RangeIndex)
+					.OnGenerateWidget(this, &SMutableParametersWidget::OnGenerateWidgetIntParameter)
+					.Content()
+					[
+						SNew(STextBlock)
+						.Text_Lambda([this, ParamIndex, RangeIndex, OptionNamesAttribute]() -> FText
+						{
+							const int32 Value = MutableParameters->GetIntValue(ParamIndex, RangeIndex);
+							const int32 ValueIndex = MutableParameters->GetIntValueIndex(ParamIndex, Value);
+
+							return FText::FromString(*OptionNamesAttribute[ValueIndex]);
+						})
+					]
 				];
 		}
 
@@ -317,7 +336,7 @@ void SMutableParametersWidget::GenerateAndAttachParameterSlate(const int32 Param
 				.MinSliderValue(0)
 				.MaxSliderValue(this, &SMutableParametersWidget::GetIntParameterValueMax, ParamIndex)
 				.Value(this, &SMutableParametersWidget::GetIntParameterValue, ParamIndex, RangeIndex)
-				.OnValueChanged(this, &SMutableParametersWidget::OnIntParameterChanged, ParamIndex, ParamComboBox,RangeIndex)
+				.OnValueChanged(this, &SMutableParametersWidget::OnIntParameterChanged, ParamIndex, ParamComboBox, RangeIndex)
 			];
 		break;
 	}
@@ -512,7 +531,7 @@ TOptional<int32> SMutableParametersWidget::GetIntParameterValueMax(int32 ParamIn
 }
 
 
-void SMutableParametersWidget::OnIntParameterChanged(int32 InValue, int32 ParamIndex, TSharedPtr<SMutableTextSearchBox> Combo, mu::RangeIndexPtrConst RangeIndex )
+void SMutableParametersWidget::OnIntParameterChanged(int32 InValue, int32 ParamIndex, TSharedPtr<SSearchableComboBox> Combo, mu::RangeIndexPtrConst RangeIndex )
 {
 	if (!MutableParameters
 		|| ParamIndex >= MutableParameters->GetCount()
@@ -531,7 +550,8 @@ void SMutableParametersWidget::OnIntParameterChanged(int32 InValue, int32 ParamI
 		if (Combo)
 		{
 			const FString Text = MutableParameters->GetIntPossibleValueName(ParamIndex, InValue);
-			Combo->SetText(FText::FromString(Text));
+			Combo->RefreshOptions();
+			Combo->SetSelectedItem(MakeShared<FString>(Text), ESelectInfo::Direct);
 		}
 		
 		OnParametersValueChanged.ExecuteIfBound(ParamIndex);
@@ -564,6 +584,12 @@ void SMutableParametersWidget::OnIntParameterTextChanged(TSharedPtr<FString> Sel
 }
 
 
+TSharedRef<SWidget> SMutableParametersWidget::OnGenerateWidgetIntParameter(TSharedPtr<FString> InItem) const
+{
+	return SNew(STextBlock).Text(FText::FromString(*InItem.Get()));
+}
+
+
 FLinearColor SMutableParametersWidget::GetColorParameterValue(int32 ParamIndex, mu::RangeIndexPtrConst RangeIndex) const
 {
 	if (!MutableParameters
diff --git a/Engine/Plugins/Experimental/Mutable/Source/CustomizableObjectEditor/Private/MuCOE/SMutableParametersWidget.h b/Engine/Plugins/Experimental/Mutable/Source/CustomizableObjectEditor/Private/MuCOE/SMutableParametersWidget.h
--- a/Engine/Plugins/Experimental/Mutable/Source/CustomizableObjectEditor/Private/MuCOE/SMutableParametersWidget.h
+++ b/Engine/Plugins/Experimental/Mutable/Source/CustomizableObjectEditor/Private/MuCOE/SMutableParametersWidget.h
@@ -13,7 +13,7 @@ namespace ESelectInfo { enum Type : int; }
 namespace ETextCommit { enum Type : int; }
 
 class FString;
-class SMutableTextSearchBox;
+class SSearchableComboBox;
 class SVerticalBox;
 struct FGeometry;
 struct FPointerEvent;
@@ -33,11 +33,11 @@ public:
 
 	SLATE_BEGIN_ARGS(SMutableParametersWidget) {}
 
-		/** Parameters to show and edit. */
-		SLATE_ARGUMENT_DEFAULT(mu::ParametersPtr, Parameters) { nullptr };
+	/** Parameters to show and edit. */
+	SLATE_ARGUMENT_DEFAULT(mu::ParametersPtr, Parameters) { nullptr };
 
-		/** Called when any parameter value has changed, with the parameter index as argument.  */
-		SLATE_EVENT(FOnMutableParameterValueChanged, OnParametersValueChanged)
+	/** Called when any parameter value has changed, with the parameter index as argument.  */
+	SLATE_EVENT(FOnMutableParameterValueChanged, OnParametersValueChanged)
 
 	SLATE_END_ARGS()
 
@@ -65,6 +65,9 @@ private:
 	/** Dynamically filled box with per-parameter widgets */
 	TSharedPtr<SVerticalBox> ParamBox;
 
+	/** Map from ParamIndexInObject to the int param's selector options */
+	TMap<int32, TSharedPtr<TArray<TSharedPtr<FString>>>> IntParameterOptions;
+
 	/** Check if the Parameter index provided is from a parameter that controls the amount of values another parameter/s
 	 * should have exposed in the UI. If the parameter does so then the redraw of the parameters widget will be scheduled
 	 * for the next update.
@@ -96,8 +99,9 @@ private:
 	void OnSetColorFromColorPicker(FLinearColor NewColor, int32 ParamIndex, mu::RangeIndexPtrConst RangeIndex);
 	TOptional<int32> GetIntParameterValue(int32 ParamIndex,  mu::RangeIndexPtrConst RangeIndex) const;
 	TOptional<int32> GetIntParameterValueMax(int32 ParamIndex) const;
-	void OnIntParameterChanged(int32 Value, int32 ParamIndex, TSharedPtr<SMutableTextSearchBox> Combo,  mu::RangeIndexPtrConst RangeIndex);
+	void OnIntParameterChanged(int32 Value, int32 ParamIndex, TSharedPtr<SSearchableComboBox> Combo,  mu::RangeIndexPtrConst RangeIndex);
 	void OnIntParameterTextChanged(TSharedPtr<FString> Selection, ESelectInfo::Type SelectInfo, int32 ParamIndex, mu::RangeIndexPtrConst RangeIndex);
+	TSharedRef<SWidget> OnGenerateWidgetIntParameter(TSharedPtr<FString> InItem) const;
 
 	/** Projector UI callbacks */
 	TOptional<FVector::FReal> GetProjectorLocation(EAxis::Type Axis, int32 ParamIndex,  mu::RangeIndexPtrConst RangeIndex) const;
diff --git a/Engine/Plugins/Experimental/Mutable/Source/CustomizableObjectEditor/Private/MuCOE/SMutableTextSearchBox.cpp b/Engine/Plugins/Experimental/Mutable/Source/CustomizableObjectEditor/Private/MuCOE/SMutableTextSearchBox.cpp
deleted file mode 100644
diff --git a/Engine/Plugins/Experimental/Mutable/Source/CustomizableObjectEditor/Private/MuCOE/SMutableTextSearchBox.h b/Engine/Plugins/Experimental/Mutable/Source/CustomizableObjectEditor/Private/MuCOE/SMutableTextSearchBox.h
deleted file mode 100644
diff --git a/Engine/Plugins/Experimental/Mutable/Source/CustomizableObjectEditor/Public/MuCOE/SCustomizableInstanceProperties.h b/Engine/Plugins/Experimental/Mutable/Source/CustomizableObjectEditor/Public/MuCOE/SCustomizableInstanceProperties.h
--- a/Engine/Plugins/Experimental/Mutable/Source/CustomizableObjectEditor/Public/MuCOE/SCustomizableInstanceProperties.h
+++ b/Engine/Plugins/Experimental/Mutable/Source/CustomizableObjectEditor/Public/MuCOE/SCustomizableInstanceProperties.h
@@ -100,6 +100,15 @@ private:
 	TArray<TSharedPtr<FString>> TextureParameterValueNames;
 	TArray<FName> TextureParameterValues;
 
+	/** Array with all the possible multilayer projector texture options */
+	TArray<TSharedPtr<TArray<TSharedPtr<FString>>>> ProjectorTextureOptions;
+
+	/** Map from ParamIndexInObject to the param's int selector options */
+	TMap<int32, TSharedPtr<TArray<TSharedPtr<FString>>>> IntParameterOptions;
+
+	/** Map from ParamIndexInObject to the projector param pose options  */
+	TMap<int32, TSharedPtr<TArray<TSharedPtr<FString>>>> ProjectorParameterPoseOptions;
+
 	struct FSliderData
 	{
 		FSliderData( TSharedPtr<SSpinBox<float>> InSlider, const FString& InParamName, int32 InRangeIndex, float InLastValueSet)
@@ -168,6 +177,8 @@ private:
 	/** SearchBox OnTextCommitted Callback to search a parameter int by name */
 	void OnIntParameterComboBoxChanged(TSharedPtr<FString> Selection, ESelectInfo::Type SelectInfo, FString ParamName);
 
+	TSharedRef<SWidget> OnGenerateWidgetIntParameter(TSharedPtr<FString> InItem) const;
+
 	/** Gets the name value of a int parameter */
 	FString GetIntParameterValue(FString ParamName, int RangeIndex = -1) const;
 	
@@ -176,6 +187,7 @@ private:
 	
 	/** SearchBox OnTextCommitted Callback to search a parameter texture by name */
 	void OnProjectorTextureParameterComboBoxChanged(TSharedPtr<FString> Selection, ESelectInfo::Type SelectInfo, FString ParamName, int32 RangeIndex) const;
+	TSharedRef<SWidget> OnGenerateWidgetProjectorParameter(TSharedPtr<FString> InItem) const;
 
 	/** Slider OnValueChanged Callbacks for float parameters */
 	void OnProjectorFloatParameterChanged(float Value, int SliderIndex);
