[mutable] Improved handling of the "Only Runtime parameters" property in CO instances so that the user is not confused when all parameters disappear. A message will be displayed if parameters are hidden and the only runtime option disabled if there are no runtime parameters when the editor is opened.
#rnx

#tests Editor

diff --git a/Engine/Plugins/Experimental/Mutable/Source/CustomizableObjectEditor/Private/MuCOE/CustomizableInstanceDetails.cpp b/Engine/Plugins/Experimental/Mutable/Source/CustomizableObjectEditor/Private/MuCOE/CustomizableInstanceDetails.cpp
--- a/Engine/Plugins/Experimental/Mutable/Source/CustomizableObjectEditor/Private/MuCOE/CustomizableInstanceDetails.cpp
+++ b/Engine/Plugins/Experimental/Mutable/Source/CustomizableObjectEditor/Private/MuCOE/CustomizableInstanceDetails.cpp
@@ -231,7 +231,26 @@ void FCustomizableInstanceDetails::CustomizeDetails(const TSharedPtr<IDetailLayo
 	];
 
 	// Parameters Widgets
-	GenerateParametersView(ParametersCategory);
+	bool bHiddenParamsRuntime = GenerateParametersView(ParametersCategory);
+
+	if (bHiddenParamsRuntime)
+	{
+		FText HiddenParamsRuntimeMessage = LOCTEXT("CustomizableInstanceDetails_HiddemParamsRuntime", "Parameters are hidden due to their Runtime type. \nUncheck the Only Runtime checkbox to see them.");
+
+		ParametersCategory.AddCustomRow(LOCTEXT("CustomizableInstanceDetails_HiddemParamsRuntimeRow", "Parameters are hidden"))
+		[
+			SNew(SHorizontalBox)
+				+ SHorizontalBox::Slot()
+				.FillWidth(2.0f)
+				.Padding(0.0f, 15.0f, 0.0f, 5.0f)
+				[
+					SNew(STextBlock)
+						.Text(HiddenParamsRuntimeMessage)
+						.ToolTipText(HiddenParamsRuntimeMessage)
+						.AutoWrapText(true)
+				]
+		];
+	}
 }
 
 
@@ -539,7 +558,7 @@ struct FParameterInfo
 
 // PARAMETERS WIDGET GENERATION -----------------------------------------------------------------------------------------------------------------
 
-void FCustomizableInstanceDetails::GenerateParametersView(IDetailCategoryBuilder& DetailsCategory)
+bool FCustomizableInstanceDetails::GenerateParametersView(IDetailCategoryBuilder& DetailsCategory)
 {
 	// Make this arrays locals
 	ParamChildren.Empty();
@@ -551,13 +570,20 @@ void FCustomizableInstanceDetails::GenerateParametersView(IDetailCategoryBuilder
 
 	if (!CustomizableObject)
 	{
-		return;
+		return false;
 	}
-	
+
+	bool bParametersHiddenRuntime = false;
+
 	if (CustomInstance->GetPrivate()->bShowOnlyRuntimeParameters)
 	{
 		const int32 NumStateParameters = CustomizableObject->GetStateParameterCount(CustomInstance->GetPrivate()->GetState());
 
+		if (NumStateParameters < CustomizableObject->GetParameterCount())
+		{
+			bParametersHiddenRuntime = true;
+		}
+
 		for (int32 ParamIndexInState = 0; ParamIndexInState < NumStateParameters; ++ParamIndexInState)
 		{
 			FParameterInfo ParameterSortInfo;
@@ -576,7 +602,7 @@ void FCustomizableInstanceDetails::GenerateParametersView(IDetailCategoryBuilder
 		for (int32 ParamIndex = 0; ParamIndex < ParametersTree.Num(); ++ParamIndex)
 		{
 			const FParameterInfo& ParamInfo = ParametersTree[ParamIndex];
-			
+
 			if (CustomInstance->GetPrivate()->bShowUISections)
 			{
 				IDetailGroup* CurrentSection = GenerateParameterSection(ParamInfo.ParamIndexInObject, DetailsCategory);
@@ -615,7 +641,7 @@ void FCustomizableInstanceDetails::GenerateParametersView(IDetailCategoryBuilder
 	else
 	{
 		const int32 NumObjectParameter = CustomizableObject->GetParameterCount();
-	
+
 		//TODO: get all parameters and sort, then make the next "for" use that sorted list as source of indexes
 		for (int32 ParamIndexInObject = 0; ParamIndexInObject < NumObjectParameter; ++ParamIndexInObject)
 		{
@@ -628,14 +654,14 @@ void FCustomizableInstanceDetails::GenerateParametersView(IDetailCategoryBuilder
 				ParametersTree.Add(ParameterSortInfo);
 			}
 		}
-	
+
 		ParametersTree.Sort();
-	
+
 		for (int32 ParamIndexInObject = 0; ParamIndexInObject < ParametersTree.Num(); ++ParamIndexInObject)
 		{
 			FillChildrenMap(ParametersTree[ParamIndexInObject].ParamIndexInObject);
 		}
-	
+
 		for (int32 ParamIndexInObject = 0; ParamIndexInObject < ParametersTree.Num(); ++ParamIndexInObject)
 		{
 			if (!ParamHasParent.Find(ParametersTree[ParamIndexInObject].ParamIndexInObject))
@@ -644,6 +670,8 @@ void FCustomizableInstanceDetails::GenerateParametersView(IDetailCategoryBuilder
 			}
 		}
 	}
+
+	return bParametersHiddenRuntime;
 }
 
 
diff --git a/Engine/Plugins/Experimental/Mutable/Source/CustomizableObjectEditor/Private/MuCOE/CustomizableInstanceDetails.h b/Engine/Plugins/Experimental/Mutable/Source/CustomizableObjectEditor/Private/MuCOE/CustomizableInstanceDetails.h
--- a/Engine/Plugins/Experimental/Mutable/Source/CustomizableObjectEditor/Private/MuCOE/CustomizableInstanceDetails.h
+++ b/Engine/Plugins/Experimental/Mutable/Source/CustomizableObjectEditor/Private/MuCOE/CustomizableInstanceDetails.h
@@ -60,7 +60,8 @@ private:
 	void OnUseUISectionsSelectionChanged(ECheckBoxState InCheckboxState);
 
 	// Main parameter generation functions
-	void GenerateParametersView(IDetailCategoryBuilder& MainCategory);
+	// Returns true if parameters have been hidden due to runtime type
+	bool GenerateParametersView(IDetailCategoryBuilder& MainCategory);
 	void RecursivelyAddParamAndChildren(const int32 ParamIndexInObject, const FString ParentName, IDetailCategoryBuilder& DetailsCategory);
 	void FillChildrenMap(int32 ParamIndexInObject);
 
diff --git a/Engine/Plugins/Experimental/Mutable/Source/CustomizableObjectEditor/Private/MuCOE/CustomizableObjectInstanceEditor.cpp b/Engine/Plugins/Experimental/Mutable/Source/CustomizableObjectEditor/Private/MuCOE/CustomizableObjectInstanceEditor.cpp
--- a/Engine/Plugins/Experimental/Mutable/Source/CustomizableObjectEditor/Private/MuCOE/CustomizableObjectInstanceEditor.cpp
+++ b/Engine/Plugins/Experimental/Mutable/Source/CustomizableObjectEditor/Private/MuCOE/CustomizableObjectInstanceEditor.cpp
@@ -321,6 +321,15 @@ void FCustomizableObjectInstanceEditor::InitCustomizableObjectInstanceEditor( co
 		CustomizableObjectPrivate->Status.GetOnStateChangedDelegate().AddRaw(this, &FCustomizableObjectInstanceEditor::OnCustomizableObjectStatusChanged);
 		const FCustomizableObjectStatusTypes::EState CurrentStatus = CustomizableObjectPrivate->Status.Get();
 		OnCustomizableObjectStatusChanged(CurrentStatus, CurrentStatus);
+
+		int32 StateParameterCount = CustomizableObject->GetStateParameterCount(CustomizableObjectInstance->GetCurrentState());
+		int32 ParameterCount = CustomizableObject->GetParameterCount();
+
+		if (StateParameterCount == 0)
+		{
+			bOnlyRuntimeParameters = false;
+			InCustomizableObjectInstance->GetPrivate()->bShowOnlyRuntimeParameters = false;
+		}
 	}
 }
 
