[Mutable] Fix tangent space decodification for Imported model.

- The decodification was assuming unsigned format but signed format was used.

#rnx

diff --git a/Engine/Plugins/Experimental/Mutable/Source/CustomizableObject/Private/MuCO/CustomizableObjectInstance.cpp b/Engine/Plugins/Experimental/Mutable/Source/CustomizableObject/Private/MuCO/CustomizableObjectInstance.cpp
--- a/Engine/Plugins/Experimental/Mutable/Source/CustomizableObject/Private/MuCO/CustomizableObjectInstance.cpp
+++ b/Engine/Plugins/Experimental/Mutable/Source/CustomizableObject/Private/MuCO/CustomizableObjectInstance.cpp
@@ -6970,7 +6970,7 @@ void UCustomizableInstancePrivate::RegenerateImportedModels()
 
 					Vertex.TangentX = TangentPtr[0].ToFVector3f();
 					Vertex.TangentZ = TangentPtr[1].ToFVector3f();
-					float TangentSign = TangentPtr[1].Vector.W == 0 ? -1.f : 1.f;
+					float TangentSign = TangentPtr[1].Vector.W < 0 ? -1.f : 1.f;
 					Vertex.TangentY = FVector3f::CrossProduct(Vertex.TangentZ, Vertex.TangentX) * TangentSign;
 
 					Vertex.UVs[0] = TexCoordPosPtr[0];
