[Mutable] Enable compiling COs using disk storage.

#rnx

diff --git a/Engine/Plugins/Experimental/Mutable/Source/CustomizableObjectEditor/Private/MuCOE/CustomizableObjectCompiler.cpp b/Engine/Plugins/Experimental/Mutable/Source/CustomizableObjectEditor/Private/MuCOE/CustomizableObjectCompiler.cpp
--- a/Engine/Plugins/Experimental/Mutable/Source/CustomizableObjectEditor/Private/MuCOE/CustomizableObjectCompiler.cpp
+++ b/Engine/Plugins/Experimental/Mutable/Source/CustomizableObjectEditor/Private/MuCOE/CustomizableObjectCompiler.cpp
@@ -871,6 +871,7 @@ void FCustomizableObjectCompiler::CompileInternal(UCustomizableObject* Object, c
 	Options.bSkinWeightProfilesEnabled = Object->bEnableAltSkinWeightProfiles;
 	Options.bPhysicsAssetMergeEnabled = Object->bEnablePhysicsAssetMerge;
 	Options.bAnimBpPhysicsManipulationEnabled = Object->bEnableAnimBpPhysicsAssetsManipualtion;
+	Options.bUseDiskCompilation = Object->CompileOptions.bUseDiskCompilation;
 	Options.ImageTiling = Object->CompileOptions.ImageTiling;
 
 	if (!InOptions.bIsCooking && IsRunningCookCommandlet())
