[Mutable] Disable LODManagment in the editor.
Add CVar to force LODManagment in the editor for testing purposes.

#rnx

diff --git a/Engine/Plugins/Experimental/Mutable/Source/CustomizableObject/Private/MuCO/CustomizableObjectSystem.cpp b/Engine/Plugins/Experimental/Mutable/Source/CustomizableObject/Private/MuCO/CustomizableObjectSystem.cpp
--- a/Engine/Plugins/Experimental/Mutable/Source/CustomizableObject/Private/MuCO/CustomizableObjectSystem.cpp
+++ b/Engine/Plugins/Experimental/Mutable/Source/CustomizableObject/Private/MuCO/CustomizableObjectSystem.cpp
@@ -121,6 +121,16 @@ TAutoConsoleVariable<bool> CVarEnableNewSplitMutableTask(
 	ECVF_Scalability);
 
 
+#if WITH_EDITOR
+bool bEnableLODManagmentInEditor = false;
+
+static FAutoConsoleVariableRef CVarMutableEnableLODManagmentInEditor(
+	TEXT("Mutable.EnableLODManagmentInEditor"),
+	bEnableLODManagmentInEditor,
+	TEXT("true/false - If true, enables custom LODManagment in the editor. "),
+	ECVF_Default);
+#endif
+
 int32 UCustomizableObjectSystemPrivate::SkeletalMeshMinLodQualityLevel = -1;
 
 
@@ -3101,10 +3111,14 @@ namespace impl
 			CustomizableObject->GetPrivate()->GetLowPriorityTextureNames(Operation->LowPriorityTextures);
 		}
 
-		bool bIsInEditorViewport = false;
+		bool bRequestAllLODs = !System->IsOnlyGenerateRequestedLODsEnabled() ||
+			!System->GetPrivate()->CurrentInstanceLODManagement->IsOnlyGenerateRequestedLODLevelsEnabled();
 
 #if WITH_EDITOR
-		for (TObjectIterator<UCustomizableObjectInstanceUsage> CustomizableObjectInstanceUsage; CustomizableObjectInstanceUsage && !bIsInEditorViewport; ++CustomizableObjectInstanceUsage)
+		// In the editor LOD Management is disabled by default. Overwrite requested LODs when disabled.
+		bRequestAllLODs |= !bEnableLODManagmentInEditor;
+
+		for (TObjectIterator<UCustomizableObjectInstanceUsage> CustomizableObjectInstanceUsage; CustomizableObjectInstanceUsage && !bRequestAllLODs; ++CustomizableObjectInstanceUsage)
 		{
 			if (IsValid(*CustomizableObjectInstanceUsage) && CustomizableObjectInstanceUsage->IsNetMode(NM_DedicatedServer))
 			{
@@ -3126,18 +3140,16 @@ namespace impl
 				switch (WorldType)
 				{
 					// Editor preview instances
-					case EWorldType::EditorPreview:
-					case EWorldType::None:
-						bIsInEditorViewport = true;
-					default: ;
+				case EWorldType::EditorPreview:
+				case EWorldType::None:
+					bRequestAllLODs = true;
+				default:;
 				}
 			}
 		}
 #endif // WITH_EDITOR
-		
-		if (!System->IsOnlyGenerateRequestedLODsEnabled() ||
-			!System->GetPrivate()->CurrentInstanceLODManagement->IsOnlyGenerateRequestedLODLevelsEnabled() ||
-			bIsInEditorViewport)
+
+		if (bRequestAllLODs)
 		{
 			TArray<uint16> RequestedLODs = Operation->GetRequestedLODs();
 			RequestedLODs.Init(0, Operation->NumComponents);
