Unshelved from pending changelist '32429739':

[Mutable] Fix crash when exiting the editor.

Mutable should not tick perform new updates once the editor is shutting down. If not, it will crash due to calling LoadSynchronous.

#rnx

diff --git a/Engine/Plugins/Experimental/Mutable/Source/CustomizableObject/Private/MuCO/CustomizableObjectSystem.cpp b/Engine/Plugins/Experimental/Mutable/Source/CustomizableObject/Private/MuCO/CustomizableObjectSystem.cpp
--- a/Engine/Plugins/Experimental/Mutable/Source/CustomizableObject/Private/MuCO/CustomizableObjectSystem.cpp
+++ b/Engine/Plugins/Experimental/Mutable/Source/CustomizableObject/Private/MuCO/CustomizableObjectSystem.cpp
@@ -3269,6 +3269,11 @@ int32 UCustomizableObjectSystem::TickInternal()
 	return 0;
 #endif
 
+	if (IsEngineExitRequested())
+	{
+		return 0;
+	}
+
 	if (!Private)
 	{
 		return 0;
