[Mutable] Fix ColorLayerFromAlpha operations. Fix ImageStorge batch iteration for first lod with offset.

#rnx

diff --git a/Engine/Plugins/Experimental/Mutable/Source/MutableRuntime/Private/MuR/OpImageLayer.h b/Engine/Plugins/Experimental/Mutable/Source/MutableRuntime/Private/MuR/OpImageLayer.h
--- a/Engine/Plugins/Experimental/Mutable/Source/MutableRuntime/Private/MuR/OpImageLayer.h
+++ b/Engine/Plugins/Experimental/Mutable/Source/MutableRuntime/Private/MuR/OpImageLayer.h
@@ -136,24 +136,24 @@ namespace mu
 		template<uint32 NumChannels, uint32 (*BLEND_FUNC)(uint32,uint32), bool Clamp>
 		FORCENOINLINE void BufferLayerColourFromAlphaGenericChannel(uint8* DestBuf, const uint8* BaseBuf, int32 NumElems, const FIntVector& Color)
 		{
-			static_assert(NumChannels >= 3 && NumChannels < 4);
+			static_assert(NumChannels > 0 && NumChannels <= 4);
 
 			for (int32 I = 0; I < NumElems; ++I)
 			{
-				for (uint32 C = 0; C < NumChannels; ++C)
+				const uint32 Alpha = Invoke([&]() -> uint32
 				{
-					const uint32 Alpha = Invoke([&]() -> uint32
+					if constexpr (NumChannels <= 3)
 					{
-						if constexpr (NumChannels == 3)
-						{
-							return 255;
-						}
-						else
-						{
-							return BaseBuf[NumChannels * I + 3];
-						}
-					});
-					
+						return 255;
+					}
+					else
+					{
+						return BaseBuf[NumChannels * I + 3];
+					}
+				});
+
+				for (uint32 C = 0; C < NumChannels; ++C)
+				{
 					uint32 Result = BLEND_FUNC(Alpha, Color[C]);
 					if constexpr (Clamp)
 					{
@@ -200,26 +200,26 @@ namespace mu
 			case EImageFormat::IF_L_UBYTE:
 			{
 				check(BytesPerElem == 1);
-				Private::BufferLayerColourGenericChannel<1, BLEND_FUNC, CLAMP>(ResultView.GetData(), BaseView.GetData(), NumElems, ColorValue);
+				Private::BufferLayerColourFromAlphaGenericChannel<1, BLEND_FUNC, CLAMP>(ResultView.GetData(), BaseView.GetData(), NumElems, ColorValue);
 				break;
 			}
 			case EImageFormat::IF_RGB_UBYTE:
 			{
 				check(BytesPerElem == 3);
-				Private::BufferLayerColourGenericChannel<3, BLEND_FUNC, CLAMP>(ResultView.GetData(), BaseView.GetData(), NumElems, ColorValue);
+				Private::BufferLayerColourFromAlphaGenericChannel<3, BLEND_FUNC, CLAMP>(ResultView.GetData(), BaseView.GetData(), NumElems, ColorValue);
 				break;
 			}
 			case EImageFormat::IF_RGBA_UBYTE:
 			{
 				check(BytesPerElem == 4);
-				Private::BufferLayerColourGenericChannel<4, BLEND_FUNC, CLAMP>(ResultView.GetData(), BaseView.GetData(), NumElems, ColorValue);
+				Private::BufferLayerColourFromAlphaGenericChannel<4, BLEND_FUNC, CLAMP>(ResultView.GetData(), BaseView.GetData(), NumElems, ColorValue);
 				break;
 			}
 			case EImageFormat::IF_BGRA_UBYTE:
 			{
 				check(BytesPerElem == 4);
 				FIntVector BGRAColorValues = FIntVector(ColorValue.Z, ColorValue.Y, ColorValue.X);
-				Private::BufferLayerColourGenericChannel<4, BLEND_FUNC, CLAMP>(ResultView.GetData(), BaseView.GetData(), NumElems , ColorValue);
+				Private::BufferLayerColourFromAlphaGenericChannel<4, BLEND_FUNC, CLAMP>(ResultView.GetData(), BaseView.GetData(), NumElems , ColorValue);
 				break;
 			}
 			default:
@@ -1636,10 +1636,12 @@ namespace mu
 		check(BaseImage->GetSizeX() == BlendImage->GetSizeX() && BaseImage->GetSizeY() == BlendImage->GetSizeY());
 		check(bOnlyFirstLOD || BaseImage->GetLODCount() <= BlendImage->GetLODCount());
 
+		bOnlyFirstLOD = bOnlyFirstLOD || BaseImage->GetLODCount() == 1;
+		
 		int32 FirstLODDataOffset = 0;	
 		int32 NumRelevantElems = -1;
 
-		if (BlendImage->m_flags & Image::IF_HAS_RELEVANCY_MAP && (bOnlyFirstLOD || BaseImage->GetLODCount() == 1))
+		if (BlendImage->m_flags & Image::IF_HAS_RELEVANCY_MAP && bOnlyFirstLOD)
 		{
 			check(BlendImage->RelevancyMaxY < BaseImage->GetSizeY());
 			check(BlendImage->RelevancyMaxY >= BlendImage->RelevancyMinY);
@@ -1663,7 +1665,7 @@ namespace mu
 				: BaseImage->DataStorage.GetNumBatchesLODRange(NumBatchElems, BytesPerElem, LODBegin, LODEnd);
 
 		// This will always be an upper-bound for bOnlyFirtsLODs, check if it performs as expected or it needs more fine tune.
-		const int32 NumRelevantBatches = bOnlyFirstLOD && NumRelevantElems != -1
+		const int32 NumRelevantBatches = NumRelevantElems != -1
 				? FMath::DivideAndRoundUp(NumRelevantElems, NumBatchElems)
 				: NumBatches;
 
@@ -1752,11 +1754,11 @@ namespace mu
 		check(bOnlyFirstLOD || BaseImage->GetLODCount() <= BlendImage->GetLODCount());
 		check(BlendAlphaSourceChannel == 3);
 
+		bOnlyFirstLOD = bOnlyFirstLOD || BaseImage->GetLODCount() == 1;
+
 		int32 FirstLODDataOffset = 0;
 		int32 NumRelevantElems = -1;
-		if (BlendImage->m_flags & Image::IF_HAS_RELEVANCY_MAP 
-			&& 
-			(bOnlyFirstLOD || BaseImage->GetLODCount()==1 ) )
+		if (BlendImage->m_flags & Image::IF_HAS_RELEVANCY_MAP && bOnlyFirstLOD)
 		{
 			check(BlendImage->RelevancyMaxY < BaseImage->GetSizeY());
 			check(BlendImage->RelevancyMaxY >= BlendImage->RelevancyMinY);
@@ -1778,7 +1780,7 @@ namespace mu
 				: BaseImage->DataStorage.GetNumBatchesLODRange(NumBatchElems, BytesPerElem, 0, NumLODs);
 
 		// This will always be an upper-bound for bOnlyFirtsLODs, check if it performs as expected or it needs more fine tune.
-		const int32 NumRelevantBatches = bOnlyFirstLOD && NumRelevantElems != -1
+		const int32 NumRelevantBatches = NumRelevantElems != -1
 				? FMath::DivideAndRoundUp(NumRelevantElems, NumBatchElems)
 				: NumBatches;
 
@@ -1835,7 +1837,6 @@ namespace mu
 		}
 	}
 
-
 	template< 
 		VectorRegister4Int (*RGB_FUNC_MASKED)(const VectorRegister4Int&, const VectorRegister4Int&, const VectorRegister4Int&),
 		int32 (*A_FUNC)(int32, int32),
diff --git a/Engine/Plugins/Experimental/Mutable/Source/MutableRuntime/Public/MuR/ImageDataStorage.h b/Engine/Plugins/Experimental/Mutable/Source/MutableRuntime/Public/MuR/ImageDataStorage.h
--- a/Engine/Plugins/Experimental/Mutable/Source/MutableRuntime/Public/MuR/ImageDataStorage.h
+++ b/Engine/Plugins/Experimental/Mutable/Source/MutableRuntime/Public/MuR/ImageDataStorage.h
@@ -208,7 +208,7 @@ public:
 	 * BatchSizeInElems*BatchElemSizeInBytes is not multiple of the buffer size.
 	 *
 	 */
-	FORCEINLINE TArrayView<const uint8> GetBatchFirstLODOffset(int32 BatchId, int32 BatchSizeInElems, int32 BatchElemSizeInBytes, int32 OffsetInBytes = 0) const
+	FORCEINLINE TArrayView<const uint8> GetBatchFirstLODOffset(int32 BatchId, int32 BatchSizeInElems, int32 BatchElemSizeInBytes, int32 OffsetInBytes) const
 	{
 		check(!IsVoid());
 		TArrayView<const uint8> FirstLODView = GetLOD(0);
@@ -219,14 +219,15 @@ public:
 	
 		TArrayView<const uint8> OffsetedLODView = MakeArrayView(FirstLODView.GetData() + OffsetInBytes, FirstLODView.Num() - OffsetInBytes);
 
-		const int32 BatchOffset = BatchId * BatchSizeInElems * BatchElemSizeInBytes;
+		const int32 BatchSizeInBytes = BatchSizeInElems * BatchElemSizeInBytes;
+		const int32 BatchOffset = BatchId * BatchSizeInBytes;
 
-		if (OffsetedLODView.Num() > BatchOffset)
+		if (BatchOffset >= OffsetedLODView.Num())
 		{
 			return TArrayView<const uint8>();
 		}
 
-		return MakeArrayView(OffsetedLODView.GetData() + BatchOffset, FMath::Min(BatchSizeInElems, OffsetedLODView.Num() - BatchOffset));
+		return MakeArrayView(OffsetedLODView.GetData() + BatchOffset, FMath::Min(BatchSizeInBytes, OffsetedLODView.Num() - BatchOffset));
 	}
 
 	/**
