[Mutable] Use mip generation setting for image surface generation

- For now only SimpleAverage and Unfiltered modes are supported.

#rnx

diff --git a/Engine/Plugins/Experimental/Mutable/Source/CustomizableObject/Internal/MuCO/CustomizableObjectPrivate.h b/Engine/Plugins/Experimental/Mutable/Source/CustomizableObject/Internal/MuCO/CustomizableObjectPrivate.h
--- a/Engine/Plugins/Experimental/Mutable/Source/CustomizableObject/Internal/MuCO/CustomizableObjectPrivate.h
+++ b/Engine/Plugins/Experimental/Mutable/Source/CustomizableObject/Internal/MuCO/CustomizableObjectPrivate.h
@@ -826,6 +826,6 @@ public:
 	// This is a manual version number for the binary blobs in this asset.
 	// Increasing it invalidates all the previously compiled models.
 	// Warning: If while merging code both versions have changed, take the highest+1.
-	static constexpr int32 CurrentSupportedVersion = 436;
+	static constexpr int32 CurrentSupportedVersion = 437;
 };
 
diff --git a/Engine/Plugins/Experimental/Mutable/Source/CustomizableObjectEditor/Private/MuCOE/GenerateMutableSource/GenerateMutableSourceSurface.cpp b/Engine/Plugins/Experimental/Mutable/Source/CustomizableObjectEditor/Private/MuCOE/GenerateMutableSource/GenerateMutableSourceSurface.cpp
--- a/Engine/Plugins/Experimental/Mutable/Source/CustomizableObjectEditor/Private/MuCOE/GenerateMutableSource/GenerateMutableSourceSurface.cpp
+++ b/Engine/Plugins/Experimental/Mutable/Source/CustomizableObjectEditor/Private/MuCOE/GenerateMutableSource/GenerateMutableSourceSurface.cpp
@@ -640,9 +640,25 @@ mu::NodeSurfacePtr GenerateMutableSourceSurface(const UEdGraphPin * Pin, FMutabl
 							const uint32 BaseLODBias = ComputeLODBiasForTexture(GenerationContext, *ReferenceTexture) + SurfaceLODBias;
 							mu::NodeImagePtr LastImage = ResizeTextureByNumMips(ImageNode, BaseLODBias);
 
+							mu::EMipmapFilterType MipGenerationFilterType = Invoke([&]()
+							{
+								if (ReferenceTexture)
+								{
+									switch(ReferenceTexture->MipGenSettings)
+									{
+									case TextureMipGenSettings::TMGS_SimpleAverage: return mu::EMipmapFilterType::SimpleAverage;
+									case TextureMipGenSettings::TMGS_Unfiltered:    return mu::EMipmapFilterType::Unfiltered;
+									default: return mu::EMipmapFilterType::SimpleAverage;
+									}
+								}
+
+								return mu::EMipmapFilterType::SimpleAverage;
+							});
+
+
 							mu::NodeImageMipmapPtr MipmapImage = new mu::NodeImageMipmap();
 							MipmapImage->SetSource(LastImage.get());
-							MipmapImage->SetMipmapGenerationSettings(mu::EMipmapFilterType::MFT_SimpleAverage, mu::EAddressMode::None, 1.0f, false);
+							MipmapImage->SetMipmapGenerationSettings(MipGenerationFilterType, mu::EAddressMode::None);
 
 							MipmapImage->SetMessageContext(Node);
 							LastImage = MipmapImage;
@@ -683,7 +699,7 @@ mu::NodeSurfacePtr GenerateMutableSourceSurface(const UEdGraphPin * Pin, FMutabl
 									mu::NodeImageMipmapPtr NormalCompositeMipmapImage = new mu::NodeImageMipmap();
 									const uint32 MipsToSkip = ComputeLODBiasForTexture(GenerationContext, *ReferenceCompositeNormalTexture, ReferenceTexture);
 									NormalCompositeMipmapImage->SetSource(ResizeTextureByNumMips(CompositeNormalImage, MipsToSkip));
-									NormalCompositeMipmapImage->SetMipmapGenerationSettings(mu::EMipmapFilterType::MFT_SimpleAverage, mu::EAddressMode::None, 1.0f, true);
+									NormalCompositeMipmapImage->SetMipmapGenerationSettings(mu::EMipmapFilterType::SimpleAverage, mu::EAddressMode::None);
 
 									CompositedImage->SetNormal(NormalCompositeMipmapImage);
 								}
diff --git a/Engine/Plugins/Experimental/Mutable/Source/MutableRuntime/Internal/MuR/Operations.h b/Engine/Plugins/Experimental/Mutable/Source/MutableRuntime/Internal/MuR/Operations.h
--- a/Engine/Plugins/Experimental/Mutable/Source/MutableRuntime/Internal/MuR/Operations.h
+++ b/Engine/Plugins/Experimental/Mutable/Source/MutableRuntime/Internal/MuR/Operations.h
@@ -607,10 +607,8 @@ namespace mu
             bool onlyTail;
 
             //! Mipmap generation settings. 
-            float sharpenFactor;
-            EAddressMode addressMode;
-            EMipmapFilterType filterType;
-			bool ditherMipmapAlpha;
+            EMipmapFilterType FilterType;
+            EAddressMode AddressMode;
         };
 
         struct ImageResizeArgs
diff --git a/Engine/Plugins/Experimental/Mutable/Source/MutableRuntime/Private/MuR/CodeRunnerTasked.cpp b/Engine/Plugins/Experimental/Mutable/Source/MutableRuntime/Private/MuR/CodeRunnerTasked.cpp
--- a/Engine/Plugins/Experimental/Mutable/Source/MutableRuntime/Private/MuR/CodeRunnerTasked.cpp
+++ b/Engine/Plugins/Experimental/Mutable/Source/MutableRuntime/Private/MuR/CodeRunnerTasked.cpp
@@ -1360,7 +1360,7 @@ namespace mu
 
 		check(StartLevel >= 0);
 
-		FMipmapGenerationSettings Settings{};
+		FMipmapGenerationSettings Settings{Args.FilterType, Args.AddressMode};
 		FImageOperator ImOp = FImageOperator::GetDefault(ImagePixelFormatFunc);
 		ImOp.ImageMipmap(Scratch, ImageCompressionQuality, Result.get(), Result.get(), StartLevel, Result->GetLODCount(), Settings);
 	}
diff --git a/Engine/Plugins/Experimental/Mutable/Source/MutableRuntime/Private/MuR/ImageTypes.cpp b/Engine/Plugins/Experimental/Mutable/Source/MutableRuntime/Private/MuR/ImageTypes.cpp
--- a/Engine/Plugins/Experimental/Mutable/Source/MutableRuntime/Private/MuR/ImageTypes.cpp
+++ b/Engine/Plugins/Experimental/Mutable/Source/MutableRuntime/Private/MuR/ImageTypes.cpp
@@ -7,6 +7,7 @@ namespace mu
 {
 	MUTABLE_IMPLEMENT_ENUM_SERIALISABLE(EBlendType);
 	MUTABLE_IMPLEMENT_ENUM_SERIALISABLE(EMipmapFilterType);
+	MUTABLE_IMPLEMENT_ENUM_SERIALISABLE(EAddressMode);
 	MUTABLE_IMPLEMENT_ENUM_SERIALISABLE(ECompositeImageMode);
 	MUTABLE_IMPLEMENT_ENUM_SERIALISABLE(ESamplingMethod);
 	MUTABLE_IMPLEMENT_ENUM_SERIALISABLE(EMinFilterMethod);
@@ -63,25 +64,37 @@ namespace mu
         return s_imageFormatData[ uint8(format) ];
     }
 
-	void FMipmapGenerationSettings::Serialise(OutputArchive& arch) const
+	void FMipmapGenerationSettings::Serialise(OutputArchive& Arch) const
 	{
-		uint32 ver = 0;
-		arch << ver;
+		uint32 Version = 1;
+		Arch << Version;
 
-		arch << m_sharpenFactor;
-		arch << m_filterType;
-		arch << m_ditherMipmapAlpha;
+		Arch << FilterType;
+		Arch << AddressMode;
 	}
 
-	void FMipmapGenerationSettings::Unserialise(InputArchive& arch)
+	void FMipmapGenerationSettings::Unserialise(InputArchive& Arch)
 	{
-		uint32 ver = 0;
-		arch >> ver;
-		check(ver == 0);
+		uint32 Version = 0;
+		Arch >> Version;
 
-		arch >> m_sharpenFactor;
-		arch >> m_filterType;
-		arch >> m_ditherMipmapAlpha;
+		check(Version <= 1);
+
+		if (Version < 1)
+		{
+			float SharpenFactor = 0.0f;
+			Arch >> SharpenFactor;
+
+			Arch >> FilterType;
+
+			bool bDitherMipmapAlpha = false;
+			Arch >> bDitherMipmapAlpha;
+		}
+		else
+		{
+			Arch >> FilterType;
+			Arch >> AddressMode;
+		}
 	}
 }
 
diff --git a/Engine/Plugins/Experimental/Mutable/Source/MutableRuntime/Private/MuR/OpImageMipmap.cpp b/Engine/Plugins/Experimental/Mutable/Source/MutableRuntime/Private/MuR/OpImageMipmap.cpp
--- a/Engine/Plugins/Experimental/Mutable/Source/MutableRuntime/Private/MuR/OpImageMipmap.cpp
+++ b/Engine/Plugins/Experimental/Mutable/Source/MutableRuntime/Private/MuR/OpImageMipmap.cpp
@@ -34,11 +34,10 @@ namespace OpImageMipmapInternal
 		{
 			for (int32 X = 0; X < DestSize.X; ++X)
 			{
-				for (int32 C = 0; C < PIXEL_SIZE; ++C)
-				{
-					DestData[(Y * DestSize.X + X) * PIXEL_SIZE + C] =
-						SourceData[((Y << 1) * SourceSize.X + (X << 1)) * PIXEL_SIZE + C];
-				}
+				uint8* DestPixel = DestData + (Y*DestSize.X + X) * PIXEL_SIZE;
+				const uint8* SourcePixel = SourceData + ((Y << 1) * SourceSize.X + (X << 1)) * PIXEL_SIZE;
+
+				FMemory::Memcpy(DestPixel, SourcePixel, PIXEL_SIZE);
 			}
 		}
 	}
@@ -315,7 +314,7 @@ namespace OpImageMipmapInternal
 
 							uint8 const * const SrcRow0 = JobStagingMemoryData + (Row0Offset.Y*BatchDecSizeInPixels.X + Row0Offset.X) * NumChannels;
 
-							if constexpr (Filter == EMipmapFilterType::MFT_SimpleAverage)
+							if constexpr (Filter == EMipmapFilterType::SimpleAverage)
 							{
 								// Use memcpy to avoid any possible but improbable UB. memcpy should be optimized away by the compiler.
 								uint64 Row0Bits; 
@@ -341,13 +340,13 @@ namespace OpImageMipmapInternal
 								
 								FMemory::Memcpy(DestPixel, &Result, NumChannels);
 							}
-							else // constexpr Filter == EMipmapFilterType::MFT_Unfiltered
+							else // constexpr Filter == EMipmapFilterType::Unfiltered
 							{
 								FMemory::Memcpy(DestPixel, SrcRow0, NumChannels);	
 							}
 							static_assert(
-								Filter == EMipmapFilterType::MFT_SimpleAverage || 
-								Filter == EMipmapFilterType::MFT_Unfiltered);
+								Filter == EMipmapFilterType::SimpleAverage || 
+								Filter == EMipmapFilterType::Unfiltered);
 						}
 					}
 				}
@@ -390,9 +389,9 @@ namespace OpImageMipmapInternal
 			return DestBeginLODSize == SrcNextLODSize;
 		}));
 
-		switch (Settings.m_filterType)
+		switch (Settings.FilterType)
 		{
-		case EMipmapFilterType::MFT_SimpleAverage:
+		case EMipmapFilterType::SimpleAverage:
 		{
 			const uint8* SrcData = SourceImage->GetLODData(SrcLOD);
 
@@ -408,7 +407,7 @@ namespace OpImageMipmapInternal
 			}
 			break;
 		}
-		case EMipmapFilterType::MFT_Unfiltered:
+		case EMipmapFilterType::Unfiltered:
 		{
 			const uint8* SrcData = SourceImage->GetLODData(SrcLOD);
 
@@ -462,14 +461,14 @@ namespace OpImageMipmapInternal
 		const EImageFormat SrcFormat = SourceImage->GetFormat();
 		const EImageFormat DestFormat = DestImage->GetFormat();
 
-		switch (Settings.m_filterType)
+		switch (Settings.FilterType)
 		{
-		case EMipmapFilterType::MFT_SimpleAverage:
+		case EMipmapFilterType::SimpleAverage:
 		{
 			const uint8* SrcData = SourceImage->GetLODData(SrcLOD);
 			uint8* DestData = DestImage->GetLODData(DestLODBegin);
 
-			GenerateNextMipBlockCompressed<PixelSize, EMipmapFilterType::MFT_SimpleAverage>(
+			GenerateNextMipBlockCompressed<PixelSize, EMipmapFilterType::SimpleAverage>(
 					SrcData, DestData, SourceSize, SrcFormat, DestFormat);
 
 			SrcData = DestData;
@@ -494,12 +493,12 @@ namespace OpImageMipmapInternal
 
 			break;
 		}
-		case EMipmapFilterType::MFT_Unfiltered:
+		case EMipmapFilterType::Unfiltered:
 		{
 			const uint8* SrcData = SourceImage->GetLODData(SrcLOD);
 			uint8* DestData = DestImage->GetLODData(DestLODBegin);
 			
-			GenerateNextMipBlockCompressed<PixelSize, EMipmapFilterType::MFT_Unfiltered>(
+			GenerateNextMipBlockCompressed<PixelSize, EMipmapFilterType::Unfiltered>(
 					SrcData, DestData, SourceSize, SrcFormat, DestFormat);
 
 			SrcData = DestData;
diff --git a/Engine/Plugins/Experimental/Mutable/Source/MutableRuntime/Public/MuR/ImageTypes.h b/Engine/Plugins/Experimental/Mutable/Source/MutableRuntime/Public/MuR/ImageTypes.h
--- a/Engine/Plugins/Experimental/Mutable/Source/MutableRuntime/Public/MuR/ImageTypes.h
+++ b/Engine/Plugins/Experimental/Mutable/Source/MutableRuntime/Public/MuR/ImageTypes.h
@@ -125,14 +125,13 @@ namespace mu
 		ClampToEdge,
 		ClampToBlack,
 	};
-
+	MUTABLE_DEFINE_ENUM_SERIALISABLE(EAddressMode);
+	
 	enum class EMipmapFilterType
 	{
-		MFT_Unfiltered,
-		MFT_SimpleAverage,
-		MFT_Sharpen,
-		_MFT_COUNT
-	};	
+		Unfiltered,
+		SimpleAverage,
+	};
 	MUTABLE_DEFINE_ENUM_SERIALISABLE(EMipmapFilterType);
 
 	enum class ECompositeImageMode
@@ -224,13 +223,11 @@ namespace mu
 
 	struct MUTABLERUNTIME_API FMipmapGenerationSettings
 	{
-		float m_sharpenFactor = 0.0f;
-		EMipmapFilterType m_filterType = EMipmapFilterType::MFT_SimpleAverage;
-		EAddressMode m_addressMode = EAddressMode::None;
-		bool m_ditherMipmapAlpha = false;
+		EMipmapFilterType FilterType = EMipmapFilterType::SimpleAverage;
+		EAddressMode AddressMode = EAddressMode::None;
 
-		void Serialise(OutputArchive& arch) const;
-		void Unserialise(InputArchive& arch);
+		void Serialise(OutputArchive& Arch) const;
+		void Unserialise(InputArchive& Arch);
 	};
 }
 
diff --git a/Engine/Plugins/Experimental/Mutable/Source/MutableTools/Private/MuT/ASTOpImageMipmap.cpp b/Engine/Plugins/Experimental/Mutable/Source/MutableTools/Private/MuT/ASTOpImageMipmap.cpp
--- a/Engine/Plugins/Experimental/Mutable/Source/MutableTools/Private/MuT/ASTOpImageMipmap.cpp
+++ b/Engine/Plugins/Experimental/Mutable/Source/MutableTools/Private/MuT/ASTOpImageMipmap.cpp
@@ -45,10 +45,8 @@ namespace mu
 				BlockLevels == other->BlockLevels &&
 				bOnlyTail == other->bOnlyTail &&
 				bPreventSplitTail == other->bPreventSplitTail &&
-				SharpenFactor == other->SharpenFactor &&
 				AddressMode == other->AddressMode &&
-				FilterType == other->FilterType &&
-				DitherMipmapAlpha == other->DitherMipmapAlpha;
+				FilterType == other->FilterType;
 		}
 		return false;
 	}
@@ -70,10 +68,8 @@ namespace mu
 		n->BlockLevels = BlockLevels;
 		n->bOnlyTail = bOnlyTail;
 		n->bPreventSplitTail = bPreventSplitTail;
-		n->SharpenFactor = SharpenFactor;
 		n->AddressMode = AddressMode;
 		n->FilterType = FilterType;
-		n->DitherMipmapAlpha = DitherMipmapAlpha;
 
 		return n;
 	}
@@ -96,10 +92,8 @@ namespace mu
 			args.levels = Levels;
 			args.blockLevels = BlockLevels;
 			args.onlyTail = bOnlyTail;
-			args.sharpenFactor = SharpenFactor;
-			args.addressMode = AddressMode;
-			args.filterType = FilterType;
-			args.ditherMipmapAlpha = DitherMipmapAlpha;
+			args.AddressMode = AddressMode;
+			args.FilterType = FilterType;
 			if (Source) args.source = Source->linkedAddress;
 
 			linkedAddress = (OP::ADDRESS)program.m_opAddress.Num();
@@ -490,4 +484,4 @@ namespace mu
 		return newAt;
 	}
 
-}
\ No newline at end of file
+}
diff --git a/Engine/Plugins/Experimental/Mutable/Source/MutableTools/Private/MuT/ASTOpImageMipmap.h b/Engine/Plugins/Experimental/Mutable/Source/MutableTools/Private/MuT/ASTOpImageMipmap.h
--- a/Engine/Plugins/Experimental/Mutable/Source/MutableTools/Private/MuT/ASTOpImageMipmap.h
+++ b/Engine/Plugins/Experimental/Mutable/Source/MutableTools/Private/MuT/ASTOpImageMipmap.h
@@ -32,10 +32,8 @@ namespace mu
 		bool bPreventSplitTail = false;
 
 		//! Mipmap generation settings. 
-		float SharpenFactor = 0.0f;
 		EAddressMode AddressMode = EAddressMode::None;
-		EMipmapFilterType FilterType = EMipmapFilterType::MFT_Unfiltered;
-		bool DitherMipmapAlpha = false;
+		EMipmapFilterType FilterType = EMipmapFilterType::SimpleAverage;
 
 	public:
 
diff --git a/Engine/Plugins/Experimental/Mutable/Source/MutableTools/Private/MuT/CodeGenerator.cpp b/Engine/Plugins/Experimental/Mutable/Source/MutableTools/Private/MuT/CodeGenerator.cpp
--- a/Engine/Plugins/Experimental/Mutable/Source/MutableTools/Private/MuT/CodeGenerator.cpp
+++ b/Engine/Plugins/Experimental/Mutable/Source/MutableTools/Private/MuT/CodeGenerator.cpp
@@ -1203,10 +1203,8 @@ namespace mu
 							op->Source = imageAd;
 							op->BlockLevels = 0;
 
-							op->AddressMode = mipmapNode->GetPrivate()->m_settings.m_addressMode;
-							op->FilterType = mipmapNode->GetPrivate()->m_settings.m_filterType;
-							op->SharpenFactor = mipmapNode->GetPrivate()->m_settings.m_sharpenFactor;
-							op->DitherMipmapAlpha = mipmapNode->GetPrivate()->m_settings.m_ditherMipmapAlpha;
+							op->AddressMode = mipmapNode->GetPrivate()->m_settings.AddressMode;
+							op->FilterType = mipmapNode->GetPrivate()->m_settings.FilterType;
 							imageAd = op;
 						}
 
@@ -1540,10 +1538,8 @@ namespace mu
 									mop->bOnlyTail = true;
 								}
 
-								mop->AddressMode = mipmapNode->GetPrivate()->m_settings.m_addressMode;
-								mop->FilterType = mipmapNode->GetPrivate()->m_settings.m_filterType;
-								mop->SharpenFactor = mipmapNode->GetPrivate()->m_settings.m_sharpenFactor;
-								mop->DitherMipmapAlpha = mipmapNode->GetPrivate()->m_settings.m_ditherMipmapAlpha;
+								mop->AddressMode = mipmapNode->GetPrivate()->m_settings.AddressMode;
+								mop->FilterType = mipmapNode->GetPrivate()->m_settings.FilterType;
 
 								imageAd = mop;
 							}
@@ -1576,9 +1572,7 @@ namespace mu
 
 								// Not important for the end of the mip tail?
 								mop->AddressMode = EAddressMode::ClampToEdge;
-								mop->FilterType = EMipmapFilterType::MFT_SimpleAverage;
-								mop->SharpenFactor = 0;
-								mop->DitherMipmapAlpha = false;
+								mop->FilterType = EMipmapFilterType::SimpleAverage;
 
 								imageAd = mop;
 							}
diff --git a/Engine/Plugins/Experimental/Mutable/Source/MutableTools/Private/MuT/CodeGenerator_Image.cpp b/Engine/Plugins/Experimental/Mutable/Source/MutableTools/Private/MuT/CodeGenerator_Image.cpp
--- a/Engine/Plugins/Experimental/Mutable/Source/MutableTools/Private/MuT/CodeGenerator_Image.cpp
+++ b/Engine/Plugins/Experimental/Mutable/Source/MutableTools/Private/MuT/CodeGenerator_Image.cpp
@@ -1731,10 +1731,8 @@ namespace mu
             op->BlockLevels = 0;
         }
 
-		op->AddressMode = node.m_settings.m_addressMode;
-		op->FilterType = node.m_settings.m_filterType;
-		op->SharpenFactor = node.m_settings.m_sharpenFactor;
-		op->DitherMipmapAlpha = node.m_settings.m_ditherMipmapAlpha;
+		op->AddressMode = node.m_settings.AddressMode;
+		op->FilterType = node.m_settings.FilterType;
 
         res = op;
 
diff --git a/Engine/Plugins/Experimental/Mutable/Source/MutableTools/Private/MuT/NodeImageMipmap.cpp b/Engine/Plugins/Experimental/Mutable/Source/MutableTools/Private/MuT/NodeImageMipmap.cpp
--- a/Engine/Plugins/Experimental/Mutable/Source/MutableTools/Private/MuT/NodeImageMipmap.cpp
+++ b/Engine/Plugins/Experimental/Mutable/Source/MutableTools/Private/MuT/NodeImageMipmap.cpp
@@ -42,12 +42,9 @@ namespace mu
 	}
 
 	//---------------------------------------------------------------------------------------------
-	void NodeImageMipmap::SetMipmapGenerationSettings( EMipmapFilterType filterType,
-													   EAddressMode addressMode,
-													   float sharpenFactor,
-													   bool mipDitherAlpha )
+	void NodeImageMipmap::SetMipmapGenerationSettings(EMipmapFilterType FilterType, EAddressMode AddressMode)
 	{
-		m_pD->m_settings = FMipmapGenerationSettings{ sharpenFactor, filterType, addressMode, mipDitherAlpha };
+		m_pD->m_settings = FMipmapGenerationSettings{ FilterType, AddressMode };
 	}
 
 }
diff --git a/Engine/Plugins/Experimental/Mutable/Source/MutableTools/Public/MuT/NodeImageMipmap.h b/Engine/Plugins/Experimental/Mutable/Source/MutableTools/Public/MuT/NodeImageMipmap.h
--- a/Engine/Plugins/Experimental/Mutable/Source/MutableTools/Public/MuT/NodeImageMipmap.h
+++ b/Engine/Plugins/Experimental/Mutable/Source/MutableTools/Public/MuT/NodeImageMipmap.h
@@ -54,10 +54,7 @@ namespace mu
 		NodeImagePtr GetSource() const;
 		void SetSource( NodeImagePtr );
 
-		void SetMipmapGenerationSettings( EMipmapFilterType filterType,
-										  EAddressMode addressMode,
-										  float sharpenFactor,  
-										  bool bDitherMipAlpha );
+		void SetMipmapGenerationSettings(EMipmapFilterType FilterType, EAddressMode AddressMode);
 
 		//-----------------------------------------------------------------------------------------
 		// Interface pattern
