[mutable] Added field to report if the level began play for the Image Updates
- It will save the state of the level when running the BuildMaterials method, so it may happen that the mip update finished after the level has loaded (while marking it as it did not).
- We do require this to be able to better parse data in MongoDB and have more consistency between image and mesh tests.
- Updated mongodb image update document.

diff --git a/Engine/Plugins/Experimental/Mutable/Source/CustomizableObject/Internal/MuCO/CustomizableObjectMipDataProvider.h b/Engine/Plugins/Experimental/Mutable/Source/CustomizableObject/Internal/MuCO/CustomizableObjectMipDataProvider.h
--- a/Engine/Plugins/Experimental/Mutable/Source/CustomizableObject/Internal/MuCO/CustomizableObjectMipDataProvider.h
+++ b/Engine/Plugins/Experimental/Mutable/Source/CustomizableObject/Internal/MuCO/CustomizableObjectMipDataProvider.h
@@ -50,7 +50,9 @@ public:
 
 	const TArray<mu::Ptr<const mu::Image>>& GetImageParameterValues() const;
 
+	// Benchmarking Utility data (it may not always be present)
 	FString CapturedDescriptor;
+	bool bLevelBegunPlay = false;
 	
 private:
 	FString CustomizableObjectPathName;
diff --git a/Engine/Plugins/Experimental/Mutable/Source/CustomizableObject/Internal/MuCO/LogBenchmarkUtil.h b/Engine/Plugins/Experimental/Mutable/Source/CustomizableObject/Internal/MuCO/LogBenchmarkUtil.h
--- a/Engine/Plugins/Experimental/Mutable/Source/CustomizableObject/Internal/MuCO/LogBenchmarkUtil.h
+++ b/Engine/Plugins/Experimental/Mutable/Source/CustomizableObject/Internal/MuCO/LogBenchmarkUtil.h
@@ -88,7 +88,7 @@ public:
 	/** Gathers update stats when it has finished. */
 	void FinishUpdateMesh(const TSharedRef<FUpdateContextPrivate>& Context);
 
-	void FinishUpdateImage(const FString& CustomizableObjectPathName, const FString& InstancePathName, const FString& InstanceDescriptor, double TaskUpdateImageTime, const int64 TaskUpdateImageMemoryPeak, const int64 TaskUpdateImageRealMemoryPeak) const;
+	void FinishUpdateImage(const FString& CustomizableObjectPathName, const FString& InstancePathName, const FString& InstanceDescriptor, const bool bDidLevelBeginPlay, double TaskUpdateImageTime, const int64 TaskUpdateImageMemoryPeak, const int64 TaskUpdateImageRealMemoryPeak) const;
 	
 	DECLARE_BENCHMARK_STAT(NumAllocatedTextures, uint32);
 	DECLARE_BENCHMARK_STAT(TextureGPUSize, uint64);
diff --git a/Engine/Plugins/Experimental/Mutable/Source/CustomizableObject/Private/MuCO/BusyWaits_Deprecated.cpp b/Engine/Plugins/Experimental/Mutable/Source/CustomizableObject/Private/MuCO/BusyWaits_Deprecated.cpp
--- a/Engine/Plugins/Experimental/Mutable/Source/CustomizableObject/Private/MuCO/BusyWaits_Deprecated.cpp
+++ b/Engine/Plugins/Experimental/Mutable/Source/CustomizableObject/Private/MuCO/BusyWaits_Deprecated.cpp
@@ -557,9 +557,10 @@ namespace CustomizableObjectMipDataProvider::ImplDeprecated
 			const FString& InstancePathName = OperationData->UpdateContext->GetInstancePathName();
 
 			const FString& Descriptor = OperationData->UpdateContext->CapturedDescriptor;
+			const bool bDidLevelBeginPlay = OperationData->UpdateContext->bLevelBegunPlay;
 
 			FFunctionGraphTask::CreateAndDispatchWhenReady(
-			[CustomizableObjectPathName, InstancePathName, Time, PeakMemory, RealMemoryPeak, Descriptor]()
+			[CustomizableObjectPathName, InstancePathName, bDidLevelBeginPlay, Time, PeakMemory, RealMemoryPeak, Descriptor]()
 			{
 				if (!UCustomizableObjectSystem::IsCreated()) // We are shutting down
 				{
@@ -572,7 +573,7 @@ namespace CustomizableObjectMipDataProvider::ImplDeprecated
 					return;
 				}
 
-				System->GetPrivate()->LogBenchmarkUtil.FinishUpdateImage(CustomizableObjectPathName, InstancePathName, Descriptor, Time, PeakMemory, RealMemoryPeak);
+				System->GetPrivate()->LogBenchmarkUtil.FinishUpdateImage(CustomizableObjectPathName, InstancePathName, Descriptor, bDidLevelBeginPlay, Time, PeakMemory, RealMemoryPeak);
 			},
 			TStatId{},
 			nullptr,
diff --git a/Engine/Plugins/Experimental/Mutable/Source/CustomizableObject/Private/MuCO/CustomizableObjectInstance.cpp b/Engine/Plugins/Experimental/Mutable/Source/CustomizableObject/Private/MuCO/CustomizableObjectInstance.cpp
--- a/Engine/Plugins/Experimental/Mutable/Source/CustomizableObject/Private/MuCO/CustomizableObjectInstance.cpp
+++ b/Engine/Plugins/Experimental/Mutable/Source/CustomizableObject/Private/MuCO/CustomizableObjectInstance.cpp
@@ -5599,6 +5599,10 @@ void UCustomizableInstancePrivate::BuildMaterials(const TSharedRef<FUpdateContex
 	if (CVarEnableBenchmark.GetValueOnAnyThread())
 	{
 		UpdateContext->CapturedDescriptor = OperationData->GetCapturedDescriptor().ToString();
+		if (GWorld)
+		{
+			UpdateContext->bLevelBegunPlay = GWorld->GetBegunPlay();
+		}
 	}
 	
 	const bool bReuseTextures = OperationData->bReuseInstanceTextures;
diff --git a/Engine/Plugins/Experimental/Mutable/Source/CustomizableObject/Private/MuCO/CustomizableObjectMipDataProvider.cpp b/Engine/Plugins/Experimental/Mutable/Source/CustomizableObject/Private/MuCO/CustomizableObjectMipDataProvider.cpp
--- a/Engine/Plugins/Experimental/Mutable/Source/CustomizableObject/Private/MuCO/CustomizableObjectMipDataProvider.cpp
+++ b/Engine/Plugins/Experimental/Mutable/Source/CustomizableObject/Private/MuCO/CustomizableObjectMipDataProvider.cpp
@@ -283,9 +283,10 @@ namespace Impl
 						const FString& InstancePathName = OperationData->UpdateContext->GetInstancePathName();
 
 						const FString& Descriptor = OperationData->UpdateContext->CapturedDescriptor;
+						const bool bDidLevelBeginPlay = OperationData->UpdateContext->bLevelBegunPlay;
 
 						ExecuteOnGameThread(UE_SOURCE_LOCATION, 
-						[CustomizableObjectPathName, InstancePathName, Time, PeakMemory, RealMemoryPeak, Descriptor ]
+						[CustomizableObjectPathName, InstancePathName, bDidLevelBeginPlay, Time, PeakMemory, RealMemoryPeak, Descriptor ]
 						{
 							if (!UCustomizableObjectSystem::IsCreated()) // We are shutting down
 							{
@@ -298,7 +299,7 @@ namespace Impl
 								return;
 							}
 
-							System->GetPrivate()->LogBenchmarkUtil.FinishUpdateImage(CustomizableObjectPathName, InstancePathName, Descriptor, Time, PeakMemory, RealMemoryPeak);
+							System->GetPrivate()->LogBenchmarkUtil.FinishUpdateImage(CustomizableObjectPathName, InstancePathName, Descriptor, bDidLevelBeginPlay, Time, PeakMemory, RealMemoryPeak);
 						});
 					}
 
diff --git a/Engine/Plugins/Experimental/Mutable/Source/CustomizableObject/Private/MuCO/LogBenchmarkUtil.cpp b/Engine/Plugins/Experimental/Mutable/Source/CustomizableObject/Private/MuCO/LogBenchmarkUtil.cpp
--- a/Engine/Plugins/Experimental/Mutable/Source/CustomizableObject/Private/MuCO/LogBenchmarkUtil.cpp
+++ b/Engine/Plugins/Experimental/Mutable/Source/CustomizableObject/Private/MuCO/LogBenchmarkUtil.cpp
@@ -337,7 +337,7 @@ void FLogBenchmarkUtil::FinishUpdateMesh(const TSharedRef<FUpdateContextPrivate>
 }
 
 
-void FLogBenchmarkUtil::FinishUpdateImage(const FString& CustomizableObjectPathName, const FString& InstancePathName, const FString& InstanceDescriptor, const double TaskUpdateImageTime, const int64 TaskUpdateImageMemoryPeak, const int64 TaskUpdateImageRealMemoryPeak) const
+void FLogBenchmarkUtil::FinishUpdateImage(const FString& CustomizableObjectPathName, const FString& InstancePathName, const FString& InstanceDescriptor, const bool bDidLevelBeginPlay, const double TaskUpdateImageTime, const int64 TaskUpdateImageMemoryPeak, const int64 TaskUpdateImageRealMemoryPeak) const
 {
 	check(IsInGameThread());
 
@@ -349,12 +349,13 @@ void FLogBenchmarkUtil::FinishUpdateImage(const FString& CustomizableObjectPathN
 	const FString& ID_CO = CustomizableObjectPathName;
 	const FString& ID_COI = InstancePathName;
 	const FString ID_Descriptor = InstanceDescriptor;
+	const FString Context_LevelBegunPlay =  bDidLevelBeginPlay ? TEXT("true") : TEXT("false");
 	const FString ID_UpdateType = TEXT("Image");
 	const double Time_TaskUpdateImage = TaskUpdateImageTime * 1000;
 	const double Memory_TaskUpdateImagePeakMB = (TaskUpdateImageMemoryPeak / 1024.0) / 1024.0;
 	const double Memory_TaskUpdateImageRealPeakMB = (TaskUpdateImageRealMemoryPeak / 1024.0) / 1024.0;
 
-	const FString UpdateString = FString::Printf(TEXT("%s;%s;%s;%s;;;;;;;;;;;;%f;%f;%f"), *ID_CO, *ID_COI, *ID_UpdateType, *ID_Descriptor, Time_TaskUpdateImage, Memory_TaskUpdateImagePeakMB,Memory_TaskUpdateImageRealPeakMB);
+	const FString UpdateString = FString::Printf(TEXT("%s;%s;%s;%s;;%s;;;;;;;;;;%f;%f;%f"), *ID_CO, *ID_COI, *ID_UpdateType, *ID_Descriptor, *Context_LevelBegunPlay, Time_TaskUpdateImage, Memory_TaskUpdateImagePeakMB,Memory_TaskUpdateImageRealPeakMB);
 	LogBenchmarkUtil::Write(*Archive, UpdateString);
 	Archive->Flush();	
 }
