Added missing includes.

diff --git a/Engine/Plugins/Enterprise/VariantManager/Source/VariantManager/Private/DisplayNodes/VariantManagerStringPropertyNode.cpp b/Engine/Plugins/Enterprise/VariantManager/Source/VariantManager/Private/DisplayNodes/VariantManagerStringPropertyNode.cpp
--- a/Engine/Plugins/Enterprise/VariantManager/Source/VariantManager/Private/DisplayNodes/VariantManagerStringPropertyNode.cpp
+++ b/Engine/Plugins/Enterprise/VariantManager/Source/VariantManager/Private/DisplayNodes/VariantManagerStringPropertyNode.cpp
@@ -10,6 +10,7 @@
 #include "VariantManagerLog.h"
 #include "PropertyHandle.h"
 #include "ISinglePropertyView.h"
+#include "UObject/StrProperty.h"
 #include "UObject/TextProperty.h"
 #include "Widgets/Layout/SBox.h"
 #include "Widgets/Text/STextBlock.h"
diff --git a/Engine/Plugins/Experimental/InstanceDataObjectFixupTool/Source/Private/InstanceDataObjectFixupPanel.cpp b/Engine/Plugins/Experimental/InstanceDataObjectFixupTool/Source/Private/InstanceDataObjectFixupPanel.cpp
--- a/Engine/Plugins/Experimental/InstanceDataObjectFixupTool/Source/Private/InstanceDataObjectFixupPanel.cpp
+++ b/Engine/Plugins/Experimental/InstanceDataObjectFixupTool/Source/Private/InstanceDataObjectFixupPanel.cpp
@@ -13,6 +13,7 @@
 
 #include "UObject/OverriddenPropertySet.h"
 #include "UObject/OverridableManager.h"
+#include "UObject/StrProperty.h"
 #include "UObject/TextProperty.h"
 
 #define LOCTEXT_NAMESPACE "InstanceDataObjectFixupPanel"
diff --git a/Engine/Plugins/Experimental/Mutable/Source/CustomizableObjectEditor/Private/MuCOE/CustomizableObjectEditor.cpp b/Engine/Plugins/Experimental/Mutable/Source/CustomizableObjectEditor/Private/MuCOE/CustomizableObjectEditor.cpp
--- a/Engine/Plugins/Experimental/Mutable/Source/CustomizableObjectEditor/Private/MuCOE/CustomizableObjectEditor.cpp
+++ b/Engine/Plugins/Experimental/Mutable/Source/CustomizableObjectEditor/Private/MuCOE/CustomizableObjectEditor.cpp
@@ -54,6 +54,7 @@
 #include "MuCO/CustomizableObjectInstancePrivate.h"
 #include "MuCO/CustomizableObjectPrivate.h"
 #include "UObject/EnumProperty.h"
+#include "UObject/StrProperty.h"
 #include "Widgets/Docking/SDockTab.h"
 #include "Widgets/Input/SSearchBox.h"
 #include "Widgets/Input/STextComboBox.h"
diff --git a/Engine/Plugins/Interchange/Editor/Source/Pipelines/Private/InterchangeEditorPipelineDetails.cpp b/Engine/Plugins/Interchange/Editor/Source/Pipelines/Private/InterchangeEditorPipelineDetails.cpp
--- a/Engine/Plugins/Interchange/Editor/Source/Pipelines/Private/InterchangeEditorPipelineDetails.cpp
+++ b/Engine/Plugins/Interchange/Editor/Source/Pipelines/Private/InterchangeEditorPipelineDetails.cpp
@@ -9,6 +9,7 @@
 #include "Nodes/InterchangeBaseNode.h"
 #include "ScopedTransaction.h"
 #include "Styling/StyleColors.h"
+#include "UObject/StrProperty.h"
 #include "Widgets/Input/SButton.h"
 #include "Widgets/Input/SNumericEntryBox.h"
 #include "Widgets/Input/STextComboBox.h"
diff --git a/Engine/Plugins/MovieScene/MovieRenderPipeline/Source/MovieRenderPipelineCore/Private/Graph/MovieGraphNode.cpp b/Engine/Plugins/MovieScene/MovieRenderPipeline/Source/MovieRenderPipelineCore/Private/Graph/MovieGraphNode.cpp
--- a/Engine/Plugins/MovieScene/MovieRenderPipeline/Source/MovieRenderPipelineCore/Private/Graph/MovieGraphNode.cpp
+++ b/Engine/Plugins/MovieScene/MovieRenderPipeline/Source/MovieRenderPipelineCore/Private/Graph/MovieGraphNode.cpp
@@ -3,6 +3,7 @@
 #include "Graph/MovieGraphNode.h"
 
 #include "Graph/MovieGraphConfig.h"
+#include "UObject/StrProperty.h"
 #include "UObject/TextProperty.h"
 
 FName UMovieGraphNode::GlobalsPinName = FName("Globals");
diff --git a/Engine/Plugins/NetcodeUnitTest/NetcodeUnitTest/Source/NetcodeUnitTest/Private/NUTUtilReflection.cpp b/Engine/Plugins/NetcodeUnitTest/NetcodeUnitTest/Source/NetcodeUnitTest/Private/NUTUtilReflection.cpp
--- a/Engine/Plugins/NetcodeUnitTest/NetcodeUnitTest/Source/NetcodeUnitTest/Private/NUTUtilReflection.cpp
+++ b/Engine/Plugins/NetcodeUnitTest/NetcodeUnitTest/Source/NetcodeUnitTest/Private/NUTUtilReflection.cpp
@@ -5,6 +5,7 @@
 #include "Misc/TVariantMeta.h"
 #include "UObject/EnumProperty.h"
 #include "Templates/ValueOrError.h"
+#include "UObject/StrProperty.h"
 #include "UObject/TextProperty.h"
 #include "NetcodeUnitTest.h"
 
diff --git a/Engine/Plugins/Runtime/Metasound/Source/MetasoundFrontend/Private/MetasoundJsonBackend.h b/Engine/Plugins/Runtime/Metasound/Source/MetasoundFrontend/Private/MetasoundJsonBackend.h
--- a/Engine/Plugins/Runtime/Metasound/Source/MetasoundFrontend/Private/MetasoundJsonBackend.h
+++ b/Engine/Plugins/Runtime/Metasound/Source/MetasoundFrontend/Private/MetasoundJsonBackend.h
@@ -14,6 +14,7 @@
 #include "UObject/Class.h"
 #include "UObject/UnrealType.h"
 #include "UObject/EnumProperty.h"
+#include "UObject/StrProperty.h"
 #include "UObject/TextProperty.h"
 
 // This file is largely copied from JsonStructSerializerBackend / JsonStructDeserializerBackend.
diff --git a/Engine/Plugins/VirtualProduction/LevelSnapshots/Source/nDisplaySupport/Private/Helpers/MapSubobjectSerializer.h b/Engine/Plugins/VirtualProduction/LevelSnapshots/Source/nDisplaySupport/Private/Helpers/MapSubobjectSerializer.h
--- a/Engine/Plugins/VirtualProduction/LevelSnapshots/Source/nDisplaySupport/Private/Helpers/MapSubobjectSerializer.h
+++ b/Engine/Plugins/VirtualProduction/LevelSnapshots/Source/nDisplaySupport/Private/Helpers/MapSubobjectSerializer.h
@@ -6,6 +6,7 @@
 #include "LevelSnapshotsLog.h"
 #include "Interfaces/ICustomObjectSnapshotSerializer.h"
 #include "Params/ObjectSnapshotSerializationData.h"
+#include "UObject/StrProperty.h"
 #include "UObject/UnrealType.h"
 
 class FScriptMapHelper;
diff --git a/Engine/Source/Editor/EditorConfig/Private/EditorConfig.cpp b/Engine/Source/Editor/EditorConfig/Private/EditorConfig.cpp
--- a/Engine/Source/Editor/EditorConfig/Private/EditorConfig.cpp
+++ b/Engine/Source/Editor/EditorConfig/Private/EditorConfig.cpp
@@ -16,6 +16,7 @@
 #include "UObject/Field.h"
 #include "UObject/NameTypes.h"
 #include "UObject/PropertyPortFlags.h"
+#include "UObject/StrProperty.h"
 #include "UObject/TextProperty.h"
 #include "UObject/UnrealType.h"
 
diff --git a/Engine/Source/Editor/Kismet/Private/BlueprintDetailsCustomization.cpp b/Engine/Source/Editor/Kismet/Private/BlueprintDetailsCustomization.cpp
--- a/Engine/Source/Editor/Kismet/Private/BlueprintDetailsCustomization.cpp
+++ b/Engine/Source/Editor/Kismet/Private/BlueprintDetailsCustomization.cpp
@@ -136,6 +136,7 @@
 #include "UObject/Package.h"
 #include "UObject/ReflectedTypeAccessors.h"
 #include "UObject/SoftObjectPath.h"
+#include "UObject/StrProperty.h"
 #include "UObject/StructOnScope.h"
 #include "UObject/TextProperty.h"
 #include "UObject/TopLevelAssetPath.h"
diff --git a/Engine/Source/Editor/PropertyEditor/Private/Presentation/PropertyTable/PropertyTableColumn.cpp b/Engine/Source/Editor/PropertyEditor/Private/Presentation/PropertyTable/PropertyTableColumn.cpp
--- a/Engine/Source/Editor/PropertyEditor/Private/Presentation/PropertyTable/PropertyTableColumn.cpp
+++ b/Engine/Source/Editor/PropertyEditor/Private/Presentation/PropertyTable/PropertyTableColumn.cpp
@@ -3,6 +3,7 @@
 #include "Presentation/PropertyTable/PropertyTableColumn.h"
 #include "Editor/EditorEngine.h"
 #include "IPropertyTableCell.h"
+#include "UObject/StrProperty.h"
 #include "UObject/TextProperty.h"
 #include "ObjectPropertyNode.h"
 #include "Presentation/PropertyTable/PropertyTableCell.h"
diff --git a/Engine/Source/Editor/PropertyEditor/Private/PropertyEditorHelpers.cpp b/Engine/Source/Editor/PropertyEditor/Private/PropertyEditorHelpers.cpp
--- a/Engine/Source/Editor/PropertyEditor/Private/PropertyEditorHelpers.cpp
+++ b/Engine/Source/Editor/PropertyEditor/Private/PropertyEditorHelpers.cpp
@@ -30,6 +30,7 @@
 #include "Engine/Selection.h"
 
 #include "UObject/PropertyOptional.h"
+#include "UObject/StrProperty.h"
 
 #define LOCTEXT_NAMESPACE "PropertyEditor"
 
diff --git a/Engine/Source/Runtime/JsonUtilities/Private/JsonObjectConverter.cpp b/Engine/Source/Runtime/JsonUtilities/Private/JsonObjectConverter.cpp
--- a/Engine/Source/Runtime/JsonUtilities/Private/JsonObjectConverter.cpp
+++ b/Engine/Source/Runtime/JsonUtilities/Private/JsonObjectConverter.cpp
@@ -7,6 +7,7 @@
 #include "UObject/Class.h"
 #include "UObject/UnrealType.h"
 #include "UObject/EnumProperty.h"
+#include "UObject/StrProperty.h"
 #include "UObject/TextProperty.h"
 #include "UObject/PropertyPortFlags.h"
 #include "UObject/Package.h"
diff --git a/Engine/Source/Runtime/Serialization/Private/Backends/CborStructDeserializerBackend.cpp b/Engine/Source/Runtime/Serialization/Private/Backends/CborStructDeserializerBackend.cpp
--- a/Engine/Source/Runtime/Serialization/Private/Backends/CborStructDeserializerBackend.cpp
+++ b/Engine/Source/Runtime/Serialization/Private/Backends/CborStructDeserializerBackend.cpp
@@ -6,6 +6,7 @@
 #include "StructSerializationUtilities.h"
 #include "UObject/Class.h"
 #include "UObject/EnumProperty.h"
+#include "UObject/StrProperty.h"
 #include "UObject/TextProperty.h"
 #include "UObject/UnrealType.h"
 
diff --git a/Engine/Source/Runtime/Serialization/Private/Backends/CborStructSerializerBackend.cpp b/Engine/Source/Runtime/Serialization/Private/Backends/CborStructSerializerBackend.cpp
--- a/Engine/Source/Runtime/Serialization/Private/Backends/CborStructSerializerBackend.cpp
+++ b/Engine/Source/Runtime/Serialization/Private/Backends/CborStructSerializerBackend.cpp
@@ -5,6 +5,7 @@
 #include "StructSerializationUtilities.h"
 #include "UObject/EnumProperty.h"
 #include "UObject/PropertyPortFlags.h"
+#include "UObject/StrProperty.h"
 #include "UObject/TextProperty.h"
 #include "UObject/UnrealType.h"
 
diff --git a/Engine/Source/Runtime/Serialization/Private/Backends/JsonStructDeserializerBackend.cpp b/Engine/Source/Runtime/Serialization/Private/Backends/JsonStructDeserializerBackend.cpp
--- a/Engine/Source/Runtime/Serialization/Private/Backends/JsonStructDeserializerBackend.cpp
+++ b/Engine/Source/Runtime/Serialization/Private/Backends/JsonStructDeserializerBackend.cpp
@@ -5,6 +5,7 @@
 #include "UObject/Class.h"
 #include "UObject/UnrealType.h"
 #include "UObject/EnumProperty.h"
+#include "UObject/StrProperty.h"
 #include "UObject/TextProperty.h"
 
 /* IStructDeserializerBackend interface
diff --git a/Engine/Source/Runtime/Serialization/Private/Backends/JsonStructSerializerBackend.cpp b/Engine/Source/Runtime/Serialization/Private/Backends/JsonStructSerializerBackend.cpp
--- a/Engine/Source/Runtime/Serialization/Private/Backends/JsonStructSerializerBackend.cpp
+++ b/Engine/Source/Runtime/Serialization/Private/Backends/JsonStructSerializerBackend.cpp
@@ -3,6 +3,7 @@
 #include "Backends/JsonStructSerializerBackend.h"
 #include "UObject/UnrealType.h"
 #include "UObject/EnumProperty.h"
+#include "UObject/StrProperty.h"
 #include "UObject/TextProperty.h"
 #include "UObject/PropertyPortFlags.h"
 
diff --git a/Engine/Source/Runtime/WebBrowser/Private/CEF/CEFJSStructSerializerBackend.cpp b/Engine/Source/Runtime/WebBrowser/Private/CEF/CEFJSStructSerializerBackend.cpp
--- a/Engine/Source/Runtime/WebBrowser/Private/CEF/CEFJSStructSerializerBackend.cpp
+++ b/Engine/Source/Runtime/WebBrowser/Private/CEF/CEFJSStructSerializerBackend.cpp
@@ -4,6 +4,7 @@
 #if WITH_CEF3
 
 #include "UObject/EnumProperty.h"
+#include "UObject/StrProperty.h"
 #include "UObject/TextProperty.h"
 #include "UObject/PropertyPortFlags.h"
 #include "Misc/CommandLine.h"
