[Mutable] Added a configuration window on create a CO asset.

#rnx

diff --git a/Engine/Plugins/Experimental/Mutable/Source/CustomizableObjectEditor/Classes/MuCOE/CustomizableObjectFactory.h b/Engine/Plugins/Experimental/Mutable/Source/CustomizableObjectEditor/Classes/MuCOE/CustomizableObjectFactory.h
--- a/Engine/Plugins/Experimental/Mutable/Source/CustomizableObjectEditor/Classes/MuCOE/CustomizableObjectFactory.h
+++ b/Engine/Plugins/Experimental/Mutable/Source/CustomizableObjectEditor/Classes/MuCOE/CustomizableObjectFactory.h
@@ -6,9 +6,45 @@
 
 #include "CustomizableObjectFactory.generated.h"
 
+enum class ECheckBoxState : uint8;
+
+namespace ESelectInfo { enum Type : int; }
+namespace ETextCommit { enum Type : int; }
+
 class FFeedbackContext;
+class FReply;
+class STextComboBox;
+class SWindow;
 class UClass;
+class UCustomizableObject;
 class UObject;
+class USkeletalMesh;
+
+struct EVisibility;
+
+
+struct FCustomizableObjectOptions
+{
+	// Determines if the CO asset to create has not been configurated
+	bool bEmptyObject = true;
+
+	// Child Object Settings
+	// Determines if the CO asset to create is child of another CO
+	bool bIsChildObject = false;
+
+	// Parent of the CO
+	TWeakObjectPtr<UCustomizableObject> ParentObject = nullptr;
+
+	// Name of the group node of the parent where this CO asset will be linked
+	FString GroupNodeName;
+
+	// Non-Child Object Settings
+	// Number of components
+	int32 NumMeshComponents = 1;
+
+	// Array with the skeletal meshes of each component
+	TArray<TWeakObjectPtr<USkeletalMesh>> ReferenceSkeletalMeshes;
+};
 
 
 UCLASS(MinimalAPI)
@@ -25,4 +61,70 @@ class UCustomizableObjectFactory : public UFactory
 	virtual bool DoesSupportClass(UClass * Class) override;
 	virtual UClass* ResolveSupportedClass() override;
 	virtual UObject* FactoryCreateNew(UClass* Class,UObject* InParent,FName Name,EObjectFlags Flags,UObject* Context,FFeedbackContext* Warn) override;
+
+	FCustomizableObjectOptions CreationSettings;
+};
+
+
+class FCustomizableObjectFactoryUI : public TSharedFromThis<FCustomizableObjectFactoryUI>
+{
+public:
+
+	// Generates the Window widget and its content
+	const FCustomizableObjectOptions ConstructFactoryUI();
+
+	// Callback that generates a CO asset using the options set
+	FReply OnCreate(bool bIsEmpty = true);
+
+	// Creates a the tooltip of the OK button
+	FText GetOKButtonTooltip() const;
+
+	// Callback that cancels the generation of the CO asset
+	FReply OnCancel();
+
+	// Enables or disables the OK button in function of the options set.
+	bool IsConfigurationValid() const;
+
+	// Returns true if the CO asset can be created when clossing the window
+	bool CanCreateObject() const;
+
+	// Callback to set if the CO is a child object
+	void OnCheckBoxChanged(ECheckBoxState State);
+
+	// Parent Selecto Widget methods
+	EVisibility GetParentWidgetsVisibility() const;
+	FString GetSelectedCustomizableObjectPath() const;
+	void OnPickedCustomizableObjectParent(const FAssetData& SelectedAsset);
+	void OnSelectGroupComboBox(TSharedPtr<FString> Selection, ESelectInfo::Type SelectInfo);
+	void GenerateGroupOptions();
+	bool IsNodeGroupSelectorEnabled() const;
+
+	// Component Selectior Widget methods
+	EVisibility GetComponentWidgetsVisibility() const;
+	FText GetMeshSelectorWidgetText() const;
+	FString GetSelectedComponentSkeletalMeshPath() const;
+	int32 GetNumComponents() const;
+	void OnNumComponentsChanged(int32 Value, ETextCommit::Type CommitInfo);
+	void OnSelectComponentComboBox(TSharedPtr<FString> Selection, ESelectInfo::Type SelectInfo);
+	void OnPickedComponentSkeletalMesh(const FAssetData& SelectedAsset);
+	void GenerateComponentOptions();
+	bool IsComponentSelectorEnabled() const;
+
+private:
+
+	// Pointer to the created widgets
+	TSharedPtr<SWindow> COSettingsWindow;
+	TSharedPtr<STextComboBox> GroupSelector;
+	TSharedPtr<STextComboBox> ComponentSelector;
+
+	// bool that determines if a CO asset can be created
+	bool bCreateObject = false;
+
+	// Setting options to generate the CO asset
+	FCustomizableObjectOptions Options;
+
+	// Combobox options
+	TArray<TSharedPtr<FString>> GroupOptions;
+	TArray<TSharedPtr<FString>> ComponentsOptions;
+
 };
diff --git a/Engine/Plugins/Experimental/Mutable/Source/CustomizableObjectEditor/Classes/MuCOE/CustomizableObjectGraph.h b/Engine/Plugins/Experimental/Mutable/Source/CustomizableObjectEditor/Classes/MuCOE/CustomizableObjectGraph.h
--- a/Engine/Plugins/Experimental/Mutable/Source/CustomizableObjectEditor/Classes/MuCOE/CustomizableObjectGraph.h
+++ b/Engine/Plugins/Experimental/Mutable/Source/CustomizableObjectEditor/Classes/MuCOE/CustomizableObjectGraph.h
@@ -28,6 +28,9 @@ public:
 
 	void PostDuplicate(bool bDuplicateForPIE) override;
 
+	/** Adds the necessary nodes for a CO to work */
+	void AddEssentialGraphNodes();
+
 private:
 
 	// Request Node Id Update Map
diff --git a/Engine/Plugins/Experimental/Mutable/Source/CustomizableObjectEditor/Private/MuCOE/CustomizableObjectEditor.cpp b/Engine/Plugins/Experimental/Mutable/Source/CustomizableObjectEditor/Private/MuCOE/CustomizableObjectEditor.cpp
--- a/Engine/Plugins/Experimental/Mutable/Source/CustomizableObjectEditor/Private/MuCOE/CustomizableObjectEditor.cpp
+++ b/Engine/Plugins/Experimental/Mutable/Source/CustomizableObjectEditor/Private/MuCOE/CustomizableObjectEditor.cpp
@@ -522,31 +522,12 @@ TSharedRef<SDockTab> FCustomizableObjectEditor::SpawnTab_InstanceProperties( con
 /** Create new tab for the supplied graph - don't call this directly, call SExplorer->FindTabForGraph.*/
 void FCustomizableObjectEditor::CreateGraphEditorWidget(UEdGraph* InGraph)
 {
-	check(InGraph != NULL);
+	UCustomizableObjectGraph* CustomizableObjectGraph = Cast<UCustomizableObjectGraph>(InGraph);
+	check(CustomizableObjectGraph != NULL);
 
-	// Check whether the graph has a base node and create one if it doesn't since it is required
-	bool bGraphHasBase = false;
+	// Add essential graph nodes if they do not exist
+	CustomizableObjectGraph->AddEssentialGraphNodes();
 
-	for (const TObjectPtr<UEdGraphNode>& AuxNode : InGraph->Nodes)
-	{
-		UCustomizableObjectNodeObject* CustomizableObjectNodeObject = Cast<UCustomizableObjectNodeObject>(AuxNode);
-
-		if (CustomizableObjectNodeObject && CustomizableObjectNodeObject->bIsBase)
-		{
-			bGraphHasBase = true;
-			break;
-		}
-	}
-
-	if (!bGraphHasBase)
-	{
-		FCustomizableObjectSchemaAction_NewNode Action;
-		UCustomizableObjectNodeObject* NodeTemplate = NewObject<UCustomizableObjectNodeObject>();
-
-		Action.NodeTemplate = NodeTemplate;
-		Action.FCustomizableObjectSchemaAction_NewNode::PerformAction(InGraph, nullptr, FVector2D::ZeroVector, false);
-	}
-	
 	GraphEditorCommands = MakeShareable(new FUICommandList);
 
 	TSharedRef<SWidget> TitleBarWidget =
diff --git a/Engine/Plugins/Experimental/Mutable/Source/CustomizableObjectEditor/Private/MuCOE/CustomizableObjectFactory.cpp b/Engine/Plugins/Experimental/Mutable/Source/CustomizableObjectEditor/Private/MuCOE/CustomizableObjectFactory.cpp
--- a/Engine/Plugins/Experimental/Mutable/Source/CustomizableObjectEditor/Private/MuCOE/CustomizableObjectFactory.cpp
+++ b/Engine/Plugins/Experimental/Mutable/Source/CustomizableObjectEditor/Private/MuCOE/CustomizableObjectFactory.cpp
@@ -3,14 +3,34 @@
 #include "MuCOE/CustomizableObjectFactory.h"
 
 #include "MuCO/CustomizableObject.h"
+#include "MuCO/CustomizableObjectPrivate.h"
 #include "MuCOE/CustomizableObjectGraph.h"
+#include "MuCOE/GraphTraversal.h"
+#include "MuCOE/Nodes/CustomizableObjectNodeObject.h"
+#include "MuCOE/Nodes/CustomizableObjectNodeObjectGroup.h"
+#include "MuCOE/EdGraphSchema_CustomizableObject.h"
 
-class FFeedbackContext;
-class UClass;
-class UObject;
+#include "ContentBrowserModule.h"
+#include "DetailLayoutBuilder.h"
+#include "Editor.h"
+#include "Editor/EditorEngine.h"
+#include "Engine/SkeletalMesh.h"
+#include "IContentBrowserSingleton.h"
+#include "Input/Reply.h"
+#include "Layout/Visibility.h"
+#include "PropertyCustomizationHelpers.h"
+#include "PropertyEditorModule.h"
+#include "Styling/SlateTypes.h"
+#include "Types/SlateEnums.h"
+#include "Widgets/Input/SButton.h"
+#include "Widgets/Input/SCheckBox.h"
+#include "Widgets/Input/SSpinBox.h"
+#include "Widgets/Input/STextComboBox.h"
+#include "Widgets/Layout/SBorder.h"
+#include "Widgets/SWindow.h"
 
-#define LOCTEXT_NAMESPACE "CustomizableObjectFactory"
 
+#define LOCTEXT_NAMESPACE "CustomizableObjectFactory"
 
 UCustomizableObjectFactory::UCustomizableObjectFactory()
 	: Super()
@@ -36,17 +56,608 @@ UClass* UCustomizableObjectFactory::ResolveSupportedClass()
 
 bool UCustomizableObjectFactory::ConfigureProperties()
 {
-	return true;
+	TSharedRef<FCustomizableObjectFactoryUI> COSettingsWindow = MakeShareable(new FCustomizableObjectFactoryUI());
+	CreationSettings = COSettingsWindow->ConstructFactoryUI();
+	
+	return COSettingsWindow->CanCreateObject();
 }
 
 
 UObject* UCustomizableObjectFactory::FactoryCreateNew(UClass* Class, UObject* InParent, FName Name, EObjectFlags Flags, UObject* Context, FFeedbackContext* Warn)
 {
 	UCustomizableObject* NewObj = NewObject<UCustomizableObject>(InParent, Class, Name, Flags);
-
 	NewObj->Source = NewObject<UCustomizableObjectGraph>(NewObj, NAME_None, RF_Transactional);
 
+	if (NewObj && NewObj->Source)
+	{
+		UCustomizableObjectGraph* CustomizableObjectGraph = Cast<UCustomizableObjectGraph>(NewObj->Source);
+		CustomizableObjectGraph->AddEssentialGraphNodes();
+
+		UCustomizableObjectNodeObject* BaseObjectNode = nullptr;
+
+		for (const TObjectPtr<UEdGraphNode>& AuxNode : NewObj->Source->Nodes)
+		{
+			UCustomizableObjectNodeObject* BaseNode = Cast<UCustomizableObjectNodeObject>(AuxNode);
+
+			if (BaseNode && BaseNode->bIsBase)
+			{
+				BaseObjectNode = BaseNode;
+				break;
+			}
+		}
+
+		if (!CreationSettings.bEmptyObject)
+		{
+			if (!BaseObjectNode)
+			{
+				return NewObj;
+			}
+
+			if (CreationSettings.bIsChildObject && CreationSettings.ParentObject.IsValid())
+			{
+				BaseObjectNode->ParentObject = CreationSettings.ParentObject.Get();
+
+				TArray<UCustomizableObjectNodeObjectGroup*> GroupNodes;
+				CreationSettings.ParentObject->Source->GetNodesOfClass<UCustomizableObjectNodeObjectGroup>(GroupNodes);
+
+				for (UCustomizableObjectNodeObjectGroup* GroupNode : GroupNodes)
+				{
+					if (GroupNode->GroupName == CreationSettings.GroupNodeName)
+					{
+						BaseObjectNode->ParentObjectGroupId = GroupNode->NodeGuid;
+
+						break;
+					}
+				}
+
+				if (NewObj->GetPrivate())
+				{
+					NewObj->GetPrivate()->SetIsChildObject(true);
+				}
+			}
+			else
+			{
+				BaseObjectNode->NumMeshComponents = CreationSettings.NumMeshComponents;
+
+				for (int32 MeshIndex = 0; MeshIndex < CreationSettings.ReferenceSkeletalMeshes.Num(); ++MeshIndex)
+				{
+					if (CreationSettings.ReferenceSkeletalMeshes[MeshIndex].IsValid())
+					{
+						NewObj->ReferenceSkeletalMeshes.Add(CreationSettings.ReferenceSkeletalMeshes[MeshIndex].Get());
+					}
+				}
+			}
+		}
+	}
+
 	return NewObj;
 }
 
+
+// Factory Settings UI ----------------------------------------------------------------------------------------------------------------------------------------------------------------
+
+const FCustomizableObjectOptions FCustomizableObjectFactoryUI::ConstructFactoryUI()
+{
+	FPropertyEditorModule& PropertyEditorModule = FModuleManager::GetModuleChecked<FPropertyEditorModule>("PropertyEditor");
+	FDetailsViewArgs DetailsViewArgs;
+	DetailsViewArgs.NameAreaSettings = FDetailsViewArgs::HideNameArea;
+	
+	TSharedRef<IDetailsView> DetailsView = PropertyEditorModule.CreateDetailView(DetailsViewArgs);
+	DetailsView->SetObject(nullptr);
+
+	// Default number of components for non-child objects
+	Options.NumMeshComponents = 1;
+	Options.ReferenceSkeletalMeshes.SetNum(1);
+	GenerateComponentOptions();
+
+	// Settings window
+	COSettingsWindow = SNew(SWindow)
+	.Title(LOCTEXT("CustomizableObjectFactoryptions", "New Costumizable Object"))
+	.ClientSize(FVector2D(300, 280))
+	.SupportsMinimize(false)
+	.SupportsMaximize(false)
+	[
+		SNew(SBorder)
+		.BorderImage(FAppStyle::GetBrush("Menu.Background")) // Adds lighter background
+		.Padding(5)
+		[
+			// Is Child Object Widgets
+			SNew(SVerticalBox)
+			+ SVerticalBox::Slot()
+			.AutoHeight()
+			.HAlign(HAlign_Left)
+			.Padding(0.0f, 5.0f, 0.0f, 0.0f)
+			[
+				SNew(SHorizontalBox)
+				+ SHorizontalBox::Slot()
+				.AutoWidth()
+				.Padding(0.0f, 2.0f, 0.0f, 0.0f)
+				[
+					SNew(STextBlock)
+					.Text(LOCTEXT("IsChildObject_CheckBoxName","Is Child Object:"))
+					.Font(IDetailLayoutBuilder::GetDetailFont())
+				]
+
+				+ SHorizontalBox::Slot()
+				.AutoWidth()
+				.Padding(5.0f, 0.0f, 0.0f, 0.0f)
+				[
+					SNew(SCheckBox)
+					.IsChecked(Options.bIsChildObject ? ECheckBoxState::Checked : ECheckBoxState::Unchecked)
+					.OnCheckStateChanged(this, &FCustomizableObjectFactoryUI::OnCheckBoxChanged)
+				]
+			]
+				
+			// Parent Selector Widgets
+			+ SVerticalBox::Slot()
+			.AutoHeight()
+			.Padding(0.0f, 10.0f, 0.0f, 10.0f)
+			.HAlign(HAlign_Fill)
+			.VAlign(VAlign_Fill)
+			[
+				SNew(SBorder)
+				.Padding(5.0f, 10.0f, 0.0f, 10.0f)
+				.BorderImage(FAppStyle::Get().GetBrush("Brushes.Panel"))
+				.HAlign(HAlign_Fill)
+				[
+					SNew(SVerticalBox)
+					+ SVerticalBox::Slot()
+					.AutoHeight()
+					.HAlign(HAlign_Left)
+					[
+						SNew(STextBlock)
+						.Visibility(this, &FCustomizableObjectFactoryUI::GetParentWidgetsVisibility)
+						.Text(LOCTEXT("SelectedParent_Text", "Parent Object: "))
+						.Font(IDetailLayoutBuilder::GetDetailFont())
+					]
+
+					+ SVerticalBox::Slot()
+					.Padding(5.0f, 5.0f, 10.0f, 0.0f)
+					.AutoHeight()
+					[
+						SNew(SObjectPropertyEntryBox)
+						.Visibility(this, &FCustomizableObjectFactoryUI::GetParentWidgetsVisibility)
+						.AllowedClass(UCustomizableObject::StaticClass())
+						.OnObjectChanged(this, &FCustomizableObjectFactoryUI::OnPickedCustomizableObjectParent)
+						.ObjectPath(this, &FCustomizableObjectFactoryUI::GetSelectedCustomizableObjectPath)
+						.DisplayThumbnail(true)
+						.ThumbnailPool(UThumbnailManager::Get().GetSharedThumbnailPool())
+						.AllowClear(false)
+						.AllowCreate(false)
+						.DisplayBrowse(false)
+						.DisplayUseSelected(false)
+					]
+
+					+ SVerticalBox::Slot()
+					.AutoHeight()
+					.HAlign(HAlign_Left)
+					.Padding(0.0f, 15.0f, 0.0f, 0.0f)
+					[
+						SNew(STextBlock)
+						.Visibility(this, &FCustomizableObjectFactoryUI::GetParentWidgetsVisibility)
+						.Text(LOCTEXT("SelecteGroup_Text", "Group Node: "))
+						.Font(IDetailLayoutBuilder::GetDetailFont())
+					]
+
+					+ SVerticalBox::Slot()
+					.Padding(5.0f, 5.0f, 10.0f, 0.0f)
+					.AutoHeight()
+					[
+						SAssignNew(GroupSelector, STextComboBox)
+						.Visibility(this, &FCustomizableObjectFactoryUI::GetParentWidgetsVisibility)
+						.OptionsSource(&GroupOptions)
+						.IsEnabled(this, &FCustomizableObjectFactoryUI::IsNodeGroupSelectorEnabled)
+						.OnSelectionChanged(this, &FCustomizableObjectFactoryUI::OnSelectGroupComboBox)
+						.Font(IDetailLayoutBuilder::GetDetailFont())
+					]
+
+					// Component Selector Widgets
+					+ SVerticalBox::Slot()
+					.AutoHeight()
+					.HAlign(HAlign_Left)
+					[
+						SNew(STextBlock)
+						.Visibility(this, &FCustomizableObjectFactoryUI::GetComponentWidgetsVisibility)
+						.Text(LOCTEXT("NumberComponents_Text", "Num Mesh Components: "))
+						.Font(IDetailLayoutBuilder::GetDetailFont())
+					]
+
+					+ SVerticalBox::Slot()
+					.AutoHeight()
+					.Padding(5.0f, 7.0f, 10.0f, 0.0f)
+					[
+						SNew(SSpinBox<int32>)
+						.Visibility(this, &FCustomizableObjectFactoryUI::GetComponentWidgetsVisibility)
+						.OnValueChanged(this, &FCustomizableObjectFactoryUI::OnNumComponentsChanged, ETextCommit::Type::Default)
+						.OnValueCommitted(this, &FCustomizableObjectFactoryUI::OnNumComponentsChanged)
+						.Value(this, &FCustomizableObjectFactoryUI::GetNumComponents)
+						.MinValue(1)
+						.MaxValue(255)
+						.MaxSliderValue(6)
+						.Delta(1)
+						.AlwaysUsesDeltaSnap(true)
+						.Font(IDetailLayoutBuilder::GetDetailFont())
+					]
+
+					+SVerticalBox::Slot()
+					.Padding(0.0f, 15.0f, 0.0f, 0.0f)
+					[
+						SNew(STextBlock)
+						.Visibility(this, &FCustomizableObjectFactoryUI::GetComponentWidgetsVisibility)
+						.Text(LOCTEXT("SelectedComponent_Text", "Component: "))
+						.Font(IDetailLayoutBuilder::GetDetailFont())
+					]
+
+					+SVerticalBox::Slot()
+					.Padding(5.0f, 5.0f, 10.0f, 0.0f)
+					[
+						SAssignNew(ComponentSelector, STextComboBox)
+						.Visibility(this, &FCustomizableObjectFactoryUI::GetComponentWidgetsVisibility)
+						.ToolTipText(LOCTEXT("SelectedComponent_Tooltip", "Select a component to set its Reference Skeletal Mesh."))
+						.OptionsSource(&ComponentsOptions)
+						.InitiallySelectedItem(ComponentsOptions[0])
+						.IsEnabled(this, &FCustomizableObjectFactoryUI::IsComponentSelectorEnabled)
+						.OnSelectionChanged(this, &FCustomizableObjectFactoryUI::OnSelectComponentComboBox)
+						.Font(IDetailLayoutBuilder::GetDetailFont())
+					]
+
+					+ SVerticalBox::Slot()
+					.Padding(0.0f, 15.0f, 0.0f, 0.0f)
+					[
+						SNew(STextBlock)
+						.Visibility(this, &FCustomizableObjectFactoryUI::GetComponentWidgetsVisibility)
+						.Text(this, &FCustomizableObjectFactoryUI::GetMeshSelectorWidgetText)
+						.Font(IDetailLayoutBuilder::GetDetailFont())
+					]
+
+					+SVerticalBox::Slot()
+					.AutoHeight()
+					.Padding(5.0f, 5.0f, 10.0f, 0.0f)
+					[
+						SNew(SObjectPropertyEntryBox)
+						.Visibility(this, &FCustomizableObjectFactoryUI::GetComponentWidgetsVisibility)
+						.IsEnabled(this, &FCustomizableObjectFactoryUI::IsComponentSelectorEnabled)
+						.AllowedClass(USkeletalMesh::StaticClass())
+						.OnObjectChanged(this, &FCustomizableObjectFactoryUI::OnPickedComponentSkeletalMesh)
+						.ObjectPath(this, &FCustomizableObjectFactoryUI::GetSelectedComponentSkeletalMeshPath)
+						.DisplayThumbnail(true)
+						.ThumbnailPool(UThumbnailManager::Get().GetSharedThumbnailPool())
+						.AllowClear(false)
+						.AllowCreate(false)
+						.DisplayBrowse(false)
+						.DisplayUseSelected(false)
+					]
+				]
+			]
+
+			+ SVerticalBox::Slot()
+			[
+				SNew(SSpacer)
+			]
+
+			+ SVerticalBox::Slot()
+			.HAlign(HAlign_Fill)
+			.AutoHeight()
+			[
+				SNew(SHorizontalBox)
+				+ SHorizontalBox::Slot()
+				.AutoWidth()
+				.HAlign(HAlign_Left)
+				[
+					SNew(SButton)
+						.Text(LOCTEXT("CreateEmptyCO", "Create Empty"))
+						.ToolTipText(LOCTEXT("CreateEmptyCO_Tooltip","Create a Customizable Object without a components or reference skeletal mesh."))
+						.OnClicked(this, &FCustomizableObjectFactoryUI::OnCreate, true)
+				]
+
+				+ SHorizontalBox::Slot()
+				[
+					SNew(SSpacer)
+				]
+
+				+ SHorizontalBox::Slot()
+				.AutoWidth()
+				.HAlign(HAlign_Right)
+				[
+					SNew(SHorizontalBox)
+					+ SHorizontalBox::Slot()
+					.AutoWidth()
+					[
+						SNew(SButton)
+							.Text(LOCTEXT("OK", "OK"))
+							.ToolTipText(this, &FCustomizableObjectFactoryUI::GetOKButtonTooltip)
+							.IsEnabled(this, &FCustomizableObjectFactoryUI::IsConfigurationValid)
+							.OnClicked(this, &FCustomizableObjectFactoryUI::OnCreate, false)
+					]
+					+ SHorizontalBox::Slot()
+					.AutoWidth()
+					[
+						SNew(SButton)
+							.Text(LOCTEXT("Cancel", "Cancel"))
+							.OnClicked(this, &FCustomizableObjectFactoryUI::OnCancel)
+					]
+				]
+			]
+		]
+	];
+
+	GEditor->EditorAddModalWindow(COSettingsWindow.ToSharedRef());
+	COSettingsWindow.Reset();
+
+	return Options;
+}
+
+
+void FCustomizableObjectFactoryUI::OnCheckBoxChanged(ECheckBoxState State)
+{
+	Options.bIsChildObject = State == ECheckBoxState::Checked;
+
+	// Reset settings
+	Options.ParentObject = nullptr;
+	Options.GroupNodeName.Empty();
+	Options.NumMeshComponents = Options.bIsChildObject ? 0 : 1;
+	Options.bIsChildObject ? Options.ReferenceSkeletalMeshes.Empty() : Options.ReferenceSkeletalMeshes.SetNum(1);
+
+	FVector2D ClientSize = Options.bIsChildObject ? FVector2D(300, 220) : FVector2D(300, 280);
+	FVector2D WindowsSize = ClientSize * COSettingsWindow->GetDPIScaleFactor();
+	
+	COSettingsWindow->Resize(WindowsSize);
+}
+
+
+FReply FCustomizableObjectFactoryUI::OnCreate(bool bIsEmpty)
+{
+	bCreateObject = true;
+	Options.bEmptyObject = bIsEmpty;
+
+	if (COSettingsWindow.IsValid())
+	{
+		COSettingsWindow->RequestDestroyWindow();
+	}
+
+	return FReply::Handled();
+}
+
+
+FText FCustomizableObjectFactoryUI::GetOKButtonTooltip() const
+{
+	FString Tooltip = "Create a Customizable Object with the settings selected.";
+
+	if (Options.bIsChildObject)
+	{
+		if (!IsConfigurationValid())
+		{
+			Tooltip = "Select a Parent Object and the name of a Group Object Node.";
+		}
+	}
+	else
+	{
+		if (!IsConfigurationValid())
+		{
+			Tooltip = "Set the number of components that the Customizable Object will have. Then select a Reference Skeletal Mesh for each component.";
+		}
+	}
+
+	return FText::FromString(Tooltip);
+}
+
+
+FReply FCustomizableObjectFactoryUI::OnCancel()
+{
+	bCreateObject = false;
+
+	if (COSettingsWindow.IsValid())
+	{
+		COSettingsWindow->RequestDestroyWindow();
+	}
+
+	return FReply::Handled();
+}
+
+
+bool FCustomizableObjectFactoryUI::IsConfigurationValid() const
+{
+	if (Options.bIsChildObject)
+	{
+		if (Options.ParentObject.IsValid() && !Options.GroupNodeName.IsEmpty())
+		{
+			return true;
+		}
+	}
+	else
+	{
+		// Check if all components have a valid Skeletal Mesh assigned
+		for (int32 CompIndex = 0; CompIndex < Options.NumMeshComponents; ++CompIndex)
+		{
+			if (!Options.ReferenceSkeletalMeshes.IsValidIndex(CompIndex) || !Options.ReferenceSkeletalMeshes[CompIndex].IsValid())
+			{
+				return false;
+			}
+		}
+
+		return true;
+	}
+
+	return false;
+}
+
+
+bool FCustomizableObjectFactoryUI::CanCreateObject() const
+{
+	return bCreateObject;
+}
+
+
+void FCustomizableObjectFactoryUI::OnPickedCustomizableObjectParent(const FAssetData& SelectedAsset)
+{
+	if (UCustomizableObject* Parent = Cast<UCustomizableObject>(SelectedAsset.GetAsset()))
+	{
+		Options.ParentObject = Parent;
+
+		if (GroupSelector.IsValid())
+		{
+			GroupSelector->ClearSelection();
+			GenerateGroupOptions();
+		}
+	}
+}
+
+
+FString FCustomizableObjectFactoryUI::GetSelectedCustomizableObjectPath() const
+{
+	return Options.ParentObject.IsValid() ? Options.ParentObject->GetPathName() : FString();
+}
+
+
+EVisibility FCustomizableObjectFactoryUI::GetParentWidgetsVisibility() const
+{
+	return Options.bIsChildObject ? EVisibility::Visible : EVisibility::Collapsed;
+}
+
+
+bool FCustomizableObjectFactoryUI::IsNodeGroupSelectorEnabled() const
+{
+	return Options.ParentObject.IsValid() ? true : false;
+}
+
+
+void FCustomizableObjectFactoryUI::GenerateGroupOptions()
+{
+	GroupOptions.Empty();
+
+	if (Options.ParentObject.IsValid())
+	{
+		TArray<UCustomizableObjectNodeObjectGroup*> GroupNodes;
+		Options.ParentObject->Source->GetNodesOfClass<UCustomizableObjectNodeObjectGroup>(GroupNodes);
+
+		for (UCustomizableObjectNodeObjectGroup* GroupNode : GroupNodes)
+		{
+			GroupOptions.Add(MakeShareable(new FString(GroupNode->GroupName)));
+		}
+
+		if (GroupSelector.IsValid())
+		{
+			GroupSelector->RefreshOptions();
+		}
+	}
+}
+
+
+void FCustomizableObjectFactoryUI::OnSelectGroupComboBox(TSharedPtr<FString> Selection, ESelectInfo::Type SelectInfo)
+{
+	if (Selection.IsValid())
+	{
+		Options.GroupNodeName = *Selection;
+		GroupSelector->SetSelectedItem(Selection);
+	}
+}
+
+
+EVisibility FCustomizableObjectFactoryUI::GetComponentWidgetsVisibility() const
+{
+	return !Options.bIsChildObject ? EVisibility::Visible : EVisibility::Collapsed;
+}
+
+
+int32 FCustomizableObjectFactoryUI::GetNumComponents() const
+{
+	return Options.NumMeshComponents;
+}
+
+
+void FCustomizableObjectFactoryUI::OnNumComponentsChanged(int32 Value, ETextCommit::Type CommitInfo)
+{
+	if (Options.NumMeshComponents != Value)
+	{
+		Options.NumMeshComponents = Value;
+		Options.ReferenceSkeletalMeshes.SetNum(Value);
+		GenerateComponentOptions();
+
+		if (Value > 0)
+		{
+			ComponentSelector->SetSelectedItem(ComponentsOptions[0]);
+		}
+		else
+		{
+			ComponentSelector->ClearSelection();
+		}
+	}
+}
+
+
+void FCustomizableObjectFactoryUI::GenerateComponentOptions()
+{
+	ComponentsOptions.Reset();
+
+	for (int32 ComponentIndex = 0; ComponentIndex < Options.NumMeshComponents; ++ComponentIndex)
+	{
+		FString ComponentName = "Component " + FString::FromInt(ComponentIndex);
+		ComponentsOptions.Add(MakeShareable(new FString(ComponentName)));
+	}
+}
+
+
+bool FCustomizableObjectFactoryUI::IsComponentSelectorEnabled() const
+{
+	return Options.NumMeshComponents > 0;
+}
+
+
+void FCustomizableObjectFactoryUI::OnSelectComponentComboBox(TSharedPtr<FString> Selection, ESelectInfo::Type SelectInfo)
+{
+	if (Selection)
+	{
+		ComponentSelector->SetSelectedItem(Selection);
+	}
+}
+
+
+FText FCustomizableObjectFactoryUI::GetMeshSelectorWidgetText() const
+{
+	FString ComponentName;
+
+	if (ComponentSelector->GetSelectedItem().IsValid() && Options.NumMeshComponents)
+	{
+		ComponentName = *ComponentSelector->GetSelectedItem();
+	}
+
+	FString Message = "Select " + ComponentName + " Skeletal Mesh:";
+
+	return FText::FromString(Message);
+}
+
+
+void FCustomizableObjectFactoryUI::OnPickedComponentSkeletalMesh(const FAssetData& SelectedAsset)
+{
+	for (int32 ComponentIndex = 0; ComponentIndex < ComponentsOptions.Num(); ++ComponentIndex)
+	{
+		if (ComponentSelector.IsValid() && ComponentSelector->GetSelectedItem().IsValid()
+			&& *ComponentsOptions[ComponentIndex] == *ComponentSelector->GetSelectedItem()
+			&& Options.ReferenceSkeletalMeshes.IsValidIndex(ComponentIndex))
+		{
+			if (USkeletalMesh* SkeletalMesh = Cast<USkeletalMesh>(SelectedAsset.GetAsset()))
+			{
+				Options.ReferenceSkeletalMeshes[ComponentIndex] = SkeletalMesh;
+			}
+		}
+	}
+}
+
+
+FString FCustomizableObjectFactoryUI::GetSelectedComponentSkeletalMeshPath() const
+{
+	for (int32 ComponentIndex = 0; ComponentIndex < ComponentsOptions.Num(); ++ComponentIndex)
+	{
+		if (ComponentSelector.IsValid() && ComponentSelector->GetSelectedItem().IsValid()
+			&& *ComponentsOptions[ComponentIndex] == *ComponentSelector->GetSelectedItem()
+			&& Options.ReferenceSkeletalMeshes.IsValidIndex(ComponentIndex))
+		{
+			return Options.ReferenceSkeletalMeshes[ComponentIndex].IsValid() ? Options.ReferenceSkeletalMeshes[ComponentIndex]->GetPathName() : FString();
+		}
+	}
+
+	return FString();
+}
+
+
 #undef LOCTEXT_NAMESPACE
diff --git a/Engine/Plugins/Experimental/Mutable/Source/CustomizableObjectEditor/Private/MuCOE/CustomizableObjectGraph.cpp b/Engine/Plugins/Experimental/Mutable/Source/CustomizableObjectEditor/Private/MuCOE/CustomizableObjectGraph.cpp
--- a/Engine/Plugins/Experimental/Mutable/Source/CustomizableObjectEditor/Private/MuCOE/CustomizableObjectGraph.cpp
+++ b/Engine/Plugins/Experimental/Mutable/Source/CustomizableObjectEditor/Private/MuCOE/CustomizableObjectGraph.cpp
@@ -143,3 +143,27 @@ void UCustomizableObjectGraph::PostDuplicate(bool bDuplicateForPIE)
 
 	Super::PostDuplicate(bDuplicateForPIE);
 }
+
+
+void UCustomizableObjectGraph::AddEssentialGraphNodes()
+{
+	// Check whether the graph has a base node and create one if it doesn't since it is required
+	bool bGraphHasBase = false;
+
+	for (const TObjectPtr<UEdGraphNode>& AuxNode : Nodes)
+	{
+		UCustomizableObjectNodeObject* CustomizableObjectNodeObject = Cast<UCustomizableObjectNodeObject>(AuxNode);
+
+		if (CustomizableObjectNodeObject && CustomizableObjectNodeObject->bIsBase)
+		{
+			bGraphHasBase = true;
+			break;
+		}
+	}
+
+	if (!bGraphHasBase)
+	{
+		UCustomizableObjectNodeObject* NodeTemplate = NewObject<UCustomizableObjectNodeObject>();
+		UCustomizableObjectNodeObject* BaseObjectNode = Cast<UCustomizableObjectNodeObject>(FCustomizableObjectSchemaAction_NewNode::CreateNode(this, nullptr, FVector2D::ZeroVector, Cast<UEdGraphNode>(NodeTemplate)));
+	}
+}
