[Mutable] Fixed compilation warning for a pointer being used before checking it against nullptr.

#rnx

diff --git a/Engine/Plugins/Experimental/Mutable/Source/CustomizableObjectEditor/Private/MuCOE/CustomizableObjectFactory.cpp b/Engine/Plugins/Experimental/Mutable/Source/CustomizableObjectEditor/Private/MuCOE/CustomizableObjectFactory.cpp
--- a/Engine/Plugins/Experimental/Mutable/Source/CustomizableObjectEditor/Private/MuCOE/CustomizableObjectFactory.cpp
+++ b/Engine/Plugins/Experimental/Mutable/Source/CustomizableObjectEditor/Private/MuCOE/CustomizableObjectFactory.cpp
@@ -66,9 +66,13 @@ bool UCustomizableObjectFactory::ConfigureProperties()
 UObject* UCustomizableObjectFactory::FactoryCreateNew(UClass* Class, UObject* InParent, FName Name, EObjectFlags Flags, UObject* Context, FFeedbackContext* Warn)
 {
 	UCustomizableObject* NewObj = NewObject<UCustomizableObject>(InParent, Class, Name, Flags);
-	NewObj->Source = NewObject<UCustomizableObjectGraph>(NewObj, NAME_None, RF_Transactional);
 
-	if (NewObj && NewObj->Source)
+	if (NewObj)
+	{
+		NewObj->Source = NewObject<UCustomizableObjectGraph>(NewObj, NAME_None, RF_Transactional);
+	}
+
+	if (NewObj->Source)
 	{
 		UCustomizableObjectGraph* CustomizableObjectGraph = Cast<UCustomizableObjectGraph>(NewObj->Source);
 		CustomizableObjectGraph->AddEssentialGraphNodes();
