[mutable] Created an interface (UClass) for running the COI baking operation from a blueprint/python.
- Created a class, CustomizableObjectInstanceBaker, to take care of the baking operation.
- Added a control flag to prevent concurrent updates for baking purposes. It does not remove the possibility of collisions with other instance updates but removes it in the context of the baking of instances.
- Updated baking test commandlet to take into consideration cases where the update fails to be scheduled (and may produce an infinite waiting loop if not taken into consideration).
- Functionality tested in Editor Object Blueprint (not in python).
- Made the only accessible way of performing a correct bake the calling of the Bake method of the instance (Editor Only).
- Moved to Private the UnrealBakeHelpers.
- Moved to Internal the MutableTestingUtils (access required from MutableValidation)

diff --git a/Engine/Plugins/Experimental/Mutable/Source/CustomizableObject/Classes/MuCO/CustomizableObjectInstance.h b/Engine/Plugins/Experimental/Mutable/Source/CustomizableObject/Classes/MuCO/CustomizableObjectInstance.h
--- a/Engine/Plugins/Experimental/Mutable/Source/CustomizableObject/Classes/MuCO/CustomizableObjectInstance.h
+++ b/Engine/Plugins/Experimental/Mutable/Source/CustomizableObject/Classes/MuCO/CustomizableObjectInstance.h
@@ -9,8 +9,6 @@
 
 #include "CustomizableObjectInstance.generated.h"
 
-struct FMutableModelImageProperties;
-
 namespace mu
 {
 	class Image;
@@ -33,10 +31,123 @@ struct FTexturePlatformData;
 struct FMutableModelImageProperties;
 
 
+
+/**
+ * Represents what kind of saving procedure was performed to save the package
+ */
+UENUM()
+enum class EPackageSaveResolutionType : uint8
+{
+	None = 0,
+
+	/** The package got saved as a new file. */
+	NewFile,
+
+	/** The package was already present in disk so the old package was deleted and a new one was saved on its place */
+	Overriden,
+
+	/** Error type : An overriding was required but due to an error or lack of user permission it could not be done. */
+	UnableToOverride
+};
+
+
+/**
+ * Data structure that exposes the path to a baked package and also what type of save was performed (an override, a standard save with a new file...)
+ */
+USTRUCT(BlueprintType, Blueprintable)
+struct FBakedResourceData
+{
+	GENERATED_BODY()
+
+	/**
+	 * The way the package represented by this object was saved onto disk.
+	 */
+	UPROPERTY(BlueprintReadOnly, Category = CustomizableObjectInstanceBaker)
+	EPackageSaveResolutionType SaveType = EPackageSaveResolutionType::None;
+
+	/**
+	 * The path used by the saved package.
+	 */
+	UPROPERTY(BlueprintReadOnly, Category = CustomizableObjectInstanceBaker)
+	FString AssetPath;
+};
+
+
+/**
+ * Structure returned as output of the baking operation. May contain a filled collection of FBakedResourceData objects and also the success end state of the
+ * baking operation. 
+ */
+USTRUCT(BlueprintType, Blueprintable)
+struct FCustomizableObjectInstanceBakeOutput
+{
+	GENERATED_BODY()
+
+	/**
+	 * Success state for the baking operation. True for success and false for failure.
+	 */
+	UPROPERTY(BlueprintReadOnly, Category = CustomizableObjectInstanceBaker)
+	bool bWasBakeSuccessful = false;
+
+	/**
+	 * Collection of FBakedResourceData representing all saved packages during the baking operation. It may be empty if the operation failed.
+	 */
+	UPROPERTY(BlueprintReadOnly, Category = CustomizableObjectInstanceBaker)
+	TArray<FBakedResourceData> SavedPackages;
+};
+
+
+DECLARE_DYNAMIC_DELEGATE_OneParam(FBakeOperationCompletedDelegate, const FCustomizableObjectInstanceBakeOutput, BakeOperationOutput);
+
+
+/**
+ * Configuration data structure designed to serve as variable container for the customizable object instance baking methods.
+ */
+USTRUCT(BlueprintType, Blueprintable)
+struct FBakingConfiguration
+{
+	GENERATED_BODY()
+
+	/**
+	 * The path where to save the baked resources. EX /Game/MyBakingTest
+	 */
+	UPROPERTY(BlueprintReadWrite, Blueprintable, Category = CustomizableObjectInstanceBaker)
+	FString OutputPath = TEXT("/Game");
+	
+	/**
+	 * The name to be used as base (prefix) during the naming of the exported resources
+	 */
+	UPROPERTY(BlueprintReadWrite, Blueprintable, Category = CustomizableObjectInstanceBaker)
+	FString OutputFilesBaseName;
+
+	/**
+	 * Determines if we want a full or partial export
+	 */
+	UPROPERTY(BlueprintReadWrite, Blueprintable, Category = CustomizableObjectInstanceBaker)
+	bool bExportAllResourcesOnBake = false;
+	
+	/**
+	 *  Determines if we want (or not) to generate constant material instances for each of the material instances found in the mutable instance
+	 */
+	UPROPERTY(BlueprintReadWrite, Blueprintable, Category = CustomizableObjectInstanceBaker)
+	bool bGenerateConstantMaterialInstancesOnBake = false;
+	
+	/**
+	 * Flag that determines if we should override already exported files or if we should not. If we encounter files to override and we have not permission to override them then the baking operation will fail.
+	 */
+	UPROPERTY(BlueprintReadWrite, Blueprintable, Category = CustomizableObjectInstanceBaker)
+	bool bAllowOverridingOfFiles = false;
+
+	/**
+	 * Callback executed once the baking operation gets completed. It will return the end success state and also some data about the assets saved.
+	 */
+	UPROPERTY(BlueprintReadWrite, Blueprintable, Category = CustomizableObjectInstanceBaker)
+	FBakeOperationCompletedDelegate OnBakeOperationCompletedCallback;
+};
+
+
 // Priority for the mutable update queue, Low is the normal distance-based priority, High is normally used for discards and Mid for LOD downgrades
 enum class EQueuePriorityType : uint8 { High, Med, Med_Low, Low };
 
-
 /** Result of all the checks just before beginning an update. */
 enum class EUpdateRequired : uint8
 {
@@ -581,6 +692,15 @@ public:
 
 	const TArray<uint16>& GetRequestedLODsPerComponent() const;
 
+#if WITH_EDITOR
+	/**
+	 * Performs the baking of the instance resources in an async fashion. Bind yourself to the callback present in InBakingConfiguration to get notified in case it fails
+	 * @param InBakingConfiguration The configuration to be using for the baking operation
+	 */
+	UFUNCTION(BlueprintCallable, Category = CustomizableObjectInstanceBaking)
+	void Bake(const FBakingConfiguration& InBakingConfiguration);
+#endif
+	
 	UCustomizableInstancePrivate* GetPrivate() const;
 
 private:
diff --git a/Engine/Plugins/Experimental/Mutable/Source/CustomizableObject/Internal/MuCO/ICustomizableObjectEditorModule.h b/Engine/Plugins/Experimental/Mutable/Source/CustomizableObject/Internal/MuCO/ICustomizableObjectEditorModule.h
--- a/Engine/Plugins/Experimental/Mutable/Source/CustomizableObject/Internal/MuCO/ICustomizableObjectEditorModule.h
+++ b/Engine/Plugins/Experimental/Mutable/Source/CustomizableObject/Internal/MuCO/ICustomizableObjectEditorModule.h
@@ -5,6 +5,7 @@
 #include "Modules/ModuleManager.h"
 #include "Templates/SharedPointer.h"
 
+struct FBakingConfiguration;
 class FCustomizableObjectEditorLogger;
 class ICustomizableObjectDebugger;
 class ICustomizableObjectEditor;
@@ -14,6 +15,7 @@ class UCustomizableObject;
 class UCustomizableObjectPrivate;
 class UCustomizableObjectInstance;
 class FExtensibilityManager;
+class FBakeOperationCompletedDelegate;
 
 extern const FName CustomizableObjectEditorAppIdentifier;
 extern const FName CustomizableObjectInstanceEditorAppIdentifier;
@@ -21,7 +23,6 @@ extern const FName CustomizableObjectPopulationEditorAppIdentifier;
 extern const FName CustomizableObjectPopulationClassEditorAppIdentifier;
 extern const FName CustomizableObjectDebuggerAppIdentifier;
 
-
 /**
  * Customizable object editor module interface
  */
@@ -49,5 +50,12 @@ public:
 	  * @param OutOfDatePackages list of out of date packages.
    	  * @return true if the compilation is out of date. */
 	virtual bool IsCompilationOutOfDate(const UCustomizableObject& Object, TArray<FName>* OutOfDatePackages = nullptr) const = 0;
+	
+	/**
+	 * Execute this method in order to bake the provided instance. It will schedule a special type of instance update before proceeding with the bake itself.
+	 * @param InTargetInstance The instance we want to bake
+	 * @param InBakingConfig Structure containing the configuration to be used for the baking
+	 */
+	virtual void BakeCustomizableObjectInstance(UCustomizableObjectInstance* InTargetInstance, const FBakingConfiguration& InBakingConfig)  = 0;
 };
  
diff --git a/Engine/Plugins/Experimental/Mutable/Source/CustomizableObject/Private/MuCO/CustomizableObjectInstance.cpp b/Engine/Plugins/Experimental/Mutable/Source/CustomizableObject/Private/MuCO/CustomizableObjectInstance.cpp
--- a/Engine/Plugins/Experimental/Mutable/Source/CustomizableObject/Private/MuCO/CustomizableObjectInstance.cpp
+++ b/Engine/Plugins/Experimental/Mutable/Source/CustomizableObject/Private/MuCO/CustomizableObjectInstance.cpp
@@ -42,6 +42,7 @@
 #include "RenderingThread.h"
 #include "SkeletalMergingLibrary.h"
 #include "UnrealMutableImageProvider.h"
+#include "MuCO/ICustomizableObjectEditorModule.h"
 #include "UObject/ObjectSaveContext.h"
 #include "UObject/UObjectIterator.h"
 #include "Widgets/Notifications/SNotificationList.h"
@@ -6451,6 +6452,29 @@ const TArray<uint16>& UCustomizableObjectInstance::GetRequestedLODsPerComponent(
 }
 
 
+#if WITH_EDITOR
+void UCustomizableObjectInstance::Bake(const FBakingConfiguration& InBakingConfiguration)
+{
+	if (ICustomizableObjectEditorModule* Module = ICustomizableObjectEditorModule::Get())
+	{
+		Module->BakeCustomizableObjectInstance(this, InBakingConfiguration);
+	}
+	else
+	{
+		// Notify of the error
+		UE_LOG(LogMutable, Error, TEXT("The module \" ICustomizableObjectEditorModule \" could not be loaded and therefore the baking operation could not be started."));
+		if (InBakingConfiguration.OnBakeOperationCompletedCallback.IsBound())
+		{
+			FCustomizableObjectInstanceBakeOutput Output;
+			Output.bWasBakeSuccessful = false;
+			Output.SavedPackages.Empty();
+			InBakingConfiguration.OnBakeOperationCompletedCallback.Execute(Output);
+		}
+	}
+}
+#endif
+
+
 USkeletalMesh* UCustomizableObjectInstance::GetSkeletalMesh(int32 ComponentIndex) const
 {
 	check(PrivateData);
diff --git a/Engine/Plugins/Experimental/Mutable/Source/CustomizableObjectEditor/Internal/MuCOE/CustomizableObjectInstanceBakingUtils.h b/Engine/Plugins/Experimental/Mutable/Source/CustomizableObjectEditor/Internal/MuCOE/CustomizableObjectInstanceBakingUtils.h
new file mode 100644
--- /dev/null
+++ b/Engine/Plugins/Experimental/Mutable/Source/CustomizableObjectEditor/Internal/MuCOE/CustomizableObjectInstanceBakingUtils.h
@@ -0,0 +1,36 @@
+﻿// Copyright Epic Games, Inc. All Rights Reserved.
+#pragma once
+
+#include "MuCO/CustomizableObjectInstance.h"
+
+/**
+ * Schedules the async update of the target instance. It will clear previous update data.
+ * This method will also take care of setting and later resetting the state of the Customizable Object System.
+ * @param InInstance The instance we want to update so we can later bake it's resources.
+ * @param InInstanceUpdateDelegate Delegate to be called after the instance gets updated.
+ */
+CUSTOMIZABLEOBJECTEDITOR_API
+void UpdateInstanceForBaking(UCustomizableObjectInstance& InInstance, FInstanceUpdateNativeDelegate& InInstanceUpdateDelegate);
+
+
+/**
+ * Serializes onto disk the resources used by the targeted Customizable Object Instance. The operation can be configured in the FInstanceBakingSettings settings object.
+ * @param InInstance The mutable COI instance whose resources we want to bake onto disk
+ * @param FileName The base name for the all the files generated during the baking process
+ * @param AssetPath The path where to save the assets
+ * @param bExportAllResources Determines if we want a full or partial export
+ * @param bGenerateConstantMaterialInstances Determines if we want to generate constant material instances or not
+ * @param bHasPermissionToOverride Do we want to allow the overriding of assets? This value will get updated internally if not Unattended. 
+ * @param bIsUnattendedExecution Determines if we want to run this method and show promps for user interaction or if we want to automatically chose the more sensible option (without the need of user interaction)
+ * @param OutSavedPackages A collection with the packages marked for save.
+ * @return True if the baking operation could be completed without issues and false otherwise
+ */
+CUSTOMIZABLEOBJECTEDITOR_API
+bool BakeCustomizableObjectInstance (UCustomizableObjectInstance& InInstance,
+ const FString& FileName, 
+ const FString& AssetPath,
+ const bool bExportAllResources,
+ const bool bGenerateConstantMaterialInstances, 
+ bool bHasPermissionToOverride,
+ bool bIsUnattendedExecution, 
+ TArray<TPair<EPackageSaveResolutionType,UPackage*>>& OutSavedPackages);
\ No newline at end of file
diff --git a/Engine/Plugins/Experimental/Mutable/Source/CustomizableObjectEditor/Private/MuCOE/CustomizableObjectEditorModule.cpp b/Engine/Plugins/Experimental/Mutable/Source/CustomizableObjectEditor/Private/MuCOE/CustomizableObjectEditorModule.cpp
--- a/Engine/Plugins/Experimental/Mutable/Source/CustomizableObjectEditor/Private/MuCOE/CustomizableObjectEditorModule.cpp
+++ b/Engine/Plugins/Experimental/Mutable/Source/CustomizableObjectEditor/Private/MuCOE/CustomizableObjectEditorModule.cpp
@@ -63,6 +63,7 @@
 #include "UObject/UObjectIterator.h"
 #include "Subsystems/PlacementSubsystem.h"
 #include "Components/SkeletalMeshComponent.h"
+#include "MuCOE/CustomizableObjectInstanceBaker.h"
 
 class AActor;
 class FString;
@@ -586,4 +587,23 @@ bool FCustomizableObjectEditorModule::IsCompilationOutOfDate(const UCustomizable
 }
 
 
+void FCustomizableObjectEditorModule::BakeCustomizableObjectInstance(UCustomizableObjectInstance* InTargetInstance, const FBakingConfiguration& InBakingConfig)
+{
+	UCustomizableObjectInstanceBaker* InstanceBaker = NewObject<UCustomizableObjectInstanceBaker>();
+
+	// Add the heap object to the root so we prevent it from being removed. It will get removed from there once it finishes it's work.
+	InstanceBaker->AddToRoot();
+	
+	// On baker operation completed just remove it from the root so it gets eventually destroyed by the GC system
+	const TSharedPtr<FOnBakerFinishedWork> OnBakerFinishedWorkCallback = MakeShared<FOnBakerFinishedWork>();
+	OnBakerFinishedWorkCallback->BindLambda([InstanceBaker]
+	{
+		InstanceBaker->RemoveFromRoot();
+	});
+	
+	// Ask for the baking of the instance
+	InstanceBaker->BakeInstance(InTargetInstance, InBakingConfig, OnBakerFinishedWorkCallback);
+}
+
+
 #undef LOCTEXT_NAMESPACE
diff --git a/Engine/Plugins/Experimental/Mutable/Source/CustomizableObjectEditor/Private/MuCOE/CustomizableObjectEditorModule.h b/Engine/Plugins/Experimental/Mutable/Source/CustomizableObjectEditor/Private/MuCOE/CustomizableObjectEditorModule.h
--- a/Engine/Plugins/Experimental/Mutable/Source/CustomizableObjectEditor/Private/MuCOE/CustomizableObjectEditorModule.h
+++ b/Engine/Plugins/Experimental/Mutable/Source/CustomizableObjectEditor/Private/MuCOE/CustomizableObjectEditorModule.h
@@ -7,12 +7,14 @@
 #include "MuCOE/CustomizableObjectEditorLogger.h"
 #include "Toolkits/AssetEditorToolkit.h"
 
+struct FBakingConfiguration;
 class USkeletalMeshComponent;
 class FPropertyEditorModule;
 class FExtensibilityManager;
 class ICustomizableObjectEditor;
 class ICustomizableObjectInstanceEditor;
 class ICustomizableObjectDebugger;
+class FBakeOperationCompletedDelegate;
 
 
 /** Get a list of packages that are used by the compilation but are not directly referenced.
@@ -35,6 +37,7 @@ public:
 	// ICustomizableObjectEditorModule interface
 	virtual FCustomizableObjectEditorLogger& GetLogger() override;
 	virtual bool IsCompilationOutOfDate(const UCustomizableObject& Object, TArray<FName>* OutOfDatePackages) const override;
+	virtual void BakeCustomizableObjectInstance(UCustomizableObjectInstance* InTargetInstance, const FBakingConfiguration& InBakingConfig) override;
 
 	virtual TSharedPtr<FExtensibilityManager> GetCustomizableObjectEditorToolBarExtensibilityManager() override { return CustomizableObjectEditor_ToolBarExtensibilityManager; }
 	virtual TSharedPtr<FExtensibilityManager> GetCustomizableObjectEditorMenuExtensibilityManager() override { return CustomizableObjectEditor_MenuExtensibilityManager; }
@@ -55,7 +58,7 @@ private:
 	FCustomizableObjectEditorLogger Logger;
 
 	// Command to look for Customizable Object Instance in the player pawn of the current world and open its Customizable Object Instance Editor
-	IConsoleCommand* LaunchCOIECommand;
+	IConsoleCommand* LaunchCOIECommand = nullptr;
 
 	FTSTicker::FDelegateHandle WarningsTickerHandle;
 	
diff --git a/Engine/Plugins/Experimental/Mutable/Source/CustomizableObjectEditor/Private/MuCOE/CustomizableObjectEditorViewportClient.cpp b/Engine/Plugins/Experimental/Mutable/Source/CustomizableObjectEditor/Private/MuCOE/CustomizableObjectEditorViewportClient.cpp
--- a/Engine/Plugins/Experimental/Mutable/Source/CustomizableObjectEditor/Private/MuCOE/CustomizableObjectEditorViewportClient.cpp
+++ b/Engine/Plugins/Experimental/Mutable/Source/CustomizableObjectEditor/Private/MuCOE/CustomizableObjectEditorViewportClient.cpp
@@ -1606,7 +1606,7 @@ void FCustomizableObjectEditorViewportClient::BakeInstance()
 	UCustomizableObjectInstance* Instance = CustomizableObjectEditorPtr.Pin()->GetPreviewInstance();
 	if (!Instance)
 	{
-		UE_LOG(LogMutable, Error, TEXT("No Mutable Customizable Object Instnace was found in the current editor."));
+		UE_LOG(LogMutable, Error, TEXT("No Mutable Customizable Object instance was found in the current editor."));
 		return;
 	}
 	
@@ -1624,7 +1624,7 @@ void FCustomizableObjectEditorViewportClient::OnInstanceForBakingUpdate(const FU
 	// Early exit if no instance was provided
 	if (!BakeTempInstance)
 	{
-		UE_LOG(LogMutable, Error, TEXT("No Mutable Customizable Object Instnace was provided for the baking."));
+		UE_LOG(LogMutable, Error, TEXT("No Mutable Customizable Object instance was provided for the baking."));
 		return;
 	}
 
@@ -1649,12 +1649,16 @@ void FCustomizableObjectEditorViewportClient::OnInstanceForBakingUpdate(const FU
 	
 	if (FolderDlg->ShowModal() != EAppReturnType::Cancel)
 	{
+		TArray<TPair<EPackageSaveResolutionType,UPackage*>> SavedPackages;
 		BakeCustomizableObjectInstance(
 			*BakeTempInstance,
 			FolderDlg->GetFileName(),
 			FolderDlg->GetAssetPath(),
 			FolderDlg->GetExportAllResources(),
-			FolderDlg->GetGenerateConstantMaterialInstances());
+			FolderDlg->GetGenerateConstantMaterialInstances(),
+			false,
+			false,
+			SavedPackages);
 	}
 	
 	BakeTempInstance = nullptr;
diff --git a/Engine/Plugins/Experimental/Mutable/Source/CustomizableObjectEditor/Private/MuCOE/CustomizableObjectInstanceBaker.cpp b/Engine/Plugins/Experimental/Mutable/Source/CustomizableObjectEditor/Private/MuCOE/CustomizableObjectInstanceBaker.cpp
new file mode 100644
--- /dev/null
+++ b/Engine/Plugins/Experimental/Mutable/Source/CustomizableObjectEditor/Private/MuCOE/CustomizableObjectInstanceBaker.cpp
@@ -0,0 +1,149 @@
+﻿// Copyright Epic Games, Inc. All Rights Reserved.
+
+#include "MuCOE/CustomizableObjectInstanceBaker.h"
+
+#include "MuCO/CustomizableObject.h"
+#include "MuCO/CustomizableObjectInstance.h"
+#include "MuCO/CustomizableObjectSystem.h"
+#include "MuCOE/CustomizableObjectInstanceBakingUtils.h"
+
+
+void UCustomizableObjectInstanceBaker::BakeInstance(UCustomizableObjectInstance* InTargetInstance, const FBakingConfiguration& InBakingConfig, const TSharedPtr<const FOnBakerFinishedWork> InOnBakerFinishedWork)
+{
+	// Cache the configuration and the callback used to report the baking results
+	Configuration = InBakingConfig;
+
+	// Cache a delegate that we will later invoke once the baking operation has finished all the work (so we can delete this object safelly)
+	if (InOnBakerFinishedWork)
+	{
+		OnWorkFinished = InOnBakerFinishedWork;
+	}
+	
+	if (!InTargetInstance)
+	{
+		UE_LOG(LogMutable,Error,TEXT("No instance has been provided."))
+		FinishBakingOperation(false, nullptr);
+		return;
+	}
+	InstanceToBake = InTargetInstance;
+	
+	if (Configuration.OutputPath.IsEmpty())
+	{
+		UE_LOG(LogMutable,Error,TEXT("No target save path has been provided."))
+		FinishBakingOperation(false, nullptr);
+		return;
+	}
+
+	// If no name is provided then compose one
+	Configuration.OutputFilesBaseName = !Configuration.OutputFilesBaseName.IsEmpty() ? Configuration.OutputFilesBaseName : InstanceToBake.GetName();	
+
+	// Schedule update of the instance for baking
+	UpdateInstance();
+}
+
+
+void UCustomizableObjectInstanceBaker::UpdateInstance()
+{
+	// Call the instance update async method
+	FInstanceUpdateNativeDelegate UpdateDelegate;
+	UpdateDelegate.AddUObject(this, &UCustomizableObjectInstanceBaker::OnInstanceUpdate);
+	
+	// Once the update finishes successfully the baking operation will engage
+	UpdateInstanceForBaking(*InstanceToBake, UpdateDelegate);
+}
+
+
+void UCustomizableObjectInstanceBaker::OnInstanceUpdate(const FUpdateContext& Result) const
+{
+	const EUpdateResult InstanceUpdateResult = Result.UpdateResult;
+	UE_LOG(LogMutable, Display, TEXT("Instance finished update with state : %s."), *UEnum::GetValueAsString(InstanceUpdateResult));
+	
+	if (UCustomizableObjectSystem::IsUpdateResultValid(InstanceUpdateResult))
+	{
+		Bake();
+	}
+	else
+	{
+		// Report a failed baking operation.
+		UE_LOG(LogMutable, Display, TEXT("Instance updating for baking failed. Aborting instance baking operation."));
+		FinishBakingOperation(false, nullptr);
+	}
+}
+
+
+void UCustomizableObjectInstanceBaker::Bake() const
+{
+	// Set this flag so we ensure that the systems that make use of it know that should not display interaction messages to the user
+	const bool GIsRunningUnattendedScript_Cached = GIsRunningUnattendedScript;
+	GIsRunningUnattendedScript = true;
+	
+	// Array with all the packages that did get saved. Packages that failed the saving procedure will not appear in this array.
+	TArray<TPair<EPackageSaveResolutionType,UPackage*>> SavedPackages;
+
+	// Ensures no interaction from the user will be required (alongside GIsRunningUnattendedScript)
+	constexpr bool bIsUnhandledExecution = true;
+	
+	const bool bBakeWasSuccessful = BakeCustomizableObjectInstance
+	(*InstanceToBake,
+		Configuration.OutputFilesBaseName,
+		Configuration.OutputPath,
+		Configuration.bExportAllResourcesOnBake,
+		Configuration.bGenerateConstantMaterialInstancesOnBake,
+		Configuration.bAllowOverridingOfFiles,
+		bIsUnhandledExecution,					
+		SavedPackages
+	);
+	
+	// Generate the array of structs to output based on the packages reported as saved
+	TArray<FBakedResourceData> SavedPackagesData;
+	{
+		// Display a list with the paths of the assets saved
+		UE_LOG(LogMutable, Log, TEXT("Saved asets (%d) :"), SavedPackages.Num());
+		SavedPackagesData.Reserve(SavedPackages.Num());
+		for (const TPair<EPackageSaveResolutionType, UPackage*>& SavedPackagePair : SavedPackages)
+		{
+			FBakedResourceData BakeData;
+			BakeData.SaveType = SavedPackagePair.Key;
+			BakeData.AssetPath = SavedPackagePair.Value->GetPathName();
+
+			// Sanity checks
+			check(BakeData.SaveType != EPackageSaveResolutionType::None);
+			check(!BakeData.AssetPath.IsEmpty());
+			
+			// Log the path for the caller to know
+			UE_LOG(LogMutable, Log, TEXT("\t%s --- %s"),  *UEnum::GetValueAsString(BakeData.SaveType), *BakeData.AssetPath);
+			
+			SavedPackagesData.Push(BakeData);
+		}
+	}
+
+	// Report that the baking operation has been completed
+	UE_LOG(LogMutable, Display, TEXT("Finsishing COI baking procedure."));
+	FinishBakingOperation(bBakeWasSuccessful, &SavedPackagesData);
+	
+	// Reset the GIsRunningUnattendedScript flag 
+	GIsRunningUnattendedScript = GIsRunningUnattendedScript_Cached;
+}
+
+
+void UCustomizableObjectInstanceBaker::FinishBakingOperation(const bool bBakeWasSuccessful, TArray<FBakedResourceData>* InSavedPackagesData) const
+{
+	// Report that the baking operation has been completed
+	if (Configuration.OnBakeOperationCompletedCallback.IsBound())
+	{
+		FCustomizableObjectInstanceBakeOutput Output;
+		Output.bWasBakeSuccessful = bBakeWasSuccessful;
+		if (InSavedPackagesData && !InSavedPackagesData->IsEmpty())
+		{
+			Output.SavedPackages = MoveTemp(*InSavedPackagesData);
+		}
+		Configuration.OnBakeOperationCompletedCallback.Execute(Output);
+	}
+	
+	// Notify whatever caller that the baker completed it's work and therefore is safe to destroy/discard
+	if (OnWorkFinished && OnWorkFinished->IsBound())
+	{
+		OnWorkFinished->Execute();
+	}
+}
+
diff --git a/Engine/Plugins/Experimental/Mutable/Source/CustomizableObjectEditor/Private/MuCOE/CustomizableObjectInstanceBaker.h b/Engine/Plugins/Experimental/Mutable/Source/CustomizableObjectEditor/Private/MuCOE/CustomizableObjectInstanceBaker.h
new file mode 100644
--- /dev/null
+++ b/Engine/Plugins/Experimental/Mutable/Source/CustomizableObjectEditor/Private/MuCOE/CustomizableObjectInstanceBaker.h
@@ -0,0 +1,88 @@
+﻿// Copyright Epic Games, Inc. All Rights Reserved.
+
+#pragma once
+
+#include "MuCO/CustomizableObjectInstance.h"
+
+#include "CustomizableObjectInstanceBaker.generated.h"
+
+// Forward declarations
+class UCustomizableObjectInstance;
+struct FUpdateContext;
+
+
+DECLARE_DELEGATE(FOnBakerFinishedWork);
+
+/**
+ * Utility class designed to allow tbe baking of the resources of a target Mutable Customizable Object instance onto disk.
+ *
+ * The actual baking operation does require some work done before the baking itself.
+ *		- Updating of the instance (async) : The mutable COInstance gets updated prior to the bake to make sure the instance resources to bake are the latest
+ *		- Baking of the instance resources (sync) : The resources of the instance get baked onto disk.
+ *
+ * Once the whole operation gets completed a callback gets executed (OnBakeOperationCompletedCallback) that will provide the caller with the bake operation result
+ * (true for success and false for failure) and a list of asset paths that got saved onto disk.
+ *
+ * Since the update of the instance for baking requires of some static changes in the instance updating system so the parallel update of instances
+ * for baking is not allowed. This prohibition is enforced by the update methods that we call prior to the baking operation. No more than one update for baking can
+ * be ran at the same time.
+ */
+UCLASS()
+class CUSTOMIZABLEOBJECTEDITOR_API UCustomizableObjectInstanceBaker : public UObject
+{
+public:
+	GENERATED_BODY()
+	
+	/**
+	 * Execute this method in order to bake the provided instance. It will schedule a special type of instance update before proceeding with the bake itself.
+	 * @param InTargetInstance The instance we want to bake
+	 * @param InBakingConfig Configuration structure that determines how the baking is going to be made
+	 * @param InOnBakerFinishedWork Delegate used to determine if the baker has completed it's operation. Does not provide any other data than that so we can remove it from
+	 * the root safely (or perform any other cleanup)
+	 */
+	void BakeInstance(UCustomizableObjectInstance* InTargetInstance, const FBakingConfiguration& InBakingConfig, const TSharedPtr<const FOnBakerFinishedWork> InOnBakerFinishedWork = nullptr);
+	
+private:
+	
+	/**
+	 * Takes care of updating and generating the instance resources so we can later bake them.
+	 */
+	void UpdateInstance();
+	
+	/**
+	 * Callback executed when the instance finishes its update. It does not matter if it is successful or not
+	 * @param Result The result returned by the instance updating operation once the operation is completed.
+	 */
+	void OnInstanceUpdate(const FUpdateContext& Result) const;
+
+	/**
+	 * Actual baking method. It will take care of baking the already updated instance and saving it's resources to disk. It will also log the paths
+	 * to the generated packages. 
+	 */
+	void Bake() const;
+
+	/**
+	 * Closes the baking operation by reporting the assets saved ans also , if bound, running the callback that reports that the baker has completed
+	 * all the work that had to be done.
+	 * @param bBakeWasSuccessful Was the baking operation successful or not?
+	 * @param InSavedPackagesData Data from the baked packages. Can be null if no package was saved.
+	 */
+	void FinishBakingOperation(const bool bBakeWasSuccessful, TArray<FBakedResourceData>* InSavedPackagesData) const;
+	
+	/**
+	 * Cached configuration provided by the user
+	 */
+	FBakingConfiguration Configuration;
+	
+	/**
+	 * The instance we want to update and then bake
+	 */
+	UPROPERTY()
+	TObjectPtr<UCustomizableObjectInstance> InstanceToBake;	
+	
+	/**
+	 * Optional delegate used to report when the baker has completed it's work (and it not running or making other systems run)
+	 * It will be called after we have given the end user the data about the success state of the baking and also the resources saved.
+	 */
+	TSharedPtr<const FOnBakerFinishedWork> OnWorkFinished;
+};
diff --git a/Engine/Plugins/Experimental/Mutable/Source/CustomizableObjectEditor/Private/MuCOE/CustomizableObjectInstanceBakingUtils.cpp b/Engine/Plugins/Experimental/Mutable/Source/CustomizableObjectEditor/Private/MuCOE/CustomizableObjectInstanceBakingUtils.cpp
new file mode 100644
--- /dev/null
+++ b/Engine/Plugins/Experimental/Mutable/Source/CustomizableObjectEditor/Private/MuCOE/CustomizableObjectInstanceBakingUtils.cpp
@@ -0,0 +1,889 @@
+﻿// Copyright Epic Games, Inc. All Rights Reserved.
+
+#include "MuCOE/CustomizableObjectInstanceBakingUtils.h"
+
+#include "MuCOE/CustomizableObjectEditor.h"
+
+#include "Misc/MessageDialog.h"
+#include "FileHelpers.h"
+#include "ObjectTools.h"
+#include "UnrealBakeHelpers.h"
+#include "Engine/SkinnedAssetCommon.h"
+#include "Materials/MaterialInstanceConstant.h"
+#include "Materials/MaterialInstanceDynamic.h"
+#include "Materials/Material.h"
+
+#include "MuCO/CustomizableObject.h"
+#include "MuCO/CustomizableObjectInstanceAssetUserData.h"
+#include "MuCO/CustomizableObjectMipDataProvider.h"
+#include "MuCO/ICustomizableObjectEditorModule.h"
+#include "MuT/UnrealPixelFormatOverride.h"
+#include "Rendering/SkeletalMeshModel.h"
+
+#define LOCTEXT_NAMESPACE "CustomizableObjectEditor" 
+
+
+/**
+ * Simple wrapper to be able to invoke the generation of a popup or log message depending on the execution context in which this code is being ran
+ * @param InMessage The message to display
+ * @param InTitle The title to be used for the popup or the log generated
+ */
+void ShowErrorNotification(const FText& InMessage,  const FText* InTitle = nullptr)
+{
+	if (!FApp::IsUnattended())
+	{
+		FMessageDialog::Open(EAppMsgType::Ok, InMessage, *InTitle);
+	}
+	else
+	{
+		if (InTitle)
+		{
+			UE_LOG(LogMutable, Error, TEXT("%s - %s"), *InTitle->ToString(), *InMessage.ToString());
+		}
+		else
+		{
+			UE_LOG(LogMutable, Error, TEXT("%s"), *InMessage.ToString());
+		}
+	}
+}
+
+/**
+ * Utility functions for the baking operation.
+ */
+
+/**
+ * Validates the filename chosen for the baking data
+ * @param FileName The filename chosen by the user
+ * @return True if validation was successful, false otherwise
+ */
+bool ValidateProvidedFileName(const FString& FileName)
+{
+	if (FileName.IsEmpty())
+	{
+		UE_LOG(LogMutable, Error, TEXT("Invalid baking configuration : FileName string is empty.."));
+		return false;
+	}
+
+	// Check for invalid characters in the name of the object to be serialized
+	{
+		TCHAR InvalidCharacter = '0';
+		{
+			FString InvalidCharacters = FPaths::GetInvalidFileSystemChars();
+			for (int32 i = 0; i < InvalidCharacters.Len(); ++i)
+			{
+				TCHAR Char = InvalidCharacters[i];
+				FString SearchedChar = FString::Chr(Char);
+				if (FileName.Contains(SearchedChar))
+				{
+					InvalidCharacter = InvalidCharacters[i];
+					break;
+				}
+			}
+		}
+
+		if (InvalidCharacter != '0')
+		{
+			const FText InvalidCharacterText = FText::FromString(FString::Chr(InvalidCharacter));
+			const FText ErrorText = FText::Format(LOCTEXT("FCustomizableObjectEditorViewportClient_BakeInstance_InvalidCharacter", "The selected contains an invalid character ({0})."), InvalidCharacterText);
+
+			ShowErrorNotification(ErrorText);
+		
+			return false;
+		}
+	}
+
+	return true;
+}
+
+
+/**
+ * Validates the AssetPath chosen for the baking data
+ * @param FileName The filename chosen by the user
+ * @param AssetPath The AssetPath chosen by the user
+ * @param InstanceCO The CustomizableObject from the provided COI
+ * @return True if validation was successful, false otherwise
+ */
+bool ValidateProvidedAssetPath(const FString& FileName, const FString& AssetPath, const UCustomizableObject* InstanceCO)
+{
+	if (AssetPath.IsEmpty())
+	{
+		UE_LOG(LogMutable, Error, TEXT("The AssetPath can not be empty!"));
+		return false;
+	}
+
+	// Ensure we are not overriding the parent CO
+	const FString FullAssetPath = AssetPath + FString("/") + FileName + FString(".") + FileName;		// Full asset path to the new asset we want to create
+	const bool bWouldOverrideParentCO = InstanceCO->GetPathName() == FullAssetPath;
+	if (bWouldOverrideParentCO)
+	{
+		const FText ErrorText = LOCTEXT("FCustomizableObjectEditorViewportClient_BakeInstance_OverwriteCO", "The selected path would overwrite the instance's parent Customizable Object.");
+
+		ShowErrorNotification(ErrorText);
+		
+		return false;
+	}
+
+	return true;
+}
+
+
+/**
+ * Outputs a string that we know it is unique.
+ * @param InResource The resource we are working with
+ * @param ResourceName The name of the resource we have provided. This should have the name of the current resource and will have the unique name for the resource once the method exits
+ * @param InCachedResources Collection with all the already processed resources.
+ * @param InCachedResourceNames Collection with all the already processed resources name's
+ * @return True if the generation of the unique resource name was successful, false otherwise.
+ */
+bool GetUniqueResourceName(const UObject* InResource, FString& ResourceName, TArray<UObject*>& InCachedResources, const TArray<FString>& InCachedResourceNames)
+{
+	check(InResource);
+
+	int32 FindResult = InCachedResourceNames.Find(ResourceName);
+	if (FindResult != INDEX_NONE)
+	{
+		if (InResource == InCachedResources[FindResult])
+		{
+			return false;
+		}
+
+		uint32 Count = 0;
+		while (FindResult != INDEX_NONE)
+		{
+			FindResult = InCachedResourceNames.Find(ResourceName + "_" + FString::FromInt(Count));
+			Count++;
+		}
+
+		ResourceName += "_" + FString::FromInt(--Count);
+	}
+
+	return true;
+}
+
+/**
+ * Ensures the resource we want to save is ready to be saved. It handles closing it's editor and warning the user about possible overriding of resources.
+ * @param InAssetSavePath The directory path where to save the baked object
+ * @param InObjName The name of the object to be baked
+ * @param bOverridePermissionGranted Control flag that determines if the user has given or not permission to override resources already in disk
+ * @param bIsUnattended
+ * @param OutSaveResolution
+ * @return True if the operation was successful, false otherwise
+ */
+bool ManageBakingAction(const FString& InAssetSavePath, const FString& InObjName, bool& bOverridePermissionGranted, const bool& bIsUnattended, EPackageSaveResolutionType& OutSaveResolution)
+{
+	// Before the value provided by "bOverridePermissionGranted" was being updated due to user request but not it is not. It will stay as is if unatended an
+	// will get updated if this gets to be an attended execution.
+	
+	const FString PackagePath = InAssetSavePath + "/" + InObjName;
+	UPackage* ExistingPackage = FindPackage(nullptr, *PackagePath);
+
+	if (!ExistingPackage)
+	{
+		const FString PackageFilePath = PackagePath + "." + InObjName;
+
+		FString PackageFileName;
+		if (FPackageName::DoesPackageExist(PackageFilePath, &PackageFileName))
+		{
+			ExistingPackage = LoadPackage(nullptr, *PackageFileName, LOAD_EditorOnly);
+		}
+		else
+		{
+			// if package does not exist
+			
+			if (!bIsUnattended)
+			{
+				// If the run is attended (the user is participating in it) then we will take care in consideration his decision what he wants.
+				bOverridePermissionGranted = false;
+			}
+
+			OutSaveResolution = EPackageSaveResolutionType::NewFile;
+			return true;
+		}
+	}
+
+	if (ExistingPackage)
+	{
+		// Checking if the asset is open in an editor
+		TArray<IAssetEditorInstance*> ObjectEditors = GEditor->GetEditorSubsystem<UAssetEditorSubsystem>()->FindEditorsForAssetAndSubObjects(ExistingPackage);
+		if (ObjectEditors.Num())
+		{
+			for (IAssetEditorInstance* ObjectEditorInstance : ObjectEditors)
+			{
+				// Close the editors that contains this asset
+				if (!ObjectEditorInstance->CloseWindow(EAssetEditorCloseReason::AssetEditorHostClosed))
+				{
+					const FText Caption = LOCTEXT("OpenExisitngFile", "Open File");
+					const FText Message = FText::Format(LOCTEXT("CantCloseAsset", "This Obejct \"{0}\" is open in an editor and can't be closed automatically. Please close the editor and try to bake it again"), FText::FromString(InObjName));
+
+					ShowErrorNotification(Message, &Caption);
+					
+					return false;
+				}
+			}
+		}
+
+		// If the execution requires user interaction and we have no permission to override the existing file ask him if he wants or not to override data
+		if (!bIsUnattended && !bOverridePermissionGranted)
+		{
+			check (!FApp::IsUnattended())
+			const FText Caption = LOCTEXT("Already existing baked files", "Already existing baked files");
+			const FText Message = FText::Format(LOCTEXT("OverwriteBakedInstance", "Instance baked files already exist in selected destination \"{0}\", this action will overwrite them."), FText::AsCultureInvariant(InAssetSavePath));
+
+			if (FMessageDialog::Open(EAppMsgType::OkCancel, Message, Caption) == EAppReturnType::Cancel)
+			{
+				return false;		// if the user cancels then we will still have no rights for overriding data
+			}
+
+			UE_LOG(LogMutable, Error, TEXT("%s - %s"), *Caption.ToString(), *Message.ToString());
+			
+			// If the user accepts the prompt then we will consider we have a green light to override the asset
+			bOverridePermissionGranted = true;
+		}
+		
+		// At this point we may or may not have permission to delete the existing asset
+		
+		// Delete the old asset if we have permission to do so
+		UObject* ExistingObject = StaticFindObject(UObject::StaticClass(), ExistingPackage, *InObjName);
+		if (ExistingObject)
+		{
+			// Based on if we have or not permission to override the file do or do not so
+			if (bOverridePermissionGranted)
+			{
+				ExistingPackage->FullyLoad();
+
+				TArray<UObject*> ObjectsToDelete;
+				ObjectsToDelete.Add(ExistingObject);
+				
+				const FText Message = FText::Format(LOCTEXT("AssetOverriden", "The COI asset \"{0}\" already exists and will be overriden due to user demand."), FText::FromString(ExistingPackage->GetName()));
+				UE_LOG(LogMutable,Warning,TEXT("%s"), *Message.ToString());
+
+				// Notify the caller we did proceed with the override (performed later)
+				OutSaveResolution = EPackageSaveResolutionType::Overriden;
+				
+				// Delete objects in the package with the same name as the one we want to create
+				const uint32 NumObjectsDeleted = ObjectTools::ForceDeleteObjects(ObjectsToDelete, false);
+				return NumObjectsDeleted == ObjectsToDelete.Num();
+			}
+			else
+			{
+				// Notify the caller that the override will not be performed
+				OutSaveResolution = EPackageSaveResolutionType::UnableToOverride;
+				
+				// Report that the file will not get overriden since we have no permission to do so
+				const FText UnableToOverrideMessage = FText::Format(LOCTEXT("AssetCanNotBeOverriden", "Could not replace the COI asset \"{0}\" as it already exists."), FText::FromString(ExistingPackage->GetName()));
+				UE_LOG(LogMutable,Error,TEXT("%s"), *UnableToOverrideMessage.ToString());
+
+				return false;
+			}
+		}
+		else
+		{
+			// Notify the caller that no override was required
+			OutSaveResolution = EPackageSaveResolutionType::NewFile;
+		}
+	}
+
+	return true;
+}
+
+
+namespace PreBakeSystemSettings
+{
+	bool bIsProgressiveMipStreamingEnabled = false;
+	bool bIsOnlyGenerateRequestedLODsEnabled = false;
+}
+
+
+// Prevents the execution of the baking in parallel for the baking operation. It will not prevent other updates from running (not baking updates)
+// So you are encouraged to halt all other updates while you are baking instances
+static bool bIsUpdateForBakingRunning = false;
+
+
+void PrepareForBaking()
+{
+	// Implementation of the bake operation
+	UCustomizableObjectSystem* System = UCustomizableObjectSystem::GetInstanceChecked();
+	
+	// The instance in the editor viewport does not have high quality mips in the platform data because streaming is enabled.
+	// Disable streaming and retry with a newly generated temp instance.
+	PreBakeSystemSettings::bIsProgressiveMipStreamingEnabled = System->IsProgressiveMipStreamingEnabled();
+	System->SetProgressiveMipStreamingEnabled(false);
+	// Disable requested LOD generation as it will prevent the new instance from having all the LODs
+	PreBakeSystemSettings::bIsOnlyGenerateRequestedLODsEnabled = System->IsOnlyGenerateRequestedLODsEnabled();
+	System->SetOnlyGenerateRequestedLODsEnabled(false);
+	// Force high quality texture compression for this instance
+	PrepareUnrealCompression();
+	System->SetImagePixelFormatOverride(UnrealPixelFormatFunc);
+}
+
+
+void RestoreCustomizableObjectSettings(const FUpdateContext& Result)
+{
+	// Reenable Mutable texture streaming and requested LOD generation as they had been disabled to bake the textures
+	UCustomizableObjectSystem* System =  UCustomizableObjectSystem::GetInstanceChecked();
+	System->SetProgressiveMipStreamingEnabled(PreBakeSystemSettings::bIsProgressiveMipStreamingEnabled);
+	System->SetOnlyGenerateRequestedLODsEnabled(PreBakeSystemSettings::bIsOnlyGenerateRequestedLODsEnabled);
+	System->SetImagePixelFormatOverride(nullptr);
+
+	// Tell the system we have finished the update and that we can schedule another "for baking" update
+	bIsUpdateForBakingRunning = false;
+}
+
+
+void UpdateInstanceForBaking(UCustomizableObjectInstance& InInstance, FInstanceUpdateNativeDelegate& InInstanceUpdateDelegate)
+{
+	if (bIsUpdateForBakingRunning)
+	{
+		UE_LOG(LogMutable, Error, TEXT("The COInstance update for baking could not be scheduled. Another instance is being updated for baking."));
+		InInstanceUpdateDelegate.Broadcast({EUpdateResult::Error});
+		return;
+	}
+
+	// Set the update for the baking of the instance as running so we prevent other baking updates while we do run our own
+	bIsUpdateForBakingRunning = true;
+	
+	// Prepare the customizable object system for baking
+	PrepareForBaking();
+	
+	// Ensure we clear the changes in the COSystem after performing the update so later updates do not get affected
+	InInstanceUpdateDelegate.AddStatic(&RestoreCustomizableObjectSettings);
+	
+	// Schedule the update
+	InInstance.UpdateSkeletalMeshAsyncResult(InInstanceUpdateDelegate,true,true);
+
+	UE_LOG(LogMutable, Display, TEXT("The COInstance Update operation for baking was succesfuly scheduled."));
+}
+
+
+bool BakeCustomizableObjectInstance(
+	UCustomizableObjectInstance& InInstance,
+	const FString& FileName,
+	const FString& AssetPath,
+	const bool bExportAllResources,
+	const bool bGenerateConstantMaterialInstances,
+	bool bHasPermissionToOverride,
+	bool bIsUnattendedExecution,
+	TArray<TPair<EPackageSaveResolutionType,UPackage*>>& OutSavedPackages)
+{
+	OutSavedPackages.Reset();
+	
+	// Ensure that the state of the COI provided is valid --------------------------------------------------------------------------------------------
+	UCustomizableObject* InstanceCO = InInstance.GetCustomizableObject();
+	check (InstanceCO);
+
+	// Ensure the CO of the COI is accessible 
+	if (!InstanceCO || InstanceCO->GetPrivate()->IsLocked())
+	{
+		FCustomizableObjectEditorLogger::CreateLog(
+		LOCTEXT("CustomizableObjectCompilingTryLater_Baking", "Please wait until the Customizable Object is compiled"))
+		.Category(ELoggerCategory::COInstanceBaking)
+		.CustomNotification()
+		.Notification(true)
+		.Log();
+
+		return false;
+	}
+	
+	if (InstanceCO->GetPrivate()->Status.Get() == FCustomizableObjectStatus::EState::Loading)
+	{
+		FCustomizableObjectEditorLogger::CreateLog(
+			LOCTEXT("CustomizableObjectCompileTryLater_BakeInstance","Please wait unitl Customizable Object is loaded"))
+		.Category(ELoggerCategory::COInstanceBaking)
+		.CustomNotification()
+		.Notification(true)
+		.Log();
+
+		return false;
+	}
+	
+	if (!ValidateProvidedFileName(FileName))
+	{
+		UE_LOG(LogMutable, Error, TEXT("The FileName for the instance baking is not valid."));
+		return false;
+	}
+
+	if (!ValidateProvidedAssetPath(FileName,AssetPath,InstanceCO))
+	{
+		UE_LOG(LogMutable, Error, TEXT("The AssetPath for the instance baking is not valid."));
+		return false;
+	}
+	
+	// Exit early if the provided instance does not have a skeletal mesh
+	if (!InInstance.HasAnySkeletalMesh())
+	{
+		UE_LOG(LogMutable, Error, TEXT("The provided instance does not have an skeletal mesh."));
+		return false;
+	}
+
+	// COI Validation completed : Proceed with the baking operation ----------------------------------------------------------------------------------
+	
+	// Notify of better configuration -> Continue operation normally
+	if (InstanceCO->CompileOptions.TextureCompression != ECustomizableObjectTextureCompression::HighQuality)
+	{
+		FCustomizableObjectEditorLogger::CreateLog(
+		LOCTEXT("CustomizableObjectBakeLowQuality", "The Customizable Object wasn't compiled with high quality textures. For the best baking results, change the Texture Compression setting and recompile it."))
+		.Category(ELoggerCategory::COInstanceBaking)
+		.CustomNotification()
+		.Notification(true)
+		.Log();
+	}
+	
+	// Set the overriding flag to true or false:
+	//	- We ask the user at least once about if he is willing to override old baked data (attended operation) and this makes the flag change 
+	//	- We never ask the user (and therefore the value in bUsedGrantedOverridingRights never changes) when we work in Unattended mode.
+	bool bUsedGrantedOverridingRights = bHasPermissionToOverride;
+	if (FApp::IsUnattended() || GIsRunningUnattendedScript)
+	{
+		bIsUnattendedExecution = true;
+	}
+	
+	
+	const int32 NumComponents = InInstance.GetNumComponents();
+	for (int32 ComponentIndex = 0; ComponentIndex < NumComponents; ++ComponentIndex)
+	{
+		USkeletalMesh* Mesh = InInstance.GetSkeletalMesh(ComponentIndex);
+
+		if (!Mesh)
+		{
+			continue;
+		}
+
+		FString ObjectName = FileName;
+		if (NumComponents > 1)
+		{
+			ObjectName = FileName + "_Component_" + FString::FromInt(ComponentIndex);
+		}
+
+		TMap<UObject*, UObject*> ReplacementMap;
+		TArray<FString> ArrayCachedElement;
+		TArray<UObject*> ArrayCachedObject;
+
+		if (bExportAllResources)
+		{
+			UMaterialInstance* Inst;
+			UMaterial* Material;
+			UTexture* Texture;
+			FString MaterialName;
+			FString ResourceName;
+			FString PackageName;
+			UObject* DuplicatedObject;
+			TArray<TMap<int, UTexture*>> TextureReplacementMaps;
+
+			// Duplicate Mutable generated textures
+			for (int32 m = 0; m < Mesh->GetMaterials().Num(); ++m)
+			{
+				UMaterialInterface* Interface = Mesh->GetMaterials()[m].MaterialInterface;
+				Material = Interface->GetMaterial();
+				MaterialName = Material ? Material->GetName() : "Material";
+				Inst = Cast<UMaterialInstance>(Mesh->GetMaterials()[m].MaterialInterface);
+
+				TMap<int, UTexture*> ReplacementTextures;
+				TextureReplacementMaps.Add(ReplacementTextures);
+
+				// The material will only have Mutable generated textures if it's actually a UMaterialInstance
+				if (Material != nullptr && Inst != nullptr)
+				{
+					TArray<FName> ParameterNames = FUnrealBakeHelpers::GetTextureParameterNames(Material);
+
+					for (int32 i = 0; i < ParameterNames.Num(); i++)
+					{
+						if (Inst->GetTextureParameterValue(ParameterNames[i], Texture))
+						{
+							UTexture2D* SrcTex = Cast<UTexture2D>(Texture);
+							if (!SrcTex) continue;
+
+							bool bIsMutableTexture = false;
+
+							for (UAssetUserData* UserData : *SrcTex->GetAssetUserDataArray())
+							{
+								UTextureMipDataProviderFactory* CustomMipDataProviderFactory = Cast<UMutableTextureMipDataProviderFactory>(UserData);
+								if (CustomMipDataProviderFactory)
+								{
+									bIsMutableTexture = true;
+								}
+							}
+
+							if ((SrcTex->GetPlatformData() != nullptr) &&
+								(SrcTex->GetPlatformData()->Mips.Num() > 0) &&
+								bIsMutableTexture)
+							{
+								FString ParameterSanitized = ParameterNames[i].GetPlainNameString();
+								RemoveRestrictedChars(ParameterSanitized);
+								ResourceName = ObjectName + "_" + MaterialName + "_" + ParameterSanitized;
+
+								if (!GetUniqueResourceName(SrcTex, ResourceName, ArrayCachedObject, ArrayCachedElement))
+								{
+									continue;
+								}
+
+								EPackageSaveResolutionType SaveType = EPackageSaveResolutionType::None;
+								if (!ManageBakingAction(AssetPath, ResourceName, bUsedGrantedOverridingRights, bIsUnattendedExecution, SaveType))
+								{
+									return false;
+								}
+
+								// Recover original name of the texture parameter value, now substituted by the generated Mutable texture
+								UTexture* OriginalTexture = nullptr;
+								UMaterialInstanceDynamic* InstDynamic = Cast<UMaterialInstanceDynamic>(Mesh->GetMaterials()[m].MaterialInterface);
+								if (InstDynamic != nullptr)
+								{
+									InstDynamic->Parent->GetTextureParameterValue(FName(*ParameterNames[i].GetPlainNameString()), OriginalTexture);
+								}
+								else
+								{
+									UMaterialInstanceConstant* InstConstant = Cast<UMaterialInstanceConstant>(Mesh->GetMaterials()[m].MaterialInterface);
+
+									if (InstConstant != nullptr)
+									{
+										InstConstant->Parent->GetTextureParameterValue(FName(*ParameterNames[i].GetPlainNameString()), OriginalTexture);
+									}
+								}
+
+								PackageName = AssetPath + FString("/") + ResourceName;
+								TMap<UObject*, UObject*> FakeReplacementMap;
+								UTexture2D* DupTex = FUnrealBakeHelpers::BakeHelper_CreateAssetTexture(SrcTex, ResourceName, PackageName, OriginalTexture, true, FakeReplacementMap, bUsedGrantedOverridingRights);
+								ArrayCachedElement.Add(ResourceName);
+								ArrayCachedObject.Add(DupTex);
+								
+								TPair<EPackageSaveResolutionType, UPackage*> PackageToSave {SaveType, DupTex->GetPackage()};
+								OutSavedPackages.Add(PackageToSave);
+
+								if (OriginalTexture != nullptr)
+								{
+									TextureReplacementMaps[m].Add(i, DupTex);
+								}
+							}
+						}
+					}
+				}
+			}
+
+			// Duplicate non-Mutable material textures
+			for (int32 m = 0; m < Mesh->GetMaterials().Num(); ++m)
+			{
+				UMaterialInterface* Interface = Mesh->GetMaterials()[m].MaterialInterface;
+				Material = Interface->GetMaterial();
+				MaterialName = Material ? Material->GetName() : "Material";
+
+				if (Material != nullptr)
+				{
+					TArray<FName> ParameterNames = FUnrealBakeHelpers::GetTextureParameterNames(Material);
+
+					for (int32 i = 0; i < ParameterNames.Num(); i++)
+					{
+						TArray<FMaterialParameterInfo> InfoArray;
+						TArray<FGuid> GuidArray;
+						Material->GetAllTextureParameterInfo(InfoArray, GuidArray);
+						
+						if (Material->GetTextureParameterValue(InfoArray[i], Texture))
+						{
+							FString ParameterSanitized = ParameterNames[i].GetPlainNameString();
+							RemoveRestrictedChars(ParameterSanitized);
+							ResourceName = ObjectName + "_" + MaterialName + "_" + ParameterSanitized;
+
+							if (ArrayCachedElement.Find(ResourceName) == INDEX_NONE)
+							{
+								EPackageSaveResolutionType SaveType = EPackageSaveResolutionType::None;
+								if (!ManageBakingAction(AssetPath, ResourceName, bUsedGrantedOverridingRights, bIsUnattendedExecution, SaveType ))
+								{
+									return false;
+								}
+
+								PackageName = AssetPath + FString("/") + ResourceName;
+								TMap<UObject*, UObject*> FakeReplacementMap;
+								DuplicatedObject = FUnrealBakeHelpers::BakeHelper_DuplicateAsset(Texture, ResourceName, PackageName, true, FakeReplacementMap, bUsedGrantedOverridingRights, false);
+								ArrayCachedElement.Add(ResourceName);
+								ArrayCachedObject.Add(DuplicatedObject);
+
+								TPair<EPackageSaveResolutionType, UPackage*> PackageToSave {SaveType, DuplicatedObject->GetPackage()};
+								OutSavedPackages.Add(PackageToSave);
+
+								UTexture* DupTexture = Cast<UTexture>(DuplicatedObject);
+								TextureReplacementMaps[m].Add(i, DupTexture);
+							}
+						}
+					}
+				}
+			}
+
+
+			// Duplicate the materials used by each material instance so that the replacement map has proper information 
+			// when duplicating the material instances
+			for (int32 m = 0; m < Mesh->GetMaterials().Num(); ++m)
+			{
+				UMaterialInterface* Interface = Mesh->GetMaterials()[m].MaterialInterface;
+				Material = Interface ? Interface->GetMaterial() : nullptr;
+
+				if (Material)
+				{
+					ResourceName = ObjectName + "_Material_" + Material->GetName();
+
+					if (!GetUniqueResourceName(Material, ResourceName, ArrayCachedObject, ArrayCachedElement))
+					{
+						continue;
+					}
+
+					EPackageSaveResolutionType SaveType = EPackageSaveResolutionType::None;
+					if (!ManageBakingAction(AssetPath, ResourceName, bUsedGrantedOverridingRights, bIsUnattendedExecution, SaveType))
+					{
+						return false;
+					}
+
+					PackageName = AssetPath + FString("/") + ResourceName;
+					TMap<UObject*, UObject*> FakeReplacementMap;
+					DuplicatedObject = FUnrealBakeHelpers::BakeHelper_DuplicateAsset(Material, ResourceName, PackageName, 
+						false, FakeReplacementMap, bUsedGrantedOverridingRights, bGenerateConstantMaterialInstances);
+					ArrayCachedElement.Add(ResourceName);
+					ArrayCachedObject.Add(DuplicatedObject);
+					ReplacementMap.Add(Interface, DuplicatedObject);
+
+					TPair<EPackageSaveResolutionType, UPackage*> PackageToSave {SaveType, DuplicatedObject->GetPackage()};
+					OutSavedPackages.Add(PackageToSave);
+
+					FUnrealBakeHelpers::CopyAllMaterialParameters(DuplicatedObject, Interface, TextureReplacementMaps[m]);
+				}
+			}
+		}
+		else
+		{
+			// Duplicate the material instances
+			for (int32 MaterialIndex = 0; MaterialIndex < Mesh->GetMaterials().Num(); ++MaterialIndex)
+			{
+				UMaterialInterface* Interface = Mesh->GetMaterials()[MaterialIndex].MaterialInterface;
+				UMaterial* ParentMaterial = Interface->GetMaterial();
+				FString MaterialName = ParentMaterial ? ParentMaterial->GetName() : "Material";
+
+				// Material
+				FString MatObjName = ObjectName + "_" + MaterialName;
+
+				if (!GetUniqueResourceName(Interface, MatObjName, ArrayCachedObject, ArrayCachedElement))
+				{
+					continue;
+				}
+
+				EPackageSaveResolutionType SaveType = EPackageSaveResolutionType::None;
+				if (!ManageBakingAction(AssetPath, MatObjName, bUsedGrantedOverridingRights, bIsUnattendedExecution, SaveType))
+				{
+					return false;
+				}
+
+				FString MatPkgName = AssetPath + FString("/") + MatObjName;
+				UObject* DupMat = FUnrealBakeHelpers::BakeHelper_DuplicateAsset(Interface, MatObjName, 
+					MatPkgName, false, ReplacementMap, bUsedGrantedOverridingRights, bGenerateConstantMaterialInstances);
+				ArrayCachedObject.Add(DupMat);
+				ArrayCachedElement.Add(MatObjName);
+
+				TPair<EPackageSaveResolutionType, UPackage*> PackageToSave {SaveType, DupMat->GetPackage()};
+				OutSavedPackages.Add(PackageToSave);
+		
+
+				UMaterialInstance* Inst = Cast<UMaterialInstance>(Interface);
+
+				// Only need to duplicate the generate textures if the original material is a dynamic instance
+				// If the material has Mutable textures, then it will be a dynamic material instance for sure
+				if (Inst)
+				{
+					// Duplicate generated textures
+					UMaterialInstanceDynamic* InstDynamic = Cast<UMaterialInstanceDynamic>(DupMat);
+					UMaterialInstanceConstant* InstConstant = Cast<UMaterialInstanceConstant>(DupMat);
+
+					if (InstDynamic || InstConstant)
+					{
+						for (int32 TextureIndex = 0; TextureIndex < Inst->TextureParameterValues.Num(); ++TextureIndex)
+						{
+							if (Inst->TextureParameterValues[TextureIndex].ParameterValue)
+							{
+								if (Inst->TextureParameterValues[TextureIndex].ParameterValue->HasAnyFlags(RF_Transient))
+								{
+									UTexture2D* SrcTex = Cast<UTexture2D>(Inst->TextureParameterValues[TextureIndex].ParameterValue);
+
+									if (SrcTex)
+									{
+										FString ParameterSanitized = Inst->TextureParameterValues[TextureIndex].ParameterInfo.Name.ToString();
+										RemoveRestrictedChars(ParameterSanitized);
+
+										FString TexObjName = ObjectName + "_" + MaterialName + "_" + ParameterSanitized;
+
+										if (!GetUniqueResourceName(SrcTex, TexObjName, ArrayCachedObject, ArrayCachedElement))
+										{
+											UTexture* PrevTexture = Cast<UTexture>(ArrayCachedObject[ArrayCachedElement.Find(TexObjName)]);
+
+											if (InstDynamic)
+											{
+												InstDynamic->SetTextureParameterValue(Inst->TextureParameterValues[TextureIndex].ParameterInfo.Name, PrevTexture);
+											}
+											else if (InstConstant)
+											{
+												InstConstant->SetTextureParameterValueEditorOnly(Inst->TextureParameterValues[TextureIndex].ParameterInfo.Name, PrevTexture);
+											}
+											
+											continue;
+										}
+
+										EPackageSaveResolutionType TextureSaveType = EPackageSaveResolutionType::None;
+										if (!ManageBakingAction(AssetPath, TexObjName, bUsedGrantedOverridingRights, bIsUnattendedExecution, TextureSaveType))
+										{
+											return false;
+										}
+
+										FString TexPkgName = AssetPath + FString("/") + TexObjName;
+										TMap<UObject*, UObject*> FakeReplacementMap;
+										UTexture2D* DupTex = FUnrealBakeHelpers::BakeHelper_CreateAssetTexture(SrcTex, TexObjName, TexPkgName, nullptr, false, FakeReplacementMap, bUsedGrantedOverridingRights);
+										ArrayCachedObject.Add(DupTex);
+										ArrayCachedElement.Add(TexObjName);
+
+										TPair<EPackageSaveResolutionType, UPackage*> TexturePackageToSave {TextureSaveType, DupTex->GetPackage()};
+										OutSavedPackages.Add(TexturePackageToSave);
+										
+
+										if (InstDynamic)
+										{
+											InstDynamic->SetTextureParameterValue(Inst->TextureParameterValues[TextureIndex].ParameterInfo.Name, DupTex);
+										}
+										else if(InstConstant)
+										{
+											InstConstant->SetTextureParameterValueEditorOnly(Inst->TextureParameterValues[TextureIndex].ParameterInfo.Name, DupTex);
+										}
+									}
+									else
+									{
+										UE_LOG(LogMutable, Error, TEXT("A Mutable texture that is not a Texture2D has been found while baking a CustomizableObjectInstance."));
+									}
+								}
+								else
+								{
+									// If it's not transient it's not a mutable texture, it's a pass-through texture
+									// Just set the original texture
+									if (InstDynamic)
+									{
+										InstDynamic->SetTextureParameterValue(Inst->TextureParameterValues[TextureIndex].ParameterInfo.Name, Inst->TextureParameterValues[TextureIndex].ParameterValue);
+									}
+									else if (InstConstant)
+									{
+										InstConstant->SetTextureParameterValueEditorOnly(Inst->TextureParameterValues[TextureIndex].ParameterInfo.Name, Inst->TextureParameterValues[TextureIndex].ParameterValue);
+									}
+								}
+							}
+						}
+					}
+				}
+			}
+		}
+		
+		// Skeletal Mesh's Skeleton
+		if (Mesh->GetSkeleton())
+		{
+			const bool bTransient = Mesh->GetSkeleton()->GetPackage() == GetTransientPackage();
+
+			// Don't duplicate if not transient or export all assets.
+			if (bTransient || bExportAllResources)
+			{
+				FString SkeletonName = ObjectName + "_Skeleton";
+				EPackageSaveResolutionType SaveType = EPackageSaveResolutionType::None;
+				if (!ManageBakingAction(AssetPath, SkeletonName, bUsedGrantedOverridingRights, bIsUnattendedExecution, SaveType))
+				{
+					return false;
+				}
+
+				FString SkeletonPkgName = AssetPath + FString("/") + SkeletonName;
+				UObject* DuplicatedSkeleton = FUnrealBakeHelpers::BakeHelper_DuplicateAsset(Mesh->GetSkeleton(), SkeletonName, 
+					SkeletonPkgName, false, ReplacementMap, bUsedGrantedOverridingRights, false);
+
+				ArrayCachedObject.Add(DuplicatedSkeleton);
+				TPair<EPackageSaveResolutionType, UPackage*> PackageToSave {SaveType, DuplicatedSkeleton->GetPackage()};
+				OutSavedPackages.Add(PackageToSave);
+				ReplacementMap.Add(Mesh->GetSkeleton(), DuplicatedSkeleton);
+			}
+		}
+
+		// Skeletal Mesh
+		EPackageSaveResolutionType SaveType = EPackageSaveResolutionType::None;
+		if (!ManageBakingAction(AssetPath, ObjectName, bUsedGrantedOverridingRights, bIsUnattendedExecution, SaveType))
+		{
+			return false;
+		}
+
+		FString PkgName = AssetPath + FString("/") + ObjectName;
+		UObject* DupObject = FUnrealBakeHelpers::BakeHelper_DuplicateAsset(Mesh, ObjectName, PkgName, 
+			false, ReplacementMap, bUsedGrantedOverridingRights, false);
+		ArrayCachedObject.Add(DupObject);
+
+		TPair<EPackageSaveResolutionType, UPackage*> PackageToSave {SaveType, DupObject->GetPackage()};
+		OutSavedPackages.Add(PackageToSave);
+
+		Mesh->Build();
+
+		USkeletalMesh* SkeletalMesh = Cast<USkeletalMesh>(DupObject);
+		if (SkeletalMesh)
+		{
+			SkeletalMesh->GetLODInfoArray() = Mesh->GetLODInfoArray();
+
+			SkeletalMesh->GetImportedModel()->SkeletalMeshModelGUID = FGuid::NewGuid();
+
+			// Duplicate AssetUserData
+			{
+				const TArray<UAssetUserData*>* AssetUserDataArray = Mesh->GetAssetUserDataArray();
+				for (const UAssetUserData* AssetUserData : *AssetUserDataArray)
+				{
+					if (AssetUserData)
+					{
+						// Duplicate to change ownership
+						UAssetUserData* NewAssetUserData = Cast<UAssetUserData>(StaticDuplicateObject(AssetUserData, SkeletalMesh));
+						SkeletalMesh->AddAssetUserData(NewAssetUserData);
+					}
+				}
+			}
+
+			// Add Instance Info in a custom AssetUserData
+			{
+				if (InInstance.GetAnimationGameplayTags().Num())
+				{
+					UCustomizableObjectInstanceUserData* InstanceData = NewObject<UCustomizableObjectInstanceUserData>(SkeletalMesh, NAME_None, RF_Public | RF_Transactional);
+					InstanceData->SetAnimationGameplayTags(InInstance.GetAnimationGameplayTags());
+					SkeletalMesh->AddAssetUserData(InstanceData);
+				}
+			}
+
+			// Generate render data
+			SkeletalMesh->Build();
+		}
+
+		// Remove duplicated UObjects from Root (previously added to avoid objects from being GC in the middle of the bake process)
+		for (UObject* Obj : ArrayCachedObject)
+		{
+			Obj->RemoveFromRoot();
+		}
+	}
+
+	// Save the packages generated during the baking operation  --------------------------------------------------------------------------------------
+	
+	// Complete the baking by saving the packages we have cached during the baking operation
+	if (OutSavedPackages.Num())
+	{
+		// Prepare the list of assets we want to provide to "PromptForCheckoutAndSave" for saving
+		TArray<UPackage*> PackagesToSaveProxy;
+		PackagesToSaveProxy.Reserve(OutSavedPackages.Num());
+		for (TPair<EPackageSaveResolutionType, UPackage*> DataToSave : OutSavedPackages)
+		{
+			PackagesToSaveProxy.Push(DataToSave.Value);
+		}
+
+		// List of packages that could not be saved
+		TArray<UPackage*> FailedToSavePackages;
+		const bool bWasSavingSuccessful = FEditorFileUtils::PromptForCheckoutAndSave(PackagesToSaveProxy, false, !bIsUnattendedExecution, &FailedToSavePackages, false, false) == FEditorFileUtils::EPromptReturnCode::PR_Success;
+
+		// Remove all packages that were going to be saved but failed to do so
+		const int32 RemovedPackagesCount = OutSavedPackages.RemoveAll([&](const TPair<EPackageSaveResolutionType, UPackage*> ToSavePackage)
+		{
+			return FailedToSavePackages.Contains(ToSavePackage.Value);
+		});
+		OutSavedPackages.Shrink();
+
+		return RemovedPackagesCount > 0 ? false : bWasSavingSuccessful;
+	}
+	
+	// The operation will fail if no packages are there to save
+	return false;
+}
+
+#undef LOCTEXT_NAMESPACE 
diff --git a/Engine/Plugins/Experimental/Mutable/Source/CustomizableObjectEditor/Private/MuCOE/UnrealBakeHelpers.h b/Engine/Plugins/Experimental/Mutable/Source/CustomizableObjectEditor/Private/MuCOE/UnrealBakeHelpers.h
new file mode 100644
--- /dev/null
+++ b/Engine/Plugins/Experimental/Mutable/Source/CustomizableObjectEditor/Private/MuCOE/UnrealBakeHelpers.h
@@ -0,0 +1,35 @@
+// Copyright Epic Games, Inc. All Rights Reserved.
+
+#pragma once
+
+#include "HAL/Platform.h"
+#include "Containers/Map.h"
+
+class UObject;
+class USkeletalMesh;
+class UTexture2D;
+class UTexture;
+class UMaterialInterface;
+class UMaterial;
+
+
+class CUSTOMIZABLEOBJECTEDITOR_API FUnrealBakeHelpers
+{
+public:
+
+	static UObject* BakeHelper_DuplicateAsset(UObject* Object, const FString& ObjName, const FString& PkgName, bool ResetDuplicatedFlags,
+											  TMap<UObject*, UObject*>& ReplacementMap, bool OverwritePackage,
+											  const bool bGenerateConstantMaterialInstances);
+
+	/**
+	 * Duplicates a texture asset. Duplicates Mutable and non Mutable textures.
+	 *
+	 * @param OrgTex Original source texture from which a Mutable texture has been generated. Only required when SrcTex is a Mutable texture.
+	 */
+	static UTexture2D* BakeHelper_CreateAssetTexture(UTexture2D* SrcTex, const FString& TexObjName, const FString& TexPkgName, const UTexture* OrgTex, bool ResetDuplicatedFlags, TMap<UObject*, UObject*>& ReplacementMap, bool OverwritePackage);
+
+	static void CopyAllMaterialParameters(UObject* DestMaterial, UMaterialInterface* OriginMaterial, const TMap<int, UTexture*>& TextureReplacementMaps);
+
+	/** Get associated texture parameter names for the material given as parameter */
+	static TArray<FName> GetTextureParameterNames(UMaterial* Material);
+};
diff --git a/Engine/Plugins/Experimental/Mutable/Source/CustomizableObjectEditor/Public/MuCOE/CustomizableObjectInstanceBakingUtils.cpp b/Engine/Plugins/Experimental/Mutable/Source/CustomizableObjectEditor/Public/MuCOE/CustomizableObjectInstanceBakingUtils.cpp
deleted file mode 100644
diff --git a/Engine/Plugins/Experimental/Mutable/Source/CustomizableObjectEditor/Public/MuCOE/CustomizableObjectInstanceBakingUtils.h b/Engine/Plugins/Experimental/Mutable/Source/CustomizableObjectEditor/Public/MuCOE/CustomizableObjectInstanceBakingUtils.h
deleted file mode 100644
diff --git a/Engine/Plugins/Experimental/Mutable/Source/CustomizableObjectEditor/Public/MuCOE/UnrealBakeHelpers.h b/Engine/Plugins/Experimental/Mutable/Source/CustomizableObjectEditor/Public/MuCOE/UnrealBakeHelpers.h
deleted file mode 100644
diff --git a/Engine/Plugins/Experimental/Mutable/Source/MutableValidation/Private/MuV/COIBakingTestCommandlet.cpp b/Engine/Plugins/Experimental/Mutable/Source/MutableValidation/Private/MuV/COIBakingTestCommandlet.cpp
--- a/Engine/Plugins/Experimental/Mutable/Source/MutableValidation/Private/MuV/COIBakingTestCommandlet.cpp
+++ b/Engine/Plugins/Experimental/Mutable/Source/MutableValidation/Private/MuV/COIBakingTestCommandlet.cpp
@@ -104,7 +104,7 @@ int32 UCOIBakingTestCommandlet::Main(const FString& Params)
 		InstanceUpdateDelegate.AddStatic(&OnInstanceUpdate);
 		
 		bIsInstanceBeingUpdated = true;
-		UE_LOG(LogMutable,Display,TEXT("Invoking instance update."));
+		UE_LOG(LogMutable, Display, TEXT("Scheduling instance update."));
 		UpdateInstanceForBaking(*TargetInstance, InstanceUpdateDelegate);
 		
 		// Now tick the engine so the instance gets updated while running in the commandlet context
@@ -130,6 +130,7 @@ int32 UCOIBakingTestCommandlet::Main(const FString& Params)
 	
 
 	// Bake the instance
+	bool bWasBakingSuccessful = false;
 	{
 		const FString BakedResourcesFileName = "MuBakedInstances";
 		
@@ -160,14 +161,17 @@ int32 UCOIBakingTestCommandlet::Main(const FString& Params)
 		{
 			UE_LOG(LogMutable,Display,TEXT("Starting Instance Baking operation."));
 			{
-				BakeCustomizableObjectInstance(
+				TArray<TPair<EPackageSaveResolutionType,UPackage*>> SavedPackages;
+				bWasBakingSuccessful = BakeCustomizableObjectInstance(
 					*TargetInstance,
 					FString::Printf(TEXT("%s_Bake"), *TargetInstance->GetName()),
 					LocalBakingDirectory,
 					true,
-					true);
+					true,
+					true,
+					true,
+					SavedPackages);
 			}
-			UE_LOG(LogMutable,Display,TEXT("Instance Baking operation has been completed."));
 			
 			// Delete the target directory where we did save the baked instance
 			if (!IFileManager::Get().DeleteDirectory(*GlobalBakingDirectory,true,true))
@@ -182,8 +186,17 @@ int32 UCOIBakingTestCommandlet::Main(const FString& Params)
 			return 1;
 		}
 	}
-	
-	return 0;
+
+	if (bWasBakingSuccessful)
+	{
+		UE_LOG(LogMutable, Display, TEXT("Instance Baking operation has been completed succesfully."));
+		return 0;
+	}
+	else
+	{
+		UE_LOG(LogMutable, Display, TEXT("Instance Baking operation has been completed with errors."));
+		return 1;
+	}
 }
 
 
