[mutable] Mesh skinning size optimizations.

- Reduce the skinning data # bones if possible.
- Moved some related structures to UE naming convention.

#rnx

diff --git a/Engine/Plugins/Experimental/Mutable/Source/CustomizableObject/Internal/MuCO/CustomizableObjectPrivate.h b/Engine/Plugins/Experimental/Mutable/Source/CustomizableObject/Internal/MuCO/CustomizableObjectPrivate.h
--- a/Engine/Plugins/Experimental/Mutable/Source/CustomizableObject/Internal/MuCO/CustomizableObjectPrivate.h
+++ b/Engine/Plugins/Experimental/Mutable/Source/CustomizableObject/Internal/MuCO/CustomizableObjectPrivate.h
@@ -826,6 +826,6 @@ public:
 	// This is a manual version number for the binary blobs in this asset.
 	// Increasing it invalidates all the previously compiled models.
 	// Warning: If while merging code both versions have changed, take the highest+1.
-	static constexpr int32 CurrentSupportedVersion = 437;
+	static constexpr int32 CurrentSupportedVersion = 438;
 };
 
diff --git a/Engine/Plugins/Experimental/Mutable/Source/CustomizableObject/Private/MuCO/MutableMeshBufferUtils.cpp b/Engine/Plugins/Experimental/Mutable/Source/CustomizableObject/Private/MuCO/MutableMeshBufferUtils.cpp
--- a/Engine/Plugins/Experimental/Mutable/Source/CustomizableObject/Private/MuCO/MutableMeshBufferUtils.cpp
+++ b/Engine/Plugins/Experimental/Mutable/Source/CustomizableObject/Private/MuCO/MutableMeshBufferUtils.cpp
@@ -12,9 +12,9 @@ void MutableMeshBufferUtils::SetupVertexPositionsBuffer(const int32& InCurrentVe
 	using namespace mu;
 	const int32 ElementSize = sizeof(FPositionVertex);
 	constexpr int32 ChannelCount = 1;
-	const MESH_BUFFER_SEMANTIC Semantics[ChannelCount] = {MBS_POSITION};
+	const EMeshBufferSemantic Semantics[ChannelCount] = {MBS_POSITION};
 	const int32 SemanticIndices[ChannelCount] = {0};
-	const MESH_BUFFER_FORMAT Formats[ChannelCount] = {MBF_FLOAT32};
+	const EMeshBufferFormat Formats[ChannelCount] = {MBF_FLOAT32};
 	const int32 Components[ChannelCount] = {3};
 	const int32 Offsets[ChannelCount] =
 	{
@@ -36,9 +36,9 @@ void MutableMeshBufferUtils::SetupTangentBuffer(const int32& InCurrentVertexBuff
 	using namespace mu;
 	const int32 ElementSize = sizeof(TangentType);
 	constexpr int32 ChannelCount = 2;
-	const MESH_BUFFER_SEMANTIC Semantics[ChannelCount] = {MBS_TANGENT, MBS_NORMAL};
+	const EMeshBufferSemantic Semantics[ChannelCount] = {MBS_TANGENT, MBS_NORMAL};
 	const int32 SemanticIndices[ChannelCount] = {0, 0};
-	const MESH_BUFFER_FORMAT Formats[ChannelCount] = {MBF_PACKEDDIRS8, MBF_PACKEDDIRS8_W_TANGENTSIGN};
+	const EMeshBufferFormat Formats[ChannelCount] = {MBF_PACKEDDIRS8, MBF_PACKEDDIRS8_W_TANGENTSIGN};
 	const int32 Components[ChannelCount] = {4, 4};
 	const int32 Offsets[ChannelCount] =
 	{
@@ -61,10 +61,10 @@ void MutableMeshBufferUtils::SetupTexCoordinatesBuffer(const int32& InCurrentVer
 	using namespace mu;
 	const int32 ElementSize = sizeof(UVType) * InChannelCount;
 	constexpr int32 MaxChannelCount = MaxTexCordChannelCount;
-	const MESH_BUFFER_SEMANTIC Semantics[MaxChannelCount] = {
+	const EMeshBufferSemantic Semantics[MaxChannelCount] = {
 		MBS_TEXCOORDS, MBS_TEXCOORDS, MBS_TEXCOORDS, MBS_TEXCOORDS
 	};
-	const MESH_BUFFER_FORMAT Formats[MaxChannelCount] = {MBF_FLOAT32, MBF_FLOAT32, MBF_FLOAT32, MBF_FLOAT32};
+	const EMeshBufferFormat Formats[MaxChannelCount] = {MBF_FLOAT32, MBF_FLOAT32, MBF_FLOAT32, MBF_FLOAT32};
 	const int32 Components[MaxChannelCount] = {2, 2, 2, 2};
 	const int32 Offsets[MaxChannelCount] =
 	{
@@ -97,10 +97,10 @@ void MutableMeshBufferUtils::SetupSkinBuffer(const int32& InCurrentVertexBuffer,
 	using namespace mu;
 	const int32 ElementSize = (MaxBoneWeightTypeSizeBytes + MaxBoneIndexTypeSizeBytes) * MaxNumBonesPerVertex;
 	constexpr int32 ChannelCount = 2;
-	const MESH_BUFFER_SEMANTIC Semantics[ChannelCount] = {MBS_BONEINDICES, MBS_BONEWEIGHTS};
+	const EMeshBufferSemantic Semantics[ChannelCount] = {MBS_BONEINDICES, MBS_BONEWEIGHTS};
 	const int32 SemanticIndices[ChannelCount] = {0, 0};
 
-	MESH_BUFFER_FORMAT Formats[ChannelCount] = {MBF_UINT8, MBF_NUINT8};
+	EMeshBufferFormat Formats[ChannelCount] = {MBF_UINT8, MBF_NUINT8};
 	switch (MaxBoneIndexTypeSizeBytes)
 	{
 	case 0: // Fallback to something in this case.
@@ -151,9 +151,9 @@ void MutableMeshBufferUtils::SetupVertexColorBuffer(const int32& InCurrentVertex
 	using namespace mu;
 	const int32 ElementSize = sizeof(FColor);
 	constexpr int32 ChannelCount = 1;
-	const MESH_BUFFER_SEMANTIC Semantics[ChannelCount] = {MBS_COLOUR};
+	const EMeshBufferSemantic Semantics[ChannelCount] = {MBS_COLOUR};
 	const int32 SemanticIndices[ChannelCount] = {0};
-	const MESH_BUFFER_FORMAT Formats[ChannelCount] = {MBF_NUINT8};
+	const EMeshBufferFormat Formats[ChannelCount] = {MBF_NUINT8};
 	const int32 Components[ChannelCount] = {4};
 	const int32 Offsets[ChannelCount] = {0};
 	check(ElementSize == 4);
@@ -171,12 +171,12 @@ void MutableMeshBufferUtils::SetupIndexBuffer(mu::FMeshBufferSet& OutTargetIndex
 	const int32 ElementSize = sizeof(uint32);
 	//SkeletalMesh->GetImportedResource()->LODModels[LOD].MultiSizeIndexContainer.GetDataTypeSize();
 	constexpr int32 ChannelCount = 1;
-	const MESH_BUFFER_SEMANTIC Semantics[ChannelCount] = {MBS_VERTEXINDEX};
+	const EMeshBufferSemantic Semantics[ChannelCount] = {MBS_VERTEXINDEX};
 	const int32 SemanticIndices[ChannelCount] = {0};
 	// We force 32 bit indices, since merging meshes may create vertex buffers bigger than the initial mesh
 	// and for now the mutable runtime doesn't handle it.
 	// \TODO: go back to 16-bit indices when possible.
-	MESH_BUFFER_FORMAT Formats[ChannelCount] = {MBF_UINT32};
+	EMeshBufferFormat Formats[ChannelCount] = {MBF_UINT32};
 	const int32 Components[ChannelCount] = {1};
 	const int32 Offsets[ChannelCount] = {0};
 
@@ -195,10 +195,10 @@ void MutableMeshBufferUtils::SetupSkinWeightProfileBuffer(const int32& InCurrent
 	using namespace mu;
 	const int32 ElementSize = sizeof(int32) + (MaxBoneIndexTypeSizeBytes + 1) * MaxNumBonesPerVertex;
 	constexpr int32 ChannelCount = 3;
-	const MESH_BUFFER_SEMANTIC Semantics[ChannelCount] = { MBS_OTHER, MBS_BONEINDICES, MBS_BONEWEIGHTS };
+	const EMeshBufferSemantic Semantics[ChannelCount] = { MBS_OTHER, MBS_BONEINDICES, MBS_BONEWEIGHTS };
 	const int32 SemanticIndices[ChannelCount] = { SemanticsIndex, SemanticsIndex, SemanticsIndex };
 
-	MESH_BUFFER_FORMAT Formats[ChannelCount] = { MBF_INT32, MBF_UINT8, MBF_NUINT8 };
+	EMeshBufferFormat Formats[ChannelCount] = { MBF_INT32, MBF_UINT8, MBF_NUINT8 };
 	switch (MaxBoneIndexTypeSizeBytes)
 	{
 	case 1: Formats[1] = mu::MBF_UINT8;
diff --git a/Engine/Plugins/Experimental/Mutable/Source/CustomizableObject/Private/MuCO/UnrealConversionUtils.cpp b/Engine/Plugins/Experimental/Mutable/Source/CustomizableObject/Private/MuCO/UnrealConversionUtils.cpp
--- a/Engine/Plugins/Experimental/Mutable/Source/CustomizableObject/Private/MuCO/UnrealConversionUtils.cpp
+++ b/Engine/Plugins/Experimental/Mutable/Source/CustomizableObject/Private/MuCO/UnrealConversionUtils.cpp
@@ -88,41 +88,73 @@ namespace UnrealConversionUtils
 			const void* InMutableData,
 			const uint32 MutableDataSize)
 		{
+			OutVertexWeightBuffer.SetNeedsCPUAccess(bNeedCPUAccess);
+
 			FSkinWeightDataVertexBuffer* VertexBuffer = OutVertexWeightBuffer.GetDataVertexBuffer();
 			VertexBuffer->SetMaxBoneInfluences(NumBoneInfluences);
-			VertexBuffer->Init(NumBones, NumVertices);
-
-			bool bIsVariableBonesPerVertex = OutVertexWeightBuffer.GetDataVertexBuffer()->GetVariableBonesPerVertex();
-			check(!FGPUBaseSkinVertexFactory::UseUnlimitedBoneInfluences(NumBoneInfluences) || bIsVariableBonesPerVertex);
-
-			OutVertexWeightBuffer.SetNeedsCPUAccess(bNeedCPUAccess);
 
-			if (NumVertices)
+			if (!NumVertices)
 			{
-				void* Data = VertexBuffer->GetWeightData();
+				return;
+			}
 
-				uint32 OutVertexWeightBufferSize = OutVertexWeightBuffer.GetDataVertexBuffer()->GetVertexDataSize();
+			int32 MaxBoneInfluences = VertexBuffer->GetMaxBoneInfluences();
+			if (MaxBoneInfluences == NumBoneInfluences)
+			{
+				VertexBuffer->Init(NumBones, NumVertices);
+				uint32 OutVertexWeightBufferSize = VertexBuffer->GetVertexDataSize();
 				ensure(MutableDataSize == OutVertexWeightBufferSize);
-
+				void* Data = VertexBuffer->GetWeightData();
 				FMemory::Memcpy(Data, InMutableData, OutVertexWeightBufferSize);
-
-				if (bIsVariableBonesPerVertex)
+			}
+			else
+			{
+				// We need to expand it with blank data interleaved
+				uint32 MutableVertexDataSize = MutableDataSize / NumVertices;
+				uint32 FinalVertexDataSize = ( MutableDataSize / NumBones) * MaxBoneInfluences;
+				VertexBuffer->Init(NumVertices*MaxBoneInfluences, NumVertices);
+				uint32 OutVertexWeightBufferSize = VertexBuffer->GetVertexDataSize();
+				ensure(FinalVertexDataSize*NumVertices == OutVertexWeightBufferSize);
+
+				int32 BoneIndexSize = OutVertexWeightBuffer.GetBoneIndexByteSize();
+				int32 WeightSize = OutVertexWeightBuffer.GetBoneWeightByteSize();
+
+				const uint8* MutableData = reinterpret_cast<const uint8*>(InMutableData);
+				uint8* Data = VertexBuffer->GetWeightData();
+				FMemory::Memzero(Data, OutVertexWeightBufferSize);
+				for (int32 V=0; V<NumVertices; ++V)
 				{
-					OutVertexWeightBuffer.RebuildLookupVertexBuffer();
+					// Bone indices
+					FMemory::Memcpy(Data, MutableData, NumBoneInfluences*BoneIndexSize);
+					MutableData += NumBoneInfluences * BoneIndexSize;
+					Data += MaxBoneInfluences * BoneIndexSize;
+
+					// Weights
+					FMemory::Memcpy(Data, MutableData, NumBoneInfluences * WeightSize);
+					MutableData += NumBoneInfluences * WeightSize;
+					Data += MaxBoneInfluences * WeightSize;
+				}
+			}
 
-					{
-						MUTABLE_CPUPROFILER_SCOPE(OptimizeVertexAndLookupBuffers);
+			bool bIsVariableBonesPerVertex = VertexBuffer->GetVariableBonesPerVertex();
+			check(!FGPUBaseSkinVertexFactory::UseUnlimitedBoneInfluences(NumBoneInfluences) || bIsVariableBonesPerVertex);
+			if (bIsVariableBonesPerVertex)
+			{
+				OutVertexWeightBuffer.RebuildLookupVertexBuffer();
 
-						// Everything in this scope is optional and makes extra copies, but will optimize the variable bone
-						// influences buffers. Without it, the vertices are assumed to have a constant NumBoneInfluences per vertex.
-						TArray<FSkinWeightInfo> TempVertices;
-						OutVertexWeightBuffer.GetSkinWeights(TempVertices);
+				{
+					MUTABLE_CPUPROFILER_SCOPE(OptimizeVertexAndLookupBuffers);
 
-						// The assignment operator actually optimizes the DataVertexBuffer
-						OutVertexWeightBuffer = TempVertices;
-					}
+					// Everything in this scope is optional and makes extra copies, but will optimize the variable bone
+					// influences buffers. Without it, the vertices are assumed to have a constant NumBoneInfluences per vertex.
+					TArray<FSkinWeightInfo> TempVertices;
+					OutVertexWeightBuffer.GetSkinWeights(TempVertices);
+
+					// The assignment operator actually optimizes the DataVertexBuffer
+					OutVertexWeightBuffer = TempVertices;
 				}
 			}
+	
 		}
 	}
 
@@ -218,7 +250,7 @@ namespace UnrealConversionUtils
 			MutableMeshVertexBuffers.GetBufferData(MUTABLE_VERTEXBUFFER_TEXCOORDS)
 		);
 
-		mu::MESH_BUFFER_FORMAT BoneIndexFormat = mu::MBF_NONE;
+		mu::EMeshBufferFormat BoneIndexFormat = mu::MBF_NONE;
 		int32 NumBoneInfluences = 0;
 		int32 BoneIndexBuffer = -1;
 		int32 BoneIndexChannel = -1;
@@ -229,7 +261,7 @@ namespace UnrealConversionUtils
 				nullptr, nullptr, &BoneIndexFormat, &NumBoneInfluences, nullptr);
 		}
 
-		mu::MESH_BUFFER_FORMAT BoneWeightFormat = mu::MBF_NONE;
+		mu::EMeshBufferFormat BoneWeightFormat = mu::MBF_NONE;
 		int32 BoneWeightBuffer = -1;
 		int32 BoneWeightChannel = -1;
 		MutableMeshVertexBuffers.FindChannel(mu::MBS_BONEWEIGHTS, 0, &BoneWeightBuffer, &BoneWeightChannel);
@@ -265,8 +297,8 @@ namespace UnrealConversionUtils
 		{
 			if (MutableMeshVertexBuffers.GetBufferChannelCount(Buffer) > 0)
 			{
-				mu::MESH_BUFFER_SEMANTIC Semantic;
-				mu::MESH_BUFFER_FORMAT Format;
+				mu::EMeshBufferSemantic Semantic;
+				mu::EMeshBufferFormat Format;
 				int32 SemanticIndex;
 				int32 ComponentCount;
 				int32 Offset;
@@ -312,7 +344,7 @@ namespace UnrealConversionUtils
 			MutableMeshVertexBuffers.GetBufferData(MUTABLE_VERTEXBUFFER_TEXCOORDS)
 		);
 
-		mu::MESH_BUFFER_FORMAT BoneIndexFormat = mu::MBF_NONE;
+		mu::EMeshBufferFormat BoneIndexFormat = mu::MBF_NONE;
 		int32 NumBoneInfluences = 0;
 		int32 BoneIndexBuffer = -1;
 		int32 BoneIndexChannel = -1;
@@ -323,7 +355,7 @@ namespace UnrealConversionUtils
 				nullptr, nullptr, &BoneIndexFormat, &NumBoneInfluences, nullptr);
 		}
 
-		mu::MESH_BUFFER_FORMAT BoneWeightFormat = mu::MBF_NONE;
+		mu::EMeshBufferFormat BoneWeightFormat = mu::MBF_NONE;
 		int32 BoneWeightBuffer = -1;
 		int32 BoneWeightChannel = -1;
 		MutableMeshVertexBuffers.FindChannel(mu::MBS_BONEWEIGHTS, 0, &BoneWeightBuffer, &BoneWeightChannel);
@@ -359,8 +391,8 @@ namespace UnrealConversionUtils
 		{
 			if (MutableMeshVertexBuffers.GetBufferChannelCount(Buffer) > 0)
 			{
-				mu::MESH_BUFFER_SEMANTIC Semantic;
-				mu::MESH_BUFFER_FORMAT Format;
+				mu::EMeshBufferSemantic Semantic;
+				mu::EMeshBufferFormat Format;
 				int32 SemanticIndex;
 				int32 ComponentCount;
 				int32 Offset;
@@ -466,7 +498,7 @@ namespace UnrealConversionUtils
 		const int32 BoneIndexSize = b16BitBoneIndices ? 2 : 1;
 
 		// BoneWeights channel info
-		mu::MESH_BUFFER_FORMAT Format;
+		mu::EMeshBufferFormat Format;
 		int32 MutableNumInfluences;
 		int32 Offset;
 		MutableMeshVertexBuffers.GetChannel(BoneIndexBuffer, 2, nullptr, nullptr, &Format, &MutableNumInfluences, &Offset);
diff --git a/Engine/Plugins/Experimental/Mutable/Source/CustomizableObjectEditor/Private/MuCOE/GenerateMutableSource/GenerateMutableSourceMesh.cpp b/Engine/Plugins/Experimental/Mutable/Source/CustomizableObjectEditor/Private/MuCOE/GenerateMutableSource/GenerateMutableSourceMesh.cpp
--- a/Engine/Plugins/Experimental/Mutable/Source/CustomizableObjectEditor/Private/MuCOE/GenerateMutableSource/GenerateMutableSourceMesh.cpp
+++ b/Engine/Plugins/Experimental/Mutable/Source/CustomizableObjectEditor/Private/MuCOE/GenerateMutableSource/GenerateMutableSourceMesh.cpp
@@ -866,9 +866,9 @@ mu::MeshPtr ConvertSkeletalMeshToMutable(const USkeletalMesh* InSkeletalMesh, co
 		using namespace mu;
 		const int ElementSize = sizeof(FSoftSkinVertex);
 		const int ChannelCount = 9;
-		const MESH_BUFFER_SEMANTIC Semantics[ChannelCount] = { MBS_POSITION, MBS_TANGENT, MBS_BINORMAL, MBS_NORMAL, MBS_TEXCOORDS, MBS_TEXCOORDS, MBS_TEXCOORDS, MBS_TEXCOORDS, MBS_COLOUR };
+		const EMeshBufferSemantic Semantics[ChannelCount] = { MBS_POSITION, MBS_TANGENT, MBS_BINORMAL, MBS_NORMAL, MBS_TEXCOORDS, MBS_TEXCOORDS, MBS_TEXCOORDS, MBS_TEXCOORDS, MBS_COLOUR };
 		const int SemanticIndices[ChannelCount] = { 0, 0, 0, 0, 0, 1, 2, 3, 0 };
-		const MESH_BUFFER_FORMAT Formats[ChannelCount] = { MBF_FLOAT32, MBF_FLOAT32, MBF_FLOAT32, MBF_FLOAT32, MBF_FLOAT32, MBF_FLOAT32, MBF_FLOAT32, MBF_FLOAT32, MBF_NUINT8 };
+		const EMeshBufferFormat Formats[ChannelCount] = { MBF_FLOAT32, MBF_FLOAT32, MBF_FLOAT32, MBF_FLOAT32, MBF_FLOAT32, MBF_FLOAT32, MBF_FLOAT32, MBF_FLOAT32, MBF_NUINT8 };
 		int Components[ChannelCount] = { 3, 3, 3, 4, 2, 2, 2, 2, 4 };
 
 		constexpr size_t SoftSkinVertexUVsElemSize = sizeof(TDecay<decltype(DeclVal<FSoftSkinVertex>().UVs[0])>::Type);
@@ -893,12 +893,12 @@ mu::MeshPtr ConvertSkeletalMeshToMutable(const USkeletalMesh* InSkeletalMesh, co
 		using namespace mu;
 		const int ElementSize = sizeof(FSoftSkinVertex);
 		const int ChannelCount = 11;
-		const MESH_BUFFER_SEMANTIC Semantics[ChannelCount] = { MBS_POSITION, MBS_TANGENT, MBS_BINORMAL, MBS_NORMAL, MBS_TEXCOORDS, MBS_TEXCOORDS, MBS_TEXCOORDS, MBS_TEXCOORDS, MBS_COLOUR, MBS_BONEINDICES, MBS_BONEWEIGHTS };
+		const EMeshBufferSemantic Semantics[ChannelCount] = { MBS_POSITION, MBS_TANGENT, MBS_BINORMAL, MBS_NORMAL, MBS_TEXCOORDS, MBS_TEXCOORDS, MBS_TEXCOORDS, MBS_TEXCOORDS, MBS_COLOUR, MBS_BONEINDICES, MBS_BONEWEIGHTS };
 		const int SemanticIndices[ChannelCount] = { 0, 0, 0, 0, 0, 1, 2, 3, 0, 0, 0 };
 
 		// TODO: Remove BoneWeightFormat after merge
-		MESH_BUFFER_FORMAT BoneWeightFormat = sizeof(TDecay<decltype(DeclVal<FSoftSkinVertex>().InfluenceWeights[0])>::Type) == 1 ? MBF_NUINT8 : MBF_NUINT16;
-		const MESH_BUFFER_FORMAT Formats[ChannelCount] = { MBF_FLOAT32, MBF_FLOAT32, MBF_FLOAT32, MBF_FLOAT32, MBF_FLOAT32, MBF_FLOAT32, MBF_FLOAT32, MBF_FLOAT32, MBF_NUINT8, MBF_UINT16, BoneWeightFormat };
+		EMeshBufferFormat BoneWeightFormat = sizeof(TDecay<decltype(DeclVal<FSoftSkinVertex>().InfluenceWeights[0])>::Type) == 1 ? MBF_NUINT8 : MBF_NUINT16;
+		const EMeshBufferFormat Formats[ChannelCount] = { MBF_FLOAT32, MBF_FLOAT32, MBF_FLOAT32, MBF_FLOAT32, MBF_FLOAT32, MBF_FLOAT32, MBF_FLOAT32, MBF_FLOAT32, MBF_NUINT8, MBF_UINT16, BoneWeightFormat };
 
 		int Components[ChannelCount] = { 3, 3, 3, 4, 2, 2, 2, 2, 4, 4, 4 };
 		if (GenerationContext.Options.CustomizableObjectNumBoneInfluences != ECustomizableObjectNumBoneInfluences::Four && 
@@ -1122,9 +1122,9 @@ mu::MeshPtr ConvertSkeletalMeshToMutable(const USkeletalMesh* InSkeletalMesh, co
 			using namespace mu;
 			const int32 ElementSize = sizeof(int32);
 			const int32 ChannelCount = 1;
-			const MESH_BUFFER_SEMANTIC Semantics[ChannelCount] = { MBS_OTHER };
+			const EMeshBufferSemantic Semantics[ChannelCount] = { MBS_OTHER };
 			const int32 SemanticIndices[ChannelCount] = { 0 };
-			const MESH_BUFFER_FORMAT Formats[ChannelCount] = { MBF_INT32 };
+			const EMeshBufferFormat Formats[ChannelCount] = { MBF_INT32 };
 			int32 Components[ChannelCount] = { 1 };
 			const int32 Offsets[ChannelCount] = { 0 };
 
@@ -1136,9 +1136,9 @@ mu::MeshPtr ConvertSkeletalMeshToMutable(const USkeletalMesh* InSkeletalMesh, co
 			using namespace mu;
 			const int32 ElementSize = sizeof(uint16);
 			const int32 ChannelCount = 1;
-			const MESH_BUFFER_SEMANTIC Semantics[ChannelCount] = { MBS_OTHER };
+			const EMeshBufferSemantic Semantics[ChannelCount] = { MBS_OTHER };
 			const int32 SemanticIndices[ChannelCount] = { 1 };
-			const MESH_BUFFER_FORMAT Formats[ChannelCount] = { MBF_UINT16 };
+			const EMeshBufferFormat Formats[ChannelCount] = { MBF_UINT16 };
 			int32 Components[ChannelCount] = { 1 };
 			const int32 Offsets[ChannelCount] = { 0 };
 
@@ -1150,9 +1150,9 @@ mu::MeshPtr ConvertSkeletalMeshToMutable(const USkeletalMesh* InSkeletalMesh, co
 			using namespace mu;
 			const int32 ElementSize = sizeof(uint16);
 			const int32 ChannelCount = 1;
-			const MESH_BUFFER_SEMANTIC Semantics[ChannelCount] = { MBS_OTHER };
+			const EMeshBufferSemantic Semantics[ChannelCount] = { MBS_OTHER };
 			const int32 SemanticIndices[ChannelCount] = { 2 };
-			const MESH_BUFFER_FORMAT Formats[ChannelCount] = { MBF_UINT16 };
+			const EMeshBufferFormat Formats[ChannelCount] = { MBF_UINT16 };
 			int32 Components[ChannelCount] = { 1 };
 			const int32 Offsets[ChannelCount] = { 0 };
 
@@ -1269,9 +1269,9 @@ mu::MeshPtr ConvertSkeletalMeshToMutable(const USkeletalMesh* InSkeletalMesh, co
 			using namespace mu;
 			const int32 ElementSize = sizeof(int32);
 			const int32 ChannelCount = 1;
-			const MESH_BUFFER_SEMANTIC Semantics[ChannelCount] = { MBS_OTHER };
+			const EMeshBufferSemantic Semantics[ChannelCount] = { MBS_OTHER };
 			const int32 SemanticIndices[ChannelCount] = { GenerationContext.Options.bRealTimeMorphTargetsEnabled ? 3 : 0 };
-			const MESH_BUFFER_FORMAT Formats[ChannelCount] = { MBF_INT32 };
+			const EMeshBufferFormat Formats[ChannelCount] = { MBF_INT32 };
 			int32 Components[ChannelCount] = { 1 };
 			const int32 Offsets[ChannelCount] = { 0 };
 
@@ -1411,7 +1411,7 @@ mu::MeshPtr ConvertSkeletalMeshToMutable(const USkeletalMesh* InSkeletalMesh, co
 
 		// TODO: Remove BoneWeightFormat after merge
 		const int32 BoneWeightTypeSizeBytes = sizeof(TDecay<decltype(DeclVal<FRawSkinWeight>().InfluenceWeights[0])>::Type);
-		MESH_BUFFER_FORMAT BoneWeightFormat = BoneWeightTypeSizeBytes == 1 ? MBF_NUINT8 : MBF_NUINT16;
+		EMeshBufferFormat BoneWeightFormat = BoneWeightTypeSizeBytes == 1 ? MBF_NUINT8 : MBF_NUINT16;
 
 		// Limit skinning weights if necessary
 		const int32 MutableBonesPerVertex = bUseUnlimitedInfluences ?
@@ -1496,9 +1496,9 @@ mu::MeshPtr ConvertSkeletalMeshToMutable(const USkeletalMesh* InSkeletalMesh, co
 			// Set up SkinWeightPRofile BufferData
 			const int32 ElementSize = sizeof(int32) + sizeof(FBoneIndexType) + BoneWeightTypeSizeBytes;
 			const int32 ChannelCount = 3;
-			const MESH_BUFFER_SEMANTIC Semantics[ChannelCount] = { MBS_OTHER, MBS_BONEINDICES, MBS_BONEWEIGHTS };
+			const EMeshBufferSemantic Semantics[ChannelCount] = { MBS_OTHER, MBS_BONEINDICES, MBS_BONEWEIGHTS };
 			const int32 SemanticIndices[ChannelCount] = { ProfileSemanticIndex, ProfileSemanticIndex, ProfileSemanticIndex };
-			const MESH_BUFFER_FORMAT Formats[ChannelCount] = { MBF_INT32, MBF_UINT16, BoneWeightFormat };
+			const EMeshBufferFormat Formats[ChannelCount] = { MBF_INT32, MBF_UINT16, BoneWeightFormat };
 			const int32 Components[ChannelCount] = { 1, MutableBonesPerVertex, MutableBonesPerVertex };
 			const int32 Offsets[ChannelCount] = { 0, sizeof(int32), sizeof(int32) + BoneIndicesSize };
 
@@ -1524,12 +1524,12 @@ mu::MeshPtr ConvertSkeletalMeshToMutable(const USkeletalMesh* InSkeletalMesh, co
 
 		const int32 FinalElementSize = sizeof(uint32_t);
 		const int32 ChannelCount = 1;
-		const MESH_BUFFER_SEMANTIC Semantics[ChannelCount] = { MBS_VERTEXINDEX };
+		const EMeshBufferSemantic Semantics[ChannelCount] = { MBS_VERTEXINDEX };
 		const int32 SemanticIndices[ChannelCount] = { 0 };
 		// We force 32 bit indices, since merging meshes may create vertex buffers bigger than the initial mesh
 		// and for now the mutable runtime doesn't handle it.
 		// \TODO: go back to 16-bit indices when possible.
-		MESH_BUFFER_FORMAT Formats[ChannelCount] = { MBF_UINT32 };
+		EMeshBufferFormat Formats[ChannelCount] = { MBF_UINT32 };
 		const int32 Components[ChannelCount] = { 1 };
 		const int32 Offsets[ChannelCount] = { 0 };
 
@@ -1840,9 +1840,9 @@ mu::MeshPtr ConvertStaticMeshToMutable(const UStaticMesh* StaticMesh, int32 LODI
 
 			const int ElementSize = 12;
 			const int ChannelCount = 1;
-			const MESH_BUFFER_SEMANTIC Semantics[ChannelCount] = { MBS_POSITION };
+			const EMeshBufferSemantic Semantics[ChannelCount] = { MBS_POSITION };
 			const int SemanticIndices[ChannelCount] = { 0 };
-			const MESH_BUFFER_FORMAT Formats[ChannelCount] = { MBF_FLOAT32 };
+			const EMeshBufferFormat Formats[ChannelCount] = { MBF_FLOAT32 };
 			const int Components[ChannelCount] = { 3 };
 			const int Offsets[ChannelCount] = { 0 };
 
@@ -1857,9 +1857,9 @@ mu::MeshPtr ConvertStaticMeshToMutable(const UStaticMesh* StaticMesh, int32 LODI
 		{
 			const FStaticMeshVertexBuffer& VertexBuffer = StaticMesh->GetRenderData()->LODResources[LODIndex].VertexBuffers.StaticMeshVertexBuffer;
 
-			MESH_BUFFER_SEMANTIC Semantics[2];
+			EMeshBufferSemantic Semantics[2];
 			int SemanticIndices[2];
-			MESH_BUFFER_FORMAT Formats[2];
+			EMeshBufferFormat Formats[2];
 			int Components[2];
 			int Offsets[2];
 
@@ -1899,9 +1899,9 @@ mu::MeshPtr ConvertStaticMeshToMutable(const UStaticMesh* StaticMesh, int32 LODI
 			int texChannels = VertexBuffer.GetNumTexCoords();
 			int ChannelCount = texChannels;
 
-			MESH_BUFFER_SEMANTIC* Semantics = new MESH_BUFFER_SEMANTIC[ChannelCount];
+			EMeshBufferSemantic* Semantics = new EMeshBufferSemantic[ChannelCount];
 			int* SemanticIndices = new int[ChannelCount];
-			MESH_BUFFER_FORMAT* Formats = new MESH_BUFFER_FORMAT[ChannelCount];
+			EMeshBufferFormat* Formats = new EMeshBufferFormat[ChannelCount];
 			int* Components = new int[ChannelCount];
 			int* Offsets = new int[ChannelCount];
 
@@ -1909,7 +1909,7 @@ mu::MeshPtr ConvertStaticMeshToMutable(const UStaticMesh* StaticMesh, int32 LODI
 			int currentOffset = 0;
 
 			int texChannelSize;
-			MESH_BUFFER_FORMAT texChannelFormat;
+			EMeshBufferFormat texChannelFormat;
 			if (VertexBuffer.GetUseFullPrecisionUVs())
 			{
 				texChannelSize = 2 * 4;
@@ -1959,9 +1959,9 @@ mu::MeshPtr ConvertStaticMeshToMutable(const UStaticMesh* StaticMesh, int32 LODI
 		using namespace mu;
 		const int ElementSize = 2;
 		const int ChannelCount = 1;
-		const MESH_BUFFER_SEMANTIC Semantics[ChannelCount] = { MBS_VERTEXINDEX };
+		const EMeshBufferSemantic Semantics[ChannelCount] = { MBS_VERTEXINDEX };
 		const int SemanticIndices[ChannelCount] = { 0 };
-		MESH_BUFFER_FORMAT Formats[ChannelCount] = { MBF_UINT16 };
+		EMeshBufferFormat Formats[ChannelCount] = { MBF_UINT16 };
 		const int Components[ChannelCount] = { 1 };
 		const int Offsets[ChannelCount] = { 0 };
 
diff --git a/Engine/Plugins/Experimental/Mutable/Source/CustomizableObjectEditor/Private/MuCOE/GenerateMutableSource/GenerateMutableSourceSurface.cpp b/Engine/Plugins/Experimental/Mutable/Source/CustomizableObjectEditor/Private/MuCOE/GenerateMutableSource/GenerateMutableSourceSurface.cpp
--- a/Engine/Plugins/Experimental/Mutable/Source/CustomizableObjectEditor/Private/MuCOE/GenerateMutableSource/GenerateMutableSourceSurface.cpp
+++ b/Engine/Plugins/Experimental/Mutable/Source/CustomizableObjectEditor/Private/MuCOE/GenerateMutableSource/GenerateMutableSourceSurface.cpp
@@ -138,9 +138,9 @@ void SetSurfaceFormat( FMutableGraphGenerationContext& GenerationContext,
 		{
 			const int32 ElementSize = sizeof(int32);
 			constexpr int32 ChannelCount = 1;
-			const MESH_BUFFER_SEMANTIC Semantics[ChannelCount] = { MBS_OTHER };
+			const EMeshBufferSemantic Semantics[ChannelCount] = { MBS_OTHER };
 			const int32 SemanticIndices[ChannelCount] = { 0 };
-			const MESH_BUFFER_FORMAT Formats[ChannelCount] = { MBF_INT32 };
+			const EMeshBufferFormat Formats[ChannelCount] = { MBF_INT32 };
 			const int32 Components[ChannelCount] = { 1 };
 			const int32 Offsets[ChannelCount] = { 0 };
 
@@ -151,9 +151,9 @@ void SetSurfaceFormat( FMutableGraphGenerationContext& GenerationContext,
 		{
 			const int32 ElementSize = sizeof(uint16);
 			constexpr int32 ChannelCount = 1;
-			const MESH_BUFFER_SEMANTIC Semantics[ChannelCount] = { MBS_OTHER };
+			const EMeshBufferSemantic Semantics[ChannelCount] = { MBS_OTHER };
 			const int32 SemanticIndices[ChannelCount] = { 1 };
-			const MESH_BUFFER_FORMAT Formats[ChannelCount] = { MBF_UINT16 };
+			const EMeshBufferFormat Formats[ChannelCount] = { MBF_UINT16 };
 			const int32 Components[ChannelCount] = { 1 };
 			const int32 Offsets[ChannelCount] = { 0 };
 
@@ -164,9 +164,9 @@ void SetSurfaceFormat( FMutableGraphGenerationContext& GenerationContext,
 		{
 			const int32 ElementSize = sizeof(uint16);
 			constexpr int32 ChannelCount = 1;
-			const MESH_BUFFER_SEMANTIC Semantics[ChannelCount] = { MBS_OTHER };
+			const EMeshBufferSemantic Semantics[ChannelCount] = { MBS_OTHER };
 			const int32 SemanticIndices[ChannelCount] = { 2 };
-			const MESH_BUFFER_FORMAT Formats[ChannelCount] = { MBF_UINT16 };
+			const EMeshBufferFormat Formats[ChannelCount] = { MBF_UINT16 };
 			const int32 Components[ChannelCount] = { 1 };
 			const int32 Offsets[ChannelCount] = { 0 };
 
@@ -181,9 +181,9 @@ void SetSurfaceFormat( FMutableGraphGenerationContext& GenerationContext,
 		using namespace mu;
 		const int32 ElementSize = sizeof(int32);
 		constexpr int32 ChannelCount = 1;
-		const MESH_BUFFER_SEMANTIC Semantics[ChannelCount] = { MBS_OTHER };
+		const EMeshBufferSemantic Semantics[ChannelCount] = { MBS_OTHER };
 		const int32 SemanticIndices[ChannelCount] = { bWithRealTimeMorphs ? 3 : 0 };
-		const MESH_BUFFER_FORMAT Formats[ChannelCount] = { MBF_INT32 };
+		const EMeshBufferFormat Formats[ChannelCount] = { MBF_INT32 };
 		const int32 Components[ChannelCount] = { 1 };
 		const int32 Offsets[ChannelCount] = { 0 };
 
@@ -367,7 +367,7 @@ mu::NodeSurfacePtr GenerateMutableSourceSurface(const UEdGraphPin * Pin, FMutabl
 			{
 				SurfNode->SetMeshCount(1);
 
-				mu::NodeMeshFormatPtr MeshFormatNode = new mu::NodeMeshFormat();
+				mu::Ptr<mu::NodeMeshFormat> MeshFormatNode = new mu::NodeMeshFormat();
 				MeshFormatNode->SetSource(MeshNode.get());
 				SetSurfaceFormat( GenerationContext,
 						MeshFormatNode->GetVertexBuffers(), MeshFormatNode->GetIndexBuffers(), MeshData,
@@ -375,6 +375,10 @@ mu::NodeSurfacePtr GenerateMutableSourceSurface(const UEdGraphPin * Pin, FMutabl
 						GenerationContext.Options.bRealTimeMorphTargetsEnabled, 
 						GenerationContext.Options.bClothingEnabled,
 						GenerationContext.Options.b16BitBoneWeightsEnabled);
+
+				// \TODO: Make it an option?
+				MeshFormatNode->SetOptimizeBuffers(true);
+
 				MeshFormatNode->SetMessageContext(Node);
 
 				SurfNode->SetMesh(0, MeshFormatNode);
@@ -1119,7 +1123,7 @@ mu::NodeSurfacePtr GenerateMutableSourceSurface(const UEdGraphPin * Pin, FMutabl
 					GenerationContext.Compiler->CompilerLog(LOCTEXT("ExtendMaterialLayoutMissing","Skeletal Mesh without Layout Node linked to an Extend Material. A 4x4 layout will be added as default layout."), Node);
 				}
 
-				mu::NodeMeshFormatPtr MeshFormat = new mu::NodeMeshFormat();
+				mu::Ptr<mu::NodeMeshFormat> MeshFormat = new mu::NodeMeshFormat();
 				SetSurfaceFormat( GenerationContext,
 						MeshFormat->GetVertexBuffers(), MeshFormat->GetIndexBuffers(), MeshData,
 						GenerationContext.Options.CustomizableObjectNumBoneInfluences,
diff --git a/Engine/Plugins/Experimental/Mutable/Source/CustomizableObjectEditor/Private/MuCOE/MutableMeshPreviewUtils.cpp b/Engine/Plugins/Experimental/Mutable/Source/CustomizableObjectEditor/Private/MuCOE/MutableMeshPreviewUtils.cpp
--- a/Engine/Plugins/Experimental/Mutable/Source/CustomizableObjectEditor/Private/MuCOE/MutableMeshPreviewUtils.cpp
+++ b/Engine/Plugins/Experimental/Mutable/Source/CustomizableObjectEditor/Private/MuCOE/MutableMeshPreviewUtils.cpp
@@ -58,8 +58,8 @@ namespace MutableMeshPreviewUtils
 					const int32 BufferChannelCount = OriginVertexBuffers.GetBufferChannelCount(BufferIndex);
 					for (int32 ChannelIndex = 0; ChannelIndex < BufferChannelCount; ChannelIndex ++)
 					{
-						mu::MESH_BUFFER_SEMANTIC BufferSemantic = mu::MESH_BUFFER_SEMANTIC::MBS_NONE;
-						mu::MESH_BUFFER_FORMAT BufferFormat = mu::MESH_BUFFER_FORMAT::MBF_NONE;
+						mu::EMeshBufferSemantic BufferSemantic = mu::EMeshBufferSemantic::MBS_NONE;
+						mu::EMeshBufferFormat BufferFormat = mu::EMeshBufferFormat::MBF_NONE;
 						int32 BufferComponentCount = 0;
 						int32 SemanticIndex = 0;
 
@@ -77,14 +77,14 @@ namespace MutableMeshPreviewUtils
 
 						switch (BufferSemantic)
 						{
-							case mu::MESH_BUFFER_SEMANTIC::MBS_COLOUR:
+							case mu::EMeshBufferSemantic::MBS_COLOUR:
 							{
 								// Does the original mesh have vert color data?
 								bOriginDoesHaveVertexColorData = true;
 								break;
 							}
 
-							case mu::MESH_BUFFER_SEMANTIC::MBS_TEXCOORDS:
+							case mu::EMeshBufferSemantic::MBS_TEXCOORDS:
 							{
 								// Store the texture semantic index to be later used during the buffer formatting process on  mu::MeshFormat
 								TextureSemanticIndices[NumbOfTextCoordChannels] = SemanticIndex;
@@ -94,18 +94,18 @@ namespace MutableMeshPreviewUtils
 								break;
 							}
 
-							case mu::MESH_BUFFER_SEMANTIC::MBS_BONEINDICES:
+							case mu::EMeshBufferSemantic::MBS_BONEINDICES:
 							{
 								// Store the amount of bones a vertex can be skinned to.
 								MaxNumBonesPerVertex = BufferComponentCount;
-								BoneIndicesSizeBytes = BufferFormat == mu::MESH_BUFFER_FORMAT::MBF_INT16 ? 2 : 1;
+								BoneIndicesSizeBytes = BufferFormat == mu::EMeshBufferFormat::MBF_INT16 ? 2 : 1;
 								break;
 							}
 
-							case mu::MESH_BUFFER_SEMANTIC::MBS_BONEWEIGHTS:
+							case mu::EMeshBufferSemantic::MBS_BONEWEIGHTS:
 							{
 								// Store the amount of bones a vertex can be skinned to.
-								BoneWeightsSizeBytes = BufferFormat == mu::MESH_BUFFER_FORMAT::MBF_NUINT16 ? 2 : 1;
+								BoneWeightsSizeBytes = BufferFormat == mu::EMeshBufferFormat::MBF_NUINT16 ? 2 : 1;
 								break;
 							}
 
@@ -201,7 +201,7 @@ namespace MutableMeshPreviewUtils
 			int32 Channel = -1;
 
 			// POSITION BUFFER INDEX CHECK
-			VertexBuffers.FindChannel(mu::MESH_BUFFER_SEMANTIC::MBS_POSITION, 0, &Buffer, &Channel);
+			VertexBuffers.FindChannel(mu::EMeshBufferSemantic::MBS_POSITION, 0, &Buffer, &Channel);
 			if (Buffer == -1)
 			{
 				if (bLogFindingsToConsole)
@@ -225,7 +225,7 @@ namespace MutableMeshPreviewUtils
 			}
 
 			// TANGENT BUFFER INDEX CHECK
-			VertexBuffers.FindChannel(mu::MESH_BUFFER_SEMANTIC::MBS_TANGENT, 0, &Buffer, &Channel);
+			VertexBuffers.FindChannel(mu::EMeshBufferSemantic::MBS_TANGENT, 0, &Buffer, &Channel);
 			if (Buffer == -1)
 			{
 				if (bLogFindingsToConsole)
@@ -249,7 +249,7 @@ namespace MutableMeshPreviewUtils
 			}
 
 			// TEX COORDS BUFFER INDEX CHECK
-			VertexBuffers.FindChannel(mu::MESH_BUFFER_SEMANTIC::MBS_TEXCOORDS, 0, &Buffer, &Channel);
+			VertexBuffers.FindChannel(mu::EMeshBufferSemantic::MBS_TEXCOORDS, 0, &Buffer, &Channel);
 			if (Buffer == -1)
 			{
 				if (bLogFindingsToConsole)
diff --git a/Engine/Plugins/Experimental/Mutable/Source/CustomizableObjectEditor/Private/MuCOE/SMutableConstantsWidget.cpp b/Engine/Plugins/Experimental/Mutable/Source/CustomizableObjectEditor/Private/MuCOE/SMutableConstantsWidget.cpp
--- a/Engine/Plugins/Experimental/Mutable/Source/CustomizableObjectEditor/Private/MuCOE/SMutableConstantsWidget.cpp
+++ b/Engine/Plugins/Experimental/Mutable/Source/CustomizableObjectEditor/Private/MuCOE/SMutableConstantsWidget.cpp
@@ -1171,7 +1171,15 @@ void SMutableConstantsWidget::LoadConstantMeshes()
 		ConstantMeshElement->MeshPtr = MutableProgramPtr->ConstantMeshes[MeshIndex].Value;
 		ConstantMeshElement->IndexOnSourceVector = MeshIndex;
 		
-		ConstantMeshesAccumulatedSize += ConstantMeshElement->MeshPtr->GetDataSize();
+		uint64 ThisMeshSize = ConstantMeshElement->MeshPtr->GetDataSize();
+
+		// Actual core disk size would be:
+		//mu::OutputMemoryStream Stream;
+		//mu::OutputArchive Archive{ &Stream };
+		//ConstantMeshElement->MeshPtr->Serialise(Archive);
+		//uint64 OtherSize = Stream.GetBufferSize();
+
+		ConstantMeshesAccumulatedSize += ThisMeshSize;
 		
 		ConstantMeshElements.Add(ConstantMeshElement);
 	}
diff --git a/Engine/Plugins/Experimental/Mutable/Source/CustomizableObjectEditor/Private/MuCOE/SMutableGraphViewer.cpp b/Engine/Plugins/Experimental/Mutable/Source/CustomizableObjectEditor/Private/MuCOE/SMutableGraphViewer.cpp
--- a/Engine/Plugins/Experimental/Mutable/Source/CustomizableObjectEditor/Private/MuCOE/SMutableGraphViewer.cpp
+++ b/Engine/Plugins/Experimental/Mutable/Source/CustomizableObjectEditor/Private/MuCOE/SMutableGraphViewer.cpp
@@ -540,7 +540,7 @@ void SMutableGraphViewer::GetChildrenForInfo(TSharedPtr<FMutableGraphTreeElement
 	{
 		mu::NodeMeshFormat* MeshFormatVar = StaticCast<mu::NodeMeshFormat*>(ParentNode);
 		mu::NodeMeshFormat::Private* Private = MeshFormatVar->GetPrivate();
-		AddChildFunc(Private->m_pSource.get(), FString::Printf(TEXT("SOURCE MESH")));
+		AddChildFunc(Private->Source.get(), FString::Printf(TEXT("SOURCE MESH")));
 	}
 
 	else if (ParentNode->GetType() == mu::NodePatchImage::GetStaticType())
diff --git a/Engine/Plugins/Experimental/Mutable/Source/CustomizableObjectEditor/Private/MuCOE/SMutableMeshViewer.cpp b/Engine/Plugins/Experimental/Mutable/Source/CustomizableObjectEditor/Private/MuCOE/SMutableMeshViewer.cpp
--- a/Engine/Plugins/Experimental/Mutable/Source/CustomizableObjectEditor/Private/MuCOE/SMutableMeshViewer.cpp
+++ b/Engine/Plugins/Experimental/Mutable/Source/CustomizableObjectEditor/Private/MuCOE/SMutableMeshViewer.cpp
@@ -613,8 +613,8 @@ void SMutableMeshViewer::FillTargetBufferSetDataArray(const mu::FMeshBufferSet&
 		{
 			TSharedPtr<FBufferChannelElement> NewChannelDefinition = MakeShareable(new FBufferChannelElement());
 
-			mu::MESH_BUFFER_SEMANTIC BufferChannelSemantic = mu::MESH_BUFFER_SEMANTIC::MBS_NONE;
-			mu::MESH_BUFFER_FORMAT BufferFormat = mu::MESH_BUFFER_FORMAT::MBF_NONE;
+			mu::EMeshBufferSemantic BufferChannelSemantic = mu::EMeshBufferSemantic::MBS_NONE;
+			mu::EMeshBufferFormat BufferFormat = mu::EMeshBufferFormat::MBF_NONE;
 			int32 BufferComponentCount = 0;
 			int32 SemanticIndex = -1;
 
diff --git a/Engine/Plugins/Experimental/Mutable/Source/MutableRuntime/Internal/MuR/ConvertData.h b/Engine/Plugins/Experimental/Mutable/Source/MutableRuntime/Internal/MuR/ConvertData.h
--- a/Engine/Plugins/Experimental/Mutable/Source/MutableRuntime/Internal/MuR/ConvertData.h
+++ b/Engine/Plugins/Experimental/Mutable/Source/MutableRuntime/Internal/MuR/ConvertData.h
@@ -17,8 +17,8 @@ namespace mu
 	inline void ConvertData
 		(
 			int channel,
-			void* pResult, MESH_BUFFER_FORMAT resultFormat,
-			const void* pSource, MESH_BUFFER_FORMAT sourceFormat
+			void* pResult, EMeshBufferFormat resultFormat,
+			const void* pSource, EMeshBufferFormat sourceFormat
 		)
 	{
 		switch ( resultFormat )
diff --git a/Engine/Plugins/Experimental/Mutable/Source/MutableRuntime/Internal/MuR/MeshPrivate.h b/Engine/Plugins/Experimental/Mutable/Source/MutableRuntime/Internal/MuR/MeshPrivate.h
--- a/Engine/Plugins/Experimental/Mutable/Source/MutableRuntime/Internal/MuR/MeshPrivate.h
+++ b/Engine/Plugins/Experimental/Mutable/Source/MutableRuntime/Internal/MuR/MeshPrivate.h
@@ -19,8 +19,8 @@ namespace mu
 	inline void GetMeshBuf
 		(
 			Mesh* pMesh,
-			MESH_BUFFER_SEMANTIC semantic,
-			MESH_BUFFER_FORMAT expectedFormat,
+			EMeshBufferSemantic semantic,
+			EMeshBufferFormat expectedFormat,
 			int expectedComponents,
             uint8_t*& pBuf,
 			int& elemSize
@@ -35,9 +35,9 @@ namespace mu
 		pMesh->GetVertexBuffers().FindChannel( semantic, 0, &buffer, &channel );
 		check( buffer>=0 && channel>=0 );
 
-		MESH_BUFFER_SEMANTIC realSemantic = MBS_NONE;
+		EMeshBufferSemantic realSemantic = MBS_NONE;
 		int realSemanticIndex = 0;
-		MESH_BUFFER_FORMAT format = MBF_NONE;
+		EMeshBufferFormat format = MBF_NONE;
 		int components = 0;
 		int offset = 0;
 		pMesh->GetVertexBuffers().GetChannel( buffer, channel, &realSemantic, &realSemanticIndex, &format, &components, &offset );
@@ -54,8 +54,8 @@ namespace mu
 	inline void GetMeshBuf
 		(
 			const Mesh* pMesh,
-			MESH_BUFFER_SEMANTIC semantic,
-			MESH_BUFFER_FORMAT expectedFormat,
+			EMeshBufferSemantic semantic,
+			EMeshBufferFormat expectedFormat,
 			int expectedComponents,
             const uint8_t*& pBuf,
 			int& elemSize
@@ -70,9 +70,9 @@ namespace mu
 		pMesh->GetVertexBuffers().FindChannel( semantic, 0, &buffer, &channel );
 		check( buffer>=0 && channel>=0 );
 
-		MESH_BUFFER_SEMANTIC realSemantic = MBS_NONE;
+		EMeshBufferSemantic realSemantic = MBS_NONE;
 		int realSemanticIndex = 0;
-		MESH_BUFFER_FORMAT format = MBF_NONE;
+		EMeshBufferFormat format = MBF_NONE;
 		int components = 0;
 		int offset = 0;
 		pMesh->GetVertexBuffers().GetChannel
@@ -106,7 +106,7 @@ namespace mu
 		inline UntypedMeshBufferIterator
 			(
 				FMeshBufferSet& bufferSet,
-				MESH_BUFFER_SEMANTIC semantic,
+				EMeshBufferSemantic semantic,
 				int semanticIndex = 0
 			)
 		{
@@ -116,7 +116,7 @@ namespace mu
 
 			if ( buffer>=0 && channel>=0 )
 			{
-				MESH_BUFFER_SEMANTIC realSemantic = MBS_NONE;
+				EMeshBufferSemantic realSemantic = MBS_NONE;
 				int realSemanticIndex = 0;
 				m_format = MBF_NONE;
 				m_components = 0;
@@ -190,7 +190,7 @@ namespace mu
 			return m_elementSize;
 		}
 
-		inline MESH_BUFFER_FORMAT GetFormat() const
+		inline EMeshBufferFormat GetFormat() const
 		{
 			return m_format;
 		}
@@ -275,7 +275,7 @@ namespace mu
 
 		int m_elementSize;
         uint8_t* m_pBuf;
-		MESH_BUFFER_FORMAT m_format;
+		EMeshBufferFormat m_format;
 		int m_components;
 	};
 
@@ -283,17 +283,17 @@ namespace mu
 	//---------------------------------------------------------------------------------------------
 	//! Class to iterate a specific buffer channel with known type
 	//---------------------------------------------------------------------------------------------
-	template<MESH_BUFFER_FORMAT FORMAT, class CTYPE, int COMPONENTS>
+	template<EMeshBufferFormat FORMAT, class CTYPE, int COMPONENTS>
 	class MeshBufferIteratorConst;
 
-	template<MESH_BUFFER_FORMAT FORMAT, class CTYPE, int COMPONENTS>
+	template<EMeshBufferFormat FORMAT, class CTYPE, int COMPONENTS>
 	class MUTABLERUNTIME_API MeshBufferIterator : public UntypedMeshBufferIterator
 	{
 	public:
 		inline MeshBufferIterator
 			(
 				FMeshBufferSet& bufferSet,
-				MESH_BUFFER_SEMANTIC semantic,
+				EMeshBufferSemantic semantic,
 				int semanticIndex = 0
 			)
 			: UntypedMeshBufferIterator( bufferSet, semantic, semanticIndex )
@@ -307,9 +307,9 @@ namespace mu
 				
 				if (buffer >= 0 && channel >= 0)
 				{
-					MESH_BUFFER_SEMANTIC realSemantic = MBS_NONE;
+					EMeshBufferSemantic realSemantic = MBS_NONE;
 					int realSemanticIndex = 0;
-					MESH_BUFFER_FORMAT format = MBF_NONE;
+					EMeshBufferFormat format = MBF_NONE;
 					int components = 0;
 					int offset = 0;
 					bufferSet.GetChannel
@@ -372,7 +372,7 @@ namespace mu
 		inline UntypedMeshBufferIteratorConst
 			(
 				const FMeshBufferSet& bufferSet,
-				MESH_BUFFER_SEMANTIC semantic,
+				EMeshBufferSemantic semantic,
 				int semanticIndex = 0
 			)
 		{
@@ -382,7 +382,7 @@ namespace mu
 
 			if ( buffer>=0 && channel>=0 )
 			{
-				MESH_BUFFER_SEMANTIC realSemantic = MBS_NONE;
+				EMeshBufferSemantic realSemantic = MBS_NONE;
 				int realSemanticIndex = 0;
 				m_format = MBF_NONE;
 				m_components = 0;
@@ -460,7 +460,7 @@ namespace mu
 			return m_elementSize;
 		}
 
-        inline MESH_BUFFER_FORMAT GetFormat() const
+        inline EMeshBufferFormat GetFormat() const
         {
             return m_format;
         }
@@ -537,7 +537,7 @@ namespace mu
 
 		int m_elementSize = 0;
         const uint8_t* m_pBuf = nullptr;
-		MESH_BUFFER_FORMAT m_format = MBF_NONE;
+		EMeshBufferFormat m_format = MBF_NONE;
 		int m_components = 0;
 	};
 
@@ -547,7 +547,7 @@ namespace mu
 	//---------------------------------------------------------------------------------------------
 	//! Class to iterate a specific buffer channel of a constant buffer set
 	//---------------------------------------------------------------------------------------------
-	template<MESH_BUFFER_FORMAT FORMAT, class CTYPE, int COMPONENTS>
+	template<EMeshBufferFormat FORMAT, class CTYPE, int COMPONENTS>
 	class MUTABLERUNTIME_API MeshBufferIteratorConst : public UntypedMeshBufferIteratorConst
 	{
 	public:
@@ -560,7 +560,7 @@ namespace mu
 		inline MeshBufferIteratorConst
 			(
 				const FMeshBufferSet& bufferSet,
-				MESH_BUFFER_SEMANTIC semantic,
+				EMeshBufferSemantic semantic,
 				int semanticIndex = 0
 			)
 			: UntypedMeshBufferIteratorConst( bufferSet, semantic, semanticIndex )
@@ -574,9 +574,9 @@ namespace mu
 
 				if (buffer >= 0 && channel >= 0)
 				{
-					MESH_BUFFER_SEMANTIC realSemantic = MBS_NONE;
+					EMeshBufferSemantic realSemantic = MBS_NONE;
 					int realSemanticIndex = 0;
-					MESH_BUFFER_FORMAT format = MBF_NONE;
+					EMeshBufferFormat format = MBF_NONE;
 					int components = 0;
 					int offset = 0;
 					bufferSet.GetChannel
diff --git a/Engine/Plugins/Experimental/Mutable/Source/MutableRuntime/Internal/MuR/OpMeshFormat.h b/Engine/Plugins/Experimental/Mutable/Source/MutableRuntime/Internal/MuR/OpMeshFormat.h
--- a/Engine/Plugins/Experimental/Mutable/Source/MutableRuntime/Internal/MuR/OpMeshFormat.h
+++ b/Engine/Plugins/Experimental/Mutable/Source/MutableRuntime/Internal/MuR/OpMeshFormat.h
@@ -2,42 +2,41 @@
 
 #pragma once
 
-#include "MuR/Mesh.h"
-
+#include "MuR/Platform.h"
 
 namespace mu
 {
-class FMeshBufferSet;
-
-	//---------------------------------------------------------------------------------------------
-	//! Convert a mesh format into another one.
-	//! Slow implementation, but it should never happen at run-time
-	//! \param keepSystemBuffers Will keep the internal system buffers even if they are not in the
-	//! original format. If they are, they will be duplicated, so be careful.
-	//---------------------------------------------------------------------------------------------
+	class Mesh;
+	class FMeshBufferSet;
+
+	/** Convert a mesh format into another one.
+	* Slow implementation, but it should never happen at run-time
+	* \param keepSystemBuffers Will keep the internal system buffers even if they are not in the
+	* original format. If they are, they will be duplicated, so be careful.
+	*/
 	MUTABLERUNTIME_API extern void MeshFormat
 		(
 			Mesh* Result,
-			const Mesh* pSource,
-			const Mesh* pFormat,
-			bool keepSystemBuffers,
-			bool formatVertices,
-			bool formatIndices,
-			bool formatFaces,
-			bool ignoreMissingChannels,
+			const Mesh* Source,
+			const Mesh* Format,
+			bool bKeepSystemBuffers,
+			bool bFormatVertices,
+			bool bFormatIndices,
+			bool bFormatFaces,
+			bool bIgnoreMissingChannels,
 			bool& bOutSuccess
 		);
 
 
-    //---------------------------------------------------------------------------------------------
-    //! Fill a meshbuffer from result with the current data in source but keeping the format
-    //! already in the result.
-    //---------------------------------------------------------------------------------------------
+    /** Fill a mesh buffer from result with the current data in source but keeping the format already in the result. */
     void MeshFormatBuffer
         (
-            const FMeshBufferSet& pSource,
-            FMeshBufferSet& pResult,
-            int bufferIndex
+            const FMeshBufferSet& Source,
+            FMeshBufferSet& Result,
+            int32 BufferIndex
         );
 
-}
\ No newline at end of file
+	/** Try to reduce the mesh size by reducing the component count, and data type of some buffers. */
+	void MeshOptimizeBuffers( Mesh* );
+
+}
diff --git a/Engine/Plugins/Experimental/Mutable/Source/MutableRuntime/Internal/MuR/OpMeshMerge.h b/Engine/Plugins/Experimental/Mutable/Source/MutableRuntime/Internal/MuR/OpMeshMerge.h
--- a/Engine/Plugins/Experimental/Mutable/Source/MutableRuntime/Internal/MuR/OpMeshMerge.h
+++ b/Engine/Plugins/Experimental/Mutable/Source/MutableRuntime/Internal/MuR/OpMeshMerge.h
@@ -46,17 +46,17 @@ namespace mu
 			check(pSecond->GetIndexBuffers().GetBufferCount() <= 1);
 			Result->GetIndexBuffers().SetBufferCount(1);
 
-			MESH_BUFFER& ResultIndexBuffer = Result->GetIndexBuffers().m_buffers[0];
+			FMeshBuffer& ResultIndexBuffer = Result->GetIndexBuffers().m_buffers[0];
 
-			const MESH_BUFFER& FirstIndexBuffer = pFirst->GetIndexBuffers().m_buffers[0];
-			const MESH_BUFFER& SecondIndexBuffer = pSecond->GetIndexBuffers().m_buffers[0];
+			const FMeshBuffer& FirstIndexBuffer = pFirst->GetIndexBuffers().m_buffers[0];
+			const FMeshBuffer& SecondIndexBuffer = pSecond->GetIndexBuffers().m_buffers[0];
 
 			// Avoid unused variable warnings
 			(void)FirstIndexBuffer;
 			(void)SecondIndexBuffer;
 
 			// This will be changed below if need to change the format of the index buffers.
-			MESH_BUFFER_FORMAT IndexBufferFormat = MBF_NONE;
+			EMeshBufferFormat IndexBufferFormat = MBF_NONE;
 
 			if (FirstCount && SecondCount)
 			{
@@ -66,7 +66,7 @@ namespace mu
 
 				// We need to know the total number of vertices in case we need to adjust the index buffer format.
 				const uint64 totalVertexCount = pFirst->GetVertexBuffers().GetElementCount() + pSecond->GetVertexBuffers().GetElementCount();
-				const uint64 maxValueBits = GetMeshFormatData(pFirst->GetIndexBuffers().m_buffers[0].m_channels[0].m_format).m_maxValueBits;
+				const uint64 maxValueBits = GetMeshFormatData(pFirst->GetIndexBuffers().m_buffers[0].m_channels[0].m_format).MaxValueBits;
 				const uint64 maxSupportedVertices = uint64(1) << maxValueBits;
 				
 				if (totalVertexCount > maxSupportedVertices)
@@ -87,7 +87,7 @@ namespace mu
 				ResultIndexBuffer.m_channels[0].m_componentCount = 1;
 				ResultIndexBuffer.m_channels[0].m_semanticIndex = 0;
 				ResultIndexBuffer.m_channels[0].m_offset = 0;
-				ResultIndexBuffer.m_elementSize = GetMeshFormatData(IndexBufferFormat).m_size;
+				ResultIndexBuffer.m_elementSize = GetMeshFormatData(IndexBufferFormat).SizeInBytes;
 			}
 			else if (FirstCount)
 			{
@@ -340,16 +340,16 @@ namespace mu
 			// Merge only the buffers present in the first mesh
 			for (int32 b = 0; b < Result->GetFaceBuffers().GetBufferCount(); ++b)
 			{
-				const MESH_BUFFER& first = pFirst->GetFaceBuffers().m_buffers[b];
+				const FMeshBuffer& first = pFirst->GetFaceBuffers().m_buffers[b];
 
-				MESH_BUFFER& result = Result->GetFaceBuffers().m_buffers[b];
+				FMeshBuffer& result = Result->GetFaceBuffers().m_buffers[b];
 				result.m_channels = first.m_channels;
 				result.m_elementSize = first.m_elementSize;
 				result.m_data.SetNum(result.m_elementSize * (FirstCount + SecondCount));
 
-				MESH_BUFFER_SEMANTIC semantic = MBS_NONE;
+				EMeshBufferSemantic semantic = MBS_NONE;
 				int semanticIndex = 0;
-				MESH_BUFFER_FORMAT type = MBF_NONE;
+				EMeshBufferFormat type = MBF_NONE;
 				int components = 0;
 				int offset = 0;
 				pFirst->GetFaceBuffers().GetChannel
@@ -373,7 +373,7 @@ namespace mu
 
 					if (otherBuffer >= 0)
 					{
-						const MESH_BUFFER& second =
+						const FMeshBuffer& second =
 							pSecond->GetFaceBuffers().m_buffers[otherBuffer];
 						check(first.m_channels == second.m_channels);
 
@@ -757,7 +757,7 @@ namespace mu
 
 			// Check if the format of the BoneIndex buffer has to change
 			bool bChangeBoneIndicesFormat = false;
-			MESH_BUFFER_FORMAT BoneIndexFormat = MaxNumBonesInBoneMaps > MAX_uint8 ? MBF_UINT16 : MBF_UINT8;
+			EMeshBufferFormat BoneIndexFormat = MaxNumBonesInBoneMaps > MAX_uint8 ? MBF_UINT16 : MBF_UINT8;
 
 			// Iterate all vertex buffers to check if we need to format bone indices
 			{
@@ -766,7 +766,7 @@ namespace mu
 					const FMeshBufferSet& VertexBuffers = pFirst->GetVertexBuffers();
 					for (int32 VertexBufferIndex = 0; !bChangeBoneIndicesFormat && VertexBufferIndex < VertexBuffers.m_buffers.Num(); ++VertexBufferIndex)
 					{
-						const MESH_BUFFER& Buffer = VertexBuffers.m_buffers[VertexBufferIndex];
+						const FMeshBuffer& Buffer = VertexBuffers.m_buffers[VertexBufferIndex];
 
 						const int32 elemSize = VertexBuffers.GetElementSize(VertexBufferIndex);
 						const int32 firstSize = FirstCount * elemSize;
@@ -788,7 +788,7 @@ namespace mu
 					const FMeshBufferSet& VertexBuffers = pSecond->GetVertexBuffers();
 					for (int32 VertexBufferIndex = 0; !bChangeBoneIndicesFormat && VertexBufferIndex < VertexBuffers.m_buffers.Num(); ++VertexBufferIndex)
 					{
-						const MESH_BUFFER& Buffer = VertexBuffers.m_buffers[VertexBufferIndex];
+						const FMeshBuffer& Buffer = VertexBuffers.m_buffers[VertexBufferIndex];
 
 						const int32 elemSize = VertexBuffers.GetElementSize(VertexBufferIndex);
 						const int32 firstSize = FirstCount * elemSize;
@@ -823,8 +823,8 @@ namespace mu
 
 				for ( int32 vb = 0; vb<vbcount; ++vb )
 				{
-					MESH_BUFFER& result = Result->GetVertexBuffers().m_buffers[vb];
-					const MESH_BUFFER& first = pFirst->GetVertexBuffers().m_buffers[vb];
+					FMeshBuffer& result = Result->GetVertexBuffers().m_buffers[vb];
+					const FMeshBuffer& first = pFirst->GetVertexBuffers().m_buffers[vb];
 
 					result.m_channels = first.m_channels;
 					result.m_elementSize = first.m_elementSize;
@@ -843,7 +843,7 @@ namespace mu
 								);
 						if ( sb>=0 )
 						{
-							const MESH_BUFFER& second = pSecond->GetVertexBuffers().m_buffers[sb];
+							const FMeshBuffer& second = pSecond->GetVertexBuffers().m_buffers[sb];
 
 							if ( second.m_channels[sc].m_componentCount
 								 >
@@ -865,7 +865,7 @@ namespace mu
                             result.m_channels[c].m_offset = (uint8_t)offset;
 							offset += result.m_channels[c].m_componentCount
 									*
-									GetMeshFormatData(result.m_channels[c].m_format).m_size;
+									GetMeshFormatData(result.m_channels[c].m_format).SizeInBytes;
 						}
 						result.m_elementSize = offset;
 					}
@@ -873,11 +873,11 @@ namespace mu
 
                 // See if we need to add additional buffers from the second mesh (like vertex colours or additional UV Channels)
                 // This is a bit ad-hoc: we only add buffers containing all new channels
-                for ( const MESH_BUFFER& buf : pSecond->GetVertexBuffers().m_buffers )
+                for ( const FMeshBuffer& buf : pSecond->GetVertexBuffers().m_buffers )
                 {
                     bool someChannel = false;
                     bool allNewChannels = true;
-					for (const MESH_BUFFER_CHANNEL& chan : buf.m_channels)
+					for (const FMeshBufferChannel& chan : buf.m_channels)
 					{
 						// Skip system buffers
 						if (chan.m_semantic == MBS_VERTEXINDEX
@@ -910,7 +910,7 @@ namespace mu
 
 							if (foundBuffer >= 0)
 							{
-								MESH_BUFFER& Buffer = VertexBuffers.m_buffers[foundBuffer];
+								FMeshBuffer& Buffer = VertexBuffers.m_buffers[foundBuffer];
 								Buffer.m_channels.Insert(chan, foundChannel + 1);
 
 								// Update offsets
@@ -920,7 +920,7 @@ namespace mu
 									Buffer.m_channels[c].m_offset = (uint8)Offset;
 									Offset += Buffer.m_channels[c].m_componentCount
 										*
-										GetMeshFormatData(Buffer.m_channels[c].m_format).m_size;
+										GetMeshFormatData(Buffer.m_channels[c].m_format).SizeInBytes;
 								}
 								Buffer.m_elementSize = Offset;
 							}
@@ -940,7 +940,7 @@ namespace mu
 					FMeshBufferSet& VertexBuffers = Result->GetVertexBuffers();
 					for (int32 VertexBufferIndex = 0; VertexBufferIndex < VertexBuffers.m_buffers.Num(); ++VertexBufferIndex)
 					{
-						MESH_BUFFER& result = VertexBuffers.m_buffers[VertexBufferIndex];
+						FMeshBuffer& result = VertexBuffers.m_buffers[VertexBufferIndex];
 
 						const int32 ChannelsCount = VertexBuffers.GetBufferChannelCount(VertexBufferIndex);
 						for (int32 ChannelIndex = 0; ChannelIndex < ChannelsCount; ++ChannelIndex)
@@ -956,7 +956,7 @@ namespace mu
 									result.m_channels[AuxChannelIndex].m_offset = (uint8_t)offset;
 									offset += result.m_channels[AuxChannelIndex].m_componentCount
 										*
-										GetMeshFormatData(result.m_channels[AuxChannelIndex].m_format).m_size;
+										GetMeshFormatData(result.m_channels[AuxChannelIndex].m_format).SizeInBytes;
 								}
 								result.m_elementSize = offset;
 							}
@@ -1021,8 +1021,8 @@ namespace mu
 				MUTABLE_CPUPROFILER_SCOPE(CopyVertexData);
 				for (int32 vb = 0; vb < Result->GetVertexBuffers().m_buffers.Num(); ++vb)
 				{
-					MESH_BUFFER& result = Result->GetVertexBuffers().m_buffers[vb];
-					const MESH_BUFFER& second = pVSecond->GetVertexBuffers().m_buffers[vb];
+					FMeshBuffer& result = Result->GetVertexBuffers().m_buffers[vb];
+					const FMeshBuffer& second = pVSecond->GetVertexBuffers().m_buffers[vb];
 
 					if (SecondCount)
 					{
@@ -1048,7 +1048,7 @@ namespace mu
 				FMeshBufferSet& VertexBuffers = Result->GetVertexBuffers();
 				for (int32 VertexBufferIndex = 0; VertexBufferIndex < VertexBuffers.m_buffers.Num(); ++VertexBufferIndex)
 				{
-					MESH_BUFFER& ResultBuffer = VertexBuffers.m_buffers[VertexBufferIndex];
+					FMeshBuffer& ResultBuffer = VertexBuffers.m_buffers[VertexBufferIndex];
 
 					const int32 ElemSize = VertexBuffers.GetElementSize(VertexBufferIndex);
 					const int32 FirstSize = FirstCount * ElemSize;
diff --git a/Engine/Plugins/Experimental/Mutable/Source/MutableRuntime/Internal/MuR/Operations.h b/Engine/Plugins/Experimental/Mutable/Source/MutableRuntime/Internal/MuR/Operations.h
--- a/Engine/Plugins/Experimental/Mutable/Source/MutableRuntime/Internal/MuR/Operations.h
+++ b/Engine/Plugins/Experimental/Mutable/Source/MutableRuntime/Internal/MuR/Operations.h
@@ -877,20 +877,24 @@ namespace mu
             ADDRESS source;
             ADDRESS format;
 
-            typedef enum
+            enum
             {
-                BT_VERTEX			= 1,
-                BT_INDEX			= 2,
-                BT_FACE				= 4,
-                //! This flag will not add blank channels for the channels in the format mesh but not
-                //! in the source mesh.
-                BT_IGNORE_MISSING	= 16,
-                //! This flag will force the reset of buffer indices to 0
-                BT_RESETBUFFERINDICES = 32
-            } BUFFER_TYPE;
-
-            //! Flag combination, selecting the buffers to reformat. The rest are left untouched.
-            uint8 buffers;
+                Vertex				= 1 << 0,
+                Index				= 1 << 1,
+                Face				= 1 << 2,
+
+                /** This flag will not add blank channels for the channels in the format mesh but not in the source mesh. */
+				IgnoreMissing		= 1 << 4,
+
+                /** This flag will force the reset of buffer indices to 0. */
+                ResetBufferIndices	= 1 << 5,
+
+				/** This flag will add a step to reduce some buffers size by removing components and changing the types if possible. */
+				OptimizeBuffers		= 1 << 6
+			} EFlags;
+
+            //! EFlags combination, selecting the buffers to reformat and other options.
+            uint8 Flags;
 
         };
 
diff --git a/Engine/Plugins/Experimental/Mutable/Source/MutableRuntime/Private/MuR/CodeRunner.cpp b/Engine/Plugins/Experimental/Mutable/Source/MutableRuntime/Private/MuR/CodeRunner.cpp
--- a/Engine/Plugins/Experimental/Mutable/Source/MutableRuntime/Private/MuR/CodeRunner.cpp
+++ b/Engine/Plugins/Experimental/Mutable/Source/MutableRuntime/Private/MuR/CodeRunner.cpp
@@ -1096,7 +1096,7 @@ namespace mu
                 Ptr<const Mesh> pBase = LoadMesh(FCacheAddress(BaseAt,item));
                 Ptr<const Mesh> pTarget = LoadMesh(FCacheAddress(TargetAt,item));
 
-				TArray<MESH_BUFFER_SEMANTIC, TInlineAllocator<8>> Semantics;
+				TArray<EMeshBufferSemantic, TInlineAllocator<8>> Semantics;
 				TArray<int32, TInlineAllocator<8>> SemanticIndices;
 
 				uint8 bIgnoreTextureCoords = 0;
@@ -1117,7 +1117,7 @@ namespace mu
 					FMemory::Memcpy(&SemanticIndex, data, sizeof(uint8)); 
 					data += sizeof(uint8);
 
-					Semantics.Add(MESH_BUFFER_SEMANTIC(Semantic));
+					Semantics.Add(EMeshBufferSemantic(Semantic));
 					SemanticIndices.Add(SemanticIndex);
                 }
 
@@ -1791,20 +1791,20 @@ namespace mu
             	MUTABLE_CPUPROFILER_SCOPE(ME_FORMAT_1)
             		
                 Ptr<const Mesh> Source = LoadMesh(FCacheAddress(args.source,item));
-                Ptr<const Mesh> pFormat = LoadMesh(FCacheAddress(args.format,item));
+                Ptr<const Mesh> Format = LoadMesh(FCacheAddress(args.format,item));
 
 				if (Source)
 				{
-					uint8 flags = args.buffers;
-					if (!pFormat && !(flags & OP::MeshFormatArgs::BT_RESETBUFFERINDICES))
+					uint8 Flags = args.Flags;
+					if (!Format && !(Flags & OP::MeshFormatArgs::ResetBufferIndices))
 					{
 						StoreMesh(item, Source);
 					}
-					else if (!pFormat)
+					else if (!Format)
 					{
 						Ptr<Mesh> Result = CloneOrTakeOver(Source);
 
-						if (flags & OP::MeshFormatArgs::BT_RESETBUFFERINDICES)
+						if (Flags & OP::MeshFormatArgs::ResetBufferIndices)
 						{
 							Result->ResetBufferIndices();
 						}
@@ -1816,29 +1816,35 @@ namespace mu
 						Ptr<Mesh> Result = CreateMesh();
 
 						bool bOutSuccess = false;
-						MeshFormat(Result.get(), Source.get(), pFormat.get(),
+						MeshFormat(Result.get(), Source.get(), Format.get(),
 							true,
-							(flags & OP::MeshFormatArgs::BT_VERTEX) != 0,
-							(flags & OP::MeshFormatArgs::BT_INDEX) != 0,
-							(flags & OP::MeshFormatArgs::BT_FACE) != 0,
-							(flags & OP::MeshFormatArgs::BT_IGNORE_MISSING) != 0,
+							(Flags & OP::MeshFormatArgs::Vertex) != 0,
+							(Flags & OP::MeshFormatArgs::Index) != 0,
+							(Flags & OP::MeshFormatArgs::Face) != 0,
+							(Flags & OP::MeshFormatArgs::IgnoreMissing) != 0,
 							bOutSuccess);
 
 						check(bOutSuccess);
 
-						if (flags & OP::MeshFormatArgs::BT_RESETBUFFERINDICES)
+						if (Flags & OP::MeshFormatArgs::ResetBufferIndices)
 						{
 							Result->ResetBufferIndices();
 						}
 
+						if (Flags & OP::MeshFormatArgs::OptimizeBuffers)
+						{
+							MUTABLE_CPUPROFILER_SCOPE(MeshOptimizeBuffers)
+							MeshOptimizeBuffers(Result.get());
+						}
+
 						Release(Source);
-						Release(pFormat);
+						Release(Format);
 						StoreMesh(item, Result);
 					}
 				}
 				else
 				{
-					Release(pFormat);
+					Release(Format);
 					StoreMesh(item, nullptr);
 				}
                 break;
diff --git a/Engine/Plugins/Experimental/Mutable/Source/MutableRuntime/Private/MuR/Mesh.cpp b/Engine/Plugins/Experimental/Mutable/Source/MutableRuntime/Private/MuR/Mesh.cpp
--- a/Engine/Plugins/Experimental/Mutable/Source/MutableRuntime/Private/MuR/Mesh.cpp
+++ b/Engine/Plugins/Experimental/Mutable/Source/MutableRuntime/Private/MuR/Mesh.cpp
@@ -684,8 +684,8 @@ bool Mesh::HasCompatibleFormat( const Mesh* pOther ) const
         check( m_IndexBuffers.GetBufferChannelCount(0) == 1 );
         check( pOther->GetIndexBuffers().GetBufferChannelCount(0) == 1 );
 
-        const MESH_BUFFER& dest = m_IndexBuffers.m_buffers[0];
-        const MESH_BUFFER& source = pOther->GetIndexBuffers().m_buffers[0];
+        const FMeshBuffer& dest = m_IndexBuffers.m_buffers[0];
+        const FMeshBuffer& source = pOther->GetIndexBuffers().m_buffers[0];
 
         compatible &= dest.m_channels[0].m_format == source.m_channels[0].m_format;
     }
@@ -700,8 +700,8 @@ bool Mesh::HasCompatibleFormat( const Mesh* pOther ) const
     //-----------------
     for ( int vb = 0; vb<m_VertexBuffers.GetBufferCount(); ++vb )
     {
-        const MESH_BUFFER& dest = m_VertexBuffers.m_buffers[vb];
-        const MESH_BUFFER& source = pOther->GetVertexBuffers().m_buffers[vb];
+        const FMeshBuffer& dest = m_VertexBuffers.m_buffers[vb];
+        const FMeshBuffer& source = pOther->GetVertexBuffers().m_buffers[vb];
 
         // TODO: More checks about channels formats and semantics
         //compatible &= GetVertexBufferElementSize(vb) == pOther->GetVertexBufferElementSize(vb);
@@ -783,9 +783,9 @@ bool Mesh::IsSameVertex
     {
         for ( int c=0; same && c<m_VertexBuffers.GetBufferChannelCount(b); ++c )
         {
-            MESH_BUFFER_SEMANTIC semantic;
+            EMeshBufferSemantic semantic;
             int semanticIndex = 0;
-            MESH_BUFFER_FORMAT format;
+            EMeshBufferFormat format;
             int components;
             int offset = 0;
             m_VertexBuffers.GetChannel( b, c, &semantic, &semanticIndex, &format, &components, &offset );
@@ -808,9 +808,9 @@ bool Mesh::IsSameVertex
             other.m_VertexBuffers.FindChannel( semantic, semanticIndex, &otherBuffer, &otherChannel );
             check( otherBuffer>=0 && otherChannel>=0 );
 
-            MESH_BUFFER_SEMANTIC otherSemantic;
+            EMeshBufferSemantic otherSemantic;
             int otherSemanticIndex = 0;
-            MESH_BUFFER_FORMAT otherFormat;
+            EMeshBufferFormat otherFormat;
             int otherComponents;
             int otherOffset;
             other.m_VertexBuffers.GetChannel
@@ -1444,9 +1444,9 @@ void Mesh::CheckIntegrity() const
 
             for ( int c=0; c<m_IndexBuffers.GetBufferChannelCount(b); ++c )
             {
-                MESH_BUFFER_SEMANTIC semantic;
+                EMeshBufferSemantic semantic;
                 int semanticIndex = 0;
-                MESH_BUFFER_FORMAT format;
+                EMeshBufferFormat format;
                 int components;
                 int offset = 0;
                 m_IndexBuffers.GetChannel
@@ -1502,9 +1502,9 @@ void Mesh::CheckIntegrity() const
             int channelCount = m_VertexBuffers.GetBufferChannelCount( b );
             for ( int c = 0; c < channelCount; ++c )
             {
-                MESH_BUFFER_SEMANTIC semantic;
+                EMeshBufferSemantic semantic;
                 int semanticIndex = 0;
-                MESH_BUFFER_FORMAT format;
+                EMeshBufferFormat format;
                 int components;
                 int offset = 0;
                 m_VertexBuffers.GetChannel
@@ -1604,7 +1604,7 @@ static bool StaticMeshFormatIdentify_Project( const Mesh* pM )
 
     if ( res )
     {
-        const MESH_BUFFER_CHANNEL& chan =
+        const FMeshBufferChannel& chan =
                 pM->m_VertexBuffers.m_buffers[0].m_channels[0];
 
         res &= chan.m_semantic == MBS_TEXCOORDS;
@@ -1617,7 +1617,7 @@ static bool StaticMeshFormatIdentify_Project( const Mesh* pM )
 
     if ( res )
     {
-        const MESH_BUFFER_CHANNEL& chan =
+        const FMeshBufferChannel& chan =
                 pM->m_VertexBuffers.m_buffers[0].m_channels[1];
 
         res &= chan.m_semantic == MBS_POSITION;
@@ -1629,7 +1629,7 @@ static bool StaticMeshFormatIdentify_Project( const Mesh* pM )
 
     if ( res )
     {
-        const MESH_BUFFER_CHANNEL& chan =
+        const FMeshBufferChannel& chan =
                 pM->m_VertexBuffers.m_buffers[0].m_channels[2];
 
         res &= chan.m_semantic == MBS_NORMAL;
@@ -1647,7 +1647,7 @@ static bool StaticMeshFormatIdentify_Project( const Mesh* pM )
 
     if ( res )
     {
-        const MESH_BUFFER_CHANNEL& chan =
+        const FMeshBufferChannel& chan =
                 pM->m_IndexBuffers.m_buffers[0].m_channels[0];
 
         res &= chan.m_semantic == MBS_VERTEXINDEX;
@@ -1678,7 +1678,7 @@ static bool StaticMeshFormatIdentify_ProjectWrapping( const Mesh* pM )
 
     if ( res )
     {
-        const MESH_BUFFER_CHANNEL& chan = pM->m_VertexBuffers.m_buffers[0].m_channels[0];
+        const FMeshBufferChannel& chan = pM->m_VertexBuffers.m_buffers[0].m_channels[0];
 
         res &= chan.m_semantic == MBS_TEXCOORDS;
         res &= chan.m_format == MBF_FLOAT32;
@@ -1690,7 +1690,7 @@ static bool StaticMeshFormatIdentify_ProjectWrapping( const Mesh* pM )
 
     if ( res )
     {
-        const MESH_BUFFER_CHANNEL& chan = pM->m_VertexBuffers.m_buffers[0].m_channels[1];
+        const FMeshBufferChannel& chan = pM->m_VertexBuffers.m_buffers[0].m_channels[1];
 
         res &= chan.m_semantic == MBS_POSITION;
         res &= chan.m_format == MBF_FLOAT32;
@@ -1701,7 +1701,7 @@ static bool StaticMeshFormatIdentify_ProjectWrapping( const Mesh* pM )
 
     if ( res )
     {
-        const MESH_BUFFER_CHANNEL& chan = pM->m_VertexBuffers.m_buffers[0].m_channels[2];
+        const FMeshBufferChannel& chan = pM->m_VertexBuffers.m_buffers[0].m_channels[2];
 
         res &= chan.m_semantic == MBS_NORMAL;
         res &= chan.m_format == MBF_FLOAT32;
@@ -1712,7 +1712,7 @@ static bool StaticMeshFormatIdentify_ProjectWrapping( const Mesh* pM )
 
     if ( res )
     {
-        const MESH_BUFFER_CHANNEL& chan = pM->m_VertexBuffers.m_buffers[0].m_channels[3];
+        const FMeshBufferChannel& chan = pM->m_VertexBuffers.m_buffers[0].m_channels[3];
 
         res &= chan.m_semantic == MBS_LAYOUTBLOCK;
         res &= chan.m_format == MBF_UINT32;
@@ -1730,7 +1730,7 @@ static bool StaticMeshFormatIdentify_ProjectWrapping( const Mesh* pM )
 
     if ( res )
     {
-        const MESH_BUFFER_CHANNEL& chan = pM->m_IndexBuffers.m_buffers[0].m_channels[0];
+        const FMeshBufferChannel& chan = pM->m_IndexBuffers.m_buffers[0].m_channels[0];
 
         res &= chan.m_semantic == MBS_VERTEXINDEX;
         res &= chan.m_format == MBF_UINT32;
@@ -1767,13 +1767,13 @@ namespace
                 + FString::Printf(TEXT("%d"), elemCount)
                 + " elements.\n";
 
-        for( const MESH_BUFFER& buf : bufset.m_buffers )
+        for( const FMeshBuffer& buf : bufset.m_buffers )
         {
             const uint8* pData = buf.m_data.GetData();
 
             out += "    Buffer with "+ FString::Printf(TEXT("%d"), buf.m_channels.Num())
                     + " channels and "+ FString::Printf(TEXT("%d"), buf.m_elementSize)+" elementsize\n";
-            for( const MESH_BUFFER_CHANNEL& chan : buf.m_channels )
+            for( const FMeshBufferChannel& chan : buf.m_channels )
             {
                 out += "      Channel with format: "+ FString::Printf(TEXT("%d"), chan.m_format)
                         + " semantic: "+ FString::Printf(TEXT("%d"), chan.m_semantic)
diff --git a/Engine/Plugins/Experimental/Mutable/Source/MutableRuntime/Private/MuR/MeshBufferSet.cpp b/Engine/Plugins/Experimental/Mutable/Source/MutableRuntime/Private/MuR/MeshBufferSet.cpp
--- a/Engine/Plugins/Experimental/Mutable/Source/MutableRuntime/Private/MuR/MeshBufferSet.cpp
+++ b/Engine/Plugins/Experimental/Mutable/Source/MutableRuntime/Private/MuR/MeshBufferSet.cpp
@@ -13,10 +13,10 @@
 
 namespace mu
 {
-	MUTABLE_IMPLEMENT_POD_SERIALISABLE(MESH_BUFFER_CHANNEL);
-	MUTABLE_IMPLEMENT_POD_VECTOR_SERIALISABLE(MESH_BUFFER_CHANNEL);
-	MUTABLE_IMPLEMENT_ENUM_SERIALISABLE(MESH_BUFFER_FORMAT);
-	MUTABLE_IMPLEMENT_ENUM_SERIALISABLE(MESH_BUFFER_SEMANTIC);
+	MUTABLE_IMPLEMENT_POD_SERIALISABLE(FMeshBufferChannel);
+	MUTABLE_IMPLEMENT_POD_VECTOR_SERIALISABLE(FMeshBufferChannel);
+	MUTABLE_IMPLEMENT_ENUM_SERIALISABLE(EMeshBufferFormat);
+	MUTABLE_IMPLEMENT_ENUM_SERIALISABLE(EMeshBufferSemantic);
 
 	
 	//---------------------------------------------------------------------------------------------
@@ -50,7 +50,7 @@ namespace mu
 
 
 	//---------------------------------------------------------------------------------------------
-	const FMeshBufferFormatData& GetMeshFormatData(MESH_BUFFER_FORMAT format)
+	const FMeshBufferFormatData& GetMeshFormatData(EMeshBufferFormat format)
 	{
 		check(format >= 0);
 		check(format < MBF_COUNT);
@@ -95,7 +95,7 @@ namespace mu
 		// allow it if no other allocation will be done.  
 		const EAllowShrinking AllowShrinking = (Count == 0) ? EAllowShrinking::Yes : EAllowShrinking::No;
 
-		for (MESH_BUFFER& buf : m_buffers)
+		for (FMeshBuffer& buf : m_buffers)
 		{
 			buf.m_data.SetNumUninitialized(buf.m_elementSize * Count, AllowShrinking);
 		}
@@ -136,9 +136,9 @@ namespace mu
 	(
 		int32 buffer,
 		int32 channel,
-		MESH_BUFFER_SEMANTIC* pSemantic,
+		EMeshBufferSemantic* pSemantic,
 		int32* pSemanticIndex,
-		MESH_BUFFER_FORMAT* pFormat,
+		EMeshBufferFormat* pFormat,
 		int32* pComponentCount,
 		int32* pOffset
 	) const
@@ -146,7 +146,7 @@ namespace mu
 		check(buffer >= 0 && buffer < m_buffers.Num());
 		check(channel >= 0 && channel < m_buffers[buffer].m_channels.Num());
 
-		const MESH_BUFFER_CHANNEL& chan = m_buffers[buffer].m_channels[channel];
+		const FMeshBufferChannel& chan = m_buffers[buffer].m_channels[channel];
 
 		if (pSemantic)
 		{
@@ -181,9 +181,9 @@ namespace mu
 		int32 buffer,
 		int32 elementSize,
 		int32 channelCount,
-		const MESH_BUFFER_SEMANTIC* pSemantics,
+		const EMeshBufferSemantic* pSemantics,
 		const int32* pSemanticIndices,
-		const MESH_BUFFER_FORMAT* pFormats,
+		const EMeshBufferFormat* pFormats,
 		const int32* pComponentCount,
 		const int32* pOffsets
 	)
@@ -191,28 +191,25 @@ namespace mu
 		check(buffer >= 0 && buffer < m_buffers.Num());
 		LLM_SCOPE_BYNAME(TEXT("MutableRuntime"));
 
-		MESH_BUFFER& buf = m_buffers[buffer];
+		FMeshBuffer& buf = m_buffers[buffer];
 
-		unsigned minElemSize = 0;
+		int32 MinElemSize = 0;
 		buf.m_channels.SetNum(channelCount);
 		for (int c = 0; c < channelCount; ++c)
 		{
-			MESH_BUFFER_CHANNEL& chan = buf.m_channels[c];
+			FMeshBufferChannel& chan = buf.m_channels[c];
 			chan.m_semantic = pSemantics ? pSemantics[c] : MBS_NONE;
 			chan.m_semanticIndex = pSemanticIndices ? ((uint8_t)pSemanticIndices[c]) : 0;
 			chan.m_format = pFormats ? pFormats[c] : MBF_NONE;
 			chan.m_componentCount = pComponentCount ? ((uint16)pComponentCount[c]) : 0;
 			chan.m_offset = pOffsets ? ((uint8_t)pOffsets[c]) : 0;
 
-			minElemSize = FMath::Max
-			(
-				minElemSize,
-				chan.m_offset + chan.m_componentCount * GetMeshFormatData(chan.m_format).m_size
-			);
+			int32 ThisChannelMinElemSize = chan.m_offset + chan.m_componentCount * GetMeshFormatData(chan.m_format).SizeInBytes;			
+			MinElemSize = FMath::Max(MinElemSize, ThisChannelMinElemSize);
 		}
 
 		// Set the user specified element size, or enlarge it if it was too small
-		buf.m_elementSize = FMath::Max(elementSize, int32(minElemSize));
+		buf.m_elementSize = FMath::Max(elementSize, MinElemSize);
 
 		// Update the buffer data
 		buf.m_data.SetNumUninitialized(buf.m_elementSize * m_elementCount, EAllowShrinking::No);
@@ -223,9 +220,9 @@ namespace mu
 	(
 		int32 buffer,
 		int32 channelIndex,
-		MESH_BUFFER_SEMANTIC semantic,
+		EMeshBufferSemantic semantic,
 		int32 semanticIndex,
-		MESH_BUFFER_FORMAT format,
+		EMeshBufferFormat format,
 		int32 componentCount,
 		int32 offset
 	)
@@ -238,7 +235,7 @@ namespace mu
 
 		LLM_SCOPE_BYNAME(TEXT("MutableRuntime"));
 
-		MESH_BUFFER& buf = m_buffers[buffer];
+		FMeshBuffer& buf = m_buffers[buffer];
 
 		if (!(channelIndex >= 0 && channelIndex < buf.m_channels.Num()))
 		{
@@ -246,7 +243,7 @@ namespace mu
 			return;
 		}
 
-		MESH_BUFFER_CHANNEL& chan = buf.m_channels[channelIndex];
+		FMeshBufferChannel& chan = buf.m_channels[channelIndex];
 		chan.m_semantic = semantic;
 		chan.m_semanticIndex = uint8_t(semanticIndex);
 		chan.m_format = format;
@@ -291,7 +288,7 @@ namespace mu
 	//---------------------------------------------------------------------------------------------
 	void FMeshBufferSet::FindChannel
 	(
-		MESH_BUFFER_SEMANTIC semantic, int32 semanticIndex,
+		EMeshBufferSemantic semantic, int32 semanticIndex,
 		int32* pBuffer, int32* pChannel
 	) const
 	{
@@ -433,7 +430,7 @@ namespace mu
 
 		if (fromIndex != toIndex)
 		{
-			for (MESH_BUFFER& b : m_buffers)
+			for (FMeshBuffer& b : m_buffers)
 			{
 				FMemory::Memcpy(&b.m_data[b.m_elementSize * toIndex],
 					&b.m_data[b.m_elementSize * fromIndex],
@@ -444,7 +441,7 @@ namespace mu
 
 
 	//-----------------------------------------------------------------------------------------
-	bool FMeshBufferSet::IsSpecialBufferToIgnoreInSimilar(const MESH_BUFFER& b) const
+	bool FMeshBufferSet::IsSpecialBufferToIgnoreInSimilar(const FMeshBuffer& b) const
 	{
 		if (b.m_channels.Num() == 1
 			&&
@@ -499,8 +496,8 @@ namespace mu
 			const int32 ThisNumChannels = m_buffers[I].m_channels.Num();
 			const int32 OtherNumChannels = m_buffers[J].m_channels.Num();
 
-			const MESH_BUFFER& ThisBuffer = m_buffers[I];
-			const MESH_BUFFER& OtherBuffer = Other.m_buffers[J];
+			const FMeshBuffer& ThisBuffer = m_buffers[I];
+			const FMeshBuffer& OtherBuffer = Other.m_buffers[J];
 
 			if (!(ThisBuffer.m_channels == OtherBuffer.m_channels && ThisBuffer.m_elementSize==OtherBuffer.m_elementSize))
 			{
@@ -516,8 +513,8 @@ namespace mu
 						continue;
 					}
 
-					const SIZE_T SizeA = GetMeshFormatData(ThisBuffer.m_channels[C].m_format).m_size * ThisBuffer.m_channels[C].m_componentCount;
-					const SIZE_T SizeB = GetMeshFormatData(OtherBuffer.m_channels[C].m_format).m_size * OtherBuffer.m_channels[C].m_componentCount;
+					const SIZE_T SizeA = GetMeshFormatData(ThisBuffer.m_channels[C].m_format).SizeInBytes * ThisBuffer.m_channels[C].m_componentCount;
+					const SIZE_T SizeB = GetMeshFormatData(OtherBuffer.m_channels[C].m_format).SizeInBytes * OtherBuffer.m_channels[C].m_componentCount;
 
 					if (SizeA != SizeB)
 					{
@@ -627,9 +624,9 @@ namespace mu
 		uint8_t currentIndices[MBS_COUNT];
 		memset(currentIndices, 0, sizeof(currentIndices));
 
-		for (MESH_BUFFER& b : m_buffers)
+		for (FMeshBuffer& b : m_buffers)
 		{
-			for (MESH_BUFFER_CHANNEL& c : b.m_channels)
+			for (FMeshBufferChannel& c : b.m_channels)
 			{
 				c.m_semanticIndex = currentIndices[c.m_semantic];
 				currentIndices[c.m_semantic]++;
@@ -641,8 +638,8 @@ namespace mu
 	//-----------------------------------------------------------------------------------------
 	void FMeshBufferSet::UpdateOffsets(int32 b)
 	{
-		uint8_t offset = 0;
-		for (MESH_BUFFER_CHANNEL& c : m_buffers[b].m_channels)
+		uint32 offset = 0;
+		for (FMeshBufferChannel& c : m_buffers[b].m_channels)
 		{
 			if (c.m_offset < offset)
 			{
@@ -652,7 +649,7 @@ namespace mu
 			{
 				offset = c.m_offset;
 			}
-			offset += uint8_t(c.m_componentCount * GetMeshFormatData(c.m_format).m_size);
+			offset += c.m_componentCount * GetMeshFormatData(c.m_format).SizeInBytes;
 		}
 
 		if (m_buffers[b].m_elementSize < offset)
@@ -661,7 +658,7 @@ namespace mu
 
 	
 	//-----------------------------------------------------------------------------------------
-	void MESH_BUFFER::Serialise(OutputArchive& arch) const
+	void FMeshBuffer::Serialise(OutputArchive& arch) const
 	{
 		arch << m_channels;
 		arch << m_data;
@@ -670,7 +667,7 @@ namespace mu
 
 	
 	//-----------------------------------------------------------------------------------------
-	void MESH_BUFFER::Unserialise(InputArchive& arch)
+	void FMeshBuffer::Unserialise(InputArchive& arch)
 	{
 		arch >> m_channels;
 		arch >> m_data;
diff --git a/Engine/Plugins/Experimental/Mutable/Source/MutableRuntime/Private/MuR/OpImageProject.cpp b/Engine/Plugins/Experimental/Mutable/Source/MutableRuntime/Private/MuR/OpImageProject.cpp
--- a/Engine/Plugins/Experimental/Mutable/Source/MutableRuntime/Private/MuR/OpImageProject.cpp
+++ b/Engine/Plugins/Experimental/Mutable/Source/MutableRuntime/Private/MuR/OpImageProject.cpp
@@ -3055,9 +3055,9 @@ void CreateMeshOptimisedForProjection(Mesh* Result, int layout)
     Result->GetVertexBuffers().SetBufferCount( 1 );
     Result->GetIndexBuffers().SetBufferCount( 1 );
 
-    MESH_BUFFER_SEMANTIC semantics[3] =	{ MBS_TEXCOORDS,	MBS_POSITION,	MBS_NORMAL };
+    EMeshBufferSemantic semantics[3] =	{ MBS_TEXCOORDS,	MBS_POSITION,	MBS_NORMAL };
     int semanticIndices[3] =			{ 0,				0,				0 };
-    MESH_BUFFER_FORMAT formats[3] =		{ MBF_FLOAT32,		MBF_FLOAT32,	MBF_FLOAT32 };
+    EMeshBufferFormat formats[3] =		{ MBF_FLOAT32,		MBF_FLOAT32,	MBF_FLOAT32 };
     int componentCounts[3] =			{ 2,				3,				3 };
     int offsets[3] =					{ 0,				8,				20 };
     semanticIndices[0] = layout;
@@ -3067,9 +3067,9 @@ void CreateMeshOptimisedForProjection(Mesh* Result, int layout)
               formats, componentCounts,
               offsets );
 
-    MESH_BUFFER_SEMANTIC isemantics[1] =	{ MBS_VERTEXINDEX };
+    EMeshBufferSemantic isemantics[1] =	{ MBS_VERTEXINDEX };
     int isemanticIndices[1] =				{ 0 };
-    MESH_BUFFER_FORMAT iformats[1] =		{ MBF_UINT32 };
+    EMeshBufferFormat iformats[1] =		{ MBF_UINT32 };
     int icomponentCounts[1] =				{ 1 };
     int ioffsets[1] =						{ 0 };
     Result->GetIndexBuffers().SetBuffer
@@ -3086,9 +3086,9 @@ void CreateMeshOptimisedForWrappingProjection(Mesh* Result, int layout)
     Result->GetVertexBuffers().SetBufferCount( 1 );
     Result->GetIndexBuffers().SetBufferCount( 1 );
 
-    MESH_BUFFER_SEMANTIC semantics[4] =	{ MBS_TEXCOORDS,	MBS_POSITION,	MBS_NORMAL,     MBS_LAYOUTBLOCK };
+    EMeshBufferSemantic semantics[4] =	{ MBS_TEXCOORDS,	MBS_POSITION,	MBS_NORMAL,     MBS_LAYOUTBLOCK };
     int semanticIndices[4] =			{ 0,				0,				0,              0 };
-    MESH_BUFFER_FORMAT formats[4] =		{ MBF_FLOAT32,		MBF_FLOAT32,	MBF_FLOAT32,    MBF_UINT32 };
+    EMeshBufferFormat formats[4] =		{ MBF_FLOAT32,		MBF_FLOAT32,	MBF_FLOAT32,    MBF_UINT32 };
     int componentCounts[4] =			{ 2,				3,				3,              1 };
     int offsets[4] =					{ 0,				8,				20,             32 };
     semanticIndices[0] = layout;
@@ -3099,9 +3099,9 @@ void CreateMeshOptimisedForWrappingProjection(Mesh* Result, int layout)
               formats, componentCounts,
               offsets );
 
-    MESH_BUFFER_SEMANTIC isemantics[1] =	{ MBS_VERTEXINDEX };
+    EMeshBufferSemantic isemantics[1] =	{ MBS_VERTEXINDEX };
     int isemanticIndices[1] =				{ 0 };
-    MESH_BUFFER_FORMAT iformats[1] =		{ MBF_UINT32 };
+    EMeshBufferFormat iformats[1] =		{ MBF_UINT32 };
     int icomponentCounts[1] =				{ 1 };
     int ioffsets[1] =						{ 0 };
     Result->GetIndexBuffers().SetBuffer
diff --git a/Engine/Plugins/Experimental/Mutable/Source/MutableRuntime/Private/MuR/OpMeshApplyLayout.h b/Engine/Plugins/Experimental/Mutable/Source/MutableRuntime/Private/MuR/OpMeshApplyLayout.h
--- a/Engine/Plugins/Experimental/Mutable/Source/MutableRuntime/Private/MuR/OpMeshApplyLayout.h
+++ b/Engine/Plugins/Experimental/Mutable/Source/MutableRuntime/Private/MuR/OpMeshApplyLayout.h
@@ -31,9 +31,9 @@ namespace mu
 		}
 
 		// Get the information about the texture coordinates channel
-		MESH_BUFFER_SEMANTIC semantic;
+		EMeshBufferSemantic semantic;
 		int32 semanticIndex;
-		MESH_BUFFER_FORMAT format;
+		EMeshBufferFormat format;
 		int32 components;
 		int32 offset;
 		pApplied->GetVertexBuffers().GetChannel( buffer, channel, &semantic, &semanticIndex, &format, &components, &offset );
diff --git a/Engine/Plugins/Experimental/Mutable/Source/MutableRuntime/Private/MuR/OpMeshBind.h b/Engine/Plugins/Experimental/Mutable/Source/MutableRuntime/Private/MuR/OpMeshBind.h
--- a/Engine/Plugins/Experimental/Mutable/Source/MutableRuntime/Private/MuR/OpMeshBind.h
+++ b/Engine/Plugins/Experimental/Mutable/Source/MutableRuntime/Private/MuR/OpMeshBind.h
@@ -137,8 +137,8 @@ namespace mu
 	{
 		constexpr static int ElementSize = sizeof(FReshapeVertexBindingData);
 		constexpr static int Channels = 4;
-		constexpr static MESH_BUFFER_SEMANTIC Semantics[Channels] = { MBS_BARYCENTRICCOORDS, MBS_DISTANCE, MBS_TRIANGLEINDEX, MBS_OTHER };
-		constexpr static MESH_BUFFER_FORMAT Formats[Channels] = { MBF_FLOAT32, MBF_FLOAT32, MBF_INT32, MBF_FLOAT32 };
+		constexpr static EMeshBufferSemantic Semantics[Channels] = { MBS_BARYCENTRICCOORDS, MBS_DISTANCE, MBS_TRIANGLEINDEX, MBS_OTHER };
+		constexpr static EMeshBufferFormat Formats[Channels] = { MBF_FLOAT32, MBF_FLOAT32, MBF_INT32, MBF_FLOAT32 };
 		constexpr static int Components[Channels] = { 2, 1, 1, 3+3+1 };
 		constexpr static int Offsets[Channels] = { 0, 8, 12, 16 };
 
@@ -164,8 +164,8 @@ namespace mu
 	{
 		constexpr static int ElementSize = sizeof(FReshapeVertexBindingData);
 		constexpr static int Channels = 4;
-		constexpr static MESH_BUFFER_SEMANTIC Semantics[Channels] = { MBS_BARYCENTRICCOORDS, MBS_DISTANCE, MBS_TRIANGLEINDEX, MBS_OTHER };
-		constexpr static MESH_BUFFER_FORMAT Formats[Channels] = { MBF_FLOAT32, MBF_FLOAT32, MBF_INT32, MBF_FLOAT32 };
+		constexpr static EMeshBufferSemantic Semantics[Channels] = { MBS_BARYCENTRICCOORDS, MBS_DISTANCE, MBS_TRIANGLEINDEX, MBS_OTHER };
+		constexpr static EMeshBufferFormat Formats[Channels] = { MBF_FLOAT32, MBF_FLOAT32, MBF_INT32, MBF_FLOAT32 };
 		constexpr static int Components[Channels] = { 2, 1, 1, 1 };
 		constexpr static int Offsets[Channels] = { 0, 8, 12, 16 };
 
@@ -181,8 +181,8 @@ namespace mu
 	{
 		constexpr static int ElementSize = sizeof(int32);
 		constexpr static int Channels = 1;
-		constexpr static MESH_BUFFER_SEMANTIC Semantics[Channels] = { MBS_OTHER };
-		constexpr static MESH_BUFFER_FORMAT Formats[Channels] = { MBF_INT32 };
+		constexpr static EMeshBufferSemantic Semantics[Channels] = { MBS_OTHER };
+		constexpr static EMeshBufferFormat Formats[Channels] = { MBF_INT32 };
 		constexpr static int Components[Channels] = { 1 };
 		constexpr static int Offsets[Channels] = { 0 };
 		constexpr static int SemanticIndices[Channels] = { 0 };
@@ -211,8 +211,8 @@ namespace mu
 	struct FClipDeformVertexBindingDataBufferDescriptor
 	{
 		constexpr static int Channels = 3;
-		constexpr static MESH_BUFFER_SEMANTIC Semantics[Channels] = { MBS_BARYCENTRICCOORDS, MBS_TRIANGLEINDEX, MBS_OTHER };
-		constexpr static MESH_BUFFER_FORMAT Formats[Channels] = { MBF_FLOAT32, MBF_INT32, MBF_FLOAT32 };
+		constexpr static EMeshBufferSemantic Semantics[Channels] = { MBS_BARYCENTRICCOORDS, MBS_TRIANGLEINDEX, MBS_OTHER };
+		constexpr static EMeshBufferFormat Formats[Channels] = { MBF_FLOAT32, MBF_INT32, MBF_FLOAT32 };
 		constexpr static int Components[Channels] = { 2, 1, 1 };
 		constexpr static int Offsets[Channels] = { 0, 8, 12 };
 
@@ -1140,8 +1140,8 @@ namespace mu
 			SkeletonBuffer.SetBuffer(0, sizeof(FReshapePointBindingData), BufDesc.Channels, BufDesc.Semantics, BufDesc.SemanticIndices, BufDesc.Formats, BufDesc.Components);
 
 			// Bone indices buffer
-			MESH_BUFFER_SEMANTIC BoneSemantics[1] = { MBS_OTHER };
-			MESH_BUFFER_FORMAT BoneFormats[1] = { MBF_INT32 };
+			EMeshBufferSemantic BoneSemantics[1] = { MBS_OTHER };
+			EMeshBufferFormat BoneFormats[1] = { MBF_INT32 };
 			int BoneSemanticIndices[1] = { 0 };
 			int BoneComponents[1] = { 1 };
 			int BoneOffsets[1] = { 0 };
diff --git a/Engine/Plugins/Experimental/Mutable/Source/MutableRuntime/Private/MuR/OpMeshClipMorphPlane.h b/Engine/Plugins/Experimental/Mutable/Source/MutableRuntime/Private/MuR/OpMeshClipMorphPlane.h
--- a/Engine/Plugins/Experimental/Mutable/Source/MutableRuntime/Private/MuR/OpMeshClipMorphPlane.h
+++ b/Engine/Plugins/Experimental/Mutable/Source/MutableRuntime/Private/MuR/OpMeshClipMorphPlane.h
@@ -140,7 +140,7 @@ namespace mu
 
 			for (int32 vb = 0; vb < pBase->GetVertexBuffers().m_buffers.Num(); ++vb)
 			{
-				const MESH_BUFFER& result = pBase->GetVertexBuffers().m_buffers[vb];
+				const FMeshBuffer& result = pBase->GetVertexBuffers().m_buffers[vb];
 
 				int elemSize = pBase->GetVertexBuffers().GetElementSize((int)vb);
 				//int firstSize = firstCount * elemSize;
@@ -148,9 +148,9 @@ namespace mu
 				for (int c = 0; c < pBase->GetVertexBuffers().GetBufferChannelCount((int)vb); ++c)
 				{
 					// Get info about the destination channel
-					MESH_BUFFER_SEMANTIC semantic = MBS_NONE;
+					EMeshBufferSemantic semantic = MBS_NONE;
 					int semanticIndex = 0;
-					MESH_BUFFER_FORMAT format = MBF_NONE;
+					EMeshBufferFormat format = MBF_NONE;
 					int components = 0;
 					int offset = 0;
 					pBase->GetVertexBuffers().GetChannel((int)vb, c, &semantic, &semanticIndex, &format, &components, &offset);
@@ -295,7 +295,7 @@ namespace mu
 		{
 			for (int32 c = 0; c < MBSPriv.m_buffers[b].m_channels.Num(); ++c)
 			{
-				MESH_BUFFER_SEMANTIC sem = MBSPriv.m_buffers[b].m_channels[c].m_semantic;
+				EMeshBufferSemantic sem = MBSPriv.m_buffers[b].m_channels[c].m_semantic;
 				int semIndex = MBSPriv.m_buffers[b].m_channels[c].m_semanticIndex;
 
 				UntypedMeshBufferIterator it(Result->GetVertexBuffers(), sem, semIndex);
diff --git a/Engine/Plugins/Experimental/Mutable/Source/MutableRuntime/Private/MuR/OpMeshClipWithMesh.cpp b/Engine/Plugins/Experimental/Mutable/Source/MutableRuntime/Private/MuR/OpMeshClipWithMesh.cpp
--- a/Engine/Plugins/Experimental/Mutable/Source/MutableRuntime/Private/MuR/OpMeshClipWithMesh.cpp
+++ b/Engine/Plugins/Experimental/Mutable/Source/MutableRuntime/Private/MuR/OpMeshClipWithMesh.cpp
@@ -426,7 +426,7 @@ namespace mu { namespace
         {
             for (int32 c = 0; c < MBSPriv2.m_buffers[b].m_channels.Num(); ++c)
             {
-                MESH_BUFFER_SEMANTIC Sem = MBSPriv2.m_buffers[b].m_channels[c].m_semantic;
+                EMeshBufferSemantic Sem = MBSPriv2.m_buffers[b].m_channels[c].m_semantic;
                 int32 SemIndex = MBSPriv2.m_buffers[b].m_channels[c].m_semanticIndex;
 
                 UntypedMeshBufferIteratorConst It(pBase->GetVertexBuffers(), Sem, SemIndex);
@@ -522,11 +522,11 @@ namespace mu { namespace
 		const FMeshBufferSet& MBSPriv = Base->GetVertexBuffers();
 		for (int32 b = 0; b < MBSPriv.m_buffers.Num(); ++b)
 		{
-			const TArray<mu::MESH_BUFFER_CHANNEL>& Channels = MBSPriv.m_buffers[b].m_channels;
+			const TArray<mu::FMeshBufferChannel>& Channels = MBSPriv.m_buffers[b].m_channels;
 
 			for (int32 c = 0; c < Channels.Num(); ++c)
 			{
-				MESH_BUFFER_SEMANTIC Sem = Channels[c].m_semantic;
+				EMeshBufferSemantic Sem = Channels[c].m_semantic;
 				if (Sem != MBS_TEXCOORDS)
 				{
 					continue;
@@ -642,9 +642,9 @@ namespace mu { namespace
 			Result->GetVertexBuffers().SetBufferCount(1);
 
 			// Vertex index channel
-			MESH_BUFFER_SEMANTIC Semantic = MBS_VERTEXINDEX;
+			EMeshBufferSemantic Semantic = MBS_VERTEXINDEX;
 			int32 SemanticIndex = 0;
-			MESH_BUFFER_FORMAT Format = MBF_UINT32;
+			EMeshBufferFormat Format = MBF_UINT32;
 			int32 Components = 1;
 			int32 Offsets = 0;
 
diff --git a/Engine/Plugins/Experimental/Mutable/Source/MutableRuntime/Private/MuR/OpMeshDifference.h b/Engine/Plugins/Experimental/Mutable/Source/MutableRuntime/Private/MuR/OpMeshDifference.h
--- a/Engine/Plugins/Experimental/Mutable/Source/MutableRuntime/Private/MuR/OpMeshDifference.h
+++ b/Engine/Plugins/Experimental/Mutable/Source/MutableRuntime/Private/MuR/OpMeshDifference.h
@@ -14,7 +14,7 @@ namespace mu
 	//---------------------------------------------------------------------------------------------
     inline void MeshDifference( Mesh* Result, const Mesh* pBase, const Mesh* pTarget,
                                    int numChannels,
-                                   const MESH_BUFFER_SEMANTIC* semantics,
+                                   const EMeshBufferSemantic* semantics,
                                    const int* semanticIndices,
                                    bool ignoreTexCoords, bool& bOutSuccess)
 
@@ -47,7 +47,7 @@ namespace mu
 		uint32 vcount = pBase->GetVertexBuffers().GetElementCount();		
 
 		// If no channels were specified, get them all
-		TArray<MESH_BUFFER_SEMANTIC> allSemantics;		
+		TArray<EMeshBufferSemantic> allSemantics;		
 		TArray<int> allSemanticIndices;
 		if ( !numChannels )
 		{
@@ -55,7 +55,7 @@ namespace mu
 			{
 				for ( int c=0; c<pBase->GetVertexBuffers().GetBufferChannelCount(vb); ++c )
 				{
-					MESH_BUFFER_SEMANTIC sem = MBS_NONE;
+					EMeshBufferSemantic sem = MBS_NONE;
 					int semIndex = 0;
 					pBase->GetVertexBuffers().GetChannel( vb, c,
 															&sem, &semIndex,
@@ -124,9 +124,9 @@ namespace mu
 			Result->GetVertexBuffers().SetElementCount( differentVertexCount );			
 			Result->GetVertexBuffers().SetBufferCount( 1 );								
 
-			TArray<MESH_BUFFER_SEMANTIC> semantic;
+			TArray<EMeshBufferSemantic> semantic;
 			TArray<int> semanticIndex;
-			TArray<MESH_BUFFER_FORMAT> format;
+			TArray<EMeshBufferFormat> format;
 			TArray<int> components;
 			TArray<int> offsets;
 			int offset = 0;										
diff --git a/Engine/Plugins/Experimental/Mutable/Source/MutableRuntime/Private/MuR/OpMeshFormat.cpp b/Engine/Plugins/Experimental/Mutable/Source/MutableRuntime/Private/MuR/OpMeshFormat.cpp
--- a/Engine/Plugins/Experimental/Mutable/Source/MutableRuntime/Private/MuR/OpMeshFormat.cpp
+++ b/Engine/Plugins/Experimental/Mutable/Source/MutableRuntime/Private/MuR/OpMeshFormat.cpp
@@ -17,142 +17,13 @@
 #include "MuR/MutableTrace.h"
 #include "MuR/Ptr.h"
 #include "MuR/RefCounted.h"
+#include "MuR/MutableRuntimeModule.h"
 
 #include "GPUSkinPublicDefs.h"
 
-
 namespace mu
 {
 
-	namespace
-	{
-
-#define MUTABLE_VERTEX_MERGE_TEX_RANGE          1024
-#define MUTABLE_TANGENT_GENERATION_EPSILON_1    0.000001f
-#define MUTABLE_TANGENT_GENERATION_EPSILON_2    0.001f
-#define MUTABLE_TANGENT_MIN_AXIS_DIFFERENCE     0.999f
-
-		//struct TVertex
-		//{
-
-		//	TVertex() = default;
-
-		//	TVertex(const FVector3f& p, const FVector3f& n, const FVector2f& t)
-		//	{
-		//		pos = p;
-		//		nor = n;
-		//		tex = t;
-		//	}
-
-		//	FVector3f pos;
-		//	FVector3f nor;
-		//	FVector2f tex;
-
-		//	inline bool operator< (const TVertex& other) const
-		//	{
-		//		if (pos < other.pos)
-		//			return true;
-		//		if (other.pos < pos)
-		//			return false;
-
-		//		if (nor < other.nor)
-		//			return true;
-		//		if (other.nor < nor)
-		//			return false;
-
-		//		// Compare the texture coordinates with a particular precission.
-		//		FIntVector2 uv0 = FIntVector2(int32(tex[0] * MUTABLE_VERTEX_MERGE_TEX_RANGE),
-		//			int32(tex[1] * MUTABLE_VERTEX_MERGE_TEX_RANGE));
-		//		FIntVector2 uv1 = FIntVector2(int32(other.tex[0] * MUTABLE_VERTEX_MERGE_TEX_RANGE),
-		//			int32(other.tex[1] * MUTABLE_VERTEX_MERGE_TEX_RANGE));
-
-		//		if (uv0 < uv1)
-		//			return true;
-		//		if (uv1 < uv0)
-		//			return false;
-
-		//		return false;
-		//	}
-
-		//	inline bool operator== (const TVertex& other) const
-		//	{
-		//		return  !((*this) < other)
-		//			&&
-		//			!(other < (*this));
-		//	}
-		//};
-
-
-		//struct TFace
-		//{
-		//	FVector3f T;
-		//	FVector3f B;
-		//	FVector3f N;
-
-		//	TFace()
-		//		: T(0, 0, 0)
-		//		, B(0, 0, 0)
-		//		, N(0, 0, 0)
-		//	{}
-
-		//	TFace
-		//	(
-		//		const FVector3f& v1,
-		//		const FVector3f& v2,
-		//		const FVector3f& v3,
-		//		const FVector2f& w1,
-		//		const FVector2f& w2,
-		//		const FVector2f& w3
-		//	)
-		//	{
-		//		FVector3f E1 = v2 - v1;
-		//		FVector3f E2 = v3 - v1;
-
-		//		FVector2f UV1 = w2 - w1;
-		//		FVector2f UV2 = w3 - w1;
-
-		//		float  UVdet = UV1[0] * UV2[1] - UV2[0] * UV1[1];
-
-		//		N = FVector3f::CrossProduct(E1, E2);
-		//		N.Normalize();
-
-		//		if (!(fabs(UVdet) <= MUTABLE_TANGENT_GENERATION_EPSILON_1))
-		//		{
-		//			double r = 1.0 / UVdet;
-
-		//			T = FVector3f
-		//				(
-		//					(UV2[1] * E1[0] - UV1[1] * E2[0]),
-		//					(UV2[1] * E1[1] - UV1[1] * E2[1]),
-		//					(UV2[1] * E1[2] - UV1[1] * E2[2])
-		//					);
-
-		//			B = FVector3f
-		//				(
-		//					(UV1[0] * E2[0] - UV2[0] * E1[0]),
-		//					(UV1[0] * E2[1] - UV2[0] * E1[1]),
-		//					(UV1[0] * E2[2] - UV2[0] * E1[2])
-		//					);
-
-		//			T = T * float(r);
-		//			B = B * float(r);
-
-		//			T.Normalize();
-		//			B.Normalize();
-		//		}
-		//		else
-		//		{
-		//			T = FVector3f(0, 0, 0);
-		//			B = FVector3f(0, 0, 0);
-		//			N = FVector3f(0, 0, 0);
-		//		}
-		//	}
-
-		//};
-
-	}
-
-
 	//-------------------------------------------------------------------------------------------------
 	void MeshFormatBuffer
 	(
@@ -169,9 +40,9 @@ namespace mu
 			for (int c = 0; c < Result.GetBufferChannelCount(b); ++c)
 			{
 				// Find this channel in the source mesh
-				MESH_BUFFER_SEMANTIC resultSemantic;
+				EMeshBufferSemantic resultSemantic;
 				int resultSemanticIndex;
-				MESH_BUFFER_FORMAT resultFormat;
+				EMeshBufferFormat resultFormat;
 				int resultComponents;
 				int resultOffset;
 				Result.GetChannel
@@ -188,7 +59,7 @@ namespace mu
 				(resultSemantic, resultSemanticIndex, &sourceBuffer, &sourceChannel);
 
 				int resultElemSize = Result.GetElementSize(b);
-				int resultChannelSize = GetMeshFormatData(resultFormat).m_size * resultComponents;
+				int resultChannelSize = GetMeshFormatData(resultFormat).SizeInBytes * resultComponents;
 				uint8_t* pResultBuf = Result.GetBufferData(b);
 				pResultBuf += resultOffset;
 
@@ -269,9 +140,9 @@ namespace mu
 				else
 				{
 					// Get the data about the source format
-					MESH_BUFFER_SEMANTIC sourceSemantic;
+					EMeshBufferSemantic sourceSemantic;
 					int sourceSemanticIndex;
-					MESH_BUFFER_FORMAT sourceFormat;
+					EMeshBufferFormat sourceFormat;
 					int sourceComponents;
 					int sourceOffset;
 					Source.GetChannel
@@ -371,8 +242,8 @@ namespace mu
 									// Add zeros. TODO: Warning?
 									FMemory::Memzero
 									(
-										pResultBuf + GetMeshFormatData(resultFormat).m_size * i,
-										GetMeshFormatData(resultFormat).m_size
+										pResultBuf + GetMeshFormatData(resultFormat).SizeInBytes * i,
+										GetMeshFormatData(resultFormat).SizeInBytes
 									);
 								}
 							}
@@ -432,9 +303,9 @@ namespace mu
 			// offsets.
 			for (int b = 0; b < Result.GetBufferCount(); ++b)
 			{
-				TArray<MESH_BUFFER_SEMANTIC> resultSemantics;
+				TArray<EMeshBufferSemantic> resultSemantics;
 				TArray<int> resultSemanticIndexs;
-				TArray<MESH_BUFFER_FORMAT> resultFormats;
+				TArray<EMeshBufferFormat> resultFormats;
 				TArray<int> resultComponentss;
 				TArray<int> resultOffsets;
 				int offset = 0;
@@ -442,9 +313,9 @@ namespace mu
 				// For every channel in this buffer
 				for (int c = 0; c < Result.GetBufferChannelCount(b); ++c)
 				{
-					MESH_BUFFER_SEMANTIC resultSemantic;
+					EMeshBufferSemantic resultSemantic;
 					int resultSemanticIndex;
-					MESH_BUFFER_FORMAT resultFormat;
+					EMeshBufferFormat resultFormat;
 					int resultComponents;
 
 					// Find this channel in the source mesh
@@ -469,7 +340,7 @@ namespace mu
 						resultComponentss.Add(resultComponents);
 						resultOffsets.Add(offset);
 
-						offset += GetMeshFormatData(resultFormat).m_size * resultComponents;
+						offset += GetMeshFormatData(resultFormat).SizeInBytes * resultComponents;
 					}
 				}
 
@@ -491,7 +362,7 @@ namespace mu
 
 
 		// For every vertex buffer in result
-		int vCount = Source.GetElementCount();
+		int32 vCount = Source.GetElementCount();
 		Result.SetElementCount(vCount);
 		for (int b = 0; b < Result.GetBufferCount(); ++b)
 		{
@@ -507,9 +378,9 @@ namespace mu
 				// Detect system buffers and clone them unmodified.
 				if (Source.GetBufferChannelCount(b) > 0)
 				{
-					MESH_BUFFER_SEMANTIC sourceSemantic;
+					EMeshBufferSemantic sourceSemantic;
 					int sourceSemanticIndex;
-					MESH_BUFFER_FORMAT sourceFormat;
+					EMeshBufferFormat sourceFormat;
 					int sourceComponents;
 					int sourceOffset;
 					Source.GetChannel
@@ -580,7 +451,7 @@ namespace mu
 				const int32 ChannelCount = VertexBuffers.GetBufferChannelCount(BufferIndex);
 				for (int32 ChannelIndex = 0; ChannelIndex < ChannelCount; ++ChannelIndex)
 				{
-					const MESH_BUFFER_CHANNEL& Channel = VertexBuffers.m_buffers[BufferIndex].m_channels[ChannelIndex];
+					const FMeshBufferChannel& Channel = VertexBuffers.m_buffers[BufferIndex].m_channels[ChannelIndex];
 
 					if (Channel.m_semantic == MBS_BONEINDICES)
 					{
@@ -605,7 +476,7 @@ namespace mu
 								++it;
 							}
 
-							MESH_BUFFER_FORMAT& format = formBuffs.m_buffers[resultBuf].m_channels[resultChan].m_format;
+							EMeshBufferFormat& format = formBuffs.m_buffers[resultBuf].m_channels[resultChan].m_format;
 							if (maxBoneIndex > 0xffff && (format == MBF_UINT8 || format == MBF_UINT16))
 							{
 								format = MBF_UINT32;
@@ -634,11 +505,8 @@ namespace mu
 
 		// \todo Make sure that the vertex indices will fit in this format, or extend it.
 
-
-
 		if (formatVertices)
 		{
-
 			FormatBufferSet(pSource->GetVertexBuffers(), Result->GetVertexBuffers(),
 				keepSystemBuffers, ignoreMissingChannels, true);
 		}
@@ -699,4 +567,89 @@ namespace mu
 		Result->EnsureSurfaceData();
 	}
 
+
+	void MeshOptimizeBuffers( Mesh* InMesh )
+	{
+		if (!InMesh)
+		{
+			return;
+		}
+
+		FMeshBufferSet& VertexBuffers = InMesh->m_VertexBuffers;
+
+		// Reduce the number of influences if possible
+		constexpr int32 SemanticIndex = 0;
+		
+		UntypedMeshBufferIteratorConst WeightIt(VertexBuffers, MBS_BONEWEIGHTS, SemanticIndex);
+		if (WeightIt.ptr())
+		{
+			int32 BufferInfluences = WeightIt.GetComponents();
+			int32 RealInfluences = 0;
+
+			for (int32 VertexIndex = 0; VertexIndex < VertexBuffers.GetElementCount(); ++VertexIndex)
+			{
+				int32 ThisVertexInfluences = 0;
+
+				switch (WeightIt.GetFormat())
+				{
+				case MBF_NUINT8:
+				{
+					const uint8* Data = reinterpret_cast<const uint8*>(WeightIt.ptr());
+					for (int32 InfluenceIndex = 0; InfluenceIndex < BufferInfluences; ++InfluenceIndex)
+					{
+						if (*Data>0)
+						{
+							++ThisVertexInfluences;
+						}
+						++Data;
+					}
+					break;
+				}
+
+				default:
+					// Unsupported
+					check(false);
+					break;
+				}
+
+				++WeightIt;
+
+				RealInfluences = FMath::Max(RealInfluences,ThisVertexInfluences);
+			}
+
+			if (RealInfluences<BufferInfluences)
+			{
+				// Remove the useless influences from the buffer.
+
+				// \todo: This is a generic innefficient way
+				FMeshBufferSet NewVertexBuffers;
+				NewVertexBuffers.m_buffers = VertexBuffers.m_buffers;
+
+				for (FMeshBuffer& Buffer: NewVertexBuffers.m_buffers)
+				{
+					int32 OffsetDelta = 0;
+					for (FMeshBufferChannel& Channel : Buffer.m_channels)
+					{
+						Channel.m_offset += OffsetDelta;
+
+						if (Channel.m_semanticIndex == SemanticIndex
+							&&
+							(Channel.m_semantic == MBS_BONEWEIGHTS || Channel.m_semantic == MBS_BONEINDICES)
+							)
+						{
+							Channel.m_componentCount = RealInfluences;
+							OffsetDelta -= (BufferInfluences - RealInfluences) * GetMeshFormatData(Channel.m_format).SizeInBytes;
+						}
+					}
+
+					Buffer.m_elementSize += OffsetDelta;
+				}
+
+				FormatBufferSet( VertexBuffers, NewVertexBuffers, true, false, true);
+
+				InMesh->m_VertexBuffers = NewVertexBuffers;
+			}
+		}
+
+	}
 }
diff --git a/Engine/Plugins/Experimental/Mutable/Source/MutableRuntime/Private/MuR/OpMeshMorph.h b/Engine/Plugins/Experimental/Mutable/Source/MutableRuntime/Private/MuR/OpMeshMorph.h
--- a/Engine/Plugins/Experimental/Mutable/Source/MutableRuntime/Private/MuR/OpMeshMorph.h
+++ b/Engine/Plugins/Experimental/Mutable/Source/MutableRuntime/Private/MuR/OpMeshMorph.h
@@ -63,13 +63,13 @@ namespace mu
 			const UntypedMeshBufferIterator& BaseTangentIter = BaseTangentFrameIters[1];
 			const UntypedMeshBufferIterator& BaseBiNormalIter = BaseTangentFrameIters[0];
 
-            const MESH_BUFFER_FORMAT NormalFormat = BaseNormalIter.GetFormat();
+            const EMeshBufferFormat NormalFormat = BaseNormalIter.GetFormat();
             const int32 NormalComps = BaseNormalIter.GetComponents();
 
-            const MESH_BUFFER_FORMAT TangentFormat = BaseTangentIter.GetFormat();
+            const EMeshBufferFormat TangentFormat = BaseTangentIter.GetFormat();
             const int32 TangentComps = BaseTangentIter.GetComponents();
 
-            const MESH_BUFFER_FORMAT BiNormalFormat = BaseBiNormalIter.GetFormat();
+            const EMeshBufferFormat BiNormalFormat = BaseBiNormalIter.GetFormat();
             const int32 BiNormalComps = BaseBiNormalIter.GetComponents();
 
 			// When normal is packed, binormal channel is not expected. It is not a big deal if it's there but we would be doing extra unused work in that case. 
@@ -189,7 +189,7 @@ namespace mu
                     UntypedMeshBufferIterator ChannelBaseIter = BaseChannelsIters[ChannelIndex] + VertexIndex;
 					UntypedMeshBufferIteratorConst ChannelMorphIter = MorphChannelsIters[ChannelIndex] + MorphIndex;
                    
-                    const MESH_BUFFER_FORMAT DestChannelFormat = BaseChannelsIters[ChannelIndex].GetFormat();
+                    const EMeshBufferFormat DestChannelFormat = BaseChannelsIters[ChannelIndex].GetFormat();
                     const int32 DestChannelComps = BaseChannelsIters[ChannelIndex].GetComponents();
 
                     // Apply Morph to range found above.
@@ -245,7 +245,7 @@ namespace mu
 			for (int32 ChannelIndex = 1; ChannelIndex < ChannelsNum; ++ChannelIndex)
 			{
 				const FMeshBufferSet& MBSPriv = RefTarget->GetVertexBuffers();
-				MESH_BUFFER_SEMANTIC Sem = MBSPriv.m_buffers[0].m_channels[ChannelIndex].m_semantic;
+				EMeshBufferSemantic Sem = MBSPriv.m_buffers[0].m_channels[ChannelIndex].m_semantic;
 				int32 SemIndex = MBSPriv.m_buffers[0].m_channels[ChannelIndex].m_semanticIndex;
 			
 				if (Sem == MBS_NORMAL && bBaseHasNormals)
diff --git a/Engine/Plugins/Experimental/Mutable/Source/MutableRuntime/Private/MuR/OpMeshOptimizeSkinning.h b/Engine/Plugins/Experimental/Mutable/Source/MutableRuntime/Private/MuR/OpMeshOptimizeSkinning.h
--- a/Engine/Plugins/Experimental/Mutable/Source/MutableRuntime/Private/MuR/OpMeshOptimizeSkinning.h
+++ b/Engine/Plugins/Experimental/Mutable/Source/MutableRuntime/Private/MuR/OpMeshOptimizeSkinning.h
@@ -42,13 +42,13 @@ namespace mu
 		bool bRequiresFormatChange = false;
 
 		// Desired format if MaxBoneMapIndex <= MAX_uint8
-		const MESH_BUFFER_FORMAT DesiredBoneIndexFormat = mu::MBF_UINT8;
+		const EMeshBufferFormat DesiredBoneIndexFormat = mu::MBF_UINT8;
 		
 		// Iterate all vertex buffers and check if the BoneIndices buffers have the desired format
 		const FMeshBufferSet& InMeshVertexBuffers = InMesh->GetVertexBuffers();
 		for (int32 VertexBufferIndex = 0; !bRequiresFormatChange && VertexBufferIndex < InMeshVertexBuffers.m_buffers.Num(); ++VertexBufferIndex)
 		{
-			const MESH_BUFFER& Buffer = InMeshVertexBuffers.m_buffers[VertexBufferIndex];
+			const FMeshBuffer& Buffer = InMeshVertexBuffers.m_buffers[VertexBufferIndex];
 
 			const int32 ChannelsCount = InMeshVertexBuffers.GetBufferChannelCount(VertexBufferIndex);
 			for (int32 ChannelIndex = 0; ChannelIndex < ChannelsCount; ++ChannelIndex)
@@ -87,19 +87,19 @@ namespace mu
 
 		for (int32 BufferIndex = 0; BufferIndex < VertexBuffersCount; ++BufferIndex)
 		{
-			const MESH_BUFFER& SourceBuffer = InMeshVertexBuffers.m_buffers[BufferIndex];
+			const FMeshBuffer& SourceBuffer = InMeshVertexBuffers.m_buffers[BufferIndex];
 			const int32 ChannelsCount = InMeshVertexBuffers.GetBufferChannelCount(BufferIndex);
 
 			int32 BoneIndexChannelIndex = INDEX_NONE;
 
 			int32 ElementSize = 0;
-			TArray<MESH_BUFFER_SEMANTIC> Semantics;
+			TArray<EMeshBufferSemantic> Semantics;
 			TArray<int32> SemanticIndices;
-			TArray<MESH_BUFFER_FORMAT> Formats;
+			TArray<EMeshBufferFormat> Formats;
 			TArray<int32> Components;
 			TArray<int32> Offsets;
 
-			MESH_BUFFER_FORMAT SourceBoneIndexFormat = mu::MBF_NONE;
+			EMeshBufferFormat SourceBoneIndexFormat = mu::MBF_NONE;
 
 			// Offset accumulator
 			int32 AuxOffset = 0;
@@ -107,8 +107,8 @@ namespace mu
 			// Copy and fix channel details
 			for (int32 ChannelIndex = 0; ChannelIndex < ChannelsCount; ++ChannelIndex)
 			{
-				const MESH_BUFFER_CHANNEL& Channel = SourceBuffer.m_channels[ChannelIndex];
-				MESH_BUFFER_FORMAT Format = Channel.m_format;
+				const FMeshBufferChannel& Channel = SourceBuffer.m_channels[ChannelIndex];
+				EMeshBufferFormat Format = Channel.m_format;
 
 				if (Channel.m_semantic == MBS_BONEINDICES)
 				{
@@ -123,7 +123,7 @@ namespace mu
 				Components.Add(Channel.m_componentCount);
 				Offsets.Add(AuxOffset);
 
-				const int32 FormatSize = GetMeshFormatData(Format).m_size;
+				const int32 FormatSize = GetMeshFormatData(Format).SizeInBytes;
 
 				ElementSize += FormatSize;
 				AuxOffset += FormatSize * Channel.m_componentCount;
@@ -138,9 +138,9 @@ namespace mu
 				uint8* Data = VertexBuffers.GetBufferData(BufferIndex);
 
 				const uint8* SourceData = (const uint8*)SourceBuffer.m_data.GetData();
-				const int32 SourceBoneIndexSize = GetMeshFormatData(SourceBoneIndexFormat).m_size;
+				const int32 SourceBoneIndexSize = GetMeshFormatData(SourceBoneIndexFormat).SizeInBytes;
 
-				const int32 BoneIndexSize = GetMeshFormatData(mu::MBF_UINT8).m_size;
+				const int32 BoneIndexSize = GetMeshFormatData(mu::MBF_UINT8).SizeInBytes;
 				const int32 BoneIndexComponentCount = Components[BoneIndexChannelIndex];
 				const int32 BoneIndexChannelOffset = Offsets[BoneIndexChannelIndex];
 
diff --git a/Engine/Plugins/Experimental/Mutable/Source/MutableRuntime/Private/MuR/OpMeshSmoothing.h b/Engine/Plugins/Experimental/Mutable/Source/MutableRuntime/Private/MuR/OpMeshSmoothing.h
--- a/Engine/Plugins/Experimental/Mutable/Source/MutableRuntime/Private/MuR/OpMeshSmoothing.h
+++ b/Engine/Plugins/Experimental/Mutable/Source/MutableRuntime/Private/MuR/OpMeshSmoothing.h
@@ -563,7 +563,7 @@ namespace mu
 		}
 
 		// Copy positions to the collapsed vertices.
-		const SIZE_T PositionElemSize = GetMeshFormatData(DstPositionBegin.GetFormat()).m_size * DstPositionBegin.GetComponents(); 
+		const SIZE_T PositionElemSize = GetMeshFormatData(DstPositionBegin.GetFormat()).SizeInBytes * DstPositionBegin.GetComponents();
 		for (int32 VertexIndex = 0; VertexIndex < NumVertices; ++VertexIndex)
 		{
 			const int32 MappedIndex = UniqueVertexMap[VertexIndex];  
diff --git a/Engine/Plugins/Experimental/Mutable/Source/MutableRuntime/Private/MuR/OpMeshTransform.h b/Engine/Plugins/Experimental/Mutable/Source/MutableRuntime/Private/MuR/OpMeshTransform.h
--- a/Engine/Plugins/Experimental/Mutable/Source/MutableRuntime/Private/MuR/OpMeshTransform.h
+++ b/Engine/Plugins/Experimental/Mutable/Source/MutableRuntime/Private/MuR/OpMeshTransform.h
@@ -36,7 +36,7 @@ namespace mu
 
             for ( int32 c=0; c<MBSPriv.m_buffers[b].m_channels.Num(); ++c )
             {
-                MESH_BUFFER_SEMANTIC sem = MBSPriv.m_buffers[b].m_channels[c].m_semantic;
+                EMeshBufferSemantic sem = MBSPriv.m_buffers[b].m_channels[c].m_semantic;
                 int semIndex = MBSPriv.m_buffers[b].m_channels[c].m_semanticIndex;
 
                 UntypedMeshBufferIterator it( Result->GetVertexBuffers(), sem, semIndex );
diff --git a/Engine/Plugins/Experimental/Mutable/Source/MutableRuntime/Public/MuR/MeshBufferSet.h b/Engine/Plugins/Experimental/Mutable/Source/MutableRuntime/Public/MuR/MeshBufferSet.h
--- a/Engine/Plugins/Experimental/Mutable/Source/MutableRuntime/Public/MuR/MeshBufferSet.h
+++ b/Engine/Plugins/Experimental/Mutable/Source/MutableRuntime/Public/MuR/MeshBufferSet.h
@@ -65,19 +65,19 @@ namespace mu
 
 		_MBF_FORCE32BITS = 0xFFFFFFFF
 
-	} MESH_BUFFER_FORMAT;
+	} EMeshBufferFormat;
 
 	//!
 	struct FMeshBufferFormatData
 	{
-		//! Size per component
-		unsigned m_size;
+		/** Size per component in bytes. */
+		uint8 SizeInBytes;
 
-		//! log 2 of the max value if integer
-		uint8_t m_maxValueBits;
+		/** log 2 of the max value if integer. */
+		uint8 MaxValueBits;
 	};
 
-	MUTABLERUNTIME_API const FMeshBufferFormatData& GetMeshFormatData(MESH_BUFFER_FORMAT format);
+	MUTABLERUNTIME_API const FMeshBufferFormatData& GetMeshFormatData(EMeshBufferFormat format);
 
 
 	//! Semantics of the mesh buffers
@@ -124,13 +124,13 @@ namespace mu
 		MBS_COUNT,
 
         _MBS_FORCE32BITS = 0xFFFFFFFF
-	} MESH_BUFFER_SEMANTIC;
+	} EMeshBufferSemantic;
 
 
 	//!
-	struct MESH_BUFFER_CHANNEL
+	struct FMeshBufferChannel
 	{
-		MESH_BUFFER_CHANNEL()
+		FMeshBufferChannel()
 		{
 			m_semantic = MBS_NONE;
 			m_format = MBF_NONE;
@@ -140,10 +140,10 @@ namespace mu
 		}
 
 		//!
-		MESH_BUFFER_SEMANTIC m_semantic;
+		EMeshBufferSemantic m_semantic;
 
 		//!
-		MESH_BUFFER_FORMAT m_format;
+		EMeshBufferFormat m_format;
 
 		//! Index of the semantic, in case there are more than one of this type.
 		uint8 m_semanticIndex;
@@ -155,7 +155,7 @@ namespace mu
 		uint16 m_componentCount;
 
 		//!
-		inline bool operator==(const MESH_BUFFER_CHANNEL& o) const
+		inline bool operator==(const FMeshBufferChannel& o) const
 		{
 			return (m_semantic == o.m_semantic) &&
 				(m_format == o.m_format) &&
@@ -167,19 +167,19 @@ namespace mu
 	};
 
 
-	struct MESH_BUFFER
+	struct FMeshBuffer
 	{
 		template<typename Type>
 		using TMemoryTrackedArray = TArray<Type, FDefaultMemoryTrackingAllocator<MemoryCounters::FMeshMemoryCounter>>;
 
 		//!
-		MESH_BUFFER()
+		FMeshBuffer()
 		{
 			m_elementSize = 0;
 		}
 
 		//!
-		TArray<MESH_BUFFER_CHANNEL> m_channels;
+		TArray<FMeshBufferChannel> m_channels;
 
 		//!
 		TMemoryTrackedArray<uint8> m_data;
@@ -194,7 +194,7 @@ namespace mu
 		inline void Unserialise(mu::InputArchive& arch);
 
 		//!
-		inline bool operator==(const MESH_BUFFER& o) const
+		inline bool operator==(const FMeshBuffer& o) const
 		{
 			bool equal = (m_channels == o.m_channels);
 			if (equal) equal = (m_elementSize == o.m_elementSize);
@@ -214,7 +214,7 @@ namespace mu
 		uint32 m_elementCount = 0;
 
 		//!
-		TArray<MESH_BUFFER> m_buffers;
+		TArray<FMeshBuffer> m_buffers;
 
 		//!
 		void Serialise(OutputArchive& arch) const;
@@ -265,9 +265,9 @@ namespace mu
 			(
 				int32 buffer,
 				int32 channel,
-				MESH_BUFFER_SEMANTIC* pSemantic,
+				EMeshBufferSemantic* pSemantic,
 				int32* pSemanticIndex,
-				MESH_BUFFER_FORMAT* pFormat,
+				EMeshBufferFormat* pFormat,
 				int32* pComponentCount,
 				int32* pOffset
 			) const;
@@ -286,9 +286,9 @@ namespace mu
 				int32 buffer,
 				int32 elementSize,
 				int32 channelCount,
-				const MESH_BUFFER_SEMANTIC* pSemantics=nullptr,
+				const EMeshBufferSemantic* pSemantics=nullptr,
 				const int32* pSemanticIndices=nullptr,
-				const MESH_BUFFER_FORMAT* pFormats=nullptr,
+				const EMeshBufferFormat* pFormats=nullptr,
 				const int32* pComponentCounts=nullptr,
 				const int32* pOffsets=nullptr
 			);
@@ -301,9 +301,9 @@ namespace mu
 			(
 				int32 buffer,
 				int32 channelIndex,
-				MESH_BUFFER_SEMANTIC semantic,
+				EMeshBufferSemantic semantic,
 				int32 semanticIndice,
-				MESH_BUFFER_FORMAT format,
+				EMeshBufferFormat format,
 				int32 componentCount,
 				int32 offset
 			);
@@ -331,9 +331,9 @@ namespace mu
 		//! channel index of the channel inside the buffer returned at [buffer]
 		void FindChannel
 			(
-				MESH_BUFFER_SEMANTIC semantic,
-				int semanticIndex,
-				int* pBuffer, int* pChannel
+				EMeshBufferSemantic semantic,
+				int32 semanticIndex,
+				int32* pBuffer, int32* pChannel
 			) const;
 
 		//! Get the offset in bytes of the data of this channel inside an element data.
@@ -368,7 +368,7 @@ namespace mu
 
 		//! Compare the mesh buffer with another one, but ignore internal data like generated
 		//! vertex indices.
-		bool IsSpecialBufferToIgnoreInSimilar(const MESH_BUFFER& b) const;
+		bool IsSpecialBufferToIgnoreInSimilar(const FMeshBuffer& b) const;
 
 		//! Compare the mesh buffer with another one, but ignore internal data like generated
 		//! vertex indices. Be aware this method compares the data byte by byte without checking
@@ -390,9 +390,9 @@ namespace mu
 	};	
 
 	
-	MUTABLE_DEFINE_POD_SERIALISABLE(MESH_BUFFER_CHANNEL);
-	MUTABLE_DEFINE_POD_VECTOR_SERIALISABLE(MESH_BUFFER_CHANNEL);
-	MUTABLE_DEFINE_ENUM_SERIALISABLE(MESH_BUFFER_FORMAT);
-	MUTABLE_DEFINE_ENUM_SERIALISABLE(MESH_BUFFER_SEMANTIC);
+	MUTABLE_DEFINE_POD_SERIALISABLE(FMeshBufferChannel);
+	MUTABLE_DEFINE_POD_VECTOR_SERIALISABLE(FMeshBufferChannel);
+	MUTABLE_DEFINE_ENUM_SERIALISABLE(EMeshBufferFormat);
+	MUTABLE_DEFINE_ENUM_SERIALISABLE(EMeshBufferSemantic);
 }
 
diff --git a/Engine/Plugins/Experimental/Mutable/Source/MutableTools/Internal/MuT/NodeMeshFormatPrivate.h b/Engine/Plugins/Experimental/Mutable/Source/MutableTools/Internal/MuT/NodeMeshFormatPrivate.h
--- a/Engine/Plugins/Experimental/Mutable/Source/MutableTools/Internal/MuT/NodeMeshFormatPrivate.h
+++ b/Engine/Plugins/Experimental/Mutable/Source/MutableTools/Internal/MuT/NodeMeshFormatPrivate.h
@@ -4,8 +4,6 @@
 
 #include "MuT/NodeMeshPrivate.h"
 #include "MuT/NodeMeshFormat.h"
-#include "MuT/AST.h"
-
 
 namespace mu
 {
@@ -18,46 +16,54 @@ namespace mu
 		static FNodeType s_type;
 
 		//! Source mesh to transform
-		NodeMeshPtr m_pSource;
+		Ptr<NodeMesh> Source;
 
-		//! New mesh format. The buffers in the sets have no elements, but they define the formats.
-		//! If they are null it means that they are left with the original format.
-		FMeshBufferSet m_VertexBuffers;
-		FMeshBufferSet m_IndexBuffers;
-		FMeshBufferSet m_FaceBuffers;
+		/** New mesh format.The buffers in the sets have no elements, but they define the formats. */
+		FMeshBufferSet VertexBuffers;
+		FMeshBufferSet IndexBuffers;
+		FMeshBufferSet FaceBuffers;
+		
+		/** */
+		bool bOptimizeBuffers = false;
 
 		//!
 		void Serialise( OutputArchive& arch ) const
 		{
-            uint32_t ver = 4;
+            uint32 ver = 5;
 			arch << ver;
 
-			arch << m_pSource;
-			arch << m_VertexBuffers;
-			arch << m_IndexBuffers;
-			arch << m_FaceBuffers;
+			arch << Source;
+			arch << VertexBuffers;
+			arch << IndexBuffers;
+			arch << FaceBuffers;
+			arch << bOptimizeBuffers;
 		}
 
 		//!
 		void Unserialise( InputArchive& arch )
 		{
-            uint32_t ver;
+            uint32 ver;
 			arch >> ver;
-            check(ver>=3 && ver<=4);
+            check(ver>=3 && ver<=5);
 
-			arch >> m_pSource;
-			arch >> m_VertexBuffers;
-			arch >> m_IndexBuffers;
-			arch >> m_FaceBuffers;
+			arch >> Source;
+			arch >> VertexBuffers;
+			arch >> IndexBuffers;
+			arch >> FaceBuffers;
 			if (ver == 3)
 			{
 				bool bDummy;
 				arch >> bDummy;
 			}
+
+			if (ver >= 5)
+			{
+				arch >> bOptimizeBuffers;
+			}
 		}
 
 		// NodeMesh::Private interface
-        NodeLayoutPtr GetLayout( int index ) const override;
+        Ptr<NodeLayout> GetLayout( int32 index ) const override;
 	};
 
 
diff --git a/Engine/Plugins/Experimental/Mutable/Source/MutableTools/Private/MuT/ASTOpMeshFormat.cpp b/Engine/Plugins/Experimental/Mutable/Source/MutableTools/Private/MuT/ASTOpMeshFormat.cpp
--- a/Engine/Plugins/Experimental/Mutable/Source/MutableTools/Private/MuT/ASTOpMeshFormat.cpp
+++ b/Engine/Plugins/Experimental/Mutable/Source/MutableTools/Private/MuT/ASTOpMeshFormat.cpp
@@ -45,7 +45,7 @@ bool ASTOpMeshFormat::IsEqual(const ASTOp& otherUntyped) const
 	if (otherUntyped.GetOpType() == GetOpType())
 	{
 		const ASTOpMeshFormat* other = static_cast<const ASTOpMeshFormat*>(&otherUntyped);
-        return Source==other->Source && Format==other->Format && Buffers==other->Buffers;
+        return Source==other->Source && Format==other->Format && Flags ==other->Flags;
     }
     return false;
 }
@@ -64,7 +64,8 @@ mu::Ptr<ASTOp> ASTOpMeshFormat::Clone(MapChildFuncRef mapChild) const
 	mu::Ptr<ASTOpMeshFormat> n = new ASTOpMeshFormat();
     n->Source = mapChild(Source.child());
     n->Format = mapChild(Format.child());
-	n->Buffers = Buffers;
+	n->Flags = Flags;
+	n->bOptimizeBuffers = bOptimizeBuffers;
     return n;
 }
 
@@ -84,13 +85,17 @@ void ASTOpMeshFormat::Link( FProgram& program, FLinkerOptions* )
         OP::MeshFormatArgs args;
         FMemory::Memzero( &args, sizeof(args) );
 
-		args.buffers = Buffers;
-        if (Source) args.source = Source->linkedAddress;
-        if (Format) args.format = Format->linkedAddress;
+		args.Flags = Flags;
+		if (bOptimizeBuffers)
+		{
+			args.Flags = args.Flags | OP::MeshFormatArgs::OptimizeBuffers;
+		}
+
+		if (Source) args.source = Source->linkedAddress;
+		if (Format) args.format = Format->linkedAddress;
 
         linkedAddress = (OP::ADDRESS)program.m_opAddress.Num();
-        //program.m_code.push_back(op);
-        program.m_opAddress.Add((uint32_t)program.m_byteCode.Num());
+        program.m_opAddress.Add((uint32)program.m_byteCode.Num());
         AppendCode(program.m_byteCode,OP_TYPE::ME_FORMAT);
         AppendCode(program.m_byteCode,args);
     }
@@ -162,9 +167,9 @@ namespace
 
 		int offset = 0;
 		int numChannels = 0;
-		TArray<MESH_BUFFER_SEMANTIC> semantics;
+		TArray<EMeshBufferSemantic> semantics;
 		TArray<int> semanticIndices;
-		TArray<MESH_BUFFER_FORMAT> formats;
+		TArray<EMeshBufferFormat> formats;
 		TArray<int> components;
 		TArray<int> offsets;
 
@@ -174,7 +179,7 @@ namespace
 		formats.Add(MBF_UINT32);
 		components.Add(1);
 		offsets.Add(offset);
-		offset += components[numChannels] * GetMeshFormatData(formats[numChannels]).m_size;
+		offset += components[numChannels] * GetMeshFormatData(formats[numChannels]).SizeInBytes;
 		numChannels++;
 
 		// Add the vertex channels from the new format
@@ -183,9 +188,9 @@ namespace
 			for (int c = 0; c < pTargetFormat->GetVertexBuffers().GetBufferChannelCount(vb); ++c)
 			{
 				// Channel info
-				MESH_BUFFER_SEMANTIC semantic;
+				EMeshBufferSemantic semantic;
 				int semanticIndex;
-				MESH_BUFFER_FORMAT format;
+				EMeshBufferFormat format;
 				int component;
 				pTargetFormat->GetVertexBuffers().GetChannel
 				(vb, c, &semantic, &semanticIndex, &format, &component, nullptr);
@@ -197,7 +202,7 @@ namespace
 				components.Add(component);
 				offsets.Add(offset);
 				offset += components[numChannels]
-					* GetMeshFormatData(formats[numChannels]).m_size;
+					* GetMeshFormatData(formats[numChannels]).SizeInBytes;
 				numChannels++;
 			}
 		}
@@ -288,9 +293,7 @@ mu::Ptr<ASTOp> Sink_MeshFormatAST::Visit(const mu::Ptr<ASTOp>& at, const ASTOpMe
 		if (newOp->Target)
 		{
 			mu::Ptr<ASTOpMeshFormat> newFormat = mu::Clone<ASTOpMeshFormat>(currentFormatOp);
-			newFormat->Buffers =
-				OP::MeshFormatArgs::BT_VERTEX
-				| OP::MeshFormatArgs::BT_IGNORE_MISSING;
+			newFormat->Flags = OP::MeshFormatArgs::Vertex | OP::MeshFormatArgs::IgnoreMissing;
 			newFormat->Format = targetMorphFormatAt;
 
 			newOp->Target = Visit(newOp->Target.child(), newFormat.get());
@@ -330,9 +333,7 @@ mu::Ptr<ASTOp> Sink_MeshFormatAST::Visit(const mu::Ptr<ASTOp>& at, const ASTOpMe
 			if (newOp->children[newOp->op.args.MeshInterpolate.targets[t]])
 			{
 				mu::Ptr<ASTOpMeshFormat> newFormat = mu::Clone<ASTOpMeshFormat>(currentFormatOp);
-				newFormat->Buffers =
-					OP::MeshFormatArgs::BT_VERTEX
-					| OP::MeshFormatArgs::BT_IGNORE_MISSING;
+				newFormat->Flags = OP::MeshFormatArgs::Vertex | OP::MeshFormatArgs::IgnoreMissing;
 				newFormat->Format = targetMorphFormatAt;
 
 				newOp->SetChild(newOp->op.args.MeshInterpolate.targets[t], Visit(newOp->children[newOp->op.args.MeshInterpolate.targets[t]].child(), newFormat.get()));
diff --git a/Engine/Plugins/Experimental/Mutable/Source/MutableTools/Private/MuT/ASTOpMeshFormat.h b/Engine/Plugins/Experimental/Mutable/Source/MutableTools/Private/MuT/ASTOpMeshFormat.h
--- a/Engine/Plugins/Experimental/Mutable/Source/MutableTools/Private/MuT/ASTOpMeshFormat.h
+++ b/Engine/Plugins/Experimental/Mutable/Source/MutableTools/Private/MuT/ASTOpMeshFormat.h
@@ -18,7 +18,8 @@ namespace mu
 
 		ASTChild Source;
 		ASTChild Format;
-		uint8_t Buffers = 0;
+		uint8 Flags = 0;
+		bool bOptimizeBuffers = false;
 
 	public:
 
diff --git a/Engine/Plugins/Experimental/Mutable/Source/MutableTools/Private/MuT/CodeGenerator_Image.cpp b/Engine/Plugins/Experimental/Mutable/Source/MutableTools/Private/MuT/CodeGenerator_Image.cpp
--- a/Engine/Plugins/Experimental/Mutable/Source/MutableTools/Private/MuT/CodeGenerator_Image.cpp
+++ b/Engine/Plugins/Experimental/Mutable/Source/MutableTools/Private/MuT/CodeGenerator_Image.cpp
@@ -1444,10 +1444,10 @@ namespace mu
                 cop->SetValue(FormatMeshResult, m_compilerOptions->OptimisationOptions.DiskCacheContext);
 
                 Ptr<ASTOpMeshFormat> FormatOp = new ASTOpMeshFormat();
-				FormatOp->Buffers = OP::MeshFormatArgs::BT_VERTEX
-                        | OP::MeshFormatArgs::BT_INDEX
-                        | OP::MeshFormatArgs::BT_FACE
-                        | OP::MeshFormatArgs::BT_RESETBUFFERINDICES;
+				FormatOp->Flags = OP::MeshFormatArgs::Vertex
+                        | OP::MeshFormatArgs::Index
+                        | OP::MeshFormatArgs::Face
+                        | OP::MeshFormatArgs::ResetBufferIndices;
 				FormatOp->Format = cop;
 				FormatOp->Source = CurrentMeshToProjectOp;
 				CurrentMeshToProjectOp = FormatOp;
@@ -1476,10 +1476,10 @@ namespace mu
                 cop->SetValue(FormatMeshResult, m_compilerOptions->OptimisationOptions.DiskCacheContext);
 
                 Ptr<ASTOpMeshFormat> FormatOp = new ASTOpMeshFormat();
-				FormatOp->Buffers = OP::MeshFormatArgs::BT_VERTEX
-					| OP::MeshFormatArgs::BT_INDEX
-					| OP::MeshFormatArgs::BT_FACE
-					| OP::MeshFormatArgs::BT_RESETBUFFERINDICES;
+				FormatOp->Flags = OP::MeshFormatArgs::Vertex
+					| OP::MeshFormatArgs::Index
+					| OP::MeshFormatArgs::Face
+					| OP::MeshFormatArgs::ResetBufferIndices;
 				FormatOp->Format = cop;
 				FormatOp->Source = CurrentMeshToProjectOp;
 				CurrentMeshToProjectOp = FormatOp;
diff --git a/Engine/Plugins/Experimental/Mutable/Source/MutableTools/Private/MuT/CodeGenerator_Mesh.cpp b/Engine/Plugins/Experimental/Mutable/Source/MutableTools/Private/MuT/CodeGenerator_Mesh.cpp
--- a/Engine/Plugins/Experimental/Mutable/Source/MutableTools/Private/MuT/CodeGenerator_Mesh.cpp
+++ b/Engine/Plugins/Experimental/Mutable/Source/MutableTools/Private/MuT/CodeGenerator_Mesh.cpp
@@ -295,9 +295,9 @@ namespace mu
 
 			// TODO
 			check(Layout->GetBlockCount() < MAX_uint16);
-			const MESH_BUFFER_SEMANTIC semantic = MBS_LAYOUTBLOCK;
+			const EMeshBufferSemantic semantic = MBS_LAYOUTBLOCK;
 			const int32 semanticIndex = int32(currentLayoutChannel);
-			const MESH_BUFFER_FORMAT format = MBF_UINT16;
+			const EMeshBufferFormat format = MBF_UINT16;
 			const int32 components = 1;
 			const int32 offset = 0;
 			currentLayoutMesh->GetVertexBuffers().SetBuffer
@@ -390,9 +390,9 @@ namespace mu
 		check(buffer >= 0);
 		check(channel >= 0);
 
-		MESH_BUFFER_SEMANTIC semantic;
+		EMeshBufferSemantic semantic;
 		int semanticIndex;
-		MESH_BUFFER_FORMAT format;
+		EMeshBufferFormat format;
 		int components;
 		int offset;
 		currentLayoutMesh->GetVertexBuffers().GetChannel
@@ -1420,9 +1420,9 @@ namespace mu
 				{
 					int newBuffer = pCloned->GetVertexBuffers().GetBufferCount();
 					pCloned->GetVertexBuffers().SetBufferCount(newBuffer + 1);
-					MESH_BUFFER_SEMANTIC semantic = MBS_VERTEXINDEX;
+					EMeshBufferSemantic semantic = MBS_VERTEXINDEX;
 					int semanticIndex = 0;
-					MESH_BUFFER_FORMAT format = MBF_UINT32;
+					EMeshBufferFormat format = MBF_UINT32;
 					int components = 1;
 					int offset = 0;
 					pCloned->GetVertexBuffers().SetBuffer
@@ -1475,42 +1475,47 @@ namespace mu
     {
         NodeMeshFormat::Private& node = *format->GetPrivate();
 
-        if ( node.m_pSource )
+        if ( node.Source )
         {
 			FMeshGenerationOptions Options = InOptions;
 
 			FMeshGenerationResult baseResult;
-			GenerateMesh(Options,baseResult, node.m_pSource);
+			GenerateMesh(Options,baseResult, node.Source);
             Ptr<ASTOpMeshFormat> op = new ASTOpMeshFormat();
             op->Source = baseResult.meshOp;
-            op->Buffers = 0;
+            op->Flags = 0;
 
-            MeshPtr pFormatMesh = new Mesh();
+            Ptr<Mesh> FormatMesh = new Mesh();
 
-            if (node.m_VertexBuffers.GetBufferCount())
+            if (node.VertexBuffers.GetBufferCount())
             {
-                op->Buffers |= OP::MeshFormatArgs::BT_VERTEX;
-                pFormatMesh->m_VertexBuffers = node.m_VertexBuffers;
+                op->Flags |= OP::MeshFormatArgs::Vertex;
+                FormatMesh->m_VertexBuffers = node.VertexBuffers;
             }
 
-            if (node.m_IndexBuffers.GetBufferCount())
+            if (node.IndexBuffers.GetBufferCount())
             {
-				op->Buffers |= OP::MeshFormatArgs::BT_INDEX;
-                pFormatMesh->m_IndexBuffers = node.m_IndexBuffers;
+				op->Flags |= OP::MeshFormatArgs::Index;
+                FormatMesh->m_IndexBuffers = node.IndexBuffers;
             }
 
-            if (node.m_FaceBuffers.GetBufferCount())
+            if (node.FaceBuffers.GetBufferCount())
             {
-                op->Buffers |= OP::MeshFormatArgs::BT_FACE;
-                pFormatMesh->m_FaceBuffers = node.m_FaceBuffers;
+                op->Flags |= OP::MeshFormatArgs::Face;
+                FormatMesh->m_FaceBuffers = node.FaceBuffers;
             }
 
+			if (node.bOptimizeBuffers)
+			{
+				op->Flags |= OP::MeshFormatArgs::OptimizeBuffers;
+			}
+
             Ptr<ASTOpConstantResource> cop = new ASTOpConstantResource();
             cop->Type = OP_TYPE::ME_CONSTANT;
-            cop->SetValue( pFormatMesh, m_compilerOptions->OptimisationOptions.DiskCacheContext );
+            cop->SetValue( FormatMesh, m_compilerOptions->OptimisationOptions.DiskCacheContext );
             op->Format = cop;
 
-            m_constantMeshes.Add(pFormatMesh);
+            m_constantMeshes.Add(FormatMesh);
 
             OutResult.meshOp = op;
             OutResult.baseMeshOp = baseResult.baseMeshOp;
diff --git a/Engine/Plugins/Experimental/Mutable/Source/MutableTools/Private/MuT/CodeOptimiser.cpp b/Engine/Plugins/Experimental/Mutable/Source/MutableTools/Private/MuT/CodeOptimiser.cpp
--- a/Engine/Plugins/Experimental/Mutable/Source/MutableTools/Private/MuT/CodeOptimiser.cpp
+++ b/Engine/Plugins/Experimental/Mutable/Source/MutableTools/Private/MuT/CodeOptimiser.cpp
@@ -678,7 +678,7 @@ namespace mu
 						const ASTOpImageRasterMesh* Raster = static_cast<const ASTOpImageRasterMesh*>(SubgraphRoot.get());
 						// If this operation is only rastering the mesh UVs, reduce it to constant. Otherwise avoid reducing it
 						// for the case of a constant projector of a large set of possible images. We don't want to generate all the
-						// projected version of the images beforehand. TODO: Make it a comptile-time option?
+						// projected version of the images beforehand. TODO: Make it a compile-time option?
 						bHasSpecialOpInSubgraph = Raster->image.child().get() != nullptr;
 						break;
 					}
diff --git a/Engine/Plugins/Experimental/Mutable/Source/MutableTools/Private/MuT/DataPacker.cpp b/Engine/Plugins/Experimental/Mutable/Source/MutableTools/Private/MuT/DataPacker.cpp
--- a/Engine/Plugins/Experimental/Mutable/Source/MutableTools/Private/MuT/DataPacker.cpp
+++ b/Engine/Plugins/Experimental/Mutable/Source/MutableTools/Private/MuT/DataPacker.cpp
@@ -677,7 +677,7 @@ namespace mu
             bool used = false;
             for (int c=0; !used && c<pMesh->GetVertexBuffers().GetBufferChannelCount(v); ++c)
             {
-                MESH_BUFFER_SEMANTIC semantic;
+                EMeshBufferSemantic semantic;
                 pMesh->GetVertexBuffers().GetChannel(v,c,&semantic,
                                                       nullptr, nullptr, nullptr, nullptr);
                 used = (( (UINT64_C(1)<<semantic) ) & usedSemantics) != 0;
@@ -685,7 +685,7 @@ namespace mu
 
             if (!used)
             {
-                TArray<MESH_BUFFER>& buffers = pMesh->GetVertexBuffers().m_buffers;
+                TArray<FMeshBuffer>& buffers = pMesh->GetVertexBuffers().m_buffers;
                 buffers.RemoveAt(v);
             }
             else
diff --git a/Engine/Plugins/Experimental/Mutable/Source/MutableTools/Private/MuT/NodeMeshFormat.cpp b/Engine/Plugins/Experimental/Mutable/Source/MutableTools/Private/MuT/NodeMeshFormat.cpp
--- a/Engine/Plugins/Experimental/Mutable/Source/MutableTools/Private/MuT/NodeMeshFormat.cpp
+++ b/Engine/Plugins/Experimental/Mutable/Source/MutableTools/Private/MuT/NodeMeshFormat.cpp
@@ -33,50 +33,50 @@ class FMeshBufferSet;
 	//---------------------------------------------------------------------------------------------
 	NodeMeshPtr NodeMeshFormat::GetSource() const
 	{
-		return m_pD->m_pSource;
+		return m_pD->Source;
 	}
 
 
-	//---------------------------------------------------------------------------------------------
 	void NodeMeshFormat::SetSource( NodeMesh* pValue )
 	{
-		m_pD->m_pSource = pValue;
+		m_pD->Source = pValue;
 	}
 
 
-	//---------------------------------------------------------------------------------------------
 	FMeshBufferSet& NodeMeshFormat::GetVertexBuffers()
 	{
-		return m_pD->m_VertexBuffers;
+		return m_pD->VertexBuffers;
 	}
 
 
-	//---------------------------------------------------------------------------------------------
 	FMeshBufferSet& NodeMeshFormat::GetIndexBuffers()
 	{
-		return m_pD->m_IndexBuffers;
+		return m_pD->IndexBuffers;
 	}
 
 
-	//---------------------------------------------------------------------------------------------
 	FMeshBufferSet& NodeMeshFormat::GetFaceBuffers()
 	{
-		return m_pD->m_FaceBuffers;
+		return m_pD->FaceBuffers;
 	}
 
+	void NodeMeshFormat::SetOptimizeBuffers(bool bEnable)
+	{
+		m_pD->bOptimizeBuffers = bEnable;
+	}
 
-	//---------------------------------------------------------------------------------------------
-	NodeLayoutPtr NodeMeshFormat::Private::GetLayout( int index ) const
+
+	Ptr<NodeLayout> NodeMeshFormat::Private::GetLayout( int32 index ) const
 	{
-		NodeLayoutPtr pResult;
+		Ptr<NodeLayout> Result;
 
-		if ( m_pSource )
+		if ( Source )
 		{
-			NodeMesh::Private* pPrivate = static_cast<NodeMesh::Private*>( m_pSource->GetBasePrivate() );
-			pResult = pPrivate->GetLayout( index );
+			NodeMesh::Private* Private = static_cast<NodeMesh::Private*>( Source->GetBasePrivate() );
+			Result = Private->GetLayout( index );
 		}
 
-		return pResult;
+		return Result;
 	}
 
 }
diff --git a/Engine/Plugins/Experimental/Mutable/Source/MutableTools/Private/MuT/NodeMeshInterpolate.cpp b/Engine/Plugins/Experimental/Mutable/Source/MutableTools/Private/MuT/NodeMeshInterpolate.cpp
--- a/Engine/Plugins/Experimental/Mutable/Source/MutableTools/Private/MuT/NodeMeshInterpolate.cpp
+++ b/Engine/Plugins/Experimental/Mutable/Source/MutableTools/Private/MuT/NodeMeshInterpolate.cpp
@@ -83,7 +83,7 @@ namespace mu
 
 
 	//---------------------------------------------------------------------------------------------
-	void NodeMeshInterpolate::SetChannel( int i, MESH_BUFFER_SEMANTIC semantic, int semanticIndex )
+	void NodeMeshInterpolate::SetChannel( int i, EMeshBufferSemantic semantic, int semanticIndex )
 	{
 		check( i>=0 && i<m_pD->m_channels.Num() );
 
diff --git a/Engine/Plugins/Experimental/Mutable/Source/MutableTools/Private/MuT/NodeMeshInterpolatePrivate.h b/Engine/Plugins/Experimental/Mutable/Source/MutableTools/Private/MuT/NodeMeshInterpolatePrivate.h
--- a/Engine/Plugins/Experimental/Mutable/Source/MutableTools/Private/MuT/NodeMeshInterpolatePrivate.h
+++ b/Engine/Plugins/Experimental/Mutable/Source/MutableTools/Private/MuT/NodeMeshInterpolatePrivate.h
@@ -30,14 +30,14 @@ namespace mu
 				semanticIndex = 0;
 			}
 
-			CHANNEL( MESH_BUFFER_SEMANTIC asemantic,
+			CHANNEL( EMeshBufferSemantic asemantic,
                      int32_t asemanticIndex )
 			{
 				semantic = asemantic;
 				semanticIndex = asemanticIndex;
 			}
 
-			MESH_BUFFER_SEMANTIC semantic;
+			EMeshBufferSemantic semantic;
             int32_t semanticIndex;
 
 			//-------------------------------------------------------------------------------------
diff --git a/Engine/Plugins/Experimental/Mutable/Source/MutableTools/Public/MuT/NodeMeshFormat.h b/Engine/Plugins/Experimental/Mutable/Source/MutableTools/Public/MuT/NodeMeshFormat.h
--- a/Engine/Plugins/Experimental/Mutable/Source/MutableTools/Public/MuT/NodeMeshFormat.h
+++ b/Engine/Plugins/Experimental/Mutable/Source/MutableTools/Public/MuT/NodeMeshFormat.h
@@ -11,16 +11,12 @@ namespace mu
 {
 
 	// Forward definitions
-	class NodeMeshFormat;
-	typedef Ptr<NodeMeshFormat> NodeMeshFormatPtr;
-	typedef Ptr<const NodeMeshFormat> NodeMeshFormatPtrConst;
-
 	class InputArchive;
 	class FMeshBufferSet;
 	class OutputArchive;
 
-	//! This node can change the buffer formats of a mesh vertices, indices and faces.
-	//! \ingroup model
+	/** This node can change the buffer formats of a mesh vertices, indices and faces.
+	*/
 	class MUTABLETOOLS_API NodeMeshFormat : public NodeMesh
 	{
 	public:
@@ -33,7 +29,7 @@ namespace mu
 
 		void SerialiseWrapper(OutputArchive& arch) const override;
 		static void Serialise( const NodeMeshFormat* pNode, OutputArchive& arch );
-		static NodeMeshFormatPtr StaticUnserialise( InputArchive& arch );
+		static Ptr<NodeMeshFormat> StaticUnserialise( InputArchive& arch );
 
 		//-----------------------------------------------------------------------------------------
 		// Node Interface
@@ -46,25 +42,34 @@ namespace mu
 		// Own Interface
 		//-----------------------------------------------------------------------------------------
 
-		//! Source mesh to be re-formatted
-		NodeMeshPtr GetSource() const;
+		/** Mesh to be reformatted. */
+		Ptr<NodeMesh> GetSource() const;
 		void SetSource( NodeMesh* );
 
-		//! Get and set the MeshBufferSet that defines the new format for the mesh vertices. These
-		//! buffers don't really contain any data (they have 0 elements) but they define the
-		//! structure. If this is null, the vertex buffers will not be changed.
+		/** Access the MeshBufferSet that defines the new format for the mesh vertices.These
+		* buffers don't really contain any data (they have 0 elements) but they define the
+		* structure. If this is null, the vertex buffers will not be changed.
+		*/
 		FMeshBufferSet& GetVertexBuffers();
 
-		//! Get and set the MeshBufferSet that defines the new format for the mesh indices. These
-		//! buffers don't really contain any data (they have 0 elements) but they define the
-		//! structure. If this is null, the vertex buffers will not be changed.
+		/** Access the MeshBufferSet that defines the new format for the mesh indices. These
+		* buffers don't really contain any data (they have 0 elements) but they define the
+		* structure. If this is null, the vertex buffers will not be changed.
+		*/
 		FMeshBufferSet& GetIndexBuffers();
 
-		//! Get and set the MeshBufferSet that defines the new format for the mesh faces. These
-		//! buffers don't really contain any data (they have 0 elements) but they define the
-		//! structure. If this is null, the vertex buffers will not be changed.
+		/** Access the MeshBufferSet that defines the new format for the mesh faces. These
+		* buffers don't really contain any data (they have 0 elements) but they define the
+		* structure. If this is null, the vertex buffers will not be changed.
+		*/
 		FMeshBufferSet& GetFaceBuffers();
 
+		/** Optimize the buffers is possible. This may change the target format to reduce the number
+		* of channels or even the type if possible to minimize size.
+		* By default it is disabled.
+		*/
+		void SetOptimizeBuffers(bool bEnable);
+
 		//-----------------------------------------------------------------------------------------
 		// Interface pattern
 		//-----------------------------------------------------------------------------------------
diff --git a/Engine/Plugins/Experimental/Mutable/Source/MutableTools/Public/MuT/NodeMeshInterpolate.h b/Engine/Plugins/Experimental/Mutable/Source/MutableTools/Public/MuT/NodeMeshInterpolate.h
--- a/Engine/Plugins/Experimental/Mutable/Source/MutableTools/Public/MuT/NodeMeshInterpolate.h
+++ b/Engine/Plugins/Experimental/Mutable/Source/MutableTools/Public/MuT/NodeMeshInterpolate.h
@@ -74,7 +74,7 @@ namespace mu
 		void SetChannelCount( int );
 
 		//! Set a channel to interpolate from the mesh.
-		void SetChannel( int index, MESH_BUFFER_SEMANTIC semantic, int semanticIndex );
+		void SetChannel( int index, EMeshBufferSemantic semantic, int semanticIndex );
 
 		//-----------------------------------------------------------------------------------------
 		// Interface pattern
