[Mutable] Fix incorrect call to GetImages.

The base function must be called.

#rnx

diff --git a/Engine/Plugins/Experimental/Mutable/Source/CustomizableObject/Private/MuCO/CustomizableObjectSystem.cpp b/Engine/Plugins/Experimental/Mutable/Source/CustomizableObject/Private/MuCO/CustomizableObjectSystem.cpp
--- a/Engine/Plugins/Experimental/Mutable/Source/CustomizableObject/Private/MuCO/CustomizableObjectSystem.cpp
+++ b/Engine/Plugins/Experimental/Mutable/Source/CustomizableObject/Private/MuCO/CustomizableObjectSystem.cpp
@@ -2832,7 +2832,7 @@ namespace impl
 				TEXT("Task_Mutable_GetImages"),
 				[OperationData]()
 				{
-					CustomizableObjectSystem::ImplDeprecated::Task_Mutable_GetImages(OperationData);
+					Task_Mutable_GetImages(OperationData);
 				});
 		}
 
