[Mutable] Fix hang when saving assets.

Now when we are in BlockTillAllRequestsFinished, LoadAdditionalAssetsAndData will be synchronous.

#rnx

diff --git a/Engine/Plugins/Experimental/Mutable/Source/CustomizableObject/Internal/MuCO/CustomizableObjectInstancePrivate.h b/Engine/Plugins/Experimental/Mutable/Source/CustomizableObject/Internal/MuCO/CustomizableObjectInstancePrivate.h
--- a/Engine/Plugins/Experimental/Mutable/Source/CustomizableObject/Internal/MuCO/CustomizableObjectInstancePrivate.h
+++ b/Engine/Plugins/Experimental/Mutable/Source/CustomizableObject/Internal/MuCO/CustomizableObjectInstancePrivate.h
@@ -219,8 +219,9 @@ public:
 
 	void ReuseTexture(UTexture2D* Texture, TSharedRef<FTexturePlatformData, ESPMode::ThreadSafe>& PlatformData);
 
-	/** Returns the task that will be called when all assets and data are loaded, may be already completed if no assets or data needs loading */
-	UE::Tasks::FTask LoadAdditionalAssetsAndDataAsync(const TSharedRef<FUpdateContextPrivate>& OperationData, FStreamableManager &StreamableManager);
+	/** Returns the task that will be called when all assets and data are loaded, may be already completed if no assets or data needs loading.
+	 * If no StreamableManager is provided, it will load assets synchronously. */
+	UE::Tasks::FTask LoadAdditionalAssetsAndData(const TSharedRef<FUpdateContextPrivate>& OperationData, FStreamableManager* StreamableManager);
 
 	void AdditionalAssetsAsyncLoaded(UCustomizableObjectInstance* Public);
 
diff --git a/Engine/Plugins/Experimental/Mutable/Source/CustomizableObject/Internal/MuCO/CustomizableObjectSystemPrivate.h b/Engine/Plugins/Experimental/Mutable/Source/CustomizableObject/Internal/MuCO/CustomizableObjectSystemPrivate.h
--- a/Engine/Plugins/Experimental/Mutable/Source/CustomizableObject/Internal/MuCO/CustomizableObjectSystemPrivate.h
+++ b/Engine/Plugins/Experimental/Mutable/Source/CustomizableObject/Internal/MuCO/CustomizableObjectSystemPrivate.h
@@ -744,6 +744,9 @@ public:
 
 	bool bAutoCompileCommandletEnabled = false;
 
+	/** We are in BlockTillAllRequestsFinished scope. */
+	bool bBlocking = false;
+
 	UPROPERTY()
 	TArray<FPendingReleaseSkeletalMeshInfo> PendingReleaseSkeletalMesh;
 	
diff --git a/Engine/Plugins/Experimental/Mutable/Source/CustomizableObject/Private/MuCO/CustomizableObjectInstance.cpp b/Engine/Plugins/Experimental/Mutable/Source/CustomizableObject/Private/MuCO/CustomizableObjectInstance.cpp
--- a/Engine/Plugins/Experimental/Mutable/Source/CustomizableObject/Private/MuCO/CustomizableObjectInstance.cpp
+++ b/Engine/Plugins/Experimental/Mutable/Source/CustomizableObject/Private/MuCO/CustomizableObjectInstance.cpp
@@ -4836,8 +4836,8 @@ FAutoConsoleVariableRef CVarMutableHighPriorityLoading(
 	bEnableHighPriorityLoading,
 	TEXT("If enabled, the request to load additional assets will have high priority."));
 
-UE::Tasks::FTask UCustomizableInstancePrivate::LoadAdditionalAssetsAndDataAsync(
-		const TSharedRef<FUpdateContextPrivate>& OperationData, FStreamableManager& StreamableManager)
+UE::Tasks::FTask UCustomizableInstancePrivate::LoadAdditionalAssetsAndData(
+		const TSharedRef<FUpdateContextPrivate>& OperationData, FStreamableManager* StreamableManager)
 {
 	MUTABLE_CPUPROFILER_SCOPE(UCustomizableInstancePrivate::LoadAdditionalAssetsAndDataAsync);
 
@@ -5125,10 +5125,22 @@ UE::Tasks::FTask UCustomizableInstancePrivate::LoadAdditionalAssetsAndDataAsync(
 	{	
 		UE::Tasks::FTaskEvent AssetAsyncLoadCompletionEvent = StreamingCompletionEvents.Emplace_GetRef(TEXT("AssetAsyncLoadCompletionEvent"));
 
-		StreamingHandle = StreamableManager.RequestAsyncLoad(
-				AssetsToStream, 
-				FStreamableDelegate::CreateUObject(this, &UCustomizableInstancePrivate::AdditionalAssetsAsyncLoaded, AssetAsyncLoadCompletionEvent),
-				bEnableHighPriorityLoading ? FStreamableManager::AsyncLoadHighPriority : FStreamableManager::DefaultAsyncLoadPriority);
+		if (StreamableManager)
+		{
+			StreamingHandle = StreamableManager->RequestAsyncLoad(
+					AssetsToStream, 
+					FStreamableDelegate::CreateUObject(this, &UCustomizableInstancePrivate::AdditionalAssetsAsyncLoaded, AssetAsyncLoadCompletionEvent),
+					bEnableHighPriorityLoading ? FStreamableManager::AsyncLoadHighPriority : FStreamableManager::DefaultAsyncLoadPriority);			
+		}
+		else
+		{
+			for (const FSoftObjectPath& Asset : AssetsToStream)
+			{
+				Asset.TryLoad();
+			}
+			
+			AdditionalAssetsAsyncLoaded(AssetAsyncLoadCompletionEvent);
+		}
 	}
 
 	
diff --git a/Engine/Plugins/Experimental/Mutable/Source/CustomizableObject/Private/MuCO/CustomizableObjectSystem.cpp b/Engine/Plugins/Experimental/Mutable/Source/CustomizableObject/Private/MuCO/CustomizableObjectSystem.cpp
--- a/Engine/Plugins/Experimental/Mutable/Source/CustomizableObject/Private/MuCO/CustomizableObjectSystem.cpp
+++ b/Engine/Plugins/Experimental/Mutable/Source/CustomizableObject/Private/MuCO/CustomizableObjectSystem.cpp
@@ -2839,7 +2839,8 @@ namespace impl
 
 		// Next Task: Load Unreal Assets
 		//-------------------------------------------------------------
-		UE::Tasks::FTask Game_LoadUnrealAssets = ObjectInstancePrivateData->LoadAdditionalAssetsAndDataAsync(OperationData, UCustomizableObjectSystem::GetInstance()->GetPrivate()->StreamableManager);
+		FStreamableManager* StreamManager = System->GetPrivate()->bBlocking ? nullptr : &System->GetPrivate()->StreamableManager;
+		UE::Tasks::FTask Game_LoadUnrealAssets = ObjectInstancePrivateData->LoadAdditionalAssetsAndData(OperationData, StreamManager);
 
 		// Next-next Task: Convert Resources
 		//-------------------------------------------------------------
@@ -4173,6 +4174,9 @@ void UCustomizableObjectSystemPrivate::UpdateResourceStreaming(float DeltaTime,
 
 int32 UCustomizableObjectSystemPrivate::BlockTillAllRequestsFinished(float TimeLimit, bool bLogResults)
 {
+	bBlocking = true;
+	ON_SCOPE_EXIT {bBlocking = false; };
+	
 	const double BlockEndTime = FPlatformTime::Seconds() + TimeLimit;
 
 	int32 RemainingWork = TNumericLimits<int32>::Max();
