[Mutable] Make source graph private.

#rnx

diff --git a/Engine/Plugins/Experimental/Mutable/Source/CustomizableObject/Internal/MuCO/CustomizableObjectPrivate.h b/Engine/Plugins/Experimental/Mutable/Source/CustomizableObject/Internal/MuCO/CustomizableObjectPrivate.h
--- a/Engine/Plugins/Experimental/Mutable/Source/CustomizableObject/Internal/MuCO/CustomizableObjectPrivate.h
+++ b/Engine/Plugins/Experimental/Mutable/Source/CustomizableObject/Internal/MuCO/CustomizableObjectPrivate.h
@@ -762,6 +762,10 @@ public:
 	TArray<FCustomizableObjectStreamedResourceData>& GetStreamedExtensionData();
 	
 	TArray<FCustomizableObjectStreamedResourceData>& GetStreamedResourceData();
+
+#if WITH_EDITORONLY_DATA
+	TObjectPtr<UEdGraph>& GetSource() const;
+#endif
 	
 	/** Cache of generated SkeletalMeshes */
 	FMeshCache MeshCache;
diff --git a/Engine/Plugins/Experimental/Mutable/Source/CustomizableObject/Private/MuCO/CustomizableObject.cpp b/Engine/Plugins/Experimental/Mutable/Source/CustomizableObject/Private/MuCO/CustomizableObject.cpp
--- a/Engine/Plugins/Experimental/Mutable/Source/CustomizableObject/Private/MuCO/CustomizableObject.cpp
+++ b/Engine/Plugins/Experimental/Mutable/Source/CustomizableObject/Private/MuCO/CustomizableObject.cpp
@@ -2002,6 +2002,14 @@ TArray<FCustomizableObjectStreamedResourceData>& UCustomizableObjectPrivate::Get
 }
 
 
+#if WITH_EDITORONLY_DATA
+TObjectPtr<UEdGraph>& UCustomizableObjectPrivate::GetSource() const
+{
+	return GetPublic()->Source;
+}
+#endif
+
+
 //-------------------------------------------------------------------------------------------------
 //-------------------------------------------------------------------------------------------------
 //-------------------------------------------------------------------------------------------------
diff --git a/Engine/Plugins/Experimental/Mutable/Source/CustomizableObjectEditor/Classes/MuCOE/Nodes/CustomizableObjectNode.h b/Engine/Plugins/Experimental/Mutable/Source/CustomizableObjectEditor/Classes/MuCOE/Nodes/CustomizableObjectNode.h
--- a/Engine/Plugins/Experimental/Mutable/Source/CustomizableObjectEditor/Classes/MuCOE/Nodes/CustomizableObjectNode.h
+++ b/Engine/Plugins/Experimental/Mutable/Source/CustomizableObjectEditor/Classes/MuCOE/Nodes/CustomizableObjectNode.h
@@ -5,6 +5,7 @@
 #include "EdGraph/EdGraph.h" // IWYU pragma: keep
 #include "EdGraph/EdGraphPin.h"
 #include "MuCO/CustomizableObject.h"
+#include "MuCO/CustomizableObjectPrivate.h"
 
 #include "CustomizableObjectNode.generated.h"
 
@@ -328,10 +329,10 @@ FNodeType* GetCustomizableObjectExternalNode(UCustomizableObject* Object, const
 {
 	FNodeType* Result = nullptr;
 
-	if (Object && Object->Source && NodeGuid.IsValid())
+	if (Object && Object->GetPrivate()->GetSource() && NodeGuid.IsValid())
 	{
 		TArray<FNodeType*> GroupNodes;
-		Object->Source->GetNodesOfClass<FNodeType>(GroupNodes);
+		Object->GetPrivate()->GetSource()->GetNodesOfClass<FNodeType>(GroupNodes);
 
 		for (FNodeType* GroupNode : GroupNodes)
 		{
diff --git a/Engine/Plugins/Experimental/Mutable/Source/CustomizableObjectEditor/Private/MuCOE/CustomizableObjectCompiler.cpp b/Engine/Plugins/Experimental/Mutable/Source/CustomizableObjectEditor/Private/MuCOE/CustomizableObjectCompiler.cpp
--- a/Engine/Plugins/Experimental/Mutable/Source/CustomizableObjectEditor/Private/MuCOE/CustomizableObjectCompiler.cpp
+++ b/Engine/Plugins/Experimental/Mutable/Source/CustomizableObjectEditor/Private/MuCOE/CustomizableObjectCompiler.cpp
@@ -123,13 +123,13 @@ bool FCustomizableObjectCompiler::IsRootObject(const UCustomizableObject* Object
 	// Look for the base object node
 	UCustomizableObjectNodeObject* Root = nullptr;
 	TArray<UCustomizableObjectNodeObject*> ObjectNodes;
-	if (!Object->Source || !Object->Source->Nodes.Num())
+	if (!Object->GetPrivate()->GetSource() || !Object->GetPrivate()->GetSource()->Nodes.Num())
 	{
 		// Conservative approach.
 		return true;
 	}
 
-	Object->Source->GetNodesOfClass<UCustomizableObjectNodeObject>(ObjectNodes);
+	Object->GetPrivate()->GetSource()->GetNodesOfClass<UCustomizableObjectNodeObject>(ObjectNodes);
 
 	for (TArray<UCustomizableObjectNodeObject*>::TIterator It(ObjectNodes); It; ++It)
 	{
@@ -384,7 +384,7 @@ void FCustomizableObjectCompiler::ProcessChildObjectsRecursively(UCustomizableOb
 			GenerationContext.GroupIdToExternalNodeMap.Add(Root->ParentObjectGroupId, Root);
 
 			TArray<UCustomizableObjectNodeObjectGroup*> GroupNodes;
-			ChildObject->Source->GetNodesOfClass<UCustomizableObjectNodeObjectGroup>(GroupNodes);
+			ChildObject->GetPrivate()->GetSource()->GetNodesOfClass<UCustomizableObjectNodeObjectGroup>(GroupNodes);
 
 			if (GroupNodes.Num() > 0) // Only grafs with group nodes should have child grafs
 			{
@@ -497,7 +497,7 @@ mu::NodeObjectPtr FCustomizableObjectCompiler::GenerateMutableRoot(
 {
 	check(Object);
 	
-	if (!Object->Source)
+	if (!Object->GetPrivate()->GetSource())
 	{
 		ErrorMsg = LOCTEXT("NoSource", "Object with no valid graph found. Object not build.");
 
@@ -600,7 +600,7 @@ mu::NodeObjectPtr FCustomizableObjectCompiler::GenerateMutableRoot(
 			(Object->MeshCompileType == EMutableCompileMeshType::AddWorkingSetAndChildren))
 		{
 			TArray<UCustomizableObjectNodeObjectGroup*> GroupNodes;
-			Object->Source->GetNodesOfClass<UCustomizableObjectNodeObjectGroup>(GroupNodes);
+			Object->GetPrivate()->GetSource()->GetNodesOfClass<UCustomizableObjectNodeObjectGroup>(GroupNodes);
 
 			if (GroupNodes.Num() > 0) // Only graphs with group nodes should have child graphs
 			{
diff --git a/Engine/Plugins/Experimental/Mutable/Source/CustomizableObjectEditor/Private/MuCOE/CustomizableObjectEditor.cpp b/Engine/Plugins/Experimental/Mutable/Source/CustomizableObjectEditor/Private/MuCOE/CustomizableObjectEditor.cpp
--- a/Engine/Plugins/Experimental/Mutable/Source/CustomizableObjectEditor/Private/MuCOE/CustomizableObjectEditor.cpp
+++ b/Engine/Plugins/Experimental/Mutable/Source/CustomizableObjectEditor/Private/MuCOE/CustomizableObjectEditor.cpp
@@ -425,7 +425,7 @@ void FCustomizableObjectEditor::CreatePreviewInstance()
 
 	if (GraphEditor.IsValid())
 	{
-		for (TObjectPtr<UEdGraphNode>& Node : CustomizableObject->Source->Nodes)
+		for (TObjectPtr<UEdGraphNode>& Node : CustomizableObject->GetPrivate()->GetSource()->Nodes)
 		{
 			GraphEditor->RefreshNode(*Node);
 		}
@@ -620,7 +620,7 @@ TSharedRef<SDockTab> FCustomizableObjectEditor::SpawnTab_Graph( const FSpawnTabA
 {
 	check( Args.GetTabId().TabType == GraphTabId );
 
-	CreateGraphEditorWidget(CustomizableObject->Source);
+	CreateGraphEditorWidget(CustomizableObject->GetPrivate()->GetSource());
 
 	TSharedRef<SDockTab> DockTab = SNew(SDockTab)
 		.Label( FText::FromString( GetTabPrefix() + LOCTEXT( "SourceGraph", "Source Graph" ).ToString() ) )
@@ -785,7 +785,7 @@ void FCustomizableObjectEditor::RedoGraphAction()
 bool FCustomizableObjectEditor::GroupNodeIsLinkedToParentByName(UCustomizableObjectNodeObject* Node, UCustomizableObject* Test, const FString& ParentGroupName)
 {
 	TArray<UCustomizableObjectNodeObjectGroup*> GroupNodes;
-	Test->Source->GetNodesOfClass<UCustomizableObjectNodeObjectGroup>(GroupNodes);
+	Test->GetPrivate()->GetSource()->GetNodesOfClass<UCustomizableObjectNodeObjectGroup>(GroupNodes);
 
 	for (UCustomizableObjectNodeObjectGroup* GroupNode : GroupNodes)
 	{
@@ -1539,7 +1539,7 @@ void FCustomizableObjectEditor::CompileObject()
 		return;
 	}
 
-	if (CustomizableObject->Source)
+	if (CustomizableObject->GetPrivate()->GetSource())
 	{
 		FCompilationOptions Options = CustomizableObject->CompileOptions;
 		Options.bSilentCompilation = false;
@@ -1757,7 +1757,7 @@ void FCustomizableObjectEditor::NotifyPostChange( const FPropertyChangedEvent& P
 	if (OuterClass && OuterClass->IsChildOf(UCustomizableObjectNode::StaticClass()))
 	{
 		FPropertyChangedEvent Event(PropertyThatChanged);
-		CustomizableObject->Source->PostEditChangeProperty(Event);
+		CustomizableObject->GetPrivate()->GetSource()->PostEditChangeProperty(Event);
 		CustomizableObject->PostEditChangeProperty(Event);
 
 		if (GraphEditor.IsValid())
@@ -1860,7 +1860,7 @@ void FCustomizableObjectEditor::PasteNodesHere(const FVector2D& Location)
 {
 	// Undo/Redo support
 	const FScopedTransaction Transaction( LOCTEXT("CustomizableObjectEditorPaste", "Customizable Object Editor Editor: Paste") );
-	CustomizableObject->Source->Modify();
+	CustomizableObject->GetPrivate()->GetSource()->Modify();
 	CustomizableObject->Modify();
 
 	// Clear the selection set (newly pasted stuff will be selected)
@@ -1872,7 +1872,7 @@ void FCustomizableObjectEditor::PasteNodesHere(const FVector2D& Location)
 
 	// Import the nodes
 	TSet<UEdGraphNode*> PastedNodes;
-	FEdGraphUtilities::ImportNodesFromText(CustomizableObject->Source, TextToImport, /*out*/ PastedNodes);
+	FEdGraphUtilities::ImportNodesFromText(CustomizableObject->GetPrivate()->GetSource(), TextToImport, /*out*/ PastedNodes);
 	
 	//Average position of nodes so we can move them while still maintaining relative distances to each other
 	FVector2D AvgNodePosition(0.0f,0.0f);
@@ -1971,7 +1971,7 @@ bool FCustomizableObjectEditor::CanPasteNodes() const
 
 	FPlatformApplicationMisc::ClipboardPaste(ClipboardContent);
 
-	return FEdGraphUtilities::CanImportNodesFromText(CustomizableObject->Source, ClipboardContent);
+	return FEdGraphUtilities::CanImportNodesFromText(CustomizableObject->GetPrivate()->GetSource(), ClipboardContent);
 }
 
 
@@ -2361,7 +2361,7 @@ UCustomizableObject* FCustomizableObjectEditor::GetAbsoluteCOParent(const UCusto
 	{
 		//Get all the NodeObjects
 		TArray<UCustomizableObjectNodeObject*> ObjectNodes;
-		Root->ParentObject->Source->GetNodesOfClass<UCustomizableObjectNodeObject>(ObjectNodes);
+		Root->ParentObject->GetPrivate()->GetSource()->GetNodesOfClass<UCustomizableObjectNodeObject>(ObjectNodes);
 		if (!ObjectNodes.IsEmpty())
 		{
 			//Getting the parent of the root
@@ -2459,7 +2459,7 @@ void FCustomizableObjectEditor::GetExternalChildObjects(const UCustomizableObjec
 			}
 
 			TArray<UCustomizableObjectNodeObjectGroup*> GroupNodes;
-			ChildObject->Source->GetNodesOfClass<UCustomizableObjectNodeObjectGroup>(GroupNodes);
+			ChildObject->GetPrivate()->GetSource()->GetNodesOfClass<UCustomizableObjectNodeObjectGroup>(GroupNodes);
 
 			if (GroupNodes.Num() > 0) // Only grafs with group nodes should have child grafs
 			{
@@ -2482,10 +2482,10 @@ void FCustomizableObjectEditor::CreatePreviewComponents()
 	// Getting the number of mesh components from the root node
 	int32 NumMeshComponents = 0;
 
-	if (CustomizableObject->Source)
+	if (CustomizableObject->GetPrivate()->GetSource())
 	{
 		TArray<UCustomizableObjectNodeObject*> RootNode;
-		CustomizableObject->Source->GetNodesOfClass< UCustomizableObjectNodeObject>(RootNode);
+		CustomizableObject->GetPrivate()->GetSource()->GetNodesOfClass< UCustomizableObjectNodeObject>(RootNode);
 
 		for (int32 i = 0; i < RootNode.Num(); ++i)
 		{
diff --git a/Engine/Plugins/Experimental/Mutable/Source/CustomizableObjectEditor/Private/MuCOE/CustomizableObjectEditorModule.cpp b/Engine/Plugins/Experimental/Mutable/Source/CustomizableObjectEditor/Private/MuCOE/CustomizableObjectEditorModule.cpp
--- a/Engine/Plugins/Experimental/Mutable/Source/CustomizableObjectEditor/Private/MuCOE/CustomizableObjectEditorModule.cpp
+++ b/Engine/Plugins/Experimental/Mutable/Source/CustomizableObjectEditor/Private/MuCOE/CustomizableObjectEditorModule.cpp
@@ -469,7 +469,7 @@ void GetReferencingPackages(const UCustomizableObject& Object, TArray<FName>& Ob
 		}
 
 		TArray<UCustomizableObjectNodeTable*> TableNodes;
-		ChildCustomizableObject->Source->GetNodesOfClass(TableNodes);
+		ChildCustomizableObject->GetPrivate()->GetSource()->GetNodesOfClass(TableNodes);
 
 		FARFilter Filter;
 		Filter.ClassPaths.Add(FTopLevelAssetPath(UDataTable::StaticClass()));
diff --git a/Engine/Plugins/Experimental/Mutable/Source/CustomizableObjectEditor/Private/MuCOE/CustomizableObjectFactory.cpp b/Engine/Plugins/Experimental/Mutable/Source/CustomizableObjectEditor/Private/MuCOE/CustomizableObjectFactory.cpp
--- a/Engine/Plugins/Experimental/Mutable/Source/CustomizableObjectEditor/Private/MuCOE/CustomizableObjectFactory.cpp
+++ b/Engine/Plugins/Experimental/Mutable/Source/CustomizableObjectEditor/Private/MuCOE/CustomizableObjectFactory.cpp
@@ -66,69 +66,72 @@ bool UCustomizableObjectFactory::ConfigureProperties()
 UObject* UCustomizableObjectFactory::FactoryCreateNew(UClass* Class, UObject* InParent, FName Name, EObjectFlags Flags, UObject* Context, FFeedbackContext* Warn)
 {
 	UCustomizableObject* NewObj = NewObject<UCustomizableObject>(InParent, Class, Name, Flags);
-
-	if (NewObj)
+	if (!NewObj)
 	{
-		NewObj->Source = NewObject<UCustomizableObjectGraph>(NewObj, NAME_None, RF_Transactional);
+		return NewObj;
 	}
-
-	if (NewObj->Source)
+	
+	UCustomizableObjectGraph* Source = NewObject<UCustomizableObjectGraph>(NewObj, NAME_None, RF_Transactional);
+	if (!Source)
 	{
-		UCustomizableObjectGraph* CustomizableObjectGraph = Cast<UCustomizableObjectGraph>(NewObj->Source);
-		CustomizableObjectGraph->AddEssentialGraphNodes();
+		return NewObj;
+	}
 
-		UCustomizableObjectNodeObject* BaseObjectNode = nullptr;
+	NewObj->GetPrivate()->GetSource() = Source;
+	
+	Source->AddEssentialGraphNodes();
 
-		for (const TObjectPtr<UEdGraphNode>& AuxNode : NewObj->Source->Nodes)
-		{
-			UCustomizableObjectNodeObject* BaseNode = Cast<UCustomizableObjectNodeObject>(AuxNode);
+	UCustomizableObjectNodeObject* BaseObjectNode = nullptr;
 
-			if (BaseNode && BaseNode->bIsBase)
-			{
-				BaseObjectNode = BaseNode;
-				break;
-			}
+	for (const TObjectPtr<UEdGraphNode>& AuxNode : Source->Nodes)
+	{
+		UCustomizableObjectNodeObject* BaseNode = Cast<UCustomizableObjectNodeObject>(AuxNode);
+
+		if (BaseNode && BaseNode->bIsBase)
+		{
+			BaseObjectNode = BaseNode;
+			break;
 		}
+	}
 
-		if (!CreationSettings.bEmptyObject)
+	if (!CreationSettings.bEmptyObject)
+	{
+		if (!BaseObjectNode)
 		{
-			if (!BaseObjectNode)
-			{
-				return NewObj;
-			}
+			return NewObj;
+		}
 
-			if (CreationSettings.bIsChildObject && CreationSettings.ParentObject.IsValid())
-			{
-				BaseObjectNode->ParentObject = CreationSettings.ParentObject.Get();
+		if (CreationSettings.bIsChildObject && CreationSettings.ParentObject.IsValid())
+		{
+			BaseObjectNode->ParentObject = CreationSettings.ParentObject.Get();
 
-				TArray<UCustomizableObjectNodeObjectGroup*> GroupNodes;
-				CreationSettings.ParentObject->Source->GetNodesOfClass<UCustomizableObjectNodeObjectGroup>(GroupNodes);
+			TArray<UCustomizableObjectNodeObjectGroup*> GroupNodes;
+			CreationSettings.ParentObject->GetPrivate()->GetSource()->GetNodesOfClass<UCustomizableObjectNodeObjectGroup>(GroupNodes);
 
-				for (UCustomizableObjectNodeObjectGroup* GroupNode : GroupNodes)
+			for (UCustomizableObjectNodeObjectGroup* GroupNode : GroupNodes)
+			{
+				if (GroupNode->GroupName == CreationSettings.GroupNodeName)
 				{
-					if (GroupNode->GroupName == CreationSettings.GroupNodeName)
-					{
-						BaseObjectNode->ParentObjectGroupId = GroupNode->NodeGuid;
-
-						break;
-					}
-				}
+					BaseObjectNode->ParentObjectGroupId = GroupNode->NodeGuid;
 
-				if (NewObj->GetPrivate())
-				{
-					NewObj->GetPrivate()->SetIsChildObject(true);
+					break;
 				}
 			}
-			else
+
+			if (NewObj->GetPrivate())
 			{
-				BaseObjectNode->NumMeshComponents = CreationSettings.NumMeshComponents;
+				NewObj->GetPrivate()->SetIsChildObject(true);
+			}
+		}
+		else
+		{
+			BaseObjectNode->NumMeshComponents = CreationSettings.NumMeshComponents;
 
-				for (int32 MeshIndex = 0; MeshIndex < CreationSettings.ReferenceSkeletalMeshes.Num(); ++MeshIndex)
+			for (int32 MeshIndex = 0; MeshIndex < CreationSettings.ReferenceSkeletalMeshes.Num(); ++MeshIndex)
+			{
+				if (CreationSettings.ReferenceSkeletalMeshes[MeshIndex].IsValid())
 				{
-					if (CreationSettings.ReferenceSkeletalMeshes[MeshIndex].IsValid())
-					{
-						NewObj->ReferenceSkeletalMeshes.Add(CreationSettings.ReferenceSkeletalMeshes[MeshIndex].Get());
-					}
+					NewObj->ReferenceSkeletalMeshes.Add(CreationSettings.ReferenceSkeletalMeshes[MeshIndex].Get());
 				}
 			}
 		}
@@ -532,7 +535,7 @@ void FCustomizableObjectFactoryUI::GenerateGroupOptions()
 	if (Options.ParentObject.IsValid())
 	{
 		TArray<UCustomizableObjectNodeObjectGroup*> GroupNodes;
-		Options.ParentObject->Source->GetNodesOfClass<UCustomizableObjectNodeObjectGroup>(GroupNodes);
+		Options.ParentObject->GetPrivate()->GetSource()->GetNodesOfClass<UCustomizableObjectNodeObjectGroup>(GroupNodes);
 
 		for (UCustomizableObjectNodeObjectGroup* GroupNode : GroupNodes)
 		{
diff --git a/Engine/Plugins/Experimental/Mutable/Source/CustomizableObjectEditor/Private/MuCOE/CustomizableObjectInstanceEditor.cpp b/Engine/Plugins/Experimental/Mutable/Source/CustomizableObjectEditor/Private/MuCOE/CustomizableObjectInstanceEditor.cpp
--- a/Engine/Plugins/Experimental/Mutable/Source/CustomizableObjectEditor/Private/MuCOE/CustomizableObjectInstanceEditor.cpp
+++ b/Engine/Plugins/Experimental/Mutable/Source/CustomizableObjectEditor/Private/MuCOE/CustomizableObjectInstanceEditor.cpp
@@ -706,7 +706,7 @@ void FCustomizableObjectInstanceEditor::Tick(float InDeltaTime)
 	if (Compiler && Compiler->Tick())
 	{
 		if (const UCustomizableObject* CustomizableObject = CustomizableObjectInstance->GetCustomizableObject();
-			CustomizableObject && CustomizableObject->Source)
+			CustomizableObject && CustomizableObject->GetPrivate()->GetSource())
 		{
 			if (PreviewCustomizableSkeletalComponents.Num() > 0)
 			{
diff --git a/Engine/Plugins/Experimental/Mutable/Source/CustomizableObjectEditor/Private/MuCOE/CustomizableObjectInstanceFactory.cpp b/Engine/Plugins/Experimental/Mutable/Source/CustomizableObjectEditor/Private/MuCOE/CustomizableObjectInstanceFactory.cpp
--- a/Engine/Plugins/Experimental/Mutable/Source/CustomizableObjectEditor/Private/MuCOE/CustomizableObjectInstanceFactory.cpp
+++ b/Engine/Plugins/Experimental/Mutable/Source/CustomizableObjectEditor/Private/MuCOE/CustomizableObjectInstanceFactory.cpp
@@ -8,6 +8,7 @@
 #include "Engine/SkeletalMesh.h"
 #include "Modules/ModuleManager.h"
 #include "MuCO/CustomizableObjectInstance.h"
+#include "MuCO/CustomizableObjectPrivate.h"
 #include "MuCO/CustomizableSkeletalComponent.h"
 #include "MuCO/CustomizableSkeletalMeshActor.h"
 #include "MuCOE/Nodes/CustomizableObjectNodeObject.h"
@@ -169,10 +170,10 @@ int32 UCustomizableObjectInstanceFactory::GetNumberOfComponents(UCustomizableObj
 
 	if (const UCustomizableObject* CustomizableObject = COInstance->GetCustomizableObject();
 		CustomizableObject &&
-		CustomizableObject->Source)
+		CustomizableObject->GetPrivate()->GetSource())
 	{
 		TArray<UCustomizableObjectNodeObject*> RootNodes;
-		CustomizableObject->Source->GetNodesOfClass<UCustomizableObjectNodeObject>(RootNodes);
+		CustomizableObject->GetPrivate()->GetSource()->GetNodesOfClass<UCustomizableObjectNodeObject>(RootNodes);
 
 		for (int32 i = 0; i < RootNodes.Num(); ++i)
 		{
diff --git a/Engine/Plugins/Experimental/Mutable/Source/CustomizableObjectEditor/Private/MuCOE/GraphTraversal.cpp b/Engine/Plugins/Experimental/Mutable/Source/CustomizableObjectEditor/Private/MuCOE/GraphTraversal.cpp
--- a/Engine/Plugins/Experimental/Mutable/Source/CustomizableObjectEditor/Private/MuCOE/GraphTraversal.cpp
+++ b/Engine/Plugins/Experimental/Mutable/Source/CustomizableObjectEditor/Private/MuCOE/GraphTraversal.cpp
@@ -3,6 +3,7 @@
 #include "MuCOE/GraphTraversal.h"
 
 #include "AssetRegistry/AssetRegistryModule.h"
+#include "MuCO/CustomizableObjectPrivate.h"
 #include "MuCOE/CustomizableObjectPin.h"
 #include "MuCOE/EdGraphSchema_CustomizableObject.h"
 #include "MuCOE/Nodes/CustomizableObjectNodeEnumParameter.h"
@@ -173,7 +174,7 @@ UCustomizableObjectNodeObject* GetRootNode(UCustomizableObject* Object, bool& bO
 	// Look for the base object node
 	UCustomizableObjectNodeObject* Root = nullptr;
 	TArray<UCustomizableObjectNodeObject*> ObjectNodes;
-	Object->Source->GetNodesOfClass<UCustomizableObjectNodeObject>(ObjectNodes);
+	Object->GetPrivate()->GetSource()->GetNodesOfClass<UCustomizableObjectNodeObject>(ObjectNodes);
 
 	bOutMultipleBaseObjectsFound = false;
 
diff --git a/Engine/Plugins/Experimental/Mutable/Source/CustomizableObjectEditor/Private/MuCOE/Nodes/CustomizableObjectNodeExternalPinDetails.cpp b/Engine/Plugins/Experimental/Mutable/Source/CustomizableObjectEditor/Private/MuCOE/Nodes/CustomizableObjectNodeExternalPinDetails.cpp
--- a/Engine/Plugins/Experimental/Mutable/Source/CustomizableObjectEditor/Private/MuCOE/Nodes/CustomizableObjectNodeExternalPinDetails.cpp
+++ b/Engine/Plugins/Experimental/Mutable/Source/CustomizableObjectEditor/Private/MuCOE/Nodes/CustomizableObjectNodeExternalPinDetails.cpp
@@ -10,6 +10,7 @@
 #include "MuCOE/UnrealEditorPortabilityHelpers.h"
 #include "PropertyCustomizationHelpers.h"
 #include "ScopedTransaction.h"
+#include "MuCO/CustomizableObjectPrivate.h"
 #include "Widgets/Input/STextComboBox.h"
 
 
@@ -55,7 +56,7 @@ void FCustomizableObjectNodeExternalPinDetails::CustomizeDetails( IDetailLayoutB
 			const UCustomizableObjectNodeExposePin* SelectedNodeExposePin = Node->GetNodeExposePin();
 
 			TArray<UCustomizableObjectNodeExposePin*> GroupNodes;
-			Node->ExternalObject->Source->GetNodesOfClass<UCustomizableObjectNodeExposePin>(GroupNodes);
+			Node->ExternalObject->GetPrivate()->GetSource()->GetNodesOfClass<UCustomizableObjectNodeExposePin>(GroupNodes);
 
 			TSharedPtr<FString> ItemToSelect;
 
@@ -118,7 +119,7 @@ void FCustomizableObjectNodeExternalPinDetails::OnGroupNodeComboBoxSelectionChan
 	if (Node->ExternalObject)
 	{
 		TArray<UCustomizableObjectNodeExposePin*> GroupNodes;
-		Node->ExternalObject->Source->GetNodesOfClass<UCustomizableObjectNodeExposePin>(GroupNodes);
+		Node->ExternalObject->GetPrivate()->GetSource()->GetNodesOfClass<UCustomizableObjectNodeExposePin>(GroupNodes);
 
 		for (UCustomizableObjectNodeExposePin* NodeExposePin: GroupNodes)
 		{
diff --git a/Engine/Plugins/Experimental/Mutable/Source/CustomizableObjectEditor/Private/MuCOE/Nodes/CustomizableObjectNodeMeshClipWithMeshDetails.cpp b/Engine/Plugins/Experimental/Mutable/Source/CustomizableObjectEditor/Private/MuCOE/Nodes/CustomizableObjectNodeMeshClipWithMeshDetails.cpp
--- a/Engine/Plugins/Experimental/Mutable/Source/CustomizableObjectEditor/Private/MuCOE/Nodes/CustomizableObjectNodeMeshClipWithMeshDetails.cpp
+++ b/Engine/Plugins/Experimental/Mutable/Source/CustomizableObjectEditor/Private/MuCOE/Nodes/CustomizableObjectNodeMeshClipWithMeshDetails.cpp
@@ -12,6 +12,7 @@
 #include "MuCOE/Nodes/CustomizableObjectNodeStaticMesh.h"
 #include "PropertyCustomizationHelpers.h"
 #include "ScopedTransaction.h"
+#include "MuCO/CustomizableObjectPrivate.h"
 #include "Widgets/Input/SButton.h"
 #include "Widgets/Input/SEditableTextBox.h"
 #include "Widgets/Input/STextComboBox.h"
@@ -131,7 +132,7 @@ void FCustomizableObjectNodeMeshClipWithMeshDetails::CustomizeDetails(IDetailLay
 		TArray<UCustomizableObjectNodeMaterial*> ArrayMaterialNodes;
 		if (SelectedCO != nullptr)
 		{
-			SelectedCO->Source->GetNodesOfClass<UCustomizableObjectNodeMaterial>(ArrayMaterialNodes);
+			SelectedCO->GetPrivate()->GetSource()->GetNodesOfClass<UCustomizableObjectNodeMaterial>(ArrayMaterialNodes);
 		}
 
 		for (const UCustomizableObjectNodeMaterial* MaterialNode : ArrayMaterialNodes)
@@ -300,7 +301,7 @@ void FCustomizableObjectNodeMeshClipWithMeshDetails::OnMeshClipWithMeshNodeCombo
 	if (Selection.IsValid() && (SelectedCO != nullptr))
 	{
 		TArray<UCustomizableObjectNodeMaterial*> ArrayMaterialNode;
-		SelectedCO->Source->GetNodesOfClass<UCustomizableObjectNodeMaterial>(ArrayMaterialNode);
+		SelectedCO->GetPrivate()->GetSource()->GetNodesOfClass<UCustomizableObjectNodeMaterial>(ArrayMaterialNode);
 
 		const FScopedTransaction Transaction(LOCTEXT("ChangedClipMeshWithMeshMaterialTransaction", "Changed Clip Mesh With Mesh Material"));
 		Node->Modify();
diff --git a/Engine/Plugins/Experimental/Mutable/Source/CustomizableObjectEditor/Private/MuCOE/Nodes/CustomizableObjectNodeObjectDetails.cpp b/Engine/Plugins/Experimental/Mutable/Source/CustomizableObjectEditor/Private/MuCOE/Nodes/CustomizableObjectNodeObjectDetails.cpp
--- a/Engine/Plugins/Experimental/Mutable/Source/CustomizableObjectEditor/Private/MuCOE/Nodes/CustomizableObjectNodeObjectDetails.cpp
+++ b/Engine/Plugins/Experimental/Mutable/Source/CustomizableObjectEditor/Private/MuCOE/Nodes/CustomizableObjectNodeObjectDetails.cpp
@@ -15,6 +15,7 @@
 #include "MuCOE/UnrealEditorPortabilityHelpers.h"
 #include "PropertyCustomizationHelpers.h"
 #include "ScopedTransaction.h"
+#include "MuCO/CustomizableObjectPrivate.h"
 #include "Widgets/Input/SComboButton.h"
 #include "Widgets/Input/STextComboBox.h"
 #include "Widgets/Text/STextBlock.h"
@@ -86,7 +87,7 @@ void FCustomizableObjectNodeObjectDetails::CustomizeDetails( IDetailLayoutBuilde
 			if (Node->ParentObject)
 			{
 				TArray<UCustomizableObjectNodeObjectGroup*> GroupNodes;
-				Node->ParentObject->Source->GetNodesOfClass<UCustomizableObjectNodeObjectGroup>(GroupNodes);
+				Node->ParentObject->GetPrivate()->GetSource()->GetNodesOfClass<UCustomizableObjectNodeObjectGroup>(GroupNodes);
 
 				TSharedPtr<FString> ItemToSelect;
 
@@ -257,7 +258,7 @@ void FCustomizableObjectNodeObjectDetails::OnGroupNodeComboBoxSelectionChanged(T
 	if (Selection.IsValid() && Node->ParentObject != nullptr)
 	{
 		TArray<UCustomizableObjectNodeObjectGroup*> GroupNodes;
-		Node->ParentObject->Source->GetNodesOfClass<UCustomizableObjectNodeObjectGroup>(GroupNodes);
+		Node->ParentObject->GetPrivate()->GetSource()->GetNodesOfClass<UCustomizableObjectNodeObjectGroup>(GroupNodes);
 
 		for (UCustomizableObjectNodeObjectGroup* GroupNode : GroupNodes)
 		{
diff --git a/Engine/Plugins/Experimental/Mutable/Source/CustomizableObjectEditor/Private/MuCOE/SCustomizableObjectEditorTagExplorer.cpp b/Engine/Plugins/Experimental/Mutable/Source/CustomizableObjectEditor/Private/MuCOE/SCustomizableObjectEditorTagExplorer.cpp
--- a/Engine/Plugins/Experimental/Mutable/Source/CustomizableObjectEditor/Private/MuCOE/SCustomizableObjectEditorTagExplorer.cpp
+++ b/Engine/Plugins/Experimental/Mutable/Source/CustomizableObjectEditor/Private/MuCOE/SCustomizableObjectEditorTagExplorer.cpp
@@ -246,9 +246,9 @@ TSharedRef<SWidget> SCustomizableObjectEditorTagExplorer::OnGetTagsMenuContent()
 
 void SCustomizableObjectEditorTagExplorer::FillTagInformation(UCustomizableObject * Object, TArray<FString>& Tags)
 {
-	if (Object->Source)
+	if (Object->GetPrivate()->GetSource())
 	{
-		for (const TObjectPtr<UEdGraphNode>& Node : Object->Source->Nodes)
+		for (const TObjectPtr<UEdGraphNode>& Node : Object->GetPrivate()->GetSource()->Nodes)
 		{
 			if (UCustomizableObjectNodeMaterial* TypedNodeMat = Cast<UCustomizableObjectNodeMaterial>(Node))
 			{
