[mutable] Compilation optimization: Load texture source data concurrently.

- Exposed a few methods of the texture api to make async loading of source data possible from other modules.

#rnx

diff --git a/Engine/Plugins/Experimental/Mutable/Source/CustomizableObject/Internal/MuCO/UnrealToMutableTextureConversionUtils.h b/Engine/Plugins/Experimental/Mutable/Source/CustomizableObject/Internal/MuCO/UnrealToMutableTextureConversionUtils.h
--- a/Engine/Plugins/Experimental/Mutable/Source/CustomizableObject/Internal/MuCO/UnrealToMutableTextureConversionUtils.h
+++ b/Engine/Plugins/Experimental/Mutable/Source/CustomizableObject/Internal/MuCO/UnrealToMutableTextureConversionUtils.h
@@ -2,9 +2,20 @@
 
 #pragma once
 
-#if WITH_EDITOR
-
 #include "HAL/Platform.h"
+#include "Engine/Texture.h"
+
+struct CUSTOMIZABLEOBJECT_API FMutableSourceTextureData
+{
+	FTextureSource Source;
+	bool bFlipGreenChannel = false;
+	bool bHasAlphaChannel = false;
+	bool bCompressionForceAlpha = false;
+	bool bIsNormalComposite = false;
+};
+
+
+#if WITH_EDITOR
 
 namespace mu
 {
@@ -23,7 +34,6 @@ enum class EUnrealToMutableConversionError
     Unknown
 };
 
-
-CUSTOMIZABLEOBJECT_API EUnrealToMutableConversionError ConvertTextureUnrealSourceToMutable(mu::Image* OutResult, UTexture2D* Texture, bool bIsNormalComposite, uint8 MipmapsToSkip);
+CUSTOMIZABLEOBJECT_API EUnrealToMutableConversionError ConvertTextureUnrealSourceToMutable(mu::Image* OutResult, FMutableSourceTextureData&, uint8 MipmapsToSkip);
 
 #endif
diff --git a/Engine/Plugins/Experimental/Mutable/Source/CustomizableObject/Private/MuCO/UnrealMutableImageProvider.cpp b/Engine/Plugins/Experimental/Mutable/Source/CustomizableObject/Private/MuCO/UnrealMutableImageProvider.cpp
--- a/Engine/Plugins/Experimental/Mutable/Source/CustomizableObject/Private/MuCO/UnrealMutableImageProvider.cpp
+++ b/Engine/Plugins/Experimental/Mutable/Source/CustomizableObject/Private/MuCO/UnrealMutableImageProvider.cpp
@@ -14,15 +14,8 @@
 //-------------------------------------------------------------------------------------------------
 namespace
 {
-
-	void ConvertTextureUnrealToMutable(mu::Image* OutResult, UTexture2D* Texture, bool bIsNormalComposite, uint8 MipmapsToSkip)
+	void ConvertTextureUnrealPlatformToMutable(mu::Image* OutResult, UTexture2D* Texture, uint8 MipmapsToSkip)
 	{		
-#if WITH_EDITOR
-
-		EUnrealToMutableConversionError Error = ConvertTextureUnrealSourceToMutable(OutResult, Texture, bIsNormalComposite, MipmapsToSkip);
-		check(Error==EUnrealToMutableConversionError::Success);
-
-#else
 		check(Texture->GetPlatformData()->Mips[MipmapsToSkip].BulkData.IsBulkDataLoaded());
 
 		int32 LODs = 1;
@@ -58,8 +51,6 @@ namespace
 			check(false);
 			OutResult->Init(SizeX, SizeY, LODs, MutableFormat, mu::EInitializationType::Black);
 		}
-
-#endif
 	}
 }
 
@@ -139,8 +130,26 @@ bool FUnrealMutableImageProvider::Tick()
 	// In the editor the src data can be directly accessed
 	int32 MipIndex = (Request->MipmapsToSkip < Texture->GetPlatformData()->Mips.Num()) ? Request->MipmapsToSkip : Texture->GetPlatformData()->Mips.Num() - 1;
 	check(MipIndex >= 0);
-	bool bIsNormalComposite = false; // TODO?
-	ConvertTextureUnrealToMutable(Request->ResultImage.get(), Texture, bIsNormalComposite, MipIndex);
+
+	FMutableSourceTextureData Tex;
+	Tex.Source = Texture->Source.CopyTornOff();
+	Tex.bFlipGreenChannel = Texture->bFlipGreenChannel;
+	Tex.bHasAlphaChannel = 
+		Texture->AdjustMinAlpha != Texture->AdjustMaxAlpha
+		&& Texture->CompressionSettings != TextureCompressionSettings::TC_Normalmap
+		&& !Texture->CompressionNoAlpha;
+	Tex.bCompressionForceAlpha = Texture->CompressionForceAlpha;
+	Tex.bIsNormalComposite = false; // TODO?
+
+	EUnrealToMutableConversionError Error = ConvertTextureUnrealSourceToMutable(Request->ResultImage.get(), Tex, MipIndex);
+	
+	if (Error != EUnrealToMutableConversionError::Success)
+	{
+		// This could happen in the editor, because some source textures may have changed while there was a background compilation.
+		// We just show a warning and move on. This cannot happen during cooks, so it is fine.
+		UE_LOG(LogMutable, Warning, TEXT("Failed to load some source texture data for [%s]. Some materials may look corrupted."), *Texture->GetName());
+	}
+
 	Request->CompletionEvent.Trigger();
 
 	return true;
@@ -225,8 +234,26 @@ TTuple<UE::Tasks::FTask, TFunction<void()>> FUnrealMutableImageProvider::GetImag
 #if WITH_EDITOR
 			// In the editor the src data can be directly accessed
 			mu::Ptr<mu::Image> Image = new mu::Image();
-			bool bIsNormalComposite = false; // TODO?
-			ConvertTextureUnrealToMutable(Image.get(), TextureToLoad, bIsNormalComposite,  MipIndex);
+
+			FMutableSourceTextureData Tex;
+			Tex.Source = TextureToLoad->Source.CopyTornOff();
+			Tex.bFlipGreenChannel = TextureToLoad->bFlipGreenChannel;
+			Tex.bHasAlphaChannel =
+				TextureToLoad->AdjustMinAlpha != TextureToLoad->AdjustMaxAlpha
+				&& TextureToLoad->CompressionSettings != TextureCompressionSettings::TC_Normalmap
+				&& !TextureToLoad->CompressionNoAlpha;
+			Tex.bCompressionForceAlpha = TextureToLoad->CompressionForceAlpha;
+			Tex.bIsNormalComposite = false; // TODO?
+
+			EUnrealToMutableConversionError Error = ConvertTextureUnrealSourceToMutable(Image.get(), Tex, MipIndex);
+
+			if (Error != EUnrealToMutableConversionError::Success)
+			{
+				// This could happen in the editor, because some source textures may have changed while there was a background compilation.
+				// We just show a warning and move on. This cannot happen during cooks, so it is fine.
+				UE_LOG(LogMutable, Warning, TEXT("Failed to load some source texture data for [%s]. Some materials may look corrupted."), *TextureToLoad->GetName());
+			}
+
 			ResultCallback(Image);
 			return Invoke(TrivialReturn);
 #else
@@ -569,8 +596,29 @@ void FUnrealMutableImageProvider::CacheImage(FName Id, bool bUser)
 					{
 						UTexture2D* UnrealTexture = Provider->GetTextureParameterValue(Id);
 						pResult = new mu::Image();
-						bool bIsNormalComposite = false;
-						ConvertTextureUnrealToMutable(pResult.get(), UnrealTexture, bIsNormalComposite, 0);
+
+#if WITH_EDITOR
+						FMutableSourceTextureData Tex;
+						Tex.Source = UnrealTexture->Source.CopyTornOff();
+						Tex.bFlipGreenChannel = UnrealTexture->bFlipGreenChannel;
+						Tex.bHasAlphaChannel =
+							UnrealTexture->AdjustMinAlpha != UnrealTexture->AdjustMaxAlpha
+							&& UnrealTexture->CompressionSettings != TextureCompressionSettings::TC_Normalmap
+							&& !UnrealTexture->CompressionNoAlpha;
+						Tex.bCompressionForceAlpha = UnrealTexture->CompressionForceAlpha;
+						Tex.bIsNormalComposite = false; // TODO?
+
+						EUnrealToMutableConversionError Error = ConvertTextureUnrealSourceToMutable(pResult.get(), Tex, 0);
+						if (Error != EUnrealToMutableConversionError::Success)
+						{
+							// This could happen in the editor, because some source textures may have changed while there was a background compilation.
+							// We just show a warning and move on. This cannot happen during cooks, so it is fine.
+							UE_LOG(LogMutable, Warning, TEXT("Failed to load some source texture data for [%s]. Some textures may be corrupted."), *UnrealTexture->GetName());
+						}
+#else
+						ConvertTextureUnrealPlatformToMutable(pResult.get(), UnrealTexture, 0);
+#endif
+
 						break;
 					}
 
diff --git a/Engine/Plugins/Experimental/Mutable/Source/CustomizableObject/Private/MuCO/UnrealToMutableTextureConversionUtils.cpp b/Engine/Plugins/Experimental/Mutable/Source/CustomizableObject/Private/MuCO/UnrealToMutableTextureConversionUtils.cpp
--- a/Engine/Plugins/Experimental/Mutable/Source/CustomizableObject/Private/MuCO/UnrealToMutableTextureConversionUtils.cpp
+++ b/Engine/Plugins/Experimental/Mutable/Source/CustomizableObject/Private/MuCO/UnrealToMutableTextureConversionUtils.cpp
@@ -11,7 +11,7 @@
 
 #if WITH_EDITOR
 
-namespace UnrealToMutableImageConversion_Interanl
+namespace UnrealToMutableImageConversion_Internal
 {
 	
 FORCEINLINE ERawImageFormat::Type ConvertFormatSourceToRaw(const ETextureSourceFormat SourceFormat)
@@ -173,44 +173,37 @@ void BlurNormalForComposite(FImage& Image)
 } //namespace UnrealToMutableImageConversion_Internal
 
 
-EUnrealToMutableConversionError ConvertTextureUnrealSourceToMutable(mu::Image* OutResult, UTexture2D* Texture, bool bIsNormalComposite, uint8 MipmapsToSkip)
+EUnrealToMutableConversionError ConvertTextureUnrealSourceToMutable(mu::Image* OutResult, FMutableSourceTextureData& Tex, uint8 MipmapsToSkip)
 {
-	MUTABLE_CPUPROFILER_SCOPE(ConvertTextureUnrealToMutableTuple);
+	MUTABLE_CPUPROFILER_SCOPE(ConvertTextureUnrealSourceToMutable);
 
-    using namespace UnrealToMutableImageConversion_Interanl;
+    using namespace UnrealToMutableImageConversion_Internal;
 
 	// Correct mips to skip to fit source data
-	MipmapsToSkip = FMath::Clamp(MipmapsToSkip, 0, Texture->Source.GetNumMips()-1);
+	MipmapsToSkip = FMath::Clamp(MipmapsToSkip, 0, Tex.Source.GetNumMips()-1);
 
     const int32 LODs = 1;
-    const int32 SizeX = Texture->Source.GetSizeX() >> MipmapsToSkip;
-    const int32 SizeY = Texture->Source.GetSizeY() >> MipmapsToSkip;
+    const int32 SizeX = Tex.Source.GetSizeX() >> MipmapsToSkip;
+    const int32 SizeY = Tex.Source.GetSizeY() >> MipmapsToSkip;
 	check(SizeX > 0 && SizeY > 0);
 
-    ETextureSourceFormat Format = Texture->Source.GetFormat();
+    ETextureSourceFormat Format = Tex.Source.GetFormat();
  
     ERawImageFormat::Type RawFormat = ConvertFormatSourceToRaw(Format);
 
-	// Not true, we will convert it.
-	// \TODO: Warn?
-    //if (RawFormat == ERawImageFormat::RGBA32F)
-    //{
-    //    return MakeTuple(nullptr, EUnrealToMutableConversionError::UnsupportedFormat);
-    //}
-
 	// What if source data is not linear?
     FImage TempImage(SizeX, SizeY, 1, RawFormat, EGammaSpace::Linear);
 	FImage TempImage2;
 
-    if (!Texture->Source.GetMipData(TempImage.RawData, MipmapsToSkip))
+    if (!Tex.Source.GetMipData(TempImage.RawData, MipmapsToSkip))
     {
         return EUnrealToMutableConversionError::Unknown;
     }
 
-    bool bFlipGreenChannel = Texture->bFlipGreenChannel;
+    bool bFlipGreenChannel = Tex.bFlipGreenChannel;
 
     // If any post processes of the image is needed, convert to RGBA32F
-    if (bIsNormalComposite)
+    if (Tex.bIsNormalComposite)
     { 
 		MUTABLE_CPUPROFILER_SCOPE(FlipOrComposite);
 
@@ -264,11 +257,8 @@ EUnrealToMutableConversionError ConvertTextureUnrealSourceToMutable(mu::Image* O
 	case ERawImageFormat::BGRA8:
 	{
 		// Try to find out if the texture has and actually makes use of the alpha channel
-		bool bHasAlphaChannel = 
-			Texture->AdjustMinAlpha != Texture->AdjustMaxAlpha
-			&& Texture->CompressionSettings != TextureCompressionSettings::TC_Normalmap
-			&& !Texture->CompressionNoAlpha
-			&& (Texture->CompressionForceAlpha 
+		bool bHasAlphaChannel = Tex.bHasAlphaChannel
+			&& (Tex.bCompressionForceAlpha
 				||
 				FImageCore::DetectAlphaChannel(TempImage));
 
diff --git a/Engine/Plugins/Experimental/Mutable/Source/CustomizableObjectEditor/Private/MuCOE/CustomizableObjectCompileRunnable.cpp b/Engine/Plugins/Experimental/Mutable/Source/CustomizableObjectEditor/Private/MuCOE/CustomizableObjectCompileRunnable.cpp
--- a/Engine/Plugins/Experimental/Mutable/Source/CustomizableObjectEditor/Private/MuCOE/CustomizableObjectCompileRunnable.cpp
+++ b/Engine/Plugins/Experimental/Mutable/Source/CustomizableObjectEditor/Private/MuCOE/CustomizableObjectCompileRunnable.cpp
@@ -45,8 +45,6 @@ FCustomizableObjectCompileRunnable::FCustomizableObjectCompileRunnable(mu::Ptr<m
 
 mu::Ptr<mu::Image> FCustomizableObjectCompileRunnable::LoadResourceReferenced(int32 ID)
 {
-	check(IsInGameThread());
-
 	MUTABLE_CPUPROFILER_SCOPE(LoadResourceReferenced);
 
 	mu::Ptr<mu::Image> Image;
@@ -58,23 +56,20 @@ mu::Ptr<mu::Image> FCustomizableObjectCompileRunnable::LoadResourceReferenced(in
 	}
 
 	// Find the texture id
-	TSoftObjectPtr<UTexture> TexturePtr = ReferencedTextures[ID];
-
-	// This can cause a stall because of loading the asset.
-	UTexture2D* Texture = Cast<UTexture2D>(TexturePtr.LoadSynchronous());
-	if (!Texture)
-	{
-		// Failed to load the texture
-		check(false);
-		return Image;
-	}
+	FMutableSourceTextureData& TextureData = ReferencedTextures[ID];
 
 	// In the editor the src data can be directly accessed
 	Image = new mu::Image();
 	int32 MipmapsToSkip = 0;
-	bool bIsNormalComposite = false; // TODO?
-	EUnrealToMutableConversionError Error = ConvertTextureUnrealSourceToMutable(Image.get(), Texture, bIsNormalComposite, MipmapsToSkip);
-	check(Error == EUnrealToMutableConversionError::Success);
+	EUnrealToMutableConversionError Error = ConvertTextureUnrealSourceToMutable(Image.get(), TextureData, MipmapsToSkip);
+
+	if (Error != EUnrealToMutableConversionError::Success)
+	{
+		// This could happen in the editor, because some source textures may have changed while there was a background compilation.
+		// We just show a warning and move on. This cannot happen during cooks, so it is fine.
+		UE_LOG(LogMutable, Warning, TEXT("Failed to load some source texture data for texture ID [%d]. Some textures may be corrupted."), ID);
+	}
+
 	return Image;
 }
 
@@ -146,31 +141,19 @@ uint32 FCustomizableObjectCompileRunnable::Run()
 		CompilerOptions->SetImagePixelFormatOverride( UnrealPixelFormatFunc );
 	}
 
-	auto ProviderTick = [this](float)
+	CompilerOptions->SetReferencedResourceCallback([this](int32 ID, TSharedPtr<mu::Ptr<mu::Image>> ResolvedImage, bool bRunImmediatlyIfPossible)
 		{
-			Tick();
-		};
-
-	CompilerOptions->SetReferencedResourceCallback([this, &ProviderTick](int32 ID, TSharedPtr<mu::Ptr<mu::Image>> ResolvedImage, bool bRunImmediatlyIfPossible)
-		{
-			// This runs in a random thread
-			UE::Tasks::FTaskEvent CompletionEvent(TEXT("ReferencedResourceCallbackCompletion"));			
-
-			if (IsInGameThread() && bRunImmediatlyIfPossible)
-			{
-				// Do everything now
-				mu::Ptr<mu::Image> Result = LoadResourceReferenced(ID);
-				*ResolvedImage = Result;
-				CompletionEvent.Trigger();
-			}
-			else
-			{
-				PendingResourceReferenceRequests.Enqueue(FReferenceResourceRequest{ ID, ResolvedImage, MakeShared<UE::Tasks::FTaskEvent>(CompletionEvent) });
-			}
-
-			return CompletionEvent;
-		}, 
-		ProviderTick
+			UE::Tasks::FTask LaunchTask = UE::Tasks::Launch(TEXT("ConstantGeneratorLaunchTasks"),
+				[ID,ResolvedImage,this]()
+				{
+					mu::Ptr<mu::Image> Result = LoadResourceReferenced(ID);
+					*ResolvedImage = Result;
+				},
+				LowLevelTasks::ETaskPriority::BackgroundLow
+			);
+
+			return LaunchTask;
+		}
 	);
 
 	const int32 MinResidentMips = UTexture::GetStaticMinTextureResidentMipCount();
diff --git a/Engine/Plugins/Experimental/Mutable/Source/CustomizableObjectEditor/Private/MuCOE/CustomizableObjectCompileRunnable.h b/Engine/Plugins/Experimental/Mutable/Source/CustomizableObjectEditor/Private/MuCOE/CustomizableObjectCompileRunnable.h
--- a/Engine/Plugins/Experimental/Mutable/Source/CustomizableObjectEditor/Private/MuCOE/CustomizableObjectCompileRunnable.h
+++ b/Engine/Plugins/Experimental/Mutable/Source/CustomizableObjectEditor/Private/MuCOE/CustomizableObjectCompileRunnable.h
@@ -5,6 +5,7 @@
 #include "HAL/Runnable.h"
 #include "MuCO/CustomizableObject.h"
 #include "MuCO/CustomizableObjectPrivate.h"
+#include "MuCO/UnrealToMutableTextureConversionUtils.h"
 #include "MuR/Ptr.h"
 #include "MuT/Node.h"
 #include "MuCOE/CustomizableObjectEditorLogger.h"
@@ -74,7 +75,7 @@ public:
 
 	FCompilationOptions Options;
 	
-	TArray<TSoftObjectPtr<UTexture>> ReferencedTextures;
+	TArray<FMutableSourceTextureData> ReferencedTextures;
 
 	FString ErrorMsg;
 
diff --git a/Engine/Plugins/Experimental/Mutable/Source/CustomizableObjectEditor/Private/MuCOE/CustomizableObjectCompiler.cpp b/Engine/Plugins/Experimental/Mutable/Source/CustomizableObjectEditor/Private/MuCOE/CustomizableObjectCompiler.cpp
--- a/Engine/Plugins/Experimental/Mutable/Source/CustomizableObjectEditor/Private/MuCOE/CustomizableObjectCompiler.cpp
+++ b/Engine/Plugins/Experimental/Mutable/Source/CustomizableObjectEditor/Private/MuCOE/CustomizableObjectCompiler.cpp
@@ -1109,11 +1109,23 @@ void FCustomizableObjectCompiler::CompileInternal(UCustomizableObject* Object, c
 		Object->GetPrivate()->GetStreamedResourceData() = MoveTemp(GenerationContext.StreamedResourceData);
 
 		// Pass-through textures
-		TArray<TSoftObjectPtr<UTexture>> NewCompileTimeReferencedTextures;
+		TArray<FMutableSourceTextureData> NewCompileTimeReferencedTextures;
 		for (const TPair<TSoftObjectPtr<UTexture>, FMutableGraphGenerationContext::FGeneratedReferencedTexture>& Pair : GenerationContext.CompileTimeTextureMap)
 		{
 			check(Pair.Value.ID == NewCompileTimeReferencedTextures.Num());
-			NewCompileTimeReferencedTextures.Add(Pair.Key);
+
+			UTexture* Texture = Pair.Key.LoadSynchronous();
+			FMutableSourceTextureData Tex;
+			Tex.Source = Texture->Source.CopyTornOff();
+			Tex.bFlipGreenChannel = Texture->bFlipGreenChannel;
+			Tex.bHasAlphaChannel =
+				Texture->AdjustMinAlpha != Texture->AdjustMaxAlpha
+				&& Texture->CompressionSettings != TextureCompressionSettings::TC_Normalmap
+				&& !Texture->CompressionNoAlpha;
+			Tex.bCompressionForceAlpha = Texture->CompressionForceAlpha;
+			Tex.bIsNormalComposite = false; // TODO?
+
+			NewCompileTimeReferencedTextures.Add(Tex);
 		}
 
 		if (!ParamNamesToSelectedOptions.Num())
@@ -1218,7 +1230,7 @@ void FCustomizableObjectCompiler::CompileInternal(UCustomizableObject* Object, c
 
 mu::NodePtr FCustomizableObjectCompiler::Export(UCustomizableObject* Object, const FCompilationOptions& InCompilerOptions, 
 	TArray<TSoftObjectPtr<UTexture>>& OutRuntimeReferencedTextures,
-	TArray<TSoftObjectPtr<UTexture>>& OutCompilerReferencedTextures )
+	TArray<FMutableSourceTextureData>& OutCompilerReferencedTextures )
 {
 	UE_LOG(LogMutable, Log, TEXT("Started Customizable Object Export %s."), *Object->GetName());
 
@@ -1265,7 +1277,19 @@ mu::NodePtr FCustomizableObjectCompiler::Export(UCustomizableObject* Object, con
 	for (const TPair<TSoftObjectPtr<UTexture>, FMutableGraphGenerationContext::FGeneratedReferencedTexture>& Pair : GenerationContext.CompileTimeTextureMap)
 	{
 		check(Pair.Value.ID == OutCompilerReferencedTextures.Num());
-		OutCompilerReferencedTextures.Add(Pair.Key);
+
+		UTexture* Texture = Pair.Key.LoadSynchronous();
+		FMutableSourceTextureData Tex;
+		Tex.Source = Texture->Source.CopyTornOff();
+		Tex.bFlipGreenChannel = Texture->bFlipGreenChannel;
+		Tex.bHasAlphaChannel =
+			Texture->AdjustMinAlpha != Texture->AdjustMaxAlpha
+			&& Texture->CompressionSettings != TextureCompressionSettings::TC_Normalmap
+			&& !Texture->CompressionNoAlpha;
+		Tex.bCompressionForceAlpha = Texture->CompressionForceAlpha;
+		Tex.bIsNormalComposite = false; // TODO?
+
+		OutCompilerReferencedTextures.Add(Tex);
 	}
 
 	return MutableRoot;
diff --git a/Engine/Plugins/Experimental/Mutable/Source/CustomizableObjectEditor/Private/MuCOE/CustomizableObjectCompiler.h b/Engine/Plugins/Experimental/Mutable/Source/CustomizableObjectEditor/Private/MuCOE/CustomizableObjectCompiler.h
--- a/Engine/Plugins/Experimental/Mutable/Source/CustomizableObjectEditor/Private/MuCOE/CustomizableObjectCompiler.h
+++ b/Engine/Plugins/Experimental/Mutable/Source/CustomizableObjectEditor/Private/MuCOE/CustomizableObjectCompiler.h
@@ -4,6 +4,7 @@
 
 #include "AssetRegistry/AssetData.h"
 #include "MuCO/CustomizableObjectPrivate.h"
+#include "MuCO/UnrealToMutableTextureConversionUtils.h"
 #include "MuCOE/CompilationMessageCache.h"
 #include "MuCOE/GenerateMutableSource/GenerateMutableSource.h"
 #include "MuCOE/CustomizableObjectEditorLogger.h"
@@ -33,7 +34,7 @@ public:
 	CUSTOMIZABLEOBJECTEDITOR_API virtual bool Tick() override;
 
 	/** Generate the Mutable Graph from the Unreal Graph. */
-	mu::NodePtr Export(UCustomizableObject* Object, const FCompilationOptions& Options, TArray<TSoftObjectPtr<UTexture>>& OutRuntimeReferencedTextures, TArray<TSoftObjectPtr<UTexture>>& OutCompilerReferencedTextures);
+	mu::NodePtr Export(UCustomizableObject* Object, const FCompilationOptions& Options, TArray<TSoftObjectPtr<UTexture>>& OutRuntimeReferencedTextures, TArray<FMutableSourceTextureData>& OutCompilerReferencedTextures);
 
 	void CompilerLog(const FText& Message, const TArray<const UObject*>& UObject, const EMessageSeverity::Type MessageSeverity = EMessageSeverity::Warning, const bool bAddBaseObjectInfo = true, const ELoggerSpamBin SpamBin = ELoggerSpamBin::ShowAll);
 	void CompilerLog(const FText& Message, const UObject* Context = nullptr, const EMessageSeverity::Type MessageSeverity = EMessageSeverity::Warning, const bool bAddBaseObjectInfo = true, const ELoggerSpamBin SpamBin = ELoggerSpamBin::ShowAll);
diff --git a/Engine/Plugins/Experimental/Mutable/Source/CustomizableObjectEditor/Private/MuCOE/GenerateMutableSource/GenerateMutableSourceImage.cpp b/Engine/Plugins/Experimental/Mutable/Source/CustomizableObjectEditor/Private/MuCOE/GenerateMutableSource/GenerateMutableSourceImage.cpp
--- a/Engine/Plugins/Experimental/Mutable/Source/CustomizableObjectEditor/Private/MuCOE/GenerateMutableSource/GenerateMutableSourceImage.cpp
+++ b/Engine/Plugins/Experimental/Mutable/Source/CustomizableObjectEditor/Private/MuCOE/GenerateMutableSource/GenerateMutableSourceImage.cpp
@@ -54,55 +54,6 @@
 #define LOCTEXT_NAMESPACE "CustomizableObjectEditor"
 
 
-mu::ImagePtr ConvertTextureUnrealToMutable(UTexture2D* Texture, const UCustomizableObjectNode* Node, FCustomizableObjectCompiler* Compiler, bool bIsNormalComposite)
-{     
-	MUTABLE_CPUPROFILER_SCOPE(ConvertTextureUnrealToMutable);
-
-	mu::Ptr<mu::Image> MutableImage = new mu::Image;
-	EUnrealToMutableConversionError Error = EUnrealToMutableConversionError::Unknown;
-
-	Error = ConvertTextureUnrealSourceToMutable(MutableImage.get(), Texture, bIsNormalComposite, 0);
-
-	switch(Error)
-	{
-	case EUnrealToMutableConversionError::Success: break;
-	case EUnrealToMutableConversionError::UnsupportedFormat:
-		{
-			Compiler->CompilerLog(
-					LOCTEXT("UnsupportedImageFormat", 
-							"Image format not supported."), 
-					Node);
-			break;
-		}
-	case EUnrealToMutableConversionError::CompositeImageDimensionMismatch:
-		{
-			Compiler->CompilerLog(
-					LOCTEXT("CompositeImageDimensionMismatch", 
-							"Composite image dimension mismatch."), 
-					Node);
-			break;
-		}
-	case EUnrealToMutableConversionError::CompositeUnsupportedFormat:
-		{
-			Compiler->CompilerLog(
-					LOCTEXT("CompositeUnsupportedFormat", 
-							"Composite image format not supported."), 
-					Node);
-			break;
-		}
-	default:
-		{
-			Compiler->CompilerLog(
-					LOCTEXT("ImageConversionUnknownError", 
-							"Image conversion unknown error."), 
-					Node);
-			break;
-		}
-	}
-
-	return MutableImage;
-}
-
 mu::Ptr<mu::NodeImage> ResizeTextureByNumMips(const mu::Ptr<mu::NodeImage>& ImageConstant, int32 MipsToSkip)
 {
 	if (MipsToSkip > 0)
diff --git a/Engine/Plugins/Experimental/Mutable/Source/CustomizableObjectEditor/Private/MuCOE/GenerateMutableSource/GenerateMutableSourceImage.h b/Engine/Plugins/Experimental/Mutable/Source/CustomizableObjectEditor/Private/MuCOE/GenerateMutableSource/GenerateMutableSourceImage.h
--- a/Engine/Plugins/Experimental/Mutable/Source/CustomizableObjectEditor/Private/MuCOE/GenerateMutableSource/GenerateMutableSourceImage.h
+++ b/Engine/Plugins/Experimental/Mutable/Source/CustomizableObjectEditor/Private/MuCOE/GenerateMutableSource/GenerateMutableSourceImage.h
@@ -20,7 +20,7 @@ namespace mu
 }
 
 
-mu::Ptr<mu::Image> ConvertTextureUnrealToMutable(UTexture2D* Texture, const UCustomizableObjectNode* Node, FCustomizableObjectCompiler* Compiler, bool bIsNormalComposite);
+//mu::Ptr<mu::Image> ConvertTextureUnrealToMutable(UTexture2D* Texture, const UCustomizableObjectNode* Node, FCustomizableObjectCompiler* Compiler, bool bIsNormalComposite);
 
 mu::Ptr<mu::NodeImage> ResizeTextureByNumMips(const mu::Ptr<mu::NodeImage>& Image, int32 MipsToSkip);
 
diff --git a/Engine/Plugins/Experimental/Mutable/Source/CustomizableObjectEditor/Private/MuCOE/SMutableCodeViewer.cpp b/Engine/Plugins/Experimental/Mutable/Source/CustomizableObjectEditor/Private/MuCOE/SMutableCodeViewer.cpp
--- a/Engine/Plugins/Experimental/Mutable/Source/CustomizableObjectEditor/Private/MuCOE/SMutableCodeViewer.cpp
+++ b/Engine/Plugins/Experimental/Mutable/Source/CustomizableObjectEditor/Private/MuCOE/SMutableCodeViewer.cpp
@@ -2560,9 +2560,25 @@ namespace
 			check(MipIndex >= 0);
 
 			mu::Ptr<mu::Image> ResultImage = new mu::Image();
-			bool bIsNormalComposite = false; // TODO?
-			EUnrealToMutableConversionError Error = ConvertTextureUnrealSourceToMutable(ResultImage.get(), Texture, bIsNormalComposite, MipmapsToSkip);
-			check(Error == EUnrealToMutableConversionError::Success);
+
+			FMutableSourceTextureData Tex;
+			Tex.Source = Texture->Source.CopyTornOff();
+			Tex.bFlipGreenChannel = Texture->bFlipGreenChannel;
+			Tex.bHasAlphaChannel =
+				Texture->AdjustMinAlpha != Texture->AdjustMaxAlpha
+				&& Texture->CompressionSettings != TextureCompressionSettings::TC_Normalmap
+				&& !Texture->CompressionNoAlpha;
+			Tex.bCompressionForceAlpha = Texture->CompressionForceAlpha;
+			Tex.bIsNormalComposite = false; // TODO?
+
+			EUnrealToMutableConversionError Error = ConvertTextureUnrealSourceToMutable(ResultImage.get(), Tex, MipmapsToSkip);
+
+			if (Error != EUnrealToMutableConversionError::Success)
+			{
+				// This could happen in the editor, because some source textures may have changed while there was a background compilation.
+				// We just show a warning and move on. This cannot happen during cooks, so it is fine.
+				UE_LOG(LogMutable, Warning, TEXT("Failed to load some source texture data for [%s]. Some textures may be corrupted."), *Texture->GetName());
+			}
 
 			ResultCallback(ResultImage);
 
diff --git a/Engine/Plugins/Experimental/Mutable/Source/CustomizableObjectEditor/Private/MuCOE/SMutableGraphViewer.h b/Engine/Plugins/Experimental/Mutable/Source/CustomizableObjectEditor/Private/MuCOE/SMutableGraphViewer.h
--- a/Engine/Plugins/Experimental/Mutable/Source/CustomizableObjectEditor/Private/MuCOE/SMutableGraphViewer.h
+++ b/Engine/Plugins/Experimental/Mutable/Source/CustomizableObjectEditor/Private/MuCOE/SMutableGraphViewer.h
@@ -36,7 +36,7 @@ public:
 	SLATE_ARGUMENT(FString, DataTag)
 
 	SLATE_ARGUMENT(TArray<TSoftObjectPtr<UTexture>>, ReferencedRuntimeTextures)
-	SLATE_ARGUMENT(TArray<TSoftObjectPtr<UTexture>>, ReferencedCompileTextures)
+	SLATE_ARGUMENT(TArray<FMutableSourceTextureData>, ReferencedCompileTextures)
 
 	SLATE_END_ARGS()
 
@@ -60,7 +60,7 @@ private:
 
 	/** Array of external referenced textures in MutableModel, indexed by id. */
 	TArray<TSoftObjectPtr<UTexture>> ReferencedRuntimeTextures;
-	TArray<TSoftObjectPtr<UTexture>> ReferencedCompileTextures;
+	TArray<FMutableSourceTextureData> ReferencedCompileTextures;
 
 	/** Object compiler. */
 	FCustomizableObjectCompiler Compiler;
diff --git a/Engine/Plugins/Experimental/Mutable/Source/CustomizableObjectEditor/Private/MuCOE/SMutableObjectViewer.cpp b/Engine/Plugins/Experimental/Mutable/Source/CustomizableObjectEditor/Private/MuCOE/SMutableObjectViewer.cpp
--- a/Engine/Plugins/Experimental/Mutable/Source/CustomizableObjectEditor/Private/MuCOE/SMutableObjectViewer.cpp
+++ b/Engine/Plugins/Experimental/Mutable/Source/CustomizableObjectEditor/Private/MuCOE/SMutableObjectViewer.cpp
@@ -105,7 +105,7 @@ void SMutableObjectViewer::GenerateMutableGraphPressed()
 {
 	// Convert from Unreal graph to Mutable graph.
 	TArray<TSoftObjectPtr<UTexture>> RuntimeTextures;
-	TArray<TSoftObjectPtr<UTexture>> CompilerTextures;
+	TArray<FMutableSourceTextureData> CompilerTextures;
 	mu::Ptr<mu::Node> RootNode = Compiler.Export(CustomizableObject, CompileOptions, RuntimeTextures, CompilerTextures);
 	if (!RootNode)
 	{
@@ -128,7 +128,7 @@ void SMutableObjectViewer::GenerateMutableGraphPressed()
 void SMutableObjectViewer::CompileMutableCodePressed()
 {
 	TArray<TSoftObjectPtr<UTexture>> RuntimeTextures;
-	TArray<TSoftObjectPtr<UTexture>> CompilerTextures;
+	TArray<FMutableSourceTextureData> CompilerTextures;
 	if (CompileOptions.bForceLargeLODBias)
 	{
 		// Debug compile with many different biases
diff --git a/Engine/Plugins/Experimental/Mutable/Source/MutableTools/Internal/MuT/AST.h b/Engine/Plugins/Experimental/Mutable/Source/MutableTools/Internal/MuT/AST.h
--- a/Engine/Plugins/Experimental/Mutable/Source/MutableTools/Internal/MuT/AST.h
+++ b/Engine/Plugins/Experimental/Mutable/Source/MutableTools/Internal/MuT/AST.h
@@ -214,9 +214,6 @@ namespace mu
 
 		/** Function used to request an engine resource (image) from the compiler. */
 		FReferencedResourceFunc ReferencedResourceProvider;
-
-		/** Function used to ensure the provision of engine resources with ReferencedResourceProvider advances. */
-		FReferencedResourceGameThreadTickFunc ReferencedResourceProviderTick;
 	};
 
 
diff --git a/Engine/Plugins/Experimental/Mutable/Source/MutableTools/Private/MuT/CodeOptimiser.cpp b/Engine/Plugins/Experimental/Mutable/Source/MutableTools/Private/MuT/CodeOptimiser.cpp
--- a/Engine/Plugins/Experimental/Mutable/Source/MutableTools/Private/MuT/CodeOptimiser.cpp
+++ b/Engine/Plugins/Experimental/Mutable/Source/MutableTools/Private/MuT/CodeOptimiser.cpp
@@ -24,7 +24,6 @@
 #include "MuR/MutableRuntimeModule.h"
 
 #include "Tasks/Task.h"
-#include "Misc/ScopeRWLock.h"
 
 #include <unordered_set>
 
@@ -767,7 +766,7 @@ namespace mu
 		if (bUseConcurrency)
 		{
 			/** Protect access to the original AST being optimized. */
-			FRWLock ASTAccessLock;
+			FCriticalSection ASTAccessLock;
 
 			// Launch the tasks.
 			UE::Tasks::FTask LaunchTask = UE::Tasks::Launch(TEXT("ConstantGeneratorLaunchTasks"), 
@@ -775,7 +774,7 @@ namespace mu
 				{
 					MUTABLE_CPUPROFILER_SCOPE(ConstantGenerator_LaunchTasks);
 
-					FReadScopeLock Lock(ASTAccessLock);
+					FScopeLock Lock(&ASTAccessLock);
 
 					FImageOperator ImOp = FImageOperator::GetDefault(InOptions->ImageFormatFunc);
 
@@ -801,19 +800,23 @@ namespace mu
 							TSharedPtr< Ptr<Image> > ResolveImage = MakeShared<Ptr<Image>>();
 
 							constexpr bool bRunImmediatlyIfPossible = false;
-							UE::Tasks::FTaskEvent ReferenceCompletionEvent = InOptions->OptimisationOptions.ReferencedResourceProvider(ImageID, ResolveImage, bRunImmediatlyIfPossible);
+							UE::Tasks::FTask ReferenceCompletion = InOptions->OptimisationOptions.ReferencedResourceProvider(ImageID, ResolveImage, bRunImmediatlyIfPossible);
 
 							UE::Tasks::FTask CompleteTask = UE::Tasks::Launch(TEXT("MutableResolveComplete"),
 								[SubgraphRoot, InOptions, ResolveImage, &ASTAccessLock]()
 								{
-									FWriteScopeLock Lock(ASTAccessLock);
+									MUTABLE_CPUPROFILER_SCOPE(MutableResolveComplete);
 
 									Ptr<ASTOpConstantResource> ConstantOp = new ASTOpConstantResource;
 									ConstantOp->Type = OP_TYPE::IM_CONSTANT;
 									ConstantOp->SetValue(ResolveImage->get(), InOptions->OptimisationOptions.DiskCacheContext);
-									ASTOp::Replace(SubgraphRoot, ConstantOp);
+
+									{
+										FScopeLock Lock(&ASTAccessLock);
+										ASTOp::Replace(SubgraphRoot, ConstantOp);
+									}
 								},
-								ReferenceCompletionEvent,
+								ReferenceCompletion,
 								LowLevelTasks::ETaskPriority::BackgroundNormal);
 
 							SubgraphCompletionEvent.AddPrerequisites(CompleteTask);
@@ -831,7 +834,8 @@ namespace mu
 							// Launch the preparation on the AST-modification pipe
 							UE::Tasks::FTask PrepareTask = UE::Tasks::Launch(TEXT("MutableConstantPrepare"), [TaskPtr, &ASTAccessLock]()
 								{
-									FReadScopeLock Lock(ASTAccessLock);
+									MUTABLE_CPUPROFILER_SCOPE(MutableConstantPrepare);
+									FScopeLock Lock(&ASTAccessLock);
 
 									// We need the clone because linking modifies ASTOp state and also to be safe for concurrency.
 									TaskPtr->SourceCloned = ASTOp::DeepClone(TaskPtr->Source);
@@ -850,11 +854,20 @@ namespace mu
 							// Launch the completion on the AST-modification pipe
 							UE::Tasks::FTask CompleteTask = UE::Tasks::Launch(TEXT("MutableConstantComplete"), [TaskPtr = MoveTemp(Task), &ASTAccessLock]()
 								{
-									FWriteScopeLock Lock(ASTAccessLock);
+									MUTABLE_CPUPROFILER_SCOPE(MutableConstantComplete);
+
+									{
+										MUTABLE_CPUPROFILER_SCOPE(Replace);
+										FScopeLock Lock(&ASTAccessLock);
+
+										ASTOp::Replace(TaskPtr->Source, TaskPtr->Result);
+										TaskPtr->Result = nullptr;
+									}
 
-									ASTOp::Replace(TaskPtr->Source, TaskPtr->Result);
-									TaskPtr->Source = nullptr;
-									TaskPtr->Result = nullptr;
+									{
+										MUTABLE_CPUPROFILER_SCOPE(Release);
+										TaskPtr->Source = nullptr;
+									}
 								},
 								RunTask,
 								LowLevelTasks::ETaskPriority::BackgroundHigh);
@@ -872,21 +885,8 @@ namespace mu
 
 			// Wait for pending tasks
 			{
-				MUTABLE_CPUPROFILER_SCOPE(ConstantGenerator_WaitPending);
-				if (!IsInGameThread())
-				{
-					LaunchTask.Wait();
-				}
-				else
-				{
-					const FReferencedResourceGameThreadTickFunc& ManualGameTick = InOptions->OptimisationOptions.ReferencedResourceProviderTick;
-
-					constexpr float FakeTimeStep = 0.1f;
-					while (!LaunchTask.Wait(FakeTimeStep))
-					{
-						ManualGameTick(FakeTimeStep);
-					}
-				}
+				MUTABLE_CPUPROFILER_SCOPE(Waiting);
+				LaunchTask.Wait();
 			}
 		}
 
@@ -921,8 +921,8 @@ namespace mu
 							{
 								TSharedPtr< Ptr<Image> > ResolveImage = MakeShared<Ptr<Image>>();
 								constexpr bool bRunImmediatlyIfPossible = true;
-								UE::Tasks::FTaskEvent ReferenceCompletionEvent = InOptions->OptimisationOptions.ReferencedResourceProvider(ImageID, ResolveImage, bRunImmediatlyIfPossible);
-								ReferenceCompletionEvent.Wait();
+								UE::Tasks::FTask ReferenceCompletion = InOptions->OptimisationOptions.ReferencedResourceProvider(ImageID, ResolveImage, bRunImmediatlyIfPossible);
+								ReferenceCompletion.Wait();
 								Ptr<ASTOpConstantResource> ConstantOp = new ASTOpConstantResource;
 								ConstantOp->Type = OP_TYPE::IM_CONSTANT;
 								ConstantOp->SetValue(ResolveImage->get(), InOptions->OptimisationOptions.DiskCacheContext);
diff --git a/Engine/Plugins/Experimental/Mutable/Source/MutableTools/Private/MuT/Compiler.cpp b/Engine/Plugins/Experimental/Mutable/Source/MutableTools/Private/MuT/Compiler.cpp
--- a/Engine/Plugins/Experimental/Mutable/Source/MutableTools/Private/MuT/Compiler.cpp
+++ b/Engine/Plugins/Experimental/Mutable/Source/MutableTools/Private/MuT/Compiler.cpp
@@ -169,10 +169,9 @@ namespace mu
 
 
 	//---------------------------------------------------------------------------------------------
-	void CompilerOptions::SetReferencedResourceCallback(const FReferencedResourceFunc& Provider, const FReferencedResourceGameThreadTickFunc& ProviderGameThreadTick)
+	void CompilerOptions::SetReferencedResourceCallback(const FReferencedResourceFunc& Provider)
 	{
 		m_pD->OptimisationOptions.ReferencedResourceProvider = Provider;
-		m_pD->OptimisationOptions.ReferencedResourceProviderTick = ProviderGameThreadTick;
 	}
 
 
diff --git a/Engine/Plugins/Experimental/Mutable/Source/MutableTools/Public/MuT/Compiler.h b/Engine/Plugins/Experimental/Mutable/Source/MutableTools/Public/MuT/Compiler.h
--- a/Engine/Plugins/Experimental/Mutable/Source/MutableTools/Public/MuT/Compiler.h
+++ b/Engine/Plugins/Experimental/Mutable/Source/MutableTools/Public/MuT/Compiler.h
@@ -32,8 +32,7 @@ namespace mu
     class Node;
     class NodeTransformedObject;
 
-	typedef TFunction<UE::Tasks::FTaskEvent(int32, TSharedPtr<Ptr<Image>>, bool)> FReferencedResourceFunc;
-	typedef TFunction<void(float)> FReferencedResourceGameThreadTickFunc;
+	typedef TFunction<UE::Tasks::FTask(int32, TSharedPtr<Ptr<Image>>, bool)> FReferencedResourceFunc;
 
     //! \brief Options used to compile the models with a compiler.
     class MUTABLETOOLS_API CompilerOptions : public RefCounted
@@ -84,7 +83,7 @@ namespace mu
 		void SetImagePixelFormatOverride(const FImageOperator::FImagePixelFormatFunc&);
 
 		/** */
-		void SetReferencedResourceCallback(const FReferencedResourceFunc&, const FReferencedResourceGameThreadTickFunc&);
+		void SetReferencedResourceCallback(const FReferencedResourceFunc&);
 
         //! Different data packing strategies
         enum class TextureLayoutStrategy : uint8
diff --git a/Engine/Source/Runtime/Engine/Classes/Engine/Texture.h b/Engine/Source/Runtime/Engine/Classes/Engine/Texture.h
--- a/Engine/Source/Runtime/Engine/Classes/Engine/Texture.h
+++ b/Engine/Source/Runtime/Engine/Classes/Engine/Texture.h
@@ -264,7 +264,7 @@ struct FTextureSource
 	);
 
 	/** Make a copy with a torn-off BulkData that has the same Guid used for DDC as this->BulkData */
-	FTextureSource CopyTornOff() const;
+	ENGINE_API FTextureSource CopyTornOff() const;
 	
 	/** PNG Compresses the source art if possible or tells the bulk data to zlib compress when it saves out to disk. */
 	ENGINE_API void Compress();
@@ -431,12 +431,12 @@ struct FTextureSource
 		// we need to support the old bulkdata code path (although previously storing these allocations as
 		// raw pointers would allow it to be assigned, this would most likely cause a mismatch in lock counts,
 		// either in FTextureSource or the underlying bulkdata and was never actually safe)
-		FMipAllocation(const FMipAllocation&) {}
-		FMipAllocation& operator =(const FMipAllocation&) { return *this; }
+		ENGINE_API FMipAllocation(const FMipAllocation&) {}
+		ENGINE_API FMipAllocation& operator =(const FMipAllocation&) { return *this; }
 
 		// We do allow rvalue assignment
-		FMipAllocation(FMipAllocation&&);
-		FMipAllocation& operator =(FMipAllocation&&);
+		ENGINE_API FMipAllocation(FMipAllocation&&);
+		ENGINE_API FMipAllocation& operator =(FMipAllocation&&);
 
 		~FMipAllocation() = default;
 
