[mutable] Improvements over Instance Baking comments:

diff --git a/Engine/Plugins/Experimental/Mutable/Source/CustomizableObject/Classes/MuCO/CustomizableObjectInstance.h b/Engine/Plugins/Experimental/Mutable/Source/CustomizableObject/Classes/MuCO/CustomizableObjectInstance.h
--- a/Engine/Plugins/Experimental/Mutable/Source/CustomizableObject/Classes/MuCO/CustomizableObjectInstance.h
+++ b/Engine/Plugins/Experimental/Mutable/Source/CustomizableObject/Classes/MuCO/CustomizableObjectInstance.h
@@ -43,10 +43,10 @@ enum class EPackageSaveResolutionType : uint8
 	/** The package got saved as a new file. */
 	NewFile,
 
-	/** The package was already present in disk so the old package was deleted and a new one was saved on its place */
+	/** The package was already present on disk so the old package was deleted and a new one was saved on its place */
 	Overriden,
 
-	/** Error type : An overriding was required but due to an error or lack of user permission it could not be done. */
+	/** Error type : An override was required but due to an error or lack of user permission it could not be done. */
 	UnableToOverride
 };
 
