[mutable] Added float slider decorators using the metadata of a parameter.
+ Now, we can add an image asset with the key "SliderImage" inside the ExtraAssets map of a parameter's metadata. This will generate a new type of widget, which is, a SSlider with the selected image as background.

#rnx

#virtualized

diff --git a/Engine/Plugins/Experimental/Mutable/Source/CustomizableObjectEditor/Private/MuCOE/CustomizableInstanceDetails.cpp b/Engine/Plugins/Experimental/Mutable/Source/CustomizableObjectEditor/Private/MuCOE/CustomizableInstanceDetails.cpp
--- a/Engine/Plugins/Experimental/Mutable/Source/CustomizableObjectEditor/Private/MuCOE/CustomizableInstanceDetails.cpp
+++ b/Engine/Plugins/Experimental/Mutable/Source/CustomizableObjectEditor/Private/MuCOE/CustomizableInstanceDetails.cpp
@@ -21,12 +21,14 @@
 #include "MuCOE/SCustomizableInstanceProperties.h"
 #include "MuCOE/UnrealEditorPortabilityHelpers.h"
 
+#include "Slate/DeferredCleanupSlateBrush.h"
 #include "SSearchableComboBox.h"
 #include "Widgets/Colors/SColorBlock.h"
 #include "Widgets/Colors/SColorPicker.h"
 #include "Widgets/Input/SButton.h"
 #include "Widgets/Input/SCheckBox.h"
 #include "Widgets/Input/SSearchBox.h"
+#include "Widgets/Input/SSlider.h"
 #include "Widgets/Input/STextComboBox.h"
 #include "Widgets/Layout/SExpandableArea.h"
 #include "Widgets/Layout/SUniformGridPanel.h"
@@ -35,6 +37,14 @@ class UObject;
 
 #define LOCTEXT_NAMESPACE "CustomizableInstanceDetails"
 
+// Define here metadata keywords used in the properties details
+namespace UIMetadataKeyWords
+{
+	// Key name for float slider decorators
+	constexpr TCHAR const* FloatDecoratorName = TEXT("SliderImage");
+}
+
+
 TAutoConsoleVariable<bool> CVarUseOldInstanceUI(
 	TEXT("mutable.UseOldInstanceUI"),
 	false,
@@ -560,10 +570,10 @@ struct FParameterInfo
 
 bool FCustomizableInstanceDetails::GenerateParametersView(IDetailCategoryBuilder& DetailsCategory)
 {
-	// Make this arrays locals
 	ParamChildren.Empty();
 	ParamHasParent.Empty();
 	GeneratedSections.Empty();
+	DynamicBrushes.Empty();
 
 	TArray<FParameterInfo> ParametersTree;
 	UCustomizableObject* CustomizableObject = CustomInstance->GetCustomizableObject();
@@ -1005,15 +1015,52 @@ void FCustomizableInstanceDetails::OnIntParameterComboBoxChanged(TSharedPtr<FStr
 TSharedRef<SWidget> FCustomizableInstanceDetails::GenerateFloatWidget(const int32 ParamIndexInObject)
 {
 	const UCustomizableObject* CustomizableObject = CustomInstance->GetCustomizableObject();
+	const FMutableParamUIMetadata& UIMetadata = CustomizableObject->GetParameterUIMetadataFromIndex(ParamIndexInObject).ParamUIMetadata;
 	FString ParamName = CustomizableObject->GetParameterName(ParamIndexInObject);
 
+	if (const TSoftObjectPtr<UObject>* FloatDecoratorAsset = UIMetadata.ExtraAssets.Find(UIMetadataKeyWords::FloatDecoratorName))
+	{
+		// Checking if there is an image decorator for the float slider:
+		if (UTexture2D* DecoratorTexture = Cast<UTexture2D>(FloatDecoratorAsset->LoadSynchronous()))
+		{
+			TSharedRef<FDeferredCleanupSlateBrush> Brush = FDeferredCleanupSlateBrush::CreateBrush(DecoratorTexture, FVector2D(DecoratorTexture->GetSizeX(), 1.0f), FLinearColor(1.0f, 1.0f, 1.0f, 1.0f), ESlateBrushTileType::NoTile, ESlateBrushImageType::Linear);
+			DynamicBrushes.Add(Brush);
+
+			// Add slider with the decorator as background
+			return SNew(SOverlay)
+			+ SOverlay::Slot()
+			[
+				SNew(SBox)
+				.Padding(0.0f, 8.0f) //Helps to shrinks the image from the borders
+				[
+					SNew(SImage)
+					.Image(Brush->GetSlateBrush())
+				]
+			]
+
+			+ SOverlay::Slot()
+			[
+				SNew(SSlider)
+				.Value(this, &FCustomizableInstanceDetails::GetFloatParameterValue, ParamName, -1)
+				.MinValue(UIMetadata.MinimumValue)
+				.MaxValue(UIMetadata.MaximumValue)
+				.OnValueChanged(this, &FCustomizableInstanceDetails::OnFloatParameterChanged, ParamName, -1)
+				.OnMouseCaptureBegin(this, &FCustomizableInstanceDetails::OnFloatParameterSliderBegin)
+				.OnMouseCaptureEnd(this, &FCustomizableInstanceDetails::OnFloatParameterSliderEnd)
+				.Style(&FAppStyle::Get().GetWidgetStyle<FSliderStyle>("ColorPicker.Slider"))
+				.IndentHandle(false)
+				.SliderBarColor(FLinearColor::Transparent)
+			];
+		}
+	}
+
 	return SNew(SSpinBox<float>)
 		.Value(this, &FCustomizableInstanceDetails::GetFloatParameterValue, ParamName, -1)
-		.MinValue(CustomizableObject->GetParameterUIMetadataFromIndex(ParamIndexInObject).ParamUIMetadata.MinimumValue)
-		.MaxValue(CustomizableObject->GetParameterUIMetadataFromIndex(ParamIndexInObject).ParamUIMetadata.MaximumValue)
+		.MinValue(UIMetadata.MinimumValue)
+		.MaxValue(UIMetadata.MaximumValue)
 		.OnValueChanged(this, &FCustomizableInstanceDetails::OnFloatParameterChanged, ParamName, -1)
 		.OnBeginSliderMovement(this, &FCustomizableInstanceDetails::OnFloatParameterSliderBegin)
-		.OnEndSliderMovement(this, &FCustomizableInstanceDetails::OnFloatParameterSliderEnd, ParamName, -1);
+		.OnEndSliderMovement(this, &FCustomizableInstanceDetails::OnFloatParameterSpinBoxEnd, ParamName, -1);
 }
 
 
@@ -1054,7 +1101,13 @@ void FCustomizableInstanceDetails::OnFloatParameterSliderBegin()
 }
 
 
-void FCustomizableInstanceDetails::OnFloatParameterSliderEnd(float Value, FString ParamName, int32 RangeIndex)
+void FCustomizableInstanceDetails::OnFloatParameterSliderEnd()
+{
+	bUpdatingSlider = false;
+}
+
+
+void FCustomizableInstanceDetails::OnFloatParameterSpinBoxEnd(float Value, FString ParamName, int32 RangeIndex)
 {
 	bUpdatingSlider = false;
 
@@ -1063,7 +1116,7 @@ void FCustomizableInstanceDetails::OnFloatParameterSliderEnd(float Value, FStrin
 	CustomInstance->SetFloatParameterSelectedOption(ParamName, Value, RangeIndex);
 	CustomInstance->UpdateSkeletalMeshAsync(true, true);
 	CustomInstance->PostEditChange();
-
+	
 	// Non-continuous change: collect garbage.
 	GEngine->ForceGarbageCollection();
 }
@@ -1682,7 +1735,7 @@ TSharedRef<SWidget> FCustomizableInstanceDetails::GenerateMultidimensionalProjec
 				.Value(this, &FCustomizableInstanceDetails::GetFloatParameterValue, OpacitySliderParamName, RangeIndex)
 				.OnValueChanged(this, &FCustomizableInstanceDetails::OnFloatParameterChanged, OpacitySliderParamName, RangeIndex)
 				.OnBeginSliderMovement(this, &FCustomizableInstanceDetails::OnFloatParameterSliderBegin)
-				.OnEndSliderMovement(this, &FCustomizableInstanceDetails::OnFloatParameterSliderEnd, OpacitySliderParamName, RangeIndex)
+				.OnEndSliderMovement(this, &FCustomizableInstanceDetails::OnFloatParameterSpinBoxEnd, OpacitySliderParamName, RangeIndex)
 				.Font(LayoutBuilder.Pin()->GetDetailFont())
 			]
 
diff --git a/Engine/Plugins/Experimental/Mutable/Source/CustomizableObjectEditor/Private/MuCOE/CustomizableInstanceDetails.h b/Engine/Plugins/Experimental/Mutable/Source/CustomizableObjectEditor/Private/MuCOE/CustomizableInstanceDetails.h
--- a/Engine/Plugins/Experimental/Mutable/Source/CustomizableObjectEditor/Private/MuCOE/CustomizableInstanceDetails.h
+++ b/Engine/Plugins/Experimental/Mutable/Source/CustomizableObjectEditor/Private/MuCOE/CustomizableInstanceDetails.h
@@ -84,7 +84,8 @@ private:
 	float GetFloatParameterValue(const FString ParamName, int32 RangeIndex) const;
 	void OnFloatParameterChanged(float Value, const FString ParamName, int32 RangeIndex);
 	void OnFloatParameterSliderBegin();
-	void OnFloatParameterSliderEnd(float Value, const FString ParamName, int32 RangeIndex);
+	void OnFloatParameterSpinBoxEnd(float Value, const FString ParamName, int32 RangeIndex);
+	void OnFloatParameterSliderEnd();
  
 	// Texture Parameters Functions
 	TSharedRef<SWidget> GenerateTextureWidget(const int32 ParamIndexInObject);
@@ -168,5 +169,8 @@ private:
 
 	/** True when a slider is being edited*/
 	bool bUpdatingSlider = false;
+
+	/** Array to store dynamic brushes. Neede because an image widget only stores a pointer to a Brush. */
+	TArray< TSharedPtr<class FDeferredCleanupSlateBrush> > DynamicBrushes;
 };
 
