[mutable] Fix binary indeterminism due to unused byte in some mesh buffers packed formats not being initialized.

- Also fix a bug that was duplicating the morph data in some cases.

#rnx

diff --git a/Engine/Plugins/Experimental/Mutable/Source/MutableRuntime/Private/MuR/OpMeshFormat.cpp b/Engine/Plugins/Experimental/Mutable/Source/MutableRuntime/Private/MuR/OpMeshFormat.cpp
--- a/Engine/Plugins/Experimental/Mutable/Source/MutableRuntime/Private/MuR/OpMeshFormat.cpp
+++ b/Engine/Plugins/Experimental/Mutable/Source/MutableRuntime/Private/MuR/OpMeshFormat.cpp
@@ -156,8 +156,8 @@ namespace mu
 						&&
 						sourceSemanticIndex == resultSemanticIndex);
 
-					int sourceElemSize = Source.GetElementSize(sourceBuffer);
-					const uint8_t* pSourceBuf = Source.GetBufferData(sourceBuffer);
+					int32 sourceElemSize = Source.GetElementSize(sourceBuffer);
+					const uint8* pSourceBuf = Source.GetBufferData(sourceBuffer);
 					pSourceBuf += sourceOffset;
 
 					// Copy element by element
@@ -165,7 +165,7 @@ namespace mu
 					{
 						if (resultFormat == sourceFormat && resultComponents == sourceComponents)
 						{
-							memcpy(pResultBuf, pSourceBuf, resultChannelSize);
+							FMemory::Memcpy(pResultBuf, pSourceBuf, resultChannelSize);
 						}
 						else if (resultFormat == MBF_PACKEDDIR8_W_TANGENTSIGN
 							||
@@ -175,7 +175,6 @@ namespace mu
 							check(resultComponents == 4);
 
 							// convert the 3 first components
-							//memcpy(pResultBuf, pSourceBuf, resultChannelSize);
 							for (int i = 0; i < 3; ++i)
 							{
 								if (i < sourceComponents)
@@ -191,10 +190,10 @@ namespace mu
 
 
 							// Add the tangent sign
-							uint8_t* pData = (uint8_t*)pResultBuf;
+							uint8* pData = reinterpret_cast<uint8*>(pResultBuf);
 
 							// Look for the full tangent space
-							int tanXBuf, tanXChan, tanYBuf, tanYChan, tanZBuf, tanZChan;
+							int32 tanXBuf, tanXChan, tanYBuf, tanYChan, tanZBuf, tanZChan;
 							Source.FindChannel(MBS_TANGENT, resultSemanticIndex, &tanXBuf, &tanXChan);
 							Source.FindChannel(MBS_BINORMAL, resultSemanticIndex, &tanYBuf, &tanYChan);
 							Source.FindChannel(MBS_NORMAL, resultSemanticIndex, &tanZBuf, &tanZChan);
@@ -211,7 +210,7 @@ namespace mu
 
 								FMatrix44f Mat(xIt.GetAsVec3f(), yIt.GetAsVec3f(), zIt.GetAsVec3f(), FVector3f(0, 0, 0));
 
-								uint8_t sign = 0;
+								uint8 sign = 0;
 								if (resultFormat == MBF_PACKEDDIR8_W_TANGENTSIGN)
 								{
 									sign = Mat.RotDeterminant() < 0 ? 0 : 255;
@@ -222,6 +221,11 @@ namespace mu
 								}
 								pData[3] = sign;
 							}
+							else
+							{
+								// At least initialize it to avoid garbage.
+								pData[3] = 0;
+							}
 						}
 						else
 						{
@@ -292,12 +296,12 @@ namespace mu
 	(
 		const FMeshBufferSet& Source,
 		FMeshBufferSet& Result,
-		bool keepSystemBuffers,
-		bool ignoreMissingChannels,
-		bool isVertexBuffer
+		bool bKeepSystemBuffers,
+		bool bIgnoreMissingChannels,
+		bool bIsVertexBuffer
 	)
 	{
-		if (ignoreMissingChannels)
+		if (bIgnoreMissingChannels)
 		{
 			// Remove from the result the channels that are not present in the source, and re-pack the
 			// offsets.
@@ -371,12 +375,14 @@ namespace mu
 
 
 		// Detect internal system buffers and clone them unmodified.
-		if (keepSystemBuffers)
+		if (bKeepSystemBuffers)
 		{
 			for (int b = 0; b < Source.GetBufferCount(); ++b)
 			{
+				bool bIsSystemBuffer = false;
+
 				// Detect system buffers and clone them unmodified.
-				if (Source.GetBufferChannelCount(b) > 0)
+				if (Source.GetBufferChannelCount(b) == 1)
 				{
 					EMeshBufferSemantic sourceSemantic;
 					int sourceSemanticIndex;
@@ -392,12 +398,17 @@ namespace mu
 					);
 
 					if (sourceSemantic == MBS_LAYOUTBLOCK
-						|| (isVertexBuffer && sourceSemantic == MBS_VERTEXINDEX)
-						)
+						||
+						(bIsVertexBuffer && sourceSemantic == MBS_VERTEXINDEX))
 					{
-						Result.AddBuffer(Source, b);
+						bIsSystemBuffer = true;
 					}
 				}
+
+				if (bIsSystemBuffer)
+				{
+					Result.AddBuffer(Source, b);
+				}
 			}
 		}
 
diff --git a/Engine/Plugins/Experimental/Mutable/Source/MutableTools/Private/MuT/ASTOpMeshFormat.cpp b/Engine/Plugins/Experimental/Mutable/Source/MutableTools/Private/MuT/ASTOpMeshFormat.cpp
--- a/Engine/Plugins/Experimental/Mutable/Source/MutableTools/Private/MuT/ASTOpMeshFormat.cpp
+++ b/Engine/Plugins/Experimental/Mutable/Source/MutableTools/Private/MuT/ASTOpMeshFormat.cpp
@@ -165,13 +165,13 @@ namespace
 		// Make a morph format by adding all the vertex channels from the base into a single
 		// vertex buffer, adding the vertex index channel
 
-		int offset = 0;
-		int numChannels = 0;
+		int32 offset = 0;
+		int32 numChannels = 0;
 		TArray<EMeshBufferSemantic> semantics;
-		TArray<int> semanticIndices;
+		TArray<int32> semanticIndices;
 		TArray<EMeshBufferFormat> formats;
-		TArray<int> components;
-		TArray<int> offsets;
+		TArray<int32> components;
+		TArray<int32> offsets;
 
 		// Add the vertex index channel
 		semantics.Add(MBS_VERTEXINDEX);
@@ -183,26 +183,25 @@ namespace
 		numChannels++;
 
 		// Add the vertex channels from the new format
-		for (int vb = 0; vb < pTargetFormat->GetVertexBuffers().GetBufferCount(); ++vb)
+		for (int32 vb = 0; vb < pTargetFormat->GetVertexBuffers().GetBufferCount(); ++vb)
 		{
-			for (int c = 0; c < pTargetFormat->GetVertexBuffers().GetBufferChannelCount(vb); ++c)
+			for (int32 c = 0; c < pTargetFormat->GetVertexBuffers().GetBufferChannelCount(vb); ++c)
 			{
 				// Channel info
 				EMeshBufferSemantic semantic;
 				int semanticIndex;
 				EMeshBufferFormat format;
 				int component;
-				pTargetFormat->GetVertexBuffers().GetChannel
-				(vb, c, &semantic, &semanticIndex, &format, &component, nullptr);
+				pTargetFormat->GetVertexBuffers().GetChannel(vb, c, &semantic, &semanticIndex, &format, &component, nullptr);
 
 				// TODO: Filter useless semantics for morphing.
+				// Maybe some formats like the ones with a packed tangent sign need to be tweaked here, to make sense of the whole buffer.
 				semantics.Add(semantic);
 				semanticIndices.Add(semanticIndex);
 				formats.Add(format);
 				components.Add(component);
 				offsets.Add(offset);
-				offset += components[numChannels]
-					* GetMeshFormatData(formats[numChannels]).SizeInBytes;
+				offset += components[numChannels] * GetMeshFormatData(formats[numChannels]).SizeInBytes;
 				numChannels++;
 			}
 		}
@@ -277,29 +276,27 @@ mu::Ptr<ASTOp> Sink_MeshFormatAST::Visit(const mu::Ptr<ASTOp>& at, const ASTOpMe
 	case OP_TYPE::ME_MORPH:
 	{
 		// Move the format down the base of the morph
-		Ptr<ASTOpMeshMorph> newOp = mu::Clone<ASTOpMeshMorph>(at);
-		newOp->Base = Visit(newOp->Base.child(), currentFormatOp);
+		Ptr<ASTOpMeshMorph> NewOp = mu::Clone<ASTOpMeshMorph>(at);
+		NewOp->Base = Visit(NewOp->Base.child(), currentFormatOp);
 
 		// Reformat the morph targets to match the new format.
-		// \TODO: Cache pTargetMorphFormat? motaop?
 		MeshPtrConst pTargetFormat = FindBaseMeshConstant(currentFormatOp->Format.child());
 		MeshPtrConst pTargetMorphFormat = MakeMorphTargetFormat(pTargetFormat);
 
-		mu::Ptr<ASTOpConstantResource> motaop = new ASTOpConstantResource();
-		motaop->Type = OP_TYPE::ME_CONSTANT;
-		motaop->SetValue(pTargetMorphFormat, nullptr);
-		auto targetMorphFormatAt = motaop;
+		mu::Ptr<ASTOpConstantResource> NewFormatConstant = new ASTOpConstantResource();
+		NewFormatConstant->Type = OP_TYPE::ME_CONSTANT;
+		NewFormatConstant->SetValue(pTargetMorphFormat, nullptr);
 
-		if (newOp->Target)
+		if (NewOp->Target)
 		{
 			mu::Ptr<ASTOpMeshFormat> newFormat = mu::Clone<ASTOpMeshFormat>(currentFormatOp);
 			newFormat->Flags = OP::MeshFormatArgs::Vertex | OP::MeshFormatArgs::IgnoreMissing;
-			newFormat->Format = targetMorphFormatAt;
+			newFormat->Format = NewFormatConstant;
 
-			newOp->Target = Visit(newOp->Target.child(), newFormat.get());
+			NewOp->Target = Visit(NewOp->Target.child(), newFormat.get());
 		}
 
-		newAt = newOp;
+		newAt = NewOp;
 		break;
 	}
 
