Exported 4x4 and 8x8 RGBA ASTC compression functions

diff --git a/Engine/Plugins/Experimental/Mutable/Source/MutableRuntime/Internal/MuR/BlockCompression/Miro/Miro.h b/Engine/Plugins/Experimental/Mutable/Source/MutableRuntime/Internal/MuR/BlockCompression/Miro/Miro.h
--- a/Engine/Plugins/Experimental/Mutable/Source/MutableRuntime/Internal/MuR/BlockCompression/Miro/Miro.h
+++ b/Engine/Plugins/Experimental/Mutable/Source/MutableRuntime/Internal/MuR/BlockCompression/Miro/Miro.h
@@ -78,7 +78,7 @@ namespace miro
 #if MIRO_INCLUDE_ASTC
 
     //! ASTC 4x4 support
-    extern void RGBA_to_ASTC4x4RGBAL( uint32 sx, uint32 sy, const uint8* From, uint8* To, int32 Quality );
+	MUTABLERUNTIME_API extern void RGBA_to_ASTC4x4RGBAL( uint32 sx, uint32 sy, const uint8* From, uint8* To, int32 Quality );
     extern void RGB_to_ASTC4x4RGBAL( uint32 sx, uint32 sy, const uint8* From, uint8* To, int32 Quality );
     extern void ASTC4x4RGBAL_to_RGBA( uint32 sx, uint32 sy, const uint8* From, uint8* To );
     extern void ASTC4x4RGBAL_to_RGB( uint32 sx, uint32 sy, const uint8* From, uint8* To );
@@ -114,7 +114,7 @@ namespace miro
 	extern void ASTC6x6RGL_to_RGBA(uint32 sx, uint32 sy, const uint8* From, uint8* To);
 
 	//! ASTC 8x8 support
-	extern void RGBA_to_ASTC8x8RGBAL(uint32 sx, uint32 sy, const uint8* From, uint8* To, int32 Quality);
+	MUTABLERUNTIME_API extern void RGBA_to_ASTC8x8RGBAL(uint32 sx, uint32 sy, const uint8* From, uint8* To, int32 Quality);
 	extern void RGB_to_ASTC8x8RGBAL(uint32 sx, uint32 sy, const uint8* From, uint8* To, int32 Quality);
 	extern void ASTC8x8RGBAL_to_RGBA(uint32 sx, uint32 sy, const uint8* From, uint8* To);
 	extern void ASTC8x8RGBAL_to_RGB(uint32 sx, uint32 sy, const uint8* From, uint8* To);
