Autortfm-ifiy Module

diff --git a/Engine/Plugins/Experimental/Mutable/Source/CustomizableObject/CustomizableObject.Build.cs b/Engine/Plugins/Experimental/Mutable/Source/CustomizableObject/CustomizableObject.Build.cs
--- a/Engine/Plugins/Experimental/Mutable/Source/CustomizableObject/CustomizableObject.Build.cs
+++ b/Engine/Plugins/Experimental/Mutable/Source/CustomizableObject/CustomizableObject.Build.cs
@@ -41,6 +41,8 @@ public class CustomizableObject : ModuleRules
 			"StructUtils"
 		});
 
+		bAllowAutoRTFMInstrumentation = true;
+
 		if (TargetRules.bBuildEditor == true)
         {
             PublicDependencyModuleNames.Add("UnrealEd");	// @todo api: Only public because of WITH_EDITOR
