[Mutable] Convert the SkinWeightProfile semantic index to a hash of the Profile Name to fix indeterminism.

#rnx

diff --git a/Engine/Plugins/Experimental/Mutable/Source/CustomizableObject/Internal/MuCO/CustomizableObjectPrivate.h b/Engine/Plugins/Experimental/Mutable/Source/CustomizableObject/Internal/MuCO/CustomizableObjectPrivate.h
--- a/Engine/Plugins/Experimental/Mutable/Source/CustomizableObject/Internal/MuCO/CustomizableObjectPrivate.h
+++ b/Engine/Plugins/Experimental/Mutable/Source/CustomizableObject/Internal/MuCO/CustomizableObjectPrivate.h
@@ -463,12 +463,15 @@ struct CUSTOMIZABLEOBJECT_API FMutableSkinWeightProfileInfo
 
 	FMutableSkinWeightProfileInfo() {};
 
-	FMutableSkinWeightProfileInfo(FName InName, bool InDefaultProfile, int8 InDefaultProfileFromLODIndex) : Name(InName),
-	DefaultProfile(InDefaultProfile), DefaultProfileFromLODIndex(InDefaultProfileFromLODIndex) {};
+	FMutableSkinWeightProfileInfo(FName InName, uint32 InNameId, bool InDefaultProfile, int8 InDefaultProfileFromLODIndex) : Name(InName),
+		NameId(InNameId), DefaultProfile(InDefaultProfile), DefaultProfileFromLODIndex(InDefaultProfileFromLODIndex) {};
 
 	UPROPERTY()
 	FName Name;
 
+	UPROPERTY()
+	uint32 NameId = 0;
+
 	UPROPERTY()
 	bool DefaultProfile = false;
 
@@ -830,6 +833,6 @@ public:
 	// This is a manual version number for the binary blobs in this asset.
 	// Increasing it invalidates all the previously compiled models.
 	// Warning: If while merging code both versions have changed, take the highest+1.
-	static constexpr int32 CurrentSupportedVersion = 438;
+	static constexpr int32 CurrentSupportedVersion = 439;
 };
 
diff --git a/Engine/Plugins/Experimental/Mutable/Source/CustomizableObject/Private/MuCO/CustomizableObject.cpp b/Engine/Plugins/Experimental/Mutable/Source/CustomizableObject/Private/MuCO/CustomizableObject.cpp
--- a/Engine/Plugins/Experimental/Mutable/Source/CustomizableObject/Private/MuCO/CustomizableObject.cpp
+++ b/Engine/Plugins/Experimental/Mutable/Source/CustomizableObject/Private/MuCO/CustomizableObject.cpp
@@ -2494,6 +2494,7 @@ FArchive& operator<<(FArchive& Ar, FMutableRefAssetUserData& Data)
 FArchive& operator<<(FArchive& Ar, FMutableSkinWeightProfileInfo& Info)
 {
 	Ar << Info.Name;
+	Ar << Info.NameId;
 	Ar << Info.DefaultProfile;
 	Ar << Info.DefaultProfileFromLODIndex;
 
diff --git a/Engine/Plugins/Experimental/Mutable/Source/CustomizableObject/Private/MuCO/CustomizableObjectInstance.cpp b/Engine/Plugins/Experimental/Mutable/Source/CustomizableObject/Private/MuCO/CustomizableObjectInstance.cpp
--- a/Engine/Plugins/Experimental/Mutable/Source/CustomizableObject/Private/MuCO/CustomizableObjectInstance.cpp
+++ b/Engine/Plugins/Experimental/Mutable/Source/CustomizableObject/Private/MuCO/CustomizableObjectInstance.cpp
@@ -4730,46 +4730,53 @@ bool UCustomizableInstancePrivate::BuildOrCopyRenderData(const TSharedRef<FUpdat
 			// Copy SkinWeightProfiles
 			if (!ModelResources.SkinWeightProfilesInfo.IsEmpty())
 			{
-				bool bHasSkinWeightProfiles = false;
-
 				const mu::FMeshBufferSet& MutableMeshVertexBuffers = Component.Mesh->GetVertexBuffers();
 
-				const int32 SkinWeightProfilesCount = ModelResources.SkinWeightProfilesInfo.Num();
-				for (int32 ProfileIndex = 0; ProfileIndex < SkinWeightProfilesCount; ++ProfileIndex)
+				bool bHasSkinWeightProfiles = false;
+			
+				const int32 NumBuffers = MutableMeshVertexBuffers.GetBufferCount();
+				for (int32 BufferIndex = 0; BufferIndex < NumBuffers; ++BufferIndex)
 				{
-					const int32 ProfileSemanticsIndex = ProfileIndex + 10;
-					int32 BoneIndicesBufferIndex, BoneIndicesBufferChannelIndex;
-					MutableMeshVertexBuffers.FindChannel(mu::MBS_BONEINDICES, ProfileSemanticsIndex, &BoneIndicesBufferIndex, &BoneIndicesBufferChannelIndex);
-
-					int32 BoneWeightsBufferIndex, BoneWeightsBufferChannelIndex;
-					MutableMeshVertexBuffers.FindChannel(mu::MBS_BONEWEIGHTS, ProfileSemanticsIndex, &BoneWeightsBufferIndex, &BoneWeightsBufferChannelIndex);
-
-					if (BoneIndicesBufferIndex < 0 || BoneIndicesBufferIndex != BoneWeightsBufferIndex)
+					// Skin Weight Profiles have 3 channels
+					if (MutableMeshVertexBuffers.GetBufferChannelCount(BufferIndex) != 3)
 					{
 						continue;
 					}
 
-					if (!bHasSkinWeightProfiles)
+					mu::EMeshBufferSemantic Semantic;
+					int32 SemanticIndex;
+					MutableMeshVertexBuffers.GetChannel(BufferIndex, 0, &Semantic, &SemanticIndex, nullptr, nullptr, nullptr);
+
+					if (Semantic != mu::EMeshBufferSemantic::MBS_ALTSKINWEIGHT)
 					{
-						LODResource.SkinWeightProfilesData.Init(&LODResource.SkinWeightVertexBuffer);
-						bHasSkinWeightProfiles = true;
+						continue;
 					}
 
-					const FMutableSkinWeightProfileInfo& Profile = ModelResources.SkinWeightProfilesInfo[ProfileIndex];
-
-					const FSkinWeightProfileInfo* ExistingProfile = SkeletalMesh->GetSkinWeightProfiles().FindByPredicate(
-						[&Profile](const FSkinWeightProfileInfo& P) { return P.Name == Profile.Name; });
+					const FMutableSkinWeightProfileInfo* ProfileInfo = ModelResources.SkinWeightProfilesInfo.FindByPredicate(
+							[&SemanticIndex](const FMutableSkinWeightProfileInfo& P) { return P.NameId == SemanticIndex; });
 
-					if (!ExistingProfile)
+					if (ensure(ProfileInfo))
 					{
-						SkeletalMesh->AddSkinWeightProfile({ Profile.Name, Profile.DefaultProfile, Profile.DefaultProfileFromLODIndex });
-					}
+						if (!bHasSkinWeightProfiles)
+						{
+							LODResource.SkinWeightProfilesData.Init(&LODResource.SkinWeightVertexBuffer);
+							bHasSkinWeightProfiles = true;
+						}
+						
+						const FSkinWeightProfileInfo* ExistingProfile = SkeletalMesh->GetSkinWeightProfiles().FindByPredicate(
+							[&ProfileInfo](const FSkinWeightProfileInfo& P) { return P.Name == ProfileInfo->Name; });
 
-					UnrealConversionUtils::CopyMutableSkinWeightProfilesBuffers(
-						LODResource,
-						Profile.Name,
-						MutableMeshVertexBuffers,
-						BoneIndicesBufferIndex);
+						if (!ExistingProfile)
+						{
+							SkeletalMesh->AddSkinWeightProfile({ ProfileInfo->Name, ProfileInfo->DefaultProfile, ProfileInfo->DefaultProfileFromLODIndex });
+						}
+
+						UnrealConversionUtils::CopyMutableSkinWeightProfilesBuffers(
+							LODResource,
+							ProfileInfo->Name,
+							MutableMeshVertexBuffers,
+							BufferIndex);
+					}
 				}
 			}
 		}
diff --git a/Engine/Plugins/Experimental/Mutable/Source/CustomizableObject/Private/MuCO/MutableMeshBufferUtils.cpp b/Engine/Plugins/Experimental/Mutable/Source/CustomizableObject/Private/MuCO/MutableMeshBufferUtils.cpp
--- a/Engine/Plugins/Experimental/Mutable/Source/CustomizableObject/Private/MuCO/MutableMeshBufferUtils.cpp
+++ b/Engine/Plugins/Experimental/Mutable/Source/CustomizableObject/Private/MuCO/MutableMeshBufferUtils.cpp
@@ -195,7 +195,7 @@ void MutableMeshBufferUtils::SetupSkinWeightProfileBuffer(const int32& InCurrent
 	using namespace mu;
 	const int32 ElementSize = sizeof(int32) + (MaxBoneIndexTypeSizeBytes + 1) * MaxNumBonesPerVertex;
 	constexpr int32 ChannelCount = 3;
-	const EMeshBufferSemantic Semantics[ChannelCount] = { MBS_OTHER, MBS_BONEINDICES, MBS_BONEWEIGHTS };
+	const EMeshBufferSemantic Semantics[ChannelCount] = { MBS_ALTSKINWEIGHT, MBS_BONEINDICES, MBS_BONEWEIGHTS };
 	const int32 SemanticIndices[ChannelCount] = { SemanticsIndex, SemanticsIndex, SemanticsIndex };
 
 	EMeshBufferFormat Formats[ChannelCount] = { MBF_INT32, MBF_UINT8, MBF_NUINT8 };
diff --git a/Engine/Plugins/Experimental/Mutable/Source/CustomizableObjectEditor/Private/MuCOE/GenerateMutableSource/GenerateMutableSource.cpp b/Engine/Plugins/Experimental/Mutable/Source/CustomizableObjectEditor/Private/MuCOE/GenerateMutableSource/GenerateMutableSource.cpp
--- a/Engine/Plugins/Experimental/Mutable/Source/CustomizableObjectEditor/Private/MuCOE/GenerateMutableSource/GenerateMutableSource.cpp
+++ b/Engine/Plugins/Experimental/Mutable/Source/CustomizableObjectEditor/Private/MuCOE/GenerateMutableSource/GenerateMutableSource.cpp
@@ -52,6 +52,7 @@
 #include "PhysicsEngine/PhysicsAsset.h"
 #include "PlatformInfo.h"
 #include "Math/NumericLimits.h"
+#include "Hash/CityHash.h"
 
 #define LOCTEXT_NAMESPACE "CustomizableObjectEditor"
 
@@ -220,6 +221,39 @@ void FMutableGraphGenerationContext::AddParameterNameUnique(const UCustomizableO
 	}
 }
 
+uint32 FMutableGraphGenerationContext::GetSkinWeightProfileIdUnique(const FName InProfileName)
+{
+	const FString ProfileNameString = InProfileName.ToString();
+	
+	uint32 UniqueProfileId = CityHash32(reinterpret_cast<const char*>(*ProfileNameString), ProfileNameString.Len() * sizeof(FString::ElementType));
+
+	bool bRemappedProfile = false;
+
+	FName& ProfileName = UniqueSkinWeightProfileIds.FindOrAdd(UniqueProfileId, InProfileName);
+	while (ProfileName != InProfileName)
+	{
+		if (uint32* RemappedProfileId = RemappedSkinWeightProfileIds.Find(InProfileName))
+		{
+			UniqueProfileId = *RemappedProfileId;
+			break;
+		}
+
+		// Id collision detected
+		bRemappedProfile = true;
+
+		// Increase Id in an attempt to make it unique again.
+		++UniqueProfileId;
+		ProfileName = UniqueSkinWeightProfileIds.FindOrAdd(UniqueProfileId, InProfileName);
+	}
+
+	if (bRemappedProfile)
+	{
+		RemappedSkinWeightProfileIds.Add(InProfileName, UniqueProfileId);
+	}
+
+	return UniqueProfileId;
+}
+
 
 const FGuid FMutableGraphGenerationContext::GetNodeIdUnique(const UCustomizableObjectNode* Node)
 {
diff --git a/Engine/Plugins/Experimental/Mutable/Source/CustomizableObjectEditor/Private/MuCOE/GenerateMutableSource/GenerateMutableSource.h b/Engine/Plugins/Experimental/Mutable/Source/CustomizableObjectEditor/Private/MuCOE/GenerateMutableSource/GenerateMutableSource.h
--- a/Engine/Plugins/Experimental/Mutable/Source/CustomizableObjectEditor/Private/MuCOE/GenerateMutableSource/GenerateMutableSource.h
+++ b/Engine/Plugins/Experimental/Mutable/Source/CustomizableObjectEditor/Private/MuCOE/GenerateMutableSource/GenerateMutableSource.h
@@ -647,6 +647,8 @@ struct FMutableGraphGenerationContext
 	/** Adds to ParameterNamesMap the node Node to the array of elements with name Name */
 	void AddParameterNameUnique(const UCustomizableObjectNode* Node, FString Name);
 
+	uint32 GetSkinWeightProfileIdUnique(const FName ProfileName);
+
 	// Check if the Id of the node Node already exists, if it's new adds it to NodeIds array, otherwise, returns new Id
 	const FGuid GetNodeIdUnique(const UCustomizableObjectNode* Node);
 
@@ -711,6 +713,9 @@ struct FMutableGraphGenerationContext
 	// Data used for SkinWeightProfiles reconstruction
 	TArray<FMutableSkinWeightProfileInfo> SkinWeightProfilesInfo;
 
+	TMap<uint32, FName> UniqueSkinWeightProfileIds;
+	TMap<FName, uint32> RemappedSkinWeightProfileIds;
+
 	TArray<FAnimBpOverridePhysicsAssetsInfo> AnimBpOverridePhysicsAssetsInfo;
 
 	// Hierarchy of current ComponentNew nodes, each stored for every LOD
diff --git a/Engine/Plugins/Experimental/Mutable/Source/CustomizableObjectEditor/Private/MuCOE/GenerateMutableSource/GenerateMutableSourceMesh.cpp b/Engine/Plugins/Experimental/Mutable/Source/CustomizableObjectEditor/Private/MuCOE/GenerateMutableSource/GenerateMutableSourceMesh.cpp
--- a/Engine/Plugins/Experimental/Mutable/Source/CustomizableObjectEditor/Private/MuCOE/GenerateMutableSource/GenerateMutableSourceMesh.cpp
+++ b/Engine/Plugins/Experimental/Mutable/Source/CustomizableObjectEditor/Private/MuCOE/GenerateMutableSource/GenerateMutableSourceMesh.cpp
@@ -1485,19 +1485,19 @@ mu::MeshPtr ConvertSkeletalMeshToMutable(const USkeletalMesh* InSkeletalMesh, co
 				MutSkinWeightData += BoneWeightsSize;
 			}
 
-			const int32 ProfileIndex = GenerationContext.SkinWeightProfilesInfo.AddUnique({Profile.Name, false, 0});
-			const int32 ProfileSemanticIndex = ProfileIndex + 10;
+			const uint32 ProfileId = GenerationContext.GetSkinWeightProfileIdUnique(Profile.Name);
+			const int32 ProfileIndex = GenerationContext.SkinWeightProfilesInfo.AddUnique({Profile.Name, ProfileId, false, 0});
+			FMutableSkinWeightProfileInfo& MutSkinWeightProfileInfo = GenerationContext.SkinWeightProfilesInfo[ProfileIndex];
 
 			const FName PlatformName = *GenerationContext.Options.TargetPlatform->PlatformName();
-			FMutableSkinWeightProfileInfo& MutSkinWeightProfileInfo = GenerationContext.SkinWeightProfilesInfo[ProfileIndex];
 			MutSkinWeightProfileInfo.DefaultProfile = MutSkinWeightProfileInfo.DefaultProfile || Profile.DefaultProfile.GetValueForPlatform(PlatformName);
 			MutSkinWeightProfileInfo.DefaultProfileFromLODIndex = FMath::Min(MutSkinWeightProfileInfo.DefaultProfileFromLODIndex, Profile.DefaultProfileFromLODIndex.GetValueForPlatform(PlatformName));
 
 			// Set up SkinWeightPRofile BufferData
 			const int32 ElementSize = sizeof(int32) + sizeof(FBoneIndexType) + BoneWeightTypeSizeBytes;
 			const int32 ChannelCount = 3;
-			const EMeshBufferSemantic Semantics[ChannelCount] = { MBS_OTHER, MBS_BONEINDICES, MBS_BONEWEIGHTS };
-			const int32 SemanticIndices[ChannelCount] = { ProfileSemanticIndex, ProfileSemanticIndex, ProfileSemanticIndex };
+			const EMeshBufferSemantic Semantics[ChannelCount] = { MBS_ALTSKINWEIGHT, MBS_BONEINDICES, MBS_BONEWEIGHTS };
+			const int32 SemanticIndices[ChannelCount] = { static_cast<int32>(ProfileId), static_cast<int32>(ProfileId), static_cast<int32>(ProfileId) };
 			const EMeshBufferFormat Formats[ChannelCount] = { MBF_INT32, MBF_UINT16, BoneWeightFormat };
 			const int32 Components[ChannelCount] = { 1, MutableBonesPerVertex, MutableBonesPerVertex };
 			const int32 Offsets[ChannelCount] = { 0, sizeof(int32), sizeof(int32) + BoneIndicesSize };
@@ -3002,14 +3002,10 @@ mu::NodeMeshPtr GenerateMutableSourceMesh(const UEdGraphPin* Pin,
 				{
 					FSkeletalMeshModel* ImportModel = TypedNodeSkel->SkeletalMesh->GetImportedModel();
 
-					const int32 SkinWeightProfilesCount = GenerationContext.SkinWeightProfilesInfo.Num();
-					for (int32 ProfileIndex = 0; ProfileIndex < SkinWeightProfilesCount; ++ProfileIndex)
+					for (const auto& SkinWeightProfile : ImportModel->LODModels[LODIndex].SkinWeightProfiles)
 					{
-						if (ImportModel->LODModels[LODIndex].SkinWeightProfiles.Find(GenerationContext.SkinWeightProfilesInfo[ProfileIndex].Name))
-						{
-							const int32 ProfileScemanticIndex = ProfileIndex + 10;
-							MeshData.SkinWeightProfilesSemanticIndices.AddUnique(ProfileScemanticIndex);
-						}
+						const int32 ProfileId = static_cast<int32>(GenerationContext.GetSkinWeightProfileIdUnique(SkinWeightProfile.Key));
+						MeshData.SkinWeightProfilesSemanticIndices.AddUnique(ProfileId);
 					}
 				}
 
diff --git a/Engine/Plugins/Experimental/Mutable/Source/MutableRuntime/Private/MuR/MeshBufferSet.cpp b/Engine/Plugins/Experimental/Mutable/Source/MutableRuntime/Private/MuR/MeshBufferSet.cpp
--- a/Engine/Plugins/Experimental/Mutable/Source/MutableRuntime/Private/MuR/MeshBufferSet.cpp
+++ b/Engine/Plugins/Experimental/Mutable/Source/MutableRuntime/Private/MuR/MeshBufferSet.cpp
@@ -199,7 +199,7 @@ namespace mu
 		{
 			FMeshBufferChannel& chan = buf.m_channels[c];
 			chan.m_semantic = pSemantics ? pSemantics[c] : MBS_NONE;
-			chan.m_semanticIndex = pSemanticIndices ? ((uint8_t)pSemanticIndices[c]) : 0;
+			chan.m_semanticIndex = pSemanticIndices ? pSemanticIndices[c] : 0;
 			chan.m_format = pFormats ? pFormats[c] : MBF_NONE;
 			chan.m_componentCount = pComponentCount ? ((uint16)pComponentCount[c]) : 0;
 			chan.m_offset = pOffsets ? ((uint8_t)pOffsets[c]) : 0;
@@ -245,7 +245,7 @@ namespace mu
 
 		FMeshBufferChannel& chan = buf.m_channels[channelIndex];
 		chan.m_semantic = semantic;
-		chan.m_semanticIndex = uint8_t(semanticIndex);
+		chan.m_semanticIndex = semanticIndex;
 		chan.m_format = format;
 		chan.m_componentCount = uint16(componentCount);
 		chan.m_offset = uint8(offset);
@@ -621,7 +621,7 @@ namespace mu
 	//-----------------------------------------------------------------------------------------
 	void FMeshBufferSet::ResetBufferIndices()
 	{
-		uint8_t currentIndices[MBS_COUNT];
+		int32 currentIndices[MBS_COUNT];
 		memset(currentIndices, 0, sizeof(currentIndices));
 
 		for (FMeshBuffer& b : m_buffers)
diff --git a/Engine/Plugins/Experimental/Mutable/Source/MutableRuntime/Public/MuR/MeshBufferSet.h b/Engine/Plugins/Experimental/Mutable/Source/MutableRuntime/Public/MuR/MeshBufferSet.h
--- a/Engine/Plugins/Experimental/Mutable/Source/MutableRuntime/Public/MuR/MeshBufferSet.h
+++ b/Engine/Plugins/Experimental/Mutable/Source/MutableRuntime/Public/MuR/MeshBufferSet.h
@@ -120,6 +120,9 @@ namespace mu
 		MBS_BARYCENTRICCOORDS,
 		MBS_DISTANCE,
 
+		//! Semantics useful for alternative skin weight profiles.
+		MBS_ALTSKINWEIGHT,
+
 		//! Utility
 		MBS_COUNT,
 
@@ -134,9 +137,9 @@ namespace mu
 		{
 			m_semantic = MBS_NONE;
 			m_format = MBF_NONE;
-			m_componentCount = 0;
 			m_semanticIndex = 0;
 			m_offset = 0;
+			m_componentCount = 0;
 		}
 
 		//!
@@ -146,10 +149,10 @@ namespace mu
 		EMeshBufferFormat m_format;
 
 		//! Index of the semantic, in case there are more than one of this type.
-		uint8 m_semanticIndex;
+		int32 m_semanticIndex;
 
 		//! Offset in bytes from the begining of a buffer element
-		uint8 m_offset;
+		uint16 m_offset;
 
 		//! Number of components of the type in m_format for every value in the channel
 		uint16 m_componentCount;
