[mutable] Fix the check for missing GUIDs generating indeterministic packaged data,
#rnx

diff --git a/Engine/Plugins/Experimental/Mutable/Source/CustomizableObjectEditor/Classes/MuCOE/Nodes/CustomizableObjectNodeObject.h b/Engine/Plugins/Experimental/Mutable/Source/CustomizableObjectEditor/Classes/MuCOE/Nodes/CustomizableObjectNodeObject.h
--- a/Engine/Plugins/Experimental/Mutable/Source/CustomizableObjectEditor/Classes/MuCOE/Nodes/CustomizableObjectNodeObject.h
+++ b/Engine/Plugins/Experimental/Mutable/Source/CustomizableObjectEditor/Classes/MuCOE/Nodes/CustomizableObjectNodeObject.h
@@ -201,9 +201,6 @@ public:
 	UPROPERTY(EditAnywhere, Category = CustomizableObject)
 	TArray<FComponentSettings> ComponentSettings;
 	
-    // To avoid any no properly saved GUIDs
-	FGuid IdentifierVerification;
-
 	// UObject interface.
 	void PostEditChangeProperty(FPropertyChangedEvent& PropertyChangedEvent) override;
 
diff --git a/Engine/Plugins/Experimental/Mutable/Source/CustomizableObjectEditor/Private/MuCOE/GenerateMutableSource/GenerateMutableSource.cpp b/Engine/Plugins/Experimental/Mutable/Source/CustomizableObjectEditor/Private/MuCOE/GenerateMutableSource/GenerateMutableSource.cpp
--- a/Engine/Plugins/Experimental/Mutable/Source/CustomizableObjectEditor/Private/MuCOE/GenerateMutableSource/GenerateMutableSource.cpp
+++ b/Engine/Plugins/Experimental/Mutable/Source/CustomizableObjectEditor/Private/MuCOE/GenerateMutableSource/GenerateMutableSource.cpp
@@ -791,7 +791,12 @@ mu::NodeObjectPtr GenerateMutableSource(const UEdGraphPin * Pin, FMutableGraphGe
 		Result = ObjectNode;
 
 		ObjectNode->SetName(TypedNodeObj->ObjectName);
-		ObjectNode->SetUid(GenerationContext.GetNodeIdUnique(TypedNodeObj).ToString());
+		FGuid FinalGuid = GenerationContext.GetNodeIdUnique(TypedNodeObj);
+		if (FinalGuid != TypedNodeObj->NodeGuid)
+		{
+			GenerationContext.Compiler->CompilerLog(FText::FromString(TEXT("Warning: Node has a duplicated GUID. A new ID has been generated, but cooked data will not be deterministic.")), Node, EMessageSeverity::Warning);
+		}
+		ObjectNode->SetUid(FinalGuid.ToString());
 
 		// LOD
 		const int32 NumLODs = TypedNodeObj->GetNumLODPins();
diff --git a/Engine/Plugins/Experimental/Mutable/Source/CustomizableObjectEditor/Private/MuCOE/Nodes/CustomizableObjectNodeObject.cpp b/Engine/Plugins/Experimental/Mutable/Source/CustomizableObjectEditor/Private/MuCOE/Nodes/CustomizableObjectNodeObject.cpp
--- a/Engine/Plugins/Experimental/Mutable/Source/CustomizableObjectEditor/Private/MuCOE/Nodes/CustomizableObjectNodeObject.cpp
+++ b/Engine/Plugins/Experimental/Mutable/Source/CustomizableObjectEditor/Private/MuCOE/Nodes/CustomizableObjectNodeObject.cpp
@@ -2,8 +2,6 @@
 
 #include "MuCOE/Nodes/CustomizableObjectNodeObject.h"
 
-#include "Containers/Queue.h"
-#include "Logging/MessageLog.h"
 #include "MuCO/ICustomizableObjectModule.h"
 #include "MuCO/CustomizableObjectExtension.h"
 #include "MuCO/CustomizableObjectPrivate.h"
@@ -14,6 +12,9 @@
 #include "MuCOE/ICustomizableObjectEditor.h"
 #include "MuCOE/Nodes/CustomizableObjectNodeMaterial.h"
 #include "MuCOE/Nodes/CustomizableObjectNodeMaterialVariation.h"
+#include "Misc/UObjectToken.h"
+#include "Logging/MessageLog.h"
+#include "Containers/Queue.h"
 
 class UCustomizableObjectNodeRemapPins;
 
@@ -27,16 +28,9 @@ UCustomizableObjectNodeObject::UCustomizableObjectNodeObject()
 	: Super()
 {
 	bIsBase = true;
-
 	ObjectName = "Unnamed Object";
 	NumLODs = 1;
-	Identifier.Invalidate();
-
-	if (!Identifier.IsValid())
-	{
-		Identifier = FGuid::NewGuid();
-		IdentifierVerification = Identifier;
-	}
+	Identifier = FGuid::NewGuid();
 }
 
 
@@ -65,6 +59,15 @@ void UCustomizableObjectNodeObject::BackwardsCompatibleFixup()
 		// This will regenerate all the Node Object Guids to finally remove the duplicated Guids warning.
 		// It is safe to do this here as Node Object do not use its node guid to link themeselves to other nodes.
 		CreateNewGuid();
+
+		// This change may make cooks to become undeterministic, if the object GUID is finally used (it is a "toggle group" option).
+		UCustomizableObject* CustomizableObject = Cast<UCustomizableObject>(GetCustomizableObjectGraph()->GetOuter());
+		if (CustomizableObject)
+		{
+			FMessageLog("Mutable").Message(EMessageSeverity::Warning)
+				->AddToken(FTextToken::Create(LOCTEXT("Indeterministic Warning", "The object was saved with an old version and it may generate indeterministic packages. Resave it to fix the problem.")))
+				->AddToken(FUObjectToken::Create(CustomizableObject));
+		}
 	}
 
 	// Update state never-stream flag from deprecated enum
@@ -275,15 +278,6 @@ void UCustomizableObjectNodeObject::PostBackwardsCompatibleFixup()
 {
 	Super::PostBackwardsCompatibleFixup();
 
-	const bool bJustBuiltCoFlag = !GetAllNonOrphanPins().Num() && !NodeGuid.IsValid();
-	if (!bJustBuiltCoFlag && Identifier == IdentifierVerification)
-	{
-		FCustomizableObjectEditorLogger::CreateLog(LOCTEXT("ResaveNode","Please re-save this Customizable Object to avoid binary differences when packaging"))
-		.Severity(EMessageSeverity::Warning)
-		.Context(*this)
-		.Log();
-	}
-
 	// Fix up ComponentSettings
 	if (ComponentSettings.IsEmpty())
 	{
