[Mutable] Re: Convert BoneId to a hash of the bone name to fix indeterminism when the order of the bones changes.

#rnx

diff --git a/Engine/Plugins/Experimental/Mutable/Source/CustomizableObject/Internal/MuCO/CustomizableObjectInstancePrivate.h b/Engine/Plugins/Experimental/Mutable/Source/CustomizableObject/Internal/MuCO/CustomizableObjectInstancePrivate.h
--- a/Engine/Plugins/Experimental/Mutable/Source/CustomizableObject/Internal/MuCO/CustomizableObjectInstancePrivate.h
+++ b/Engine/Plugins/Experimental/Mutable/Source/CustomizableObject/Internal/MuCO/CustomizableObjectInstancePrivate.h
@@ -273,7 +273,7 @@ private:
 	USkeleton* MergeSkeletons(UCustomizableObject& CustomizableObject, const FMutableRefSkeletalMeshData& RefSkeletalMeshData, int32 ComponentIndex);
 
 	//
-	UPhysicsAsset* GetOrBuildMainPhysicsAsset(TObjectPtr<class UPhysicsAsset> TamplateAsset, const mu::PhysicsBody* PhysicsBody, const UCustomizableObject& CustomizableObject, int32 ComponentIndex, bool bDisableCollisionBetweenAssets);
+	UPhysicsAsset* GetOrBuildMainPhysicsAsset(const TSharedRef<FUpdateContextPrivate>& OperationData, TObjectPtr<class UPhysicsAsset> TamplateAsset, const mu::PhysicsBody* PhysicsBody, bool bDisableCollisionBetweenAssets, int32 ComponentIndex);
 	
 	// Create a transient texture and add it to the TextureTrackerArray
 	UTexture2D* CreateTexture();
diff --git a/Engine/Plugins/Experimental/Mutable/Source/CustomizableObject/Internal/MuCO/CustomizableObjectPrivate.h b/Engine/Plugins/Experimental/Mutable/Source/CustomizableObject/Internal/MuCO/CustomizableObjectPrivate.h
--- a/Engine/Plugins/Experimental/Mutable/Source/CustomizableObject/Internal/MuCO/CustomizableObjectPrivate.h
+++ b/Engine/Plugins/Experimental/Mutable/Source/CustomizableObject/Internal/MuCO/CustomizableObjectPrivate.h
@@ -17,7 +17,12 @@
 
 #include "CustomizableObjectPrivate.generated.h"
 
-namespace mu { class Model; }
+namespace mu 
+{ 
+	class Model;
+	struct FBoneName;
+}
+
 class USkeletalMesh;
 class USkeleton;
 class UPhysicsAsset;
@@ -92,6 +97,29 @@ struct FCustomizableObjectStreameableResourceId
 };
 static_assert(sizeof(FCustomizableObjectStreameableResourceId) == sizeof(uint32));
 
+
+USTRUCT()
+struct FMutableRemappedBone
+{
+	GENERATED_USTRUCT_BODY()
+
+	UPROPERTY()
+	FName Name;
+	
+	UPROPERTY()
+	uint32 Hash = 0;
+	
+	bool operator==(const FName& InName)
+	{
+		return Name == InName;
+	}
+
+#if WITH_EDITORONLY_DATA
+	friend FArchive& operator<<(FArchive& Ar, FMutableRemappedBone& RemappedBone);
+#endif
+};
+
+
 USTRUCT()
 struct FMutableModelParameterValue
 {
@@ -553,9 +581,9 @@ struct FModelResources
 	UPROPERTY()
 	TArray<FName> MaterialSlotNames;
 
-	/** Bone names of all the bones that can possibly use the generated meshes */
+	/** Bones remapped due to a hash collision at compile time. FName to (remapped)Hash. */
 	UPROPERTY()
-	TArray<FName> BoneNames;
+	TArray<FMutableRemappedBone> RemappedBoneNames;
 
 	/** Mesh sockets provided by the part skeletal meshes, to be merged in the generated meshes */
 	UPROPERTY()
@@ -833,6 +861,6 @@ public:
 	// This is a manual version number for the binary blobs in this asset.
 	// Increasing it invalidates all the previously compiled models.
 	// Warning: If while merging code both versions have changed, take the highest+1.
-	static constexpr int32 CurrentSupportedVersion = 439;
+	static constexpr int32 CurrentSupportedVersion = 440;
 };
 
diff --git a/Engine/Plugins/Experimental/Mutable/Source/CustomizableObject/Internal/MuCO/CustomizableObjectSystemPrivate.h b/Engine/Plugins/Experimental/Mutable/Source/CustomizableObject/Internal/MuCO/CustomizableObjectSystemPrivate.h
--- a/Engine/Plugins/Experimental/Mutable/Source/CustomizableObject/Internal/MuCO/CustomizableObjectSystemPrivate.h
+++ b/Engine/Plugins/Experimental/Mutable/Source/CustomizableObject/Internal/MuCO/CustomizableObjectSystemPrivate.h
@@ -16,6 +16,7 @@
 #include "MuR/Mesh.h"
 #include "MuR/Parameters.h"
 #include "MuR/System.h"
+#include "MuR/Skeleton.h"
 #include "MuR/Image.h"
 #include "UObject/GCObject.h"
 #include "WorldCollision.h"
@@ -375,11 +376,9 @@ struct FInstanceUpdateData
 		uint16 FirstSurface = 0;
 		uint16 SurfaceCount = 0;
 
-		// \TODO: Flatten
-		TArray<uint16> ActiveBones;
 		/** Range in the external Bones array */
-		//uint32 FirstActiveBone;
-		//uint32 ActiveBoneCount;
+		uint32 FirstActiveBone = 0;
+		uint32 ActiveBoneCount = 0;
 
 		/** Range in the external Bones array */
 		uint32 FirstBoneMap = 0;
@@ -400,18 +399,26 @@ struct FInstanceUpdateData
 	TArray<FVector> Vectors;
 	TArray<FScalar> Scalars;
 
-	TArray<uint16> BoneMaps;
+	TArray<mu::FBoneName> ActiveBones;
+	TArray<mu::FBoneName> BoneMaps;
 	
 	TMap<uint32, TArray<FMorphTargetVertexData>> MorphTargetsVertexData;
 
-	struct FSkeletonData
+	struct FBone
 	{
-		int16 ComponentIndex = INDEX_NONE;
+		mu::FBoneName Name;
+		FMatrix44f MatrixWithScale;
+
+		bool operator==(const mu::FBoneName& OtherName) const { return Name == OtherName; };
+	};
 
+	struct FSkeletonData
+	{
 		TArray<uint16> SkeletonIds;
 
-		TArray<uint16> BoneIds;
-		TArray<FMatrix44f> BoneMatricesWithScale;
+		TArray<FBone> BonePose;
+
+		TMap<mu::FBoneName, TPair<FName, uint16>> BoneInfoMap;
 	};
 
 	TArray<FSkeletonData> Skeletons;
diff --git a/Engine/Plugins/Experimental/Mutable/Source/CustomizableObject/Internal/MuCO/UnrealConversionUtils.h b/Engine/Plugins/Experimental/Mutable/Source/CustomizableObject/Internal/MuCO/UnrealConversionUtils.h
--- a/Engine/Plugins/Experimental/Mutable/Source/CustomizableObject/Internal/MuCO/UnrealConversionUtils.h
+++ b/Engine/Plugins/Experimental/Mutable/Source/CustomizableObject/Internal/MuCO/UnrealConversionUtils.h
@@ -9,11 +9,13 @@
 
 #include "MuR/Ptr.h"
 #include "Containers/Array.h"
+#include "Containers/Map.h"
 
 namespace mu
 {
 	class Mesh;
 	class FMeshBufferSet;
+	struct FBoneName;
 }
 struct FReferenceSkeleton;
 
@@ -40,7 +42,8 @@ namespace UnrealConversionUtils
 	CUSTOMIZABLEOBJECT_API void SetupRenderSections(
 		FSkeletalMeshLODRenderData& LODResource,
 		const mu::Ptr<const mu::Mesh> InMutableMesh,
-		const TArray<uint16>& InBoneMap,
+		const TArray<mu::FBoneName>& InBoneMap,
+		const TMap<mu::FBoneName, TPair<FName, uint16>>& BoneInfoMap,
 		const int32 InFirstBoneMapIndex);
 
 
diff --git a/Engine/Plugins/Experimental/Mutable/Source/CustomizableObject/Private/MuCO/CustomizableObject.cpp b/Engine/Plugins/Experimental/Mutable/Source/CustomizableObject/Private/MuCO/CustomizableObject.cpp
--- a/Engine/Plugins/Experimental/Mutable/Source/CustomizableObject/Private/MuCO/CustomizableObject.cpp
+++ b/Engine/Plugins/Experimental/Mutable/Source/CustomizableObject/Private/MuCO/CustomizableObject.cpp
@@ -550,7 +550,7 @@ void UCustomizableObjectPrivate::SaveCompiledData(FArchive& MemoryWriter, bool b
 	MemoryWriter << LocalModelResources.AnimBpOverridePhysiscAssetsInfo;
 
 	MemoryWriter << LocalModelResources.MaterialSlotNames;
-	MemoryWriter << LocalModelResources.BoneNames;
+	MemoryWriter << LocalModelResources.RemappedBoneNames;
 	MemoryWriter << LocalModelResources.SocketArray;
 
 	MemoryWriter << LocalModelResources.SkinWeightProfilesInfo;
@@ -671,7 +671,7 @@ void UCustomizableObjectPrivate::LoadCompiledData(FArchive& MemoryReader, const
 		MemoryReader << LocalModelResource.AnimBpOverridePhysiscAssetsInfo;
 
 		MemoryReader << LocalModelResource.MaterialSlotNames;
-		MemoryReader << LocalModelResource.BoneNames;
+		MemoryReader << LocalModelResource.RemappedBoneNames;
 		MemoryReader << LocalModelResource.SocketArray;
 
 		MemoryReader << LocalModelResource.SkinWeightProfilesInfo;
@@ -2360,6 +2360,13 @@ bool FMutableSkinWeightProfileInfo::operator==(const FMutableSkinWeightProfileIn
 
 
 #if WITH_EDITORONLY_DATA
+FArchive& operator<<(FArchive& Ar, FMutableRemappedBone& RemappedBone)
+{
+	Ar << RemappedBone.Name;
+	Ar << RemappedBone.Hash;
+	return Ar;
+}
+
 FArchive& operator<<(FArchive& Ar, FMutableModelImageProperties& ImageProps)
 {
 	Ar << ImageProps.TextureParameterName;
diff --git a/Engine/Plugins/Experimental/Mutable/Source/CustomizableObject/Private/MuCO/CustomizableObjectInstance.cpp b/Engine/Plugins/Experimental/Mutable/Source/CustomizableObject/Private/MuCO/CustomizableObjectInstance.cpp
--- a/Engine/Plugins/Experimental/Mutable/Source/CustomizableObject/Private/MuCO/CustomizableObjectInstance.cpp
+++ b/Engine/Plugins/Experimental/Mutable/Source/CustomizableObject/Private/MuCO/CustomizableObjectInstance.cpp
@@ -49,6 +49,7 @@
 #include "PhysicsEngine/AggregateGeom.h"
 #include "PhysicsEngine/PhysicsConstraintTemplate.h"
 #include "PhysicsEngine/SkeletalBodySetup.h"
+#include "Hash/CityHash.h"
 
 #include UE_INLINE_GENERATED_CPP_BY_NAME(CustomizableObjectInstance)
 
@@ -1187,7 +1188,10 @@ namespace
 	};
 
 	TObjectPtr<UPhysicsAsset> MakePhysicsAssetFromTemplateAndMutableBody(
-		TObjectPtr<UPhysicsAsset> TemplateAsset, const mu::PhysicsBody* MutablePhysics, const UCustomizableObject& CustomizableObject) 
+		const TSharedRef<FUpdateContextPrivate>& OperationData,
+		TObjectPtr<UPhysicsAsset> TemplateAsset,
+		const mu::PhysicsBody* MutablePhysics,
+		int32 ComponentIndex)
 	{
 		check(TemplateAsset);
 		TObjectPtr<UPhysicsAsset> Result = NewObject<UPhysicsAsset>();
@@ -1202,18 +1206,16 @@ namespace
 
 		Result->bNotForDedicatedServer = TemplateAsset->bNotForDedicatedServer;
 
-		const FModelResources& ModelResources = CustomizableObject.GetPrivate()->GetModelResources();
+		const TMap<mu::FBoneName, TPair<FName, uint16>>& BoneInfoMap = OperationData->InstanceUpdateData.Skeletons[ComponentIndex].BoneInfoMap;
 		TMap<FName, int32> BonesInUse;
 
 		const int32 MutablePhysicsBodyCount = MutablePhysics->GetBodyCount();
 		BonesInUse.Reserve(MutablePhysicsBodyCount);
 		for ( int32 I = 0; I < MutablePhysicsBodyCount; ++I )
 		{
-			const uint16 BoneNameId = MutablePhysics->GetBodyBoneId(I);
-			if (ModelResources.BoneNames.IsValidIndex(BoneNameId))
+			if (const TPair<FName, uint16>* BoneInfo = BoneInfoMap.Find(MutablePhysics->GetBodyBoneId(I)))
 			{
-				FName BoneName = ModelResources.BoneNames[BoneNameId];
-				BonesInUse.Add(BoneName, I);
+				BonesInUse.Add(BoneInfo->Key, I);
 			}
 		}
 
@@ -1313,11 +1315,11 @@ namespace
 }
 
 UPhysicsAsset* UCustomizableInstancePrivate::GetOrBuildMainPhysicsAsset(
+	const TSharedRef<FUpdateContextPrivate>& OperationData,
 	TObjectPtr<UPhysicsAsset> TemplateAsset,
 	const mu::PhysicsBody* MutablePhysics,
-	const UCustomizableObject& CustomizableObject,
-	int32 ComponentId,
-	bool bDisableCollisionsBetweenDifferentAssets)
+	bool bDisableCollisionsBetweenDifferentAssets,
+	int32 ComponentIndex)
 {
 
 	MUTABLE_CPUPROFILER_SCOPE(MergePhysicsAssets);
@@ -1326,7 +1328,7 @@ UPhysicsAsset* UCustomizableInstancePrivate::GetOrBuildMainPhysicsAsset(
 
 	UPhysicsAsset* Result = nullptr;
 
-	FCustomizableInstanceComponentData* ComponentData = GetComponentData(ComponentId);
+	FCustomizableInstanceComponentData* ComponentData = GetComponentData(ComponentIndex);
 	check(ComponentData);
 
 	TArray<TObjectPtr<UPhysicsAsset>>& PhysicsAssets = ComponentData->PhysicsAssets.PhysicsAssetsToMerge;
@@ -1370,18 +1372,16 @@ UPhysicsAsset* UCustomizableInstancePrivate::GetOrBuildMainPhysicsAsset(
 
 	Result->bNotForDedicatedServer = TemplateAsset->bNotForDedicatedServer;
 
-	const FModelResources& ModelResources = CustomizableObject.GetPrivate()->GetModelResources();
+	const TMap<mu::FBoneName, TPair<FName, uint16>>& BoneInfoMap = OperationData->InstanceUpdateData.Skeletons[ComponentIndex].BoneInfoMap;
 	TMap<FName, int32> BonesInUse;
 
 	const int32 MutablePhysicsBodyCount = MutablePhysics->GetBodyCount();
 	BonesInUse.Reserve(MutablePhysicsBodyCount);
 	for ( int32 I = 0; I < MutablePhysicsBodyCount; ++I )
 	{
-		const uint16 BoneNameId = MutablePhysics->GetBodyBoneId(I);
-		if (ModelResources.BoneNames.IsValidIndex(BoneNameId))
+		if (const TPair<FName, uint16>* BoneInfo = BoneInfoMap.Find(MutablePhysics->GetBodyBoneId(I)))
 		{
-			FName BoneName = ModelResources.BoneNames[BoneNameId];
-			BonesInUse.Add(BoneName, I);
+			BonesInUse.Add(BoneInfo->Key, I);
 		}
 	}
 
@@ -1856,7 +1856,7 @@ bool UCustomizableInstancePrivate::UpdateSkeletalMesh_PostBeginUpdate0(UCustomiz
 			{
 				constexpr bool bDisallowCollisionBetweenAssets = true;
 				UPhysicsAsset* PhysicsAssetResult = GetOrBuildMainPhysicsAsset(
-					RefSkeletalMeshData.PhysicsAsset, MutablePhysics.get(), *CustomizableObject, Component.Id, bDisallowCollisionBetweenAssets);
+					OperationData, RefSkeletalMeshData.PhysicsAsset, MutablePhysics.get(), bDisallowCollisionBetweenAssets, Component.Id);
 
 				SkeletalMesh->SetPhysicsAsset(PhysicsAssetResult);
 
@@ -1905,7 +1905,8 @@ bool UCustomizableInstancePrivate::UpdateSkeletalMesh_PostBeginUpdate0(UCustomiz
 					PhysicsAssetsUsedByAnimBp.AnimInstancePropertyIndexAndPhysicsAssets.Emplace_GetRef();
 
 				Entry.PropertyIndex = Info.PropertyIndex;
-				Entry.PhysicsAsset = MakePhysicsAssetFromTemplateAndMutableBody(PhysicsAssetTemplate, AdditionalPhysiscBody.get(), *CustomizableObject);
+				Entry.PhysicsAsset = MakePhysicsAssetFromTemplateAndMutableBody(
+					OperationData,PhysicsAssetTemplate, AdditionalPhysiscBody.get(), Component.Id);
 			}
 
 			// Add sockets from the SkeletalMesh of reference and from the MutableMesh
@@ -3311,96 +3312,72 @@ bool UCustomizableInstancePrivate::BuildSkeletonData(const TSharedRef<FUpdateCon
 	SkeletalMesh.SetRefSkeleton(Skeleton->GetReferenceSkeleton());
 	FReferenceSkeleton& ReferenceSkeleton = SkeletalMesh.GetRefSkeleton();
 
-	// Check that the bones we need are present in the current Skeleton
-	FInstanceUpdateData::FSkeletonData& MutSkeletonData = OperationData->InstanceUpdateData.Skeletons[ComponentIndex];
-	const int32 MutBoneCount = MutSkeletonData.BoneIds.Num();
-	
-	TMap<uint16, uint16> BoneToFinalBoneIndexMap;
-	BoneToFinalBoneIndexMap.Reserve(MutBoneCount);
+	const TArray<FMeshBoneInfo>& RawRefBoneInfo = ReferenceSkeleton.GetRawRefBoneInfo();
+	const int32 RawRefBoneCount = ReferenceSkeleton.GetRawBoneNum();
+
+	const TArray<FInstanceUpdateData::FBone>& BonePose = OperationData->InstanceUpdateData.Skeletons[ComponentIndex].BonePose;
+	TMap<mu::FBoneName, TPair<FName,uint16>>& BoneInfoMap = OperationData->InstanceUpdateData.Skeletons[ComponentIndex].BoneInfoMap;
 
 	{
-		MUTABLE_CPUPROFILER_SCOPE(BuildSkeletonData_EnsureBonesExist);
+		MUTABLE_CPUPROFILER_SCOPE(BuildSkeletonData_BuildBoneInfoMap);
 
-		const FModelResources& ModelResources = CustomizableObject.GetPrivate()->GetModelResources();
+		BoneInfoMap.Reserve(RawRefBoneCount);
 
-		// Ensure all the required bones are present in the skeleton
-		for (int32 BoneIndex = 0; BoneIndex < MutBoneCount; ++BoneIndex)
+		const FModelResources& ModelResources = CustomizableObject.GetPrivate()->GetModelResources();
+		for (int32 Index = 0; Index < RawRefBoneCount; ++Index)
 		{
-			const uint16 BoneId = MutSkeletonData.BoneIds[BoneIndex];
-			check(ModelResources.BoneNames.IsValidIndex(BoneId));
-
-			const FName BoneName = ModelResources.BoneNames[BoneId];
-			check(BoneName != NAME_None);
-
-			const int32 SourceBoneIndex = ReferenceSkeleton.FindRawBoneIndex(BoneName);
-			if (SourceBoneIndex == INDEX_NONE)
+			const FName BoneName = RawRefBoneInfo[Index].Name;
+			if (const FMutableRemappedBone* RemappedBone = ModelResources.RemappedBoneNames.FindByKey(BoneName))
 			{
-				// Merged skeleton is missing some bones! This happens if one of the skeletons involved in the merge is discarded due to being incompatible with the rest
-				// or if the source mesh is not in sync with the skeleton. 
-				UE_LOG(LogMutable, Warning, TEXT("Building instance: generated mesh has a bone [%s] not present in the reference mesh [%s]. Failing to generate mesh. "),
-					*BoneName.ToString(), *SkeletalMesh.GetName());
-				return false;
+				TPair<FName, uint16>& BoneInfo = BoneInfoMap.Add(mu::FBoneName(RemappedBone->Hash));
+				BoneInfo.Key = BoneName;
+				BoneInfo.Value = Index;
+			}
+			else
+			{
+				const FString BoneNameString = BoneName.ToString();
+				const mu::FBoneName Bone(CityHash32(reinterpret_cast<const char*>(*BoneNameString), BoneNameString.Len() * sizeof(FString::ElementType)));
+				TPair<FName, uint16>& BoneInfo = BoneInfoMap.Add(Bone);
+				BoneInfo.Key = BoneName;
+				BoneInfo.Value = Index;
 			}
-
-			BoneToFinalBoneIndexMap.Add(BoneId, SourceBoneIndex);
 		}
 	}
 
 	{
-		MUTABLE_CPUPROFILER_SCOPE(BuildSkeletonData_FixBoneIndices);
+		MUTABLE_CPUPROFILER_SCOPE(BuildSkeletonData_EnsureBonesExist);
 
-		// Fix up BoneMaps and ActiveBones indices
-		for (FInstanceUpdateData::FComponent& Component : OperationData->InstanceUpdateData.Components)
+		// Ensure all required bones are present in the skeleton
+		for (const FInstanceUpdateData::FBone& Bone : BonePose)
 		{
-			if (Component.Id != ComponentIndex)
-			{
-				continue;
-			}
-
-			for (uint32 BoneMapIndex = Component.FirstBoneMap; BoneMapIndex < Component.FirstBoneMap + Component.BoneMapCount; ++BoneMapIndex)
-			{
-				const int32 BoneId = OperationData->InstanceUpdateData.BoneMaps[BoneMapIndex];
-				OperationData->InstanceUpdateData.BoneMaps[BoneMapIndex] = BoneToFinalBoneIndexMap[BoneId];
-			}
-
-			for (uint16& BoneId : Component.ActiveBones)
+			if (!BoneInfoMap.Find(Bone.Name))
 			{
-				BoneId = BoneToFinalBoneIndexMap[BoneId];
+				UE_LOG(LogMutable, Warning, TEXT("The skeleton of skeletal mesh [%s] is missing a bone with ID [%d], which the mesh requires."),
+					*SkeletalMesh.GetName(), Bone.Name.Id);
+				return false;
 			}
-			Component.ActiveBones.Sort();
 		}
 	}
-	
+
 	{
 		MUTABLE_CPUPROFILER_SCOPE(BuildSkeletonData_ApplyPose);
 		
-		const int32 RefRawBoneCount = ReferenceSkeleton.GetRawBoneNum();
-
 		TArray<FMatrix44f>& RefBasesInvMatrix = SkeletalMesh.GetRefBasesInvMatrix();
-		RefBasesInvMatrix.Empty(RefRawBoneCount);
+		RefBasesInvMatrix.Empty(RawRefBoneCount);
 
-		// Initialize the base matrices
-		if (RefRawBoneCount == MutBoneCount)
-		{
-			RefBasesInvMatrix.AddUninitialized(MutBoneCount);
-		}
-		else
-		{
-			// Bad case, some bone poses are missing, calculate the InvRefMatrices to ensure all transforms are there for the second step 
-			MUTABLE_CPUPROFILER_SCOPE(BuildSkeletonData_CalcInvRefMatrices0);
-			SkeletalMesh.CalculateInvRefMatrices();
-		}
+		// Calculate the InvRefMatrices to ensure all transforms are there for the second step 
+		SkeletalMesh.CalculateInvRefMatrices();
 
 		// First step is to update the RefBasesInvMatrix for the bones.
-		for (int32 BoneIndex = 0; BoneIndex < MutBoneCount; ++BoneIndex)
+		for (const FInstanceUpdateData::FBone& Bone : BonePose)
 		{
-			const int32 RefSkelBoneIndex = BoneToFinalBoneIndexMap[MutSkeletonData.BoneIds[BoneIndex]];
-			RefBasesInvMatrix[RefSkelBoneIndex] = MutSkeletonData.BoneMatricesWithScale[BoneIndex];
+			const int32 BoneIndex = BoneInfoMap[Bone.Name].Value;
+			RefBasesInvMatrix[BoneIndex] = Bone.MatrixWithScale;
 		}
 
 		// The second step is to update the pose transforms in the ref skeleton from the BasesInvMatrix
 		FReferenceSkeletonModifier SkeletonModifier(ReferenceSkeleton, Skeleton);
-		for (int32 RefSkelBoneIndex = 0; RefSkelBoneIndex < ReferenceSkeleton.GetRawBoneNum(); ++RefSkelBoneIndex)
+		for (int32 RefSkelBoneIndex = 0; RefSkelBoneIndex < RawRefBoneCount; ++RefSkelBoneIndex)
 		{
 			int32 ParentBoneIndex = ReferenceSkeleton.GetParentIndex(RefSkelBoneIndex);
 			if (ParentBoneIndex >= 0)
@@ -3413,7 +3390,7 @@ bool UCustomizableInstancePrivate::BuildSkeletonData(const TSharedRef<FUpdateCon
 		}
 
 		// Force a CalculateInvRefMatrices
-		RefBasesInvMatrix.Empty(RefRawBoneCount);
+		RefBasesInvMatrix.Empty(RawRefBoneCount);
 	}
 
 	{
@@ -4701,15 +4678,29 @@ bool UCustomizableInstancePrivate::BuildOrCopyRenderData(const TSharedRef<FUpdat
 
 		FSkeletalMeshLODRenderData& LODResource = RenderData->LODRenderData[LODIndex];
 
+		const TMap<mu::FBoneName, TPair<FName, uint16>>& BoneInfoMap = OperationData->InstanceUpdateData.Skeletons[ComponentIndex].BoneInfoMap;
+		
 		// Set active and required bones
-		LODResource.ActiveBoneIndices.Append(Component.ActiveBones);
-		LODResource.RequiredBones.Append(Component.ActiveBones);
+		{
+			const TArray<mu::FBoneName>& ActiveBones = OperationData->InstanceUpdateData.ActiveBones;
+			LODResource.ActiveBoneIndices.Reserve(Component.ActiveBoneCount);
+
+			for (uint32 Index = 0; Index < Component.ActiveBoneCount; ++Index)
+			{
+				const uint16 ActiveBoneIndex = BoneInfoMap[ActiveBones[Component.FirstActiveBone + Index]].Value;
+				LODResource.ActiveBoneIndices.Add(ActiveBoneIndex);
+			}
+
+			LODResource.RequiredBones = LODResource.ActiveBoneIndices;
+			LODResource.RequiredBones.Sort();
+		}
 
 		// Set RenderSections
 		UnrealConversionUtils::SetupRenderSections(
 			LODResource,
 			Component.Mesh,
 			OperationData->InstanceUpdateData.BoneMaps,
+			BoneInfoMap,
 			Component.FirstBoneMap);
 
 		if (LODResource.bStreamedDataInlined) // Non-streamable LOD
@@ -4891,9 +4882,11 @@ UE::Tasks::FTask UCustomizableInstancePrivate::LoadAdditionalAssetsAndData(
 
 
 	// Load Skeletons required by the SubMeshes of the newly generated Mesh, will be merged later
-	for (FInstanceUpdateData::FSkeletonData& SkeletonData : OperationData->InstanceUpdateData.Skeletons)
+	for (int32 ComponentIndex = 0; ComponentIndex < OperationData->NumComponents; ++ ComponentIndex)
 	{
-		FCustomizableInstanceComponentData* ComponentData = GetComponentData(SkeletonData.ComponentIndex);
+		const FInstanceUpdateData::FSkeletonData& SkeletonData = OperationData->InstanceUpdateData.Skeletons[ComponentIndex];
+		
+		FCustomizableInstanceComponentData* ComponentData = GetComponentData(ComponentIndex);
 		if (!ComponentData)
 		{
 			check(false);
diff --git a/Engine/Plugins/Experimental/Mutable/Source/CustomizableObject/Private/MuCO/CustomizableObjectSystem.cpp b/Engine/Plugins/Experimental/Mutable/Source/CustomizableObject/Private/MuCO/CustomizableObjectSystem.cpp
--- a/Engine/Plugins/Experimental/Mutable/Source/CustomizableObject/Private/MuCO/CustomizableObjectSystem.cpp
+++ b/Engine/Plugins/Experimental/Mutable/Source/CustomizableObject/Private/MuCO/CustomizableObjectSystem.cpp
@@ -1793,22 +1793,18 @@ namespace impl
 	{
 		MUTABLE_CPUPROFILER_SCOPE(Subtask_Mutable_PrepareSkeletonData);
 
+		OperationData->InstanceUpdateData.Skeletons.SetNum(OperationData->NumComponents);
+
 		for (FInstanceUpdateData::FComponent& Component : OperationData->InstanceUpdateData.Components)
 		{
-			if (!OperationData->InstanceUpdateData.Skeletons.IsValidIndex(Component.Id))
-			{
-				OperationData->InstanceUpdateData.Skeletons.SetNum(Component.Id + 1);
-			}
-
-			FInstanceUpdateData::FSkeletonData& SkeletonData = OperationData->InstanceUpdateData.Skeletons[Component.Id];
-			SkeletonData.ComponentIndex = Component.Id;
-
 			mu::MeshPtrConst Mesh = Component.Mesh;
 			if (!Mesh)
 			{
 				continue;
 			}
 
+			FInstanceUpdateData::FSkeletonData& SkeletonData = OperationData->InstanceUpdateData.Skeletons[Component.Id];
+
 			// Add SkeletonIds 
 			const int32 SkeletonIDsCount = Mesh->GetSkeletonIDsCount();
 			for (int32 SkeletonIndex = 0; SkeletonIndex < SkeletonIDsCount; ++SkeletonIndex)
@@ -1817,27 +1813,25 @@ namespace impl
 			}
 
 			// Append BoneMap to the array of BoneMaps
-			const TArray<uint16>& BoneMap = Mesh->GetBoneMap();
+			const TArray<mu::FBoneName>& BoneMap = Mesh->GetBoneMap();
 			Component.FirstBoneMap = OperationData->InstanceUpdateData.BoneMaps.Num();
 			Component.BoneMapCount = BoneMap.Num();
 			OperationData->InstanceUpdateData.BoneMaps.Append(BoneMap);
 
 			// Add active bone indices and poses
-			const int32 MaxBoneIndex = Mesh->GetBonePoseCount();
-			Component.ActiveBones.Reserve(MaxBoneIndex);
-			for (int32 BonePoseIndex = 0; BonePoseIndex < MaxBoneIndex; ++BonePoseIndex)
+			Component.FirstActiveBone = OperationData->InstanceUpdateData.ActiveBones.Num();
+			Component.ActiveBoneCount = Mesh->GetBonePoseCount();
+			for (uint32 BoneIndex = 0; BoneIndex < Component.ActiveBoneCount; ++BoneIndex)
 			{
-				const uint16 BoneId = Mesh->GetBonePoseBoneId(BonePoseIndex);
+				const mu::FBoneName& BoneId = Mesh->GetBonePoseId(BoneIndex);
 
-				Component.ActiveBones.Add(BoneId);
+				OperationData->InstanceUpdateData.ActiveBones.Add(BoneId);
 
-				if (SkeletonData.BoneIds.Find(BoneId) == INDEX_NONE)
+				if (!SkeletonData.BonePose.FindByKey(BoneId))
 				{
-					SkeletonData.BoneIds.Add(BoneId);
-
 					FTransform3f Transform;
-					Mesh->GetBoneTransform(BonePoseIndex, Transform);
-					SkeletonData.BoneMatricesWithScale.Emplace(Transform.Inverse().ToMatrixWithScale());
+					Mesh->GetBonePoseTransform(BoneIndex, Transform);
+					SkeletonData.BonePose.Add({ BoneId,Transform.Inverse().ToMatrixWithScale() });
 				}
 			}
 		}
diff --git a/Engine/Plugins/Experimental/Mutable/Source/CustomizableObject/Private/MuCO/UnrealConversionUtils.cpp b/Engine/Plugins/Experimental/Mutable/Source/CustomizableObject/Private/MuCO/UnrealConversionUtils.cpp
--- a/Engine/Plugins/Experimental/Mutable/Source/CustomizableObject/Private/MuCO/UnrealConversionUtils.cpp
+++ b/Engine/Plugins/Experimental/Mutable/Source/CustomizableObject/Private/MuCO/UnrealConversionUtils.cpp
@@ -7,6 +7,7 @@
 #include "Engine/SkeletalMesh.h"
 #include "GPUSkinVertexFactory.h"
 #include "MuCO/CustomizableObject.h"
+#include "MuR/Skeleton.h"
 #include "MuR/Mesh.h"
 #include "MuR/MeshBufferSet.h"
 #include "MuR/MutableTrace.h"
@@ -162,7 +163,8 @@ namespace UnrealConversionUtils
 	void SetupRenderSections(
 		FSkeletalMeshLODRenderData& LODResource,
 		const mu::MeshPtrConst InMutableMesh,
-		const TArray<uint16>& InBoneMap,
+		const TArray<mu::FBoneName>& InBoneMap,
+		const TMap<mu::FBoneName, TPair<FName, uint16>>& BoneInfoMap,
 		const int32 InFirstBoneMapIndex)
 	{
 		check(InMutableMesh);
@@ -218,7 +220,7 @@ namespace UnrealConversionUtils
 			Section.BoneMap.Reserve(BoneCount);
 			for (int32 BoneMapIndex = 0; BoneMapIndex < BoneCount; ++BoneMapIndex, ++FirstBone)
 			{
-				Section.BoneMap.Add(InBoneMap[FirstBone]);
+				Section.BoneMap.Add(BoneInfoMap[InBoneMap[FirstBone]].Value);
 			}
 		}
 	}
diff --git a/Engine/Plugins/Experimental/Mutable/Source/CustomizableObjectEditor/Private/MuCOE/CustomizableObjectCompiler.cpp b/Engine/Plugins/Experimental/Mutable/Source/CustomizableObjectEditor/Private/MuCOE/CustomizableObjectCompiler.cpp
--- a/Engine/Plugins/Experimental/Mutable/Source/CustomizableObjectEditor/Private/MuCOE/CustomizableObjectCompiler.cpp
+++ b/Engine/Plugins/Experimental/Mutable/Source/CustomizableObjectEditor/Private/MuCOE/CustomizableObjectCompiler.cpp
@@ -971,9 +971,15 @@ void FCustomizableObjectCompiler::CompileInternal(UCustomizableObject* Object, c
 		ModelResources.AnimBpOverridePhysiscAssetsInfo = MoveTemp(GenerationContext.AnimBpOverridePhysicsAssetsInfo);
 
 		ModelResources.MaterialSlotNames = MoveTemp(GenerationContext.ReferencedMaterialSlotNames);
-		ModelResources.BoneNames = MoveTemp(GenerationContext.BoneNames);
 		ModelResources.SocketArray = MoveTemp(GenerationContext.SocketArray);
 
+		for (auto& It : GenerationContext.RemappedBoneNames)
+		{
+			FMutableRemappedBone& RemappedBone = ModelResources.RemappedBoneNames.AddDefaulted_GetRef();
+			RemappedBone.Name = It.Key; // BoneName
+			RemappedBone.Hash = It.Value.Id; // mu::FBoneName
+		}
+
 		ModelResources.SkinWeightProfilesInfo = MoveTemp(GenerationContext.SkinWeightProfilesInfo);
 
 		TArray<FGeneratedImageProperties> ImageProperties;
diff --git a/Engine/Plugins/Experimental/Mutable/Source/CustomizableObjectEditor/Private/MuCOE/GenerateMutableSource/GenerateMutableSource.cpp b/Engine/Plugins/Experimental/Mutable/Source/CustomizableObjectEditor/Private/MuCOE/GenerateMutableSource/GenerateMutableSource.cpp
--- a/Engine/Plugins/Experimental/Mutable/Source/CustomizableObjectEditor/Private/MuCOE/GenerateMutableSource/GenerateMutableSource.cpp
+++ b/Engine/Plugins/Experimental/Mutable/Source/CustomizableObjectEditor/Private/MuCOE/GenerateMutableSource/GenerateMutableSource.cpp
@@ -374,6 +374,57 @@ void FMutableGraphGenerationContext::GenerateSharedSurfacesUniqueIds()
 	}
 }
 
+bool FMutableGraphGenerationContext::FindBone(const FName& InBoneName, mu::FBoneName& OutBoneName) const
+{
+	const FString BoneNameString = InBoneName.ToString();
+	OutBoneName.Id = CityHash32(reinterpret_cast<const char*>(*BoneNameString), BoneNameString.Len() * sizeof(FString::ElementType));
+	if (UniqueBoneNames.Contains(OutBoneName))
+	{
+		return true;
+	}
+
+	if (const mu::FBoneName* BoneName = RemappedBoneNames.Find(InBoneName))
+	{
+		OutBoneName = *BoneName;
+		return true;
+	}
+
+	return false;
+}
+
+
+mu::FBoneName FMutableGraphGenerationContext::GetBoneUnique(const FName& InBoneName)
+{
+	const FString BoneNameString = InBoneName.ToString();
+	mu::FBoneName Bone(CityHash32(reinterpret_cast<const char*>(*BoneNameString), BoneNameString.Len() * sizeof(FString::ElementType)));
+
+	bool bNewRemappedBoneName = false;
+
+	FName& BoneName = UniqueBoneNames.FindOrAdd(Bone, InBoneName);
+	while (BoneName != InBoneName)
+	{
+		if (mu::FBoneName* RemappedBoneName = RemappedBoneNames.Find(InBoneName))
+		{
+			Bone.Id = RemappedBoneName->Id;
+			break;
+		}
+
+		// Id collision detected
+		bNewRemappedBoneName = true;
+
+		// Increase Id in an attempt to make it unique again.
+		++Bone.Id;
+		BoneName = UniqueBoneNames.FindOrAdd(Bone, InBoneName);
+	}
+
+	if (bNewRemappedBoneName)
+	{
+		RemappedBoneNames.Add(InBoneName, Bone);
+	}
+
+	return Bone;
+}
+
 
 //void FMutableGraphGenerationContext::CheckPhysicsAssetInSkeletalMesh(const USkeletalMesh* SkeletalMesh)
 //{
@@ -741,10 +792,10 @@ mu::NodeMeshApplyPosePtr CreateNodeMeshApplyPose(FMutableGraphGenerationContext&
 	for (int32 i = 0; i < ArrayBoneName.Num(); ++i)
 	{
 		const FName BoneName = ArrayBoneName[i];
-		const uint16 BoneId = uint16(GenerationContext.BoneNames.AddUnique(BoneName));
+		const mu::FBoneName& BoneId = GenerationContext.GetBoneUnique(BoneName);
 
-		MutableSkeleton->SetBoneFName(i, BoneName);
-		MutableSkeleton->SetBoneId(i, BoneId);
+		MutableSkeleton->SetDebugName(i, BoneName);
+		MutableSkeleton->SetBoneName(i, { BoneId });
 		MutableMesh->SetBonePose(i, BoneId, (FTransform3f)ArrayTransform[i], mu::EBoneUsageFlags::Skinning);
 	}
 
diff --git a/Engine/Plugins/Experimental/Mutable/Source/CustomizableObjectEditor/Private/MuCOE/GenerateMutableSource/GenerateMutableSource.h b/Engine/Plugins/Experimental/Mutable/Source/CustomizableObjectEditor/Private/MuCOE/GenerateMutableSource/GenerateMutableSource.h
--- a/Engine/Plugins/Experimental/Mutable/Source/CustomizableObjectEditor/Private/MuCOE/GenerateMutableSource/GenerateMutableSource.h
+++ b/Engine/Plugins/Experimental/Mutable/Source/CustomizableObjectEditor/Private/MuCOE/GenerateMutableSource/GenerateMutableSource.h
@@ -12,6 +12,7 @@
 #include "MuCOE/Nodes/CustomizableObjectNodeMaterialBase.h"
 #include "MuCOE/Nodes/CustomizableObjectNodeObject.h"
 #include "MuR/MutableMemory.h"
+#include "MuR/Skeleton.h"
 #include "MuT/NodeComponentNew.h"
 #include "MuR/Ptr.h"
 #include "MuT/NodeImageConstant.h"
@@ -660,6 +661,11 @@ struct FMutableGraphGenerationContext
 	/** Generates shared surface IDs for all surface nodes. If one or more nodes are equal, they will use the same SharedSurfaceId */
 	void GenerateSharedSurfacesUniqueIds();
 
+	bool FindBone(const FName& BoneName, mu::FBoneName& OutBone) const;
+
+	/** Get unique identifier for BoneName built from its FString. */
+	mu::FBoneName GetBoneUnique(const FName& BoneName);
+
 	/** Check if the PhysicsAsset of a given SkeletalMesh has any SkeletalBodySetup with BoneNames not present in the
 	* InSkeletalMesh's RefSkeleton, if so, adds the PhysicsAsset to the DiscartedPhysicsAssetMap to display a warning later on */
 	//void CheckPhysicsAssetInSkeletalMesh(const USkeletalMesh* InSkeletalMesh);
@@ -685,8 +691,9 @@ struct FMutableGraphGenerationContext
 
 	TArray<const USkeleton*> ReferencedSkeletons;
 
-	// Array of unique BoneNames. All bones from mu::Skeletons will point to the names of it.
-	TArray<FName> BoneNames;
+	// Array of unique Bone identifiers. 
+	TMap<mu::FBoneName, FName> UniqueBoneNames;
+	TMap<FName, mu::FBoneName> RemappedBoneNames; // Bone identifiers that had a collision.
 
 	// Used to aviod Nodes with duplicated ids
 	TMap<FGuid, TArray<const UObject*>> NodeIdsMap;
diff --git a/Engine/Plugins/Experimental/Mutable/Source/CustomizableObjectEditor/Private/MuCOE/GenerateMutableSource/GenerateMutableSourceMesh.cpp b/Engine/Plugins/Experimental/Mutable/Source/CustomizableObjectEditor/Private/MuCOE/GenerateMutableSource/GenerateMutableSourceMesh.cpp
--- a/Engine/Plugins/Experimental/Mutable/Source/CustomizableObjectEditor/Private/MuCOE/GenerateMutableSource/GenerateMutableSourceMesh.cpp
+++ b/Engine/Plugins/Experimental/Mutable/Source/CustomizableObjectEditor/Private/MuCOE/GenerateMutableSource/GenerateMutableSourceMesh.cpp
@@ -408,7 +408,7 @@ void SetAndPropagatePoseBoneUsage(mu::Mesh& MutableMesh, int32 PoseIndex, mu::EB
 
 	while (BoneIndex != INDEX_NONE)
 	{
-		PoseIndex = MutableMesh.FindBonePose(MutableSkeleton.GetBoneId(BoneIndex));
+		PoseIndex = MutableMesh.FindBonePose(MutableSkeleton.GetBoneName(BoneIndex));
 
 		if (PoseIndex == INDEX_NONE)
 		{
@@ -472,8 +472,8 @@ TArray<uint8> MakePhysicsAssetBodySetupRelevancyMap(const FMutableGraphGeneratio
 
 	for (int32 BodyIndex = 0; BodyIndex < BodySetupsNum; ++BodyIndex)
 	{
-		const int32 BoneNameId = GenerationContext.BoneNames.Find(Asset->SkeletalBodySetups[BodyIndex]->BoneName);
-		RelevancyMap[BodyIndex] = BoneNameId != INDEX_NONE && Mesh->GetSkeleton()->FindBone(uint16(BoneNameId)) != INDEX_NONE;
+		mu::FBoneName Bone; 
+		RelevancyMap[BodyIndex] = GenerationContext.FindBone(Asset->SkeletalBodySetups[BodyIndex]->BoneName, Bone);
 	}
 
 	return RelevancyMap;
@@ -512,7 +512,7 @@ mu::Ptr<mu::PhysicsBody> MakePhysicsBodyFromAsset(FMutableGraphGenerationContext
 
 		TObjectPtr<USkeletalBodySetup>& BodySetup = SkeletalBodySetups[SourceBodyIndex++];
 
-		const uint16 BodyBoneId = GenerationContext.BoneNames.AddUnique(BodySetup->BoneName);
+		const mu::FBoneName& BodyBoneId = GenerationContext.GetBoneUnique(BodySetup->BoneName);
 		PhysicsBody->SetBodyBoneId(B, BodyBoneId);
 		
 		const int32 NumSpheres = BodySetup->AggGeom.SphereElems.Num();
@@ -801,7 +801,7 @@ mu::MeshPtr ConvertSkeletalMeshToMutable(const USkeletalMesh* InSkeletalMesh, co
 		MutableSkeleton->SetBoneCount(NumRequiredBones);
 
 		// MutableBoneMap will not keep an index to the Skeleton, but to the BoneName
-		TArray<uint16> MutableBoneMap;
+		TArray<mu::FBoneName> MutableBoneMap;
 		MutableBoneMap.SetNum(BoneMap.Num());
 		
 		const TArray<FMeshBoneInfo>& RefBoneInfo = InSkeletalMesh->GetRefSkeleton().GetRefBoneInfo();
@@ -813,19 +813,19 @@ mu::MeshPtr ConvertSkeletalMeshToMutable(const USkeletalMesh* InSkeletalMesh, co
 			const int32 ParentBoneIndex = RequiredBones.Find(BoneInfo.ParentIndex);
 
 			// Set bone hierarchy
-			const uint16 BoneId = uint16(GenerationContext.BoneNames.AddUnique(BoneInfo.Name));
+			const mu::FBoneName& BoneName = GenerationContext.GetBoneUnique(BoneInfo.Name);
 
-			MutableSkeleton->SetBoneId(BoneIndex, BoneId);
+			MutableSkeleton->SetBoneName(BoneIndex, BoneName);
 			MutableSkeleton->SetBoneParent(BoneIndex, ParentBoneIndex);
 
 			// Debug. Will not be serialized
-			MutableSkeleton->SetBoneFName(BoneIndex, BoneInfo.Name);
+			MutableSkeleton->SetDebugName(BoneIndex, BoneInfo.Name);
 
 			// BoneMap: Convert RefSkeletonBoneIndex to BoneId
 			const int32 BoneMapIndex = BoneMap.Find(RefSkeletonBoneIndex);
 			if (BoneMapIndex != INDEX_NONE)
 			{
-				MutableBoneMap[BoneMapIndex] = BoneId;
+				MutableBoneMap[BoneMapIndex] = BoneName;
 			}
 
 			// Set bone pose
@@ -837,7 +837,7 @@ mu::MeshPtr ConvertSkeletalMeshToMutable(const USkeletalMesh* InSkeletalMesh, co
 			EnumAddFlags(BoneUsageFlags, BoneMapIndex != INDEX_NONE ? mu::EBoneUsageFlags::Skinning : mu::EBoneUsageFlags::None);
 			EnumAddFlags(BoneUsageFlags, ParentBoneIndex == INDEX_NONE ? mu::EBoneUsageFlags::Root : mu::EBoneUsageFlags::None);
 
-			MutableMesh->SetBonePose(BoneIndex, BoneId, BaseInvTransform.Inverse(), BoneUsageFlags);
+			MutableMesh->SetBonePose(BoneIndex, BoneName, BaseInvTransform.Inverse(), BoneUsageFlags);
 		}
 
 		MutableMesh->SetBoneMap(MutableBoneMap);
@@ -1584,8 +1584,8 @@ mu::MeshPtr ConvertSkeletalMeshToMutable(const USkeletalMesh* InSkeletalMesh, co
 				continue;
 			}
 
-			const int32 BoneNameIndex = GenerationContext.BoneNames.Find(BodySetup->BoneName);
-			const int32 BonePoseIndex = MutableMesh->FindBonePose(uint16(BoneNameIndex));
+			const mu::FBoneName& BoneName = GenerationContext.GetBoneUnique(BodySetup->BoneName);
+			const int32 BonePoseIndex = MutableMesh->FindBonePose(BoneName);
 
 			if (BonePoseIndex == INDEX_NONE)
 			{
@@ -1648,7 +1648,7 @@ mu::MeshPtr ConvertSkeletalMeshToMutable(const USkeletalMesh* InSkeletalMesh, co
 		{
 			TObjectPtr<USkeletalBodySetup>& BodySetup = RelevantBodySetups[B];
 			
-			const uint16 BoneId = GenerationContext.BoneNames.AddUnique(BodySetup->BoneName);
+			const mu::FBoneName& BoneId = GenerationContext.GetBoneUnique(BodySetup->BoneName);
 			PhysicsBody->SetBodyBoneId( B, BoneId);
 			
 			const int32 NumSpheres = BodySetup->AggGeom.SphereElems.Num();
@@ -1750,7 +1750,7 @@ mu::MeshPtr ConvertSkeletalMeshToMutable(const USkeletalMesh* InSkeletalMesh, co
 				mu::EBoneUsageFlags::Skinning | mu::EBoneUsageFlags::Physics | mu::EBoneUsageFlags::Deform;
 		if (EnumHasAnyFlags(BonePose.BoneUsageFlags, FlagsToPropagate))
 		{
-			const int32 Index = MutableMesh->GetSkeleton()->FindBone(MutableMesh->GetBonePoseBoneId(I));
+			const int32 Index = MutableMesh->GetSkeleton()->FindBone(BonePose.BoneId);
 
 			if (Index == INDEX_NONE)
 			{
@@ -2787,7 +2787,7 @@ mu::NodeMeshPtr GenerateMorphMesh(const UEdGraphPin* Pin,
 
 						for (const FName BoneName : BonesToDeform)
 						{
-							Result->AddBoneToDeform(GenerationContext.BoneNames.AddUnique(BoneName));
+							Result->AddBoneToDeform(GenerationContext.GetBoneUnique(BoneName));
 						}
 					}
 
@@ -2803,7 +2803,7 @@ mu::NodeMeshPtr GenerateMorphMesh(const UEdGraphPin* Pin,
 	
 						for (const FName& PhysicsBoneName : PhysicsToDeform)
 						{
-							Result->AddPhysicsBodyToDeform(GenerationContext.BoneNames.AddUnique(PhysicsBoneName));
+							Result->AddPhysicsBodyToDeform(GenerationContext.GetBoneUnique(PhysicsBoneName));
 						}
 					}
 						
@@ -3531,7 +3531,7 @@ mu::NodeMeshPtr GenerateMutableSourceMesh(const UEdGraphPin* Pin,
 				
 				for (const FName& BoneName : BonesToDeform)
 				{
-					MeshNode->AddBoneToDeform(GenerationContext.BoneNames.Find(BoneName));
+					MeshNode->AddBoneToDeform(GenerationContext.GetBoneUnique(BoneName));
 				}
 			}
 
@@ -3546,7 +3546,7 @@ mu::NodeMeshPtr GenerateMutableSourceMesh(const UEdGraphPin* Pin,
 
 				for (const FName& PhysicsBoneName : PhysicsToDeform)
 				{
-					MeshNode->AddPhysicsBodyToDeform(GenerationContext.BoneNames.Find(PhysicsBoneName));
+					MeshNode->AddPhysicsBodyToDeform(GenerationContext.GetBoneUnique(PhysicsBoneName));
 				}	
 			}
 			
diff --git a/Engine/Plugins/Experimental/Mutable/Source/CustomizableObjectEditor/Private/MuCOE/GenerateMutableSource/GenerateMutableSourceModifier.cpp b/Engine/Plugins/Experimental/Mutable/Source/CustomizableObjectEditor/Private/MuCOE/GenerateMutableSource/GenerateMutableSourceModifier.cpp
--- a/Engine/Plugins/Experimental/Mutable/Source/CustomizableObjectEditor/Private/MuCOE/GenerateMutableSource/GenerateMutableSourceModifier.cpp
+++ b/Engine/Plugins/Experimental/Mutable/Source/CustomizableObjectEditor/Private/MuCOE/GenerateMutableSource/GenerateMutableSourceModifier.cpp
@@ -61,7 +61,7 @@ mu::NodeModifierPtr GenerateMutableSourceModifier(const UEdGraphPin * Pin, FMuta
 		ClipNode->SetParams(TypedNodeClip->B, TypedNodeClip->Exponent);
 		ClipNode->SetMorphEllipse(TypedNodeClip->Radius, TypedNodeClip->Radius2, TypedNodeClip->RotationAngle);
 
-		ClipNode->SetVertexSelectionBone(GenerationContext.BoneNames.AddUnique(TypedNodeClip->BoneName), TypedNodeClip->MaxEffectRadius);
+		ClipNode->SetVertexSelectionBone(GenerationContext.GetBoneUnique(TypedNodeClip->BoneName), TypedNodeClip->MaxEffectRadius);
 
 		ClipNode->SetMultipleTagPolicy( TypedNodeClip->MultipleTagPolicy );
 		for (const FString& Tag : TypedNodeClip->Tags)
diff --git a/Engine/Plugins/Experimental/Mutable/Source/CustomizableObjectEditor/Private/MuCOE/MutableMeshPreviewUtils.cpp b/Engine/Plugins/Experimental/Mutable/Source/CustomizableObjectEditor/Private/MuCOE/MutableMeshPreviewUtils.cpp
--- a/Engine/Plugins/Experimental/Mutable/Source/CustomizableObjectEditor/Private/MuCOE/MutableMeshPreviewUtils.cpp
+++ b/Engine/Plugins/Experimental/Mutable/Source/CustomizableObjectEditor/Private/MuCOE/MutableMeshPreviewUtils.cpp
@@ -9,6 +9,7 @@
 #include "MuCO/MutableMeshBufferUtils.h"
 #include "MuCO/UnrealConversionUtils.h"
 #include "MuCO/UnrealPortabilityHelpers.h"
+#include "MuR/Skeleton.h"
 #include "MuR/OpMeshFormat.h"
 #include "Engine/SkeletalMesh.h"
 #include "MuCO/CustomizableObject.h"
@@ -314,7 +315,7 @@ namespace MutableMeshPreviewUtils
 				for (int32 BoneIndex = 0; BoneIndex < BonePoseCount; ++BoneIndex)
 				{
 					FTransform3f Transform;
-					InMutableMesh->GetBoneTransform(BoneIndex, Transform);
+					InMutableMesh->GetBonePoseTransform(BoneIndex, Transform);
 
 					Points.Add(FVector(Transform.GetTranslation()));
 				}
@@ -374,9 +375,12 @@ namespace MutableMeshPreviewUtils
 			const int32 NumBonesInBoneMap = !InMutableMesh->GetBoneMap().IsEmpty() ? InMutableMesh->GetBoneMap().Num() : InMutableMesh->GetBonePoseCount();
 			
 			// Fill the bonemap with zeros
-			TArray<uint16> BoneMap;
+			TArray<mu::FBoneName> BoneMap;
 			BoneMap.SetNumZeroed(FMath::Max(NumBonesInBoneMap, 1));
 
+			TMap <mu::FBoneName, TPair<FName, uint16>> BoneInfoMap;
+			BoneInfoMap.Add(mu::FBoneName(0), {NAME_None, 0});
+
 			FSkeletalMeshLODRenderData& LODResource = OutSkeletalMesh->GetResourceForRendering()->LODRenderData[0];
 
 			// Load buffer data found on the mutable model onto the out skeletal mesh
@@ -385,6 +389,7 @@ namespace MutableMeshPreviewUtils
 				LODResource,
 				InMutableMesh,
 				BoneMap,
+				BoneInfoMap,
 				0);
 
 			UnrealConversionUtils::CopyMutableVertexBuffers(
diff --git a/Engine/Plugins/Experimental/Mutable/Source/CustomizableObjectEditor/Private/MuCOE/SMutableSkeletonViewer.cpp b/Engine/Plugins/Experimental/Mutable/Source/CustomizableObjectEditor/Private/MuCOE/SMutableSkeletonViewer.cpp
--- a/Engine/Plugins/Experimental/Mutable/Source/CustomizableObjectEditor/Private/MuCOE/SMutableSkeletonViewer.cpp
+++ b/Engine/Plugins/Experimental/Mutable/Source/CustomizableObjectEditor/Private/MuCOE/SMutableSkeletonViewer.cpp
@@ -123,14 +123,14 @@ void SMutableSkeletonViewer::FillBoneDefinitionArrays()
 
 		// Save the bone data
 		NewBoneDefinition->BoneIndex = BoneIndex;
-		NewBoneDefinition->BoneName = Skeleton->GetBoneFName(BoneIndex);
+		NewBoneDefinition->BoneName = Skeleton->GetDebugName(BoneIndex);
 
 		// Get the parent node if any
 		const int32 ParentBoneIndex = Skeleton->GetBoneParent(BoneIndex);
 		if (ParentBoneIndex != INDEX_NONE)
 		{
 			NewBoneDefinition->ParentBoneIndex = ParentBoneIndex;
-			NewBoneDefinition->ParentBoneName = Skeleton->GetBoneFName(ParentBoneIndex);
+			NewBoneDefinition->ParentBoneName = Skeleton->GetDebugName(ParentBoneIndex);
 
 			// Store child bones
 			ChildBones.Add(NewBoneDefinition);
diff --git a/Engine/Plugins/Experimental/Mutable/Source/MutableRuntime/Internal/MuR/OpMeshMerge.h b/Engine/Plugins/Experimental/Mutable/Source/MutableRuntime/Internal/MuR/OpMeshMerge.h
--- a/Engine/Plugins/Experimental/Mutable/Source/MutableRuntime/Internal/MuR/OpMeshMerge.h
+++ b/Engine/Plugins/Experimental/Mutable/Source/MutableRuntime/Internal/MuR/OpMeshMerge.h
@@ -457,7 +457,7 @@ namespace mu
 			// Merge pSecond and build the remap array 
 			for (int32 SecondBoneIndex = 0; SecondBoneIndex < NumBonesSecond; ++SecondBoneIndex)
 			{
-				const uint16 BoneNameId = pSecondSkeleton->BoneIds[SecondBoneIndex];
+				const FBoneName& BoneNameId = pSecondSkeleton->BoneIds[SecondBoneIndex];
 				int32 Index = pResultSkeleton->FindBone(BoneNameId);
 
 				// Add a new bone
@@ -467,11 +467,12 @@ namespace mu
 
 					// Add an incorrect index, to be fixed below in case the parent index is later in the bone array.
 					pResultSkeleton->BoneParents.Add(pSecondSkeleton->BoneParents[SecondBoneIndex]);
-
-					if (pSecondSkeleton->BoneNames.IsValidIndex(SecondBoneIndex))
+#if WITH_EDITOR
+					if (pSecondSkeleton->DebugBoneNames.IsValidIndex(SecondBoneIndex))
 					{
-						pResultSkeleton->BoneNames.Add(pSecondSkeleton->BoneNames[SecondBoneIndex]);
+						pResultSkeleton->DebugBoneNames.Add(pSecondSkeleton->DebugBoneNames[SecondBoneIndex]);
 					}
+#endif
 				}
 
 				SecondToResultBoneIndices[SecondBoneIndex] = (uint16)Index;
@@ -611,11 +612,11 @@ namespace mu
 			// Appends InPhysicsBody to OutPhysicsBody removing Bodies that are equal, have same bone and customId and its properies are identical.	
 			auto AppendPhysicsBodiesUnique = [](PhysicsBody& OutPhysicsBody, const PhysicsBody& InPhysicsBody) -> bool
 			{
-				TArray<uint16>& OutBones = OutPhysicsBody.BoneIds;
+				TArray<FBoneName>& OutBones = OutPhysicsBody.BoneIds;
 				TArray<int32>& OutCustomIds = OutPhysicsBody.BodiesCustomIds;
 				TArray<FPhysicsBodyAggregate>& OutBodies = OutPhysicsBody.Bodies;
 
-				const TArray<uint16>& InBones = InPhysicsBody.BoneIds;
+				const TArray<FBoneName>& InBones = InPhysicsBody.BoneIds;
 				const TArray<int32>& InCustomIds = InPhysicsBody.BodiesCustomIds;
 				const TArray<FPhysicsBodyAggregate>& InBodies = InPhysicsBody.Bodies;
 
@@ -1165,28 +1166,30 @@ namespace mu
         TMap<int,int> otherToResult;
 
         int initialBones = pBase->GetBoneCount();
-        for ( int b=0; pOther && b<pOther->GetBoneCount(); ++b)
-        {
-            int32 resultBoneIndex = pBase->FindBone( pOther->GetBoneId(b) );
-            if ( resultBoneIndex<0 )
-            {
-                int32 newIndex = pBase->BoneIds.Num();
-                otherToResult.Add(b,newIndex);
-                pBase->BoneIds.Add( pOther->BoneIds[b] );
+		for (int b = 0; pOther && b < pOther->GetBoneCount(); ++b)
+		{
+			int32 resultBoneIndex = pBase->FindBone(pOther->GetBoneName(b));
+			if (resultBoneIndex < 0)
+			{
+				int32 newIndex = pBase->BoneIds.Num();
+				otherToResult.Add(b, newIndex);
+				pBase->BoneIds.Add(pOther->BoneIds[b]);
 
-                // Will be remapped below
-                pBase->BoneParents.Add(pOther->BoneParents[b] );
+				// Will be remapped below
+				pBase->BoneParents.Add(pOther->BoneParents[b]);
 
-				if (pOther->BoneNames.IsValidIndex(b))
+#if WITH_EDITOR
+				if (pOther->DebugBoneNames.IsValidIndex(b))
 				{
-					pBase->BoneNames.Add(pOther->BoneNames[b]);
+					pBase->DebugBoneNames.Add(pOther->DebugBoneNames[b]);
 				}
-            }
-            else
-            {
-                otherToResult.Add(b,resultBoneIndex);
-            }
-        }
+#endif
+			}
+			else
+			{
+				otherToResult.Add(b, resultBoneIndex);
+			}
+		}
 
         // Fix bone parent indices of the bones added from pOther
         for ( int b=initialBones;b<pBase->GetBoneCount(); ++b)
@@ -1213,34 +1216,8 @@ namespace mu
             return;
         }
 
-  //      mu::SkeletonPtrConst SourceSkeleton = SourceMesh->GetSkeleton();
-		//const TArray<uint16>& SourceBoneMaps = SourceMesh->GetBoneMap();
-
-  //      // Remap the indices of the bonemap to those of the new skeleton
-		//TArray<uint16> RemappedBoneMap;
-
-		//const int32 NumBonesBoneMap = SourceBoneMaps.Num();
-		//RemappedBoneMap.Reserve(NumBonesBoneMap);
-
-		//bool bBonesRemapped = false;
-
-		//for (const uint16& SourceBoneIndex : SourceBoneMaps)
-		//{
-		//	const uint16 BoneIndex = Skeleton->FindBone(SourceSkeleton->GetBoneId(SourceBoneIndex));
-
-		//	bBonesRemapped = bBonesRemapped || BoneIndex != SourceBoneIndex;
-		//	RemappedBoneMap.Add(BoneIndex);
-		//}
-
-  //      if (!bBonesRemapped)
-  //      {
-		//	bOutSuccess = false;
-  //          return;
-  //      }
-
 		Result->CopyFrom(*SourceMesh);
 		Result->SetSkeleton(Skeleton);
-		//Result->SetBoneMap(RemappedBoneMap);
     }
 	
 }
diff --git a/Engine/Plugins/Experimental/Mutable/Source/MutableRuntime/Private/MuR/CodeRunner.cpp b/Engine/Plugins/Experimental/Mutable/Source/MutableRuntime/Private/MuR/CodeRunner.cpp
--- a/Engine/Plugins/Experimental/Mutable/Source/MutableRuntime/Private/MuR/CodeRunner.cpp
+++ b/Engine/Plugins/Experimental/Mutable/Source/MutableRuntime/Private/MuR/CodeRunner.cpp
@@ -2022,7 +2022,7 @@ namespace mu
 					Ptr<Mesh> Result = CreateMesh(Source ? Source->GetDataSize() : 0);
 
 					bool bOutSuccess = false;
-					MeshClipMorphPlane(Result.get(), Source.get(), origin, normal, args.dist, args.factor, morphShape.size[0], morphShape.size[1], morphShape.size[2], selectionShape, bOutSuccess, INDEX_NONE, -1);
+					MeshClipMorphPlane(Result.get(), Source.get(), origin, normal, args.dist, args.factor, morphShape.size[0], morphShape.size[1], morphShape.size[2], selectionShape, bOutSuccess, nullptr, -1);
 					
 					if (!bOutSuccess)
 					{
@@ -2043,8 +2043,11 @@ namespace mu
 
 					Ptr<Mesh> Result = CreateMesh(Source->GetDataSize());
 
+					check(args.vertexSelectionShapeOrBone <= MAX_uint32);
+					const FBoneName Bone(args.vertexSelectionShapeOrBone);
+
 					bool bOutSuccess = false;
-					MeshClipMorphPlane(Result.get(), Source.get(), origin, normal, args.dist, args.factor, morphShape.size[0], morphShape.size[1], morphShape.size[2], selectionShape, bOutSuccess, args.vertexSelectionShapeOrBone, args.maxBoneRadius);
+					MeshClipMorphPlane(Result.get(), Source.get(), origin, normal, args.dist, args.factor, morphShape.size[0], morphShape.size[1], morphShape.size[2], selectionShape, bOutSuccess, &Bone, args.maxBoneRadius);
 
 					if (!bOutSuccess)
 					{
@@ -2066,7 +2069,7 @@ namespace mu
 					Ptr<Mesh> Result = CreateMesh(Source ? Source->GetDataSize() : 0);
 
 					bool bOutSuccess = false;
-					MeshClipMorphPlane(Result.get(), Source.get(), origin, normal, args.dist, args.factor, morphShape.size[0], morphShape.size[1], morphShape.size[2], selectionShape, bOutSuccess, INDEX_NONE, -1.0f);
+					MeshClipMorphPlane(Result.get(), Source.get(), origin, normal, args.dist, args.factor, morphShape.size[0], morphShape.size[1], morphShape.size[2], selectionShape, bOutSuccess, nullptr, -1.0f);
 
 					if (!bOutSuccess)
 					{
@@ -2376,19 +2379,19 @@ namespace mu
 					FMemory::Memcpy(&NumBones, Data, sizeof(int32)); 
 					Data += sizeof(int32);
 					
-					TArray<uint16> BonesToDeform;
+					TArray<FBoneName> BonesToDeform;
 					BonesToDeform.SetNumUninitialized(NumBones);
-					FMemory::Memcpy(BonesToDeform.GetData(), Data, NumBones * sizeof(uint16));
-					Data += NumBones * sizeof(uint16);
+					FMemory::Memcpy(BonesToDeform.GetData(), Data, NumBones * sizeof(FBoneName));
+					Data += NumBones * sizeof(FBoneName);
 
 					int32 NumPhysicsBodies;
 					FMemory::Memcpy(&NumPhysicsBodies, Data, sizeof(int32)); 
 					Data += sizeof(int32);
 
-					TArray<uint16> PhysicsToDeform;
+					TArray<FBoneName> PhysicsToDeform;
 					PhysicsToDeform.SetNumUninitialized(NumPhysicsBodies);
-					FMemory::Memcpy(PhysicsToDeform.GetData(), Data, NumPhysicsBodies * sizeof(uint16));
-					Data += NumPhysicsBodies * sizeof(uint16);
+					FMemory::Memcpy(PhysicsToDeform.GetData(), Data, NumPhysicsBodies * sizeof(FBoneName));
+					Data += NumPhysicsBodies * sizeof(FBoneName);
 
 					const EMeshBindShapeFlags BindFlags = static_cast<EMeshBindShapeFlags>(Args.flags);
 
@@ -2627,7 +2630,7 @@ namespace mu
                 {
                     if ( Source->GetSkeleton()
                          &&
-                         !Source->GetSkeleton()->BoneIds.IsEmpty() )
+                         Source->GetSkeleton()->GetBoneCount() > 0 )
                     {
                         // For some reason we already have bone data, so we can't just overwrite it
                         // or the skinning may break. This may happen because of a problem in the
diff --git a/Engine/Plugins/Experimental/Mutable/Source/MutableRuntime/Private/MuR/Mesh.cpp b/Engine/Plugins/Experimental/Mutable/Source/MutableRuntime/Private/MuR/Mesh.cpp
--- a/Engine/Plugins/Experimental/Mutable/Source/MutableRuntime/Private/MuR/Mesh.cpp
+++ b/Engine/Plugins/Experimental/Mutable/Source/MutableRuntime/Private/MuR/Mesh.cpp
@@ -545,7 +545,7 @@ const TArray<uint32>& Mesh::GetStreamedResources() const
 
 
 //---------------------------------------------------------------------------------------------
-int32 Mesh::FindBonePose(const uint16 BoneId) const
+int32 Mesh::FindBonePose(const FBoneName& BoneId) const
 {
 	return BonePoses.IndexOfByPredicate([BoneId](const FBonePose& Pose) { return Pose.BoneId == BoneId; });
 }
@@ -567,7 +567,7 @@ int32 mu::Mesh::GetBonePoseCount() const
 
 
 //---------------------------------------------------------------------------------------------
-void mu::Mesh::SetBonePose(int32 Index, uint16 BoneId, FTransform3f Transform, EBoneUsageFlags BoneUsageFlags)
+void mu::Mesh::SetBonePose(int32 Index, const FBoneName& BoneId, FTransform3f Transform, EBoneUsageFlags BoneUsageFlags)
 {
 	check(BonePoses.IsValidIndex(Index));
 	if (BonePoses.IsValidIndex(Index))
@@ -578,20 +578,15 @@ void mu::Mesh::SetBonePose(int32 Index, uint16 BoneId, FTransform3f Transform, E
 
 
 //---------------------------------------------------------------------------------------------
-int32 Mesh::GetBonePoseBoneId(int32 Index) const
+const FBoneName& Mesh::GetBonePoseId(int32 Index) const
 {
 	check(BonePoses.IsValidIndex(Index));
-	if (BonePoses.IsValidIndex(Index))
-	{
-		return BonePoses[Index].BoneId;
-	}
-
-	return INDEX_NONE;
+	return BonePoses[Index].BoneId;
 }
 
 
 //---------------------------------------------------------------------------------------------
-void mu::Mesh::GetBoneTransform(int32 BoneIndex, FTransform3f& Transform) const
+void mu::Mesh::GetBonePoseTransform(int32 BoneIndex, FTransform3f& Transform) const
 {
 	check(BoneIndex >= 0 && BoneIndex < BonePoses.Num());
 	Transform = BoneIndex > INDEX_NONE ? BonePoses[BoneIndex].BoneTransform : FTransform3f::Identity;
@@ -607,14 +602,14 @@ EBoneUsageFlags Mesh::GetBoneUsageFlags(int32 BoneIndex) const
 
 
 //---------------------------------------------------------------------------------------------
-void Mesh::SetBoneMap(const TArray<uint16>& InBoneMap)
+void Mesh::SetBoneMap(const TArray<FBoneName>& InBoneMap)
 {
 	BoneMap = InBoneMap;
 }
 
 
 //---------------------------------------------------------------------------------------------
-const TArray<uint16>& Mesh::GetBoneMap() const
+const TArray<FBoneName>& Mesh::GetBoneMap() const
 {
 	return BoneMap;
 }
@@ -1179,7 +1174,7 @@ void Mesh::FBonePose::Unserialise(InputArchive& arch)
 		std::string DeprecatedBoneName;
 		arch >> DeprecatedBoneName;
 
-		BoneId = 0;
+		BoneId = FBoneName(0);
 	}
 	else
 	{
@@ -1322,7 +1317,7 @@ void Mesh::Unserialise(InputArchive& arch)
 
 		for (int32 BoneIndex = 0; BoneIndex < NumBones; ++BoneIndex)
 		{
-			BonePoses[BoneIndex].BoneId = BoneIndex;
+			BonePoses[BoneIndex].BoneId = FBoneName(static_cast<uint32>(BoneIndex));
 			BonePoses[BoneIndex].BoneUsageFlags = EBoneUsageFlags::Skinning;
 			BonePoses[BoneIndex].BoneTransform = m_pSkeleton->m_boneTransforms_DEPRECATED[BoneIndex];
 		}
@@ -1338,7 +1333,7 @@ void Mesh::Unserialise(InputArchive& arch)
 		BoneMap.SetNum(NumBonePoses);
 		for (int32 BoneIndex = 0; BoneIndex < NumBonePoses; ++BoneIndex)
 		{
-			BoneMap[BoneIndex] = BoneIndex;
+			BoneMap[BoneIndex] = FBoneName(static_cast<uint32>(BoneIndex));
 		}
 
 		for (MESH_SURFACE& Surface : m_surfaces)
diff --git a/Engine/Plugins/Experimental/Mutable/Source/MutableRuntime/Private/MuR/OpMeshApplyPose.h b/Engine/Plugins/Experimental/Mutable/Source/MutableRuntime/Private/MuR/OpMeshApplyPose.h
--- a/Engine/Plugins/Experimental/Mutable/Source/MutableRuntime/Private/MuR/OpMeshApplyPose.h
+++ b/Engine/Plugins/Experimental/Mutable/Source/MutableRuntime/Private/MuR/OpMeshApplyPose.h
@@ -38,7 +38,7 @@ namespace mu
 
 		for (int32 BoneIndex = 0; BoneIndex < NumBones; ++BoneIndex)
 		{
-			const int32 BonePoseIndex = pPose->FindBonePose(pSkeleton->GetBoneId(BoneIndex));
+			const int32 BonePoseIndex = pPose->FindBonePose(pSkeleton->GetBoneName(BoneIndex));
 			if (BonePoseIndex != INDEX_NONE)
 			{
 				BoneToPoseIndex[BoneIndex] = BonePoseIndex;
@@ -54,7 +54,7 @@ namespace mu
 		}
 
 
-		TArray<uint16> BoneMap = pBase->GetBoneMap();
+		const TArray<FBoneName>& BoneMap = pBase->GetBoneMap();
 		const int32 NumBonesBoneMap = BoneMap.Num();
 
         // Prepare the skin matrices. They may be in different order, and we only need the ones
diff --git a/Engine/Plugins/Experimental/Mutable/Source/MutableRuntime/Private/MuR/OpMeshBind.h b/Engine/Plugins/Experimental/Mutable/Source/MutableRuntime/Private/MuR/OpMeshBind.h
--- a/Engine/Plugins/Experimental/Mutable/Source/MutableRuntime/Private/MuR/OpMeshBind.h
+++ b/Engine/Plugins/Experimental/Mutable/Source/MutableRuntime/Private/MuR/OpMeshBind.h
@@ -428,7 +428,7 @@ namespace mu
 
 	inline TTuple<TArray<FReshapePointBindingData>, TArray<int32>, TArray<int32>> BindPhysicsBodies( 
 			TArray<const PhysicsBody*> PhysicsBodies, FShapeMeshTree& ShapeMeshTree, const Mesh* pMesh, 
-			const TArray<uint16>& PhysicsToDeform )
+			const TArray<FBoneName>& PhysicsToDeform )
 	{
 		TTuple<TArray<FReshapePointBindingData>, TArray<int32>, TArray<int32>> ReturnValue;
 
@@ -519,7 +519,7 @@ namespace mu
 				FTransform3f T = FTransform3f::Identity;
 				if (BoneIndex > 0)
 				{
-					pMesh->GetBoneTransform(BoneIndex, T);
+					pMesh->GetBonePoseTransform(BoneIndex, T);
 				}
 
 				const int32 SphereCount = Body->GetSphereCount(B);
@@ -900,7 +900,7 @@ namespace mu
 
 
 	inline TTuple<TArray<FReshapePointBindingData>, TArray<int32>> BindPose(
-			const Mesh* Mesh, FShapeMeshTree& ShapeMeshTree, const TArray<uint16>& BonesToDeform )
+			const Mesh* Mesh, FShapeMeshTree& ShapeMeshTree, const TArray<FBoneName>& BonesToDeform )
 	{
 		UE::Geometry::FAxisAlignedBox3d ShapeAABBox = ShapeMeshTree.GetBoundingBox();
 
@@ -981,7 +981,7 @@ namespace mu
     inline void MeshBindShapeReshape(
 			Mesh* Result,
 			const Mesh* BaseMesh, const Mesh* ShapeMesh, 
-			const TArray<uint16>& BonesToDeform, const TArray<uint16>& PhysicsToDeform, 
+			const TArray<FBoneName>& BonesToDeform, const TArray<FBoneName>& PhysicsToDeform,
 			EMeshBindShapeFlags BindFlags, FMeshBindColorChannelUsages ColorChannelUsages,
 			bool& bOutSuccess)
     {
diff --git a/Engine/Plugins/Experimental/Mutable/Source/MutableRuntime/Private/MuR/OpMeshClipMorphPlane.h b/Engine/Plugins/Experimental/Mutable/Source/MutableRuntime/Private/MuR/OpMeshClipMorphPlane.h
--- a/Engine/Plugins/Experimental/Mutable/Source/MutableRuntime/Private/MuR/OpMeshClipMorphPlane.h
+++ b/Engine/Plugins/Experimental/Mutable/Source/MutableRuntime/Private/MuR/OpMeshClipMorphPlane.h
@@ -70,7 +70,7 @@ namespace mu
 	//! Reference version
 	//---------------------------------------------------------------------------------------------
 	inline void MeshClipMorphPlane(Mesh* Result, const Mesh* pBase, const FVector3f& origin, const FVector3f& normal, float dist, float factor, float radius,
-		float radius2, float angle, const FShape& selectionShape, bool& bOutSuccess, const int32 BoneId = INDEX_NONE, float vertexSelectionBoneMaxRadius = -1.f)
+		float radius2, float angle, const FShape& selectionShape, bool& bOutSuccess, const FBoneName* BoneId = nullptr, float vertexSelectionBoneMaxRadius = -1.f)
 	{
 		bOutSuccess = true;
 		//float radius = 8.f;
@@ -101,10 +101,10 @@ namespace mu
         Ptr<const Skeleton> BaseSkeleton = pBase->GetSkeleton();
 
 		TArray<bool> AffectedBoneMapIndices;
-		const int32 BaseBoneIndex = BaseSkeleton ? BaseSkeleton->FindBone(BoneId) : INDEX_NONE;
+		const int32 BaseBoneIndex = BaseSkeleton && BoneId ? BaseSkeleton->FindBone(*BoneId) : INDEX_NONE;
         if (BaseBoneIndex != INDEX_NONE)
 		{
-			const TArray<uint16>& BoneMap = pBase->BoneMap;
+			const TArray<FBoneName>& BoneMap = pBase->BoneMap;
 			AffectedBoneMapIndices.SetNum(BoneMap.Num());
 
 			const int32 BoneCount = BaseSkeleton->GetBoneCount();
@@ -126,8 +126,8 @@ namespace mu
 				
 				AffectedSkeletonBones[BoneIndex] = true;
 
-				const uint16 AffectedBoneId = BaseSkeleton->GetBoneId(BoneIndex);
-				const int32 AffectedBoneMapIndex = BoneMap.Find(AffectedBoneId);
+				const FBoneName& AffectedBone = BaseSkeleton->GetBoneName(BoneIndex);
+				const int32 AffectedBoneMapIndex = BoneMap.Find(AffectedBone);
 				if (AffectedBoneMapIndex != INDEX_NONE)
 				{
 					AffectedBoneMapIndices[AffectedBoneMapIndex] = true;
diff --git a/Engine/Plugins/Experimental/Mutable/Source/MutableRuntime/Private/MuR/OpMeshReshape.h b/Engine/Plugins/Experimental/Mutable/Source/MutableRuntime/Private/MuR/OpMeshReshape.h
--- a/Engine/Plugins/Experimental/Mutable/Source/MutableRuntime/Private/MuR/OpMeshReshape.h
+++ b/Engine/Plugins/Experimental/Mutable/Source/MutableRuntime/Private/MuR/OpMeshReshape.h
@@ -549,11 +549,11 @@ namespace mu
 		// Retrieve them in the same order the boxes where put in, so they can be linked to the physics body volumes.
 		for ( const int32 B : UsedIndices )
 		{	
-			int32 BoneIdx = BaseMesh.FindBonePose(PBody.GetBodyBoneId(B));
+			int32 BoneIndex = BaseMesh.FindBonePose(PBody.GetBodyBoneId(B));
 			FTransform3f BoneTransform = FTransform3f::Identity;
-			if (BoneIdx >= 0)
+			if (BoneIndex >= 0)
 			{
-				BaseMesh.GetBoneTransform(BoneIdx, BoneTransform);
+				BaseMesh.GetBonePoseTransform(BoneIndex, BoneTransform);
 			}
 			
 			FTransform3f InvBoneTransform = BoneTransform.Inverse();
diff --git a/Engine/Plugins/Experimental/Mutable/Source/MutableRuntime/Private/MuR/PhysicsBody.cpp b/Engine/Plugins/Experimental/Mutable/Source/MutableRuntime/Private/MuR/PhysicsBody.cpp
--- a/Engine/Plugins/Experimental/Mutable/Source/MutableRuntime/Private/MuR/PhysicsBody.cpp
+++ b/Engine/Plugins/Experimental/Mutable/Source/MutableRuntime/Private/MuR/PhysicsBody.cpp
@@ -279,14 +279,14 @@ namespace mu
 	}
 
 	//-------------------------------------------------------------------------------------------------
-	uint16 PhysicsBody::GetBodyBoneId(int32 B) const
+	const FBoneName& PhysicsBody::GetBodyBoneId(int32 B) const
 	{
 		check(BoneIds.IsValidIndex(B));
 		return BoneIds[B];
 	}
 
 	//-------------------------------------------------------------------------------------------------
-	void PhysicsBody::SetBodyBoneId(int32 B, const uint16 BoneId)
+	void PhysicsBody::SetBodyBoneId(int32 B, const FBoneName& BoneId)
 	{
 		check(BoneIds.IsValidIndex(B));
 		BoneIds[B] = BoneId;
@@ -719,7 +719,7 @@ namespace mu
 	//-------------------------------------------------------------------------------------------------
 	void PhysicsBody::Serialise(OutputArchive& arch) const
 	{
-		uint32 ver = 3;
+		uint32 ver = 4;
 		arch << ver;
         	
 		arch << CustomId;
@@ -735,7 +735,7 @@ namespace mu
 	{
 		uint32 ver;
 		arch >> ver;
-		check(ver <= 3);
+		check(ver <= 4);
         	
 		if (ver >= 2)
 		{
@@ -744,10 +744,22 @@ namespace mu
         	
 		arch >> Bodies;
 
-		if (ver >= 3)
+		if (ver >= 4)
 		{
 			arch >> BoneIds;
 		}
+		else if (ver == 3)
+		{
+			TArray<uint16> BoneIds_DEPRECATED;
+			arch >> BoneIds_DEPRECATED;
+
+			const int32 NumBoneNames = BoneIds_DEPRECATED.Num();
+			BoneIds.SetNum(NumBoneNames);
+			for (int32 BoneIndex = 0; BoneIndex < NumBoneNames; ++BoneIndex)
+			{
+				BoneIds[BoneIndex].Id = BoneIds_DEPRECATED[BoneIndex];
+			}
+		}
 		else
 		{
 			TArray<std::string> BoneNames;
@@ -759,7 +771,7 @@ namespace mu
 			BoneIds.Reserve(NumBoneNames);
 			for (int32 BoneIndex = 0; BoneIndex < NumBoneNames; ++BoneIndex)
 			{
-				BoneIds.Add((uint16)UniqueBoneNames.AddUnique(BoneNames[BoneIndex]));
+				BoneIds.Add(FBoneName(BoneIndex));
 			}
 		}
 		arch >> BodiesCustomIds;
diff --git a/Engine/Plugins/Experimental/Mutable/Source/MutableRuntime/Private/MuR/Skeleton.cpp b/Engine/Plugins/Experimental/Mutable/Source/MutableRuntime/Private/MuR/Skeleton.cpp
--- a/Engine/Plugins/Experimental/Mutable/Source/MutableRuntime/Private/MuR/Skeleton.cpp
+++ b/Engine/Plugins/Experimental/Mutable/Source/MutableRuntime/Private/MuR/Skeleton.cpp
@@ -35,7 +35,7 @@ namespace mu
 		pResult->BoneParents = BoneParents;
 
 		// For debug
-		pResult->BoneNames = BoneNames;
+		pResult->DebugBoneNames = DebugBoneNames;
 
 		return pResult;
 	}
@@ -51,18 +51,18 @@ namespace mu
 	//---------------------------------------------------------------------------------------------
 	void Skeleton::SetBoneCount(int32 NumBones)
 	{
-		BoneNames.SetNum(NumBones);
+		DebugBoneNames.SetNum(NumBones);
 		BoneIds.SetNum(NumBones);
 		BoneParents.Init(INDEX_NONE, NumBones);
 	}
 
 
 	//---------------------------------------------------------------------------------------------
-	const FName Skeleton::GetBoneFName(int32 Index) const
+	const FName Skeleton::GetDebugName(int32 Index) const
 	{
-		if (BoneNames.IsValidIndex(Index))
+		if (DebugBoneNames.IsValidIndex(Index))
 		{
-			return BoneNames[Index];
+			return DebugBoneNames[Index];
 		}
 
 		return FName("Unknown Bone");
@@ -70,11 +70,11 @@ namespace mu
 
 
 	//---------------------------------------------------------------------------------------------
-	void Skeleton::SetBoneFName(const int32 Index, const FName BoneName)
+	void Skeleton::SetDebugName(const int32 Index, const FName BoneName)
 	{
-		if (BoneNames.IsValidIndex(Index))
+		if (DebugBoneNames.IsValidIndex(Index))
 		{
-			BoneNames[Index] = BoneName;
+			DebugBoneNames[Index] = BoneName;
 		}
 	}
 
@@ -105,7 +105,7 @@ namespace mu
 
 
 	//-------------------------------------------------------------------------------------------------
-	uint16 Skeleton::GetBoneId(const int32 Index) const
+	const FBoneName& Skeleton::GetBoneName(int32 Index) const
 	{
 		check(BoneIds.IsValidIndex(Index));
 		return BoneIds[Index];
@@ -113,27 +113,27 @@ namespace mu
 
 	
 	//-------------------------------------------------------------------------------------------------
-	void Skeleton::SetBoneId(const int32 Index, const uint16 BoneId)
+	void Skeleton::SetBoneName(int32 Index, const FBoneName& BoneName)
 	{
 		check(BoneIds.IsValidIndex(Index));
 		if (BoneIds.IsValidIndex(Index))
 		{
-			BoneIds[Index] = BoneId;
+			BoneIds[Index] = BoneName;
 		}
 	}
 
 
 	//-------------------------------------------------------------------------------------------------
-	int32 Skeleton::FindBone(const uint16 BoneId) const
+	int32 Skeleton::FindBone(const FBoneName& BoneName) const
 	{
-		return BoneIds.Find(BoneId);
+		return BoneIds.Find(BoneName);
 	}
 
 
 	//-------------------------------------------------------------------------------------------------
 	void Skeleton::Serialise(OutputArchive& arch) const
 	{
-		uint32 ver = 6;
+		uint32 ver = 7;
 		arch << ver;
 
 		arch << BoneIds;
@@ -148,10 +148,22 @@ namespace mu
 		arch >> ver;
 		check(ver >= 3);
 
-		if(ver >= 6)
+		if (ver >= 7)
 		{
 			arch >> BoneIds;
 		}
+		else if(ver == 6)
+		{
+			TArray<uint16> BoneIds_DEPRECATED;
+			arch >> BoneIds_DEPRECATED;
+
+			const int32 NumBones = BoneIds_DEPRECATED.Num();
+			BoneIds.Reserve(NumBones);
+			for (int32 Index = 0; Index < NumBones; ++Index)
+			{
+				BoneIds.Add({ BoneIds_DEPRECATED[Index]});
+			}
+		}
 		else
 		{
 			TArray<std::string> OldBoneNames;
@@ -162,7 +174,7 @@ namespace mu
 			BoneIds.SetNum(NumBones);
 			for (int32 Index = 0; Index < NumBones; ++Index)
 			{
-				BoneIds[Index] = Index;
+				BoneIds[Index].Id = Index;
 			}
 		}
 
diff --git a/Engine/Plugins/Experimental/Mutable/Source/MutableRuntime/Public/MuR/Mesh.h b/Engine/Plugins/Experimental/Mutable/Source/MutableRuntime/Public/MuR/Mesh.h
--- a/Engine/Plugins/Experimental/Mutable/Source/MutableRuntime/Public/MuR/Mesh.h
+++ b/Engine/Plugins/Experimental/Mutable/Source/MutableRuntime/Public/MuR/Mesh.h
@@ -15,7 +15,6 @@
 #include "MuR/Ptr.h"
 #include "MuR/RefCounted.h"
 #include "MuR/Serialisation.h"
-#include "MuR/Serialisation.h"
 #include "MuR/Skeleton.h"
 #include "Templates/Tuple.h"
 
@@ -290,7 +289,7 @@ namespace mu
 		const TArray<uint32>& GetStreamedResources() const;
 
 		//!
-		int32 FindBonePose(uint16 BoneId) const;
+		int32 FindBonePose(const FBoneName& BoneName) const;
 		
 		//!
 		void SetBonePoseCount(int32 count);
@@ -299,22 +298,22 @@ namespace mu
 		int32 GetBonePoseCount() const;
 
 		//!
-		void SetBonePose(int32 Index, uint16 BoneId, FTransform3f Transform, EBoneUsageFlags BoneUsageFlags);
+		void SetBonePose(int32 Index, const FBoneName& BoneName, FTransform3f Transform, EBoneUsageFlags BoneUsageFlags);
 
-		//! Return BoneId (uint16) of the pose at index PoseIndex or INDEX_NONE
-		int32 GetBonePoseBoneId(int32 PoseIndex) const;
+		//! @return - Bone identifier of the pose at 'Index'.
+		const FBoneName& GetBonePoseId(int32 BoneIndex) const;
 
 		//! Return a matrix stored per bone. It is a set of 16-float values.
-		void GetBoneTransform(int32 BoneIndex, FTransform3f& Transform) const;
+		void GetBonePoseTransform(int32 BoneIndex, FTransform3f& Transform) const;
 
 		//! 
 		EBoneUsageFlags GetBoneUsageFlags(int32 BoneIndex) const;
 
 		//! Set the bonemap of this mesh
-		void SetBoneMap(const TArray<uint16>& InBoneMap);
+		void SetBoneMap(const TArray<FBoneName>& InBoneMap);
 
-		//! Return an array containing the bonemap indices of all surfaces in the mesh.
-		const TArray<uint16>& GetBoneMap() const;
+		//! Return an array containing the bonemaps of all surfaces in the mesh.
+		const TArray<FBoneName>& GetBoneMap() const;
 
 		//!
 		int32 GetSkeletonIDsCount() const;
@@ -390,18 +389,15 @@ namespace mu
 
 		struct FBonePose
 		{
-			// Index of the bone in the CO BoneNames array
-			uint16 BoneId;
+			// Identifier built from the bone FName.
+			FBoneName BoneId;
 
 			EBoneUsageFlags BoneUsageFlags = EBoneUsageFlags::None;
 			FTransform3f BoneTransform;
 
 			inline void Serialise(OutputArchive& arch) const;
-
-
 			inline void Unserialise(InputArchive& arch);
 
-
 			//!
 			inline bool operator==(const FBonePose& Other) const
 			{
@@ -413,7 +409,7 @@ namespace mu
 		TMemoryTrackedArray<FBonePose> BonePoses;
 
 		// Array containing the bonemaps of all surfaces in the mesh.
-		TArray<uint16> BoneMap;
+		TArray<FBoneName> BoneMap;
 
 		//!
 		inline void Serialise(OutputArchive& arch) const;
diff --git a/Engine/Plugins/Experimental/Mutable/Source/MutableRuntime/Public/MuR/PhysicsBody.h b/Engine/Plugins/Experimental/Mutable/Source/MutableRuntime/Public/MuR/PhysicsBody.h
--- a/Engine/Plugins/Experimental/Mutable/Source/MutableRuntime/Public/MuR/PhysicsBody.h
+++ b/Engine/Plugins/Experimental/Mutable/Source/MutableRuntime/Public/MuR/PhysicsBody.h
@@ -8,6 +8,7 @@
 #include "MuR/Ptr.h"
 #include "MuR/RefCounted.h"
 #include "MuR/Serialisation.h"
+#include "MuR/Skeleton.h"
 #include "Containers/Array.h"
 #include "Math/Quat.h"
 #include "Math/Transform.h"
@@ -187,8 +188,8 @@ namespace mu
 		void SetBodyCount(int32 B);
 		int32 GetBodyCount() const;
 	
-		void SetBodyBoneId(int32 B, uint16 BoneId);
-		uint16 GetBodyBoneId(int32 B) const;
+		void SetBodyBoneId(int32 B, const FBoneName& BoneId);
+		const FBoneName& GetBodyBoneId(int32 B) const;
 		
 		void SetBodyCustomId(int32 B, int32 BodyCustomId);
 		int32 GetBodyCustomId(int32 B) const;
@@ -295,7 +296,7 @@ namespace mu
 		int32 CustomId = -1;
 
 		// Bone name the physics volume aggregate is bound to. 
-		TArray<uint16> BoneIds;
+		TArray<FBoneName> BoneIds;
 		TArray<FPhysicsBodyAggregate> Bodies;
 		TArray<int32> BodiesCustomIds;
 
diff --git a/Engine/Plugins/Experimental/Mutable/Source/MutableRuntime/Public/MuR/Skeleton.h b/Engine/Plugins/Experimental/Mutable/Source/MutableRuntime/Public/MuR/Skeleton.h
--- a/Engine/Plugins/Experimental/Mutable/Source/MutableRuntime/Public/MuR/Skeleton.h
+++ b/Engine/Plugins/Experimental/Mutable/Source/MutableRuntime/Public/MuR/Skeleton.h
@@ -23,6 +23,37 @@ namespace mu
     typedef Ptr<Skeleton> SkeletonPtr;
     typedef Ptr<const Skeleton> SkeletonPtrConst;
 
+	// Bone name identifier
+	struct FBoneName
+	{
+		FBoneName() {};
+		FBoneName(uint32 InID) : Id(InID) {};
+
+		// Hash built from the bone name (FString)
+		uint32 Id = 0;
+
+		inline void Serialise(OutputArchive& arch) const
+		{
+			arch << Id;
+		}
+
+		inline void Unserialise(InputArchive& arch)
+		{
+			arch >> Id;
+		}
+
+		//!
+		inline bool operator==(const FBoneName& Other) const
+		{
+			return Id == Other.Id;
+		}
+	};
+
+	inline uint32 GetTypeHash(const FBoneName& Bone)
+	{
+		return Bone.Id;
+	}
+
 
     //! \brief Skeleton object.
 	//! \ingroup runtime
@@ -48,24 +79,24 @@ namespace mu
 		// Own interface
 		//-----------------------------------------------------------------------------------------
 
-        //! Return the number of bones used by this skeleton
+        //! @return - Number of bones in the Skeleton
 		int32 GetBoneCount() const;
         void SetBoneCount(int32 c);
 
-		//! DEBUG. Return the FName of the bone at index BoneIndex. Only valid in the editor
-		const FName GetBoneFName(int32 Index) const;
-		void SetBoneFName(const int32 Index, const FName BoneName);
+		//! @return - FName of the bone at 'Index'. Only valid in the editor
+		const FName GetDebugName(int32 Index) const;
+		void SetDebugName(const int32 Index, const FName BoneName);
 
         //! Get and set the parent bone of each bone. The parent can be -1 if the bone is a root.
         int32 GetBoneParent(int32 boneIndex) const;
         void SetBoneParent(int32 boneIndex, int32 parentBoneIndex);
 
-		//! Return the BoneIndex of the bone at index from the BoneIndices array
-		uint16 GetBoneId(int32 Index) const;
-		void SetBoneId(int32 Index, uint16 BoneIndex);
+		//! @return - BoneName of the Bone at 'Index'.
+		const FBoneName& GetBoneName(int32 Index) const;
+		void SetBoneName(int32 Index, const FBoneName& BoneName);
 
-		//! Return the index of BoneIndex inside the BoneIndices array.
-		int32 FindBone(const uint16 BoneIndex) const;
+		//! @return - Index in the Skeleton. INDEX_NONE if not found.
+		int32 FindBone(const FBoneName& BoneName) const;
 		
 
 	protected:
@@ -79,13 +110,11 @@ namespace mu
 		TArray<std::string> m_bones_DEPRECATED;
 		TArray<FTransform3f> m_boneTransforms_DEPRECATED;
 
-
-
 		//! DEBUG. FNames of the bones. Only valid in the editor. Do not serialize.
-		TArray<FName> BoneNames;
+		TArray<FName> DebugBoneNames;
 
-		//! Ids of the bones. The Id is the BoneName index in the CO BoneNames array
-		TArray<uint16> BoneIds;
+		//! Array of bone identifiers. 
+		TArray<FBoneName> BoneIds;
 
 		//! For each bone, index of the parent bone in the bone vectors. -1 means no parent.
 		//! This array must have the same size than the m_bones array.
diff --git a/Engine/Plugins/Experimental/Mutable/Source/MutableTools/Internal/MuT/NodeMeshMorphPrivate.h b/Engine/Plugins/Experimental/Mutable/Source/MutableTools/Internal/MuT/NodeMeshMorphPrivate.h
--- a/Engine/Plugins/Experimental/Mutable/Source/MutableTools/Internal/MuT/NodeMeshMorphPrivate.h
+++ b/Engine/Plugins/Experimental/Mutable/Source/MutableTools/Internal/MuT/NodeMeshMorphPrivate.h
@@ -6,6 +6,7 @@
 #include "MuT/NodeScalar.h"
 #include "MuT/NodeMeshPrivate.h"
 #include "MuT/AST.h"
+#include "MuR/Skeleton.h"
 
 namespace mu
 {
@@ -24,13 +25,13 @@ namespace mu
 		bool bReshapeSkeleton = false;
 		bool bReshapePhysicsVolumes = false;
 		
-		TArray<uint16> BonesToDeform;
-		TArray<uint16> PhysicsToDeform;
+		TArray<FBoneName> BonesToDeform;
+		TArray<FBoneName> PhysicsToDeform;
 
         //!
 		void Serialise( OutputArchive& arch ) const
 		{
-            uint32_t ver = 7;
+            uint32_t ver = 8;
 			arch << ver;
 
 			arch << Factor;
@@ -86,10 +87,22 @@ namespace mu
 					arch >> bReshapePhysicsVolumes;
 				}
 
-				if (ver >= 7)
+				if (ver >= 8)
 				{
 					arch >> BonesToDeform;
 				}
+				else if (ver == 7)
+				{
+					TArray<uint16> BonesToDeform_DEPRECATED;
+					arch >> BonesToDeform_DEPRECATED;
+
+					const int32 NumBonesToDeform = BonesToDeform_DEPRECATED.Num();
+					BonesToDeform.SetNum(NumBonesToDeform);
+					for (int32 Index = 0; Index < NumBonesToDeform; ++Index)
+					{
+						BonesToDeform[Index].Id = BonesToDeform_DEPRECATED[Index];
+					}
+				}
 				else
 				{
 					TArray<string> BonesToDeform_DEPRECATED;
@@ -99,7 +112,7 @@ namespace mu
 					BonesToDeform.SetNumUninitialized(NumBonesToDeform);
 					for (int32 Index = 0; Index < NumBonesToDeform; ++Index)
 					{
-						BonesToDeform[Index] = Index;
+						BonesToDeform[Index].Id = Index;
 					}
 				}
 			}
@@ -122,10 +135,22 @@ namespace mu
 				arch >> bDeformAllPhysics_DEPRECATED;
 			}
 
-			if (ver >= 7)
+			if (ver >= 8)
 			{
 				arch >> PhysicsToDeform;
 			}
+			else if (ver == 7)
+			{
+				TArray<uint16> PhysicsToDeform_DEPRECATED;
+				arch >> PhysicsToDeform_DEPRECATED;
+
+				const int32 NumPhysicsToDeform = PhysicsToDeform_DEPRECATED.Num();
+				PhysicsToDeform.SetNum(NumPhysicsToDeform);
+				for (int32 Index = 0; Index < NumPhysicsToDeform; ++Index)
+				{
+					PhysicsToDeform[Index].Id = PhysicsToDeform_DEPRECATED[Index];
+				}
+			}
 			else if (ver >= 3)
 			{
 				TArray<string> PhysicsToDeform_DEPRECATED;
@@ -135,7 +160,7 @@ namespace mu
 				PhysicsToDeform.SetNumUninitialized(NumPhysicsToDeform);
 				for (int32 Index = 0; Index < NumPhysicsToDeform; ++Index)
 				{
-					PhysicsToDeform[Index] = Index;
+					PhysicsToDeform[Index].Id = Index;
 				}
 			}
 			else
diff --git a/Engine/Plugins/Experimental/Mutable/Source/MutableTools/Internal/MuT/NodeModifierMeshClipMorphPlanePrivate.h b/Engine/Plugins/Experimental/Mutable/Source/MutableTools/Internal/MuT/NodeModifierMeshClipMorphPlanePrivate.h
--- a/Engine/Plugins/Experimental/Mutable/Source/MutableTools/Internal/MuT/NodeModifierMeshClipMorphPlanePrivate.h
+++ b/Engine/Plugins/Experimental/Mutable/Source/MutableTools/Internal/MuT/NodeModifierMeshClipMorphPlanePrivate.h
@@ -7,6 +7,7 @@
 #include "MuT/AST.h"
 
 #include "MuR/MutableMath.h"
+#include "MuR/Skeleton.h"
 
 
 namespace mu
@@ -55,7 +56,7 @@ namespace mu
         uint8 m_vertexSelectionType;
 		FVector3f m_selectionBoxOrigin;
 		FVector3f m_selectionBoxRadius;
-        uint16 m_vertexSelectionBone;
+        FBoneName m_vertexSelectionBone;
 
 		// Max distance a vertex can have to the bone in order to be affected. A negative value
 		// means no limit.
@@ -66,7 +67,7 @@ namespace mu
         {
             NodeModifier::Private::Serialise(arch);
 
-			uint32 ver = 3;
+			uint32 ver = 4;
             arch << ver;
 
             arch << m_origin;
@@ -102,15 +103,21 @@ namespace mu
             arch >> m_vertexSelectionType;
             arch >> m_selectionBoxOrigin;
             arch >> m_selectionBoxRadius;
-			if (ver >= 3)
+			if (ver >= 4)
 			{
 				arch >> m_vertexSelectionBone;
 			}
+			else if (ver == 3)
+			{
+				uint16 BoneId = 0;
+				arch >> BoneId;
+				m_vertexSelectionBone.Id = BoneId;
+			}
 			else
 			{
 				string OldVertexSelectionBone;
 				arch >> OldVertexSelectionBone;
-				m_vertexSelectionBone = 0;
+				m_vertexSelectionBone.Id = 0;
 			}
 			arch >> m_maxEffectRadius;
         }
diff --git a/Engine/Plugins/Experimental/Mutable/Source/MutableTools/Private/MuT/ASTOpMeshBindShape.cpp b/Engine/Plugins/Experimental/Mutable/Source/MutableTools/Private/MuT/ASTOpMeshBindShape.cpp
--- a/Engine/Plugins/Experimental/Mutable/Source/MutableTools/Private/MuT/ASTOpMeshBindShape.cpp
+++ b/Engine/Plugins/Experimental/Mutable/Source/MutableTools/Private/MuT/ASTOpMeshBindShape.cpp
@@ -79,14 +79,14 @@ namespace mu
 		hash_combine(Result, static_cast<uint32>(BChannelUsage));
 		hash_combine(Result, static_cast<uint32>(AChannelUsage));
 
-		for (const uint16 S : BonesToDeform)
+		for (const FBoneName& S : BonesToDeform)
 		{
-			hash_combine(Result, S);
+			hash_combine(Result, S.Id);
 		}
 
-		for (const uint16 S : PhysicsToDeform)
+		for (const FBoneName& S : PhysicsToDeform)
 		{
-			hash_combine(Result, S);
+			hash_combine(Result, S.Id);
 		}
 
 		return Result;
@@ -179,13 +179,13 @@ namespace mu
 			AppendCode(program.m_byteCode, Args);
 
 			AppendCode(program.m_byteCode, (int32)BonesToDeform.Num());
-			for (const uint16 S : BonesToDeform)
+			for (const FBoneName& S : BonesToDeform)
 			{
 				AppendCode(program.m_byteCode, S);
 			}
 
 			AppendCode(program.m_byteCode, (int32)PhysicsToDeform.Num());
-			for (const uint16 S : PhysicsToDeform)
+			for (const FBoneName& S : PhysicsToDeform)
 			{
 				AppendCode(program.m_byteCode, S);
 			}
diff --git a/Engine/Plugins/Experimental/Mutable/Source/MutableTools/Private/MuT/ASTOpMeshBindShape.h b/Engine/Plugins/Experimental/Mutable/Source/MutableTools/Private/MuT/ASTOpMeshBindShape.h
--- a/Engine/Plugins/Experimental/Mutable/Source/MutableTools/Private/MuT/ASTOpMeshBindShape.h
+++ b/Engine/Plugins/Experimental/Mutable/Source/MutableTools/Private/MuT/ASTOpMeshBindShape.h
@@ -6,13 +6,13 @@
 #include "HAL/Platform.h"
 #include "MuR/Operations.h"
 #include "MuR/Ptr.h"
+#include "MuR/Skeleton.h"
 #include "MuT/AST.h"
 
 
 namespace mu
 {
-struct FProgram;
-
+	struct FProgram;
 
 	//---------------------------------------------------------------------------------------------
 	//!
@@ -24,8 +24,8 @@ struct FProgram;
 		ASTChild Mesh;
 		ASTChild Shape;
 
-		TArray<uint16> BonesToDeform;
-		TArray<uint16> PhysicsToDeform;
+		TArray<FBoneName> BonesToDeform;
+		TArray<FBoneName> PhysicsToDeform;
 
 		uint32 BindingMethod = 0;
 		
diff --git a/Engine/Plugins/Experimental/Mutable/Source/MutableTools/Private/MuT/ASTOpMeshClipMorphPlane.cpp b/Engine/Plugins/Experimental/Mutable/Source/MutableTools/Private/MuT/ASTOpMeshClipMorphPlane.cpp
--- a/Engine/Plugins/Experimental/Mutable/Source/MutableTools/Private/MuT/ASTOpMeshClipMorphPlane.cpp
+++ b/Engine/Plugins/Experimental/Mutable/Source/MutableTools/Private/MuT/ASTOpMeshClipMorphPlane.cpp
@@ -93,7 +93,8 @@ namespace mu
 			args.vertexSelectionType = (uint8_t)vertexSelectionType;
 			if (vertexSelectionType == OP::MeshClipMorphPlaneArgs::VS_BONE_HIERARCHY)
 			{
-				args.vertexSelectionShapeOrBone = vertexSelectionBone;
+				check(vertexSelectionBone.Id < MAX_uint32);
+				args.vertexSelectionShapeOrBone = vertexSelectionBone.Id;
 			}
 			else if (vertexSelectionType == OP::MeshClipMorphPlaneArgs::VS_SHAPE)
 			{
diff --git a/Engine/Plugins/Experimental/Mutable/Source/MutableTools/Private/MuT/ASTOpMeshClipMorphPlane.h b/Engine/Plugins/Experimental/Mutable/Source/MutableTools/Private/MuT/ASTOpMeshClipMorphPlane.h
--- a/Engine/Plugins/Experimental/Mutable/Source/MutableTools/Private/MuT/ASTOpMeshClipMorphPlane.h
+++ b/Engine/Plugins/Experimental/Mutable/Source/MutableTools/Private/MuT/ASTOpMeshClipMorphPlane.h
@@ -24,7 +24,7 @@ struct FProgram;
 
 		FShape morphShape;
 		FShape selectionShape;
-		uint16 vertexSelectionBone;
+		FBoneName vertexSelectionBone;
 
 		OP::MeshClipMorphPlaneArgs::VERTEX_SELECTION_TYPE vertexSelectionType = OP::MeshClipMorphPlaneArgs::VS_NONE;
 
diff --git a/Engine/Plugins/Experimental/Mutable/Source/MutableTools/Private/MuT/NodeMeshClipMorphPlane.cpp b/Engine/Plugins/Experimental/Mutable/Source/MutableTools/Private/MuT/NodeMeshClipMorphPlane.cpp
--- a/Engine/Plugins/Experimental/Mutable/Source/MutableTools/Private/MuT/NodeMeshClipMorphPlane.cpp
+++ b/Engine/Plugins/Experimental/Mutable/Source/MutableTools/Private/MuT/NodeMeshClipMorphPlane.cpp
@@ -75,7 +75,7 @@ namespace mu
 	}
 
 	//---------------------------------------------------------------------------------------------
-	void NodeMeshClipMorphPlane::SetVertexSelectionBone(uint16 BoneId, float maxEffectRadius)
+	void NodeMeshClipMorphPlane::SetVertexSelectionBone(const FBoneName& BoneId, float maxEffectRadius)
 	{
 		m_pD->m_vertexSelectionType = Private::VS_BONE_HIERARCHY;
 		m_pD->m_vertexSelectionBone = BoneId;
diff --git a/Engine/Plugins/Experimental/Mutable/Source/MutableTools/Private/MuT/NodeMeshClipMorphPlanePrivate.h b/Engine/Plugins/Experimental/Mutable/Source/MutableTools/Private/MuT/NodeMeshClipMorphPlanePrivate.h
--- a/Engine/Plugins/Experimental/Mutable/Source/MutableTools/Private/MuT/NodeMeshClipMorphPlanePrivate.h
+++ b/Engine/Plugins/Experimental/Mutable/Source/MutableTools/Private/MuT/NodeMeshClipMorphPlanePrivate.h
@@ -7,7 +7,7 @@
 #include "MuT/AST.h"
 
 #include "MuR/MutableMath.h"
-
+#include "MuR/Skeleton.h"
 
 namespace mu
 {
@@ -60,7 +60,7 @@ namespace mu
         uint8 m_vertexSelectionType;
 		FVector3f m_selectionBoxOrigin;
 		FVector3f m_selectionBoxRadius;
-		uint16 m_vertexSelectionBone;
+		FBoneName m_vertexSelectionBone;
 
 		TArray<FString> m_tags;
 
@@ -71,7 +71,7 @@ namespace mu
 		//!
 		void Serialise( OutputArchive& arch ) const
 		{
-            uint32 ver = 5;
+            uint32 ver = 6;
 			arch << ver;
 
 			arch << m_pSource;
@@ -124,15 +124,21 @@ namespace mu
 				}
 			}
 
-			if (ver >= 4)
+			if (ver >= 6)
 			{
 				arch >> m_vertexSelectionBone;
 			}
+			else if (ver >= 4)
+			{
+				uint16 BoneId = 0;
+				arch >> BoneId;
+				m_vertexSelectionBone.Id = BoneId;
+			}
 			else
 			{
 				string OldVertexSelectionBone;
 				arch >> OldVertexSelectionBone;
-				m_vertexSelectionBone = 0;
+				m_vertexSelectionBone.Id = 0;
 			}
 			arch >> m_maxEffectRadius;
 		}
diff --git a/Engine/Plugins/Experimental/Mutable/Source/MutableTools/Private/MuT/NodeMeshMorph.cpp b/Engine/Plugins/Experimental/Mutable/Source/MutableTools/Private/MuT/NodeMeshMorph.cpp
--- a/Engine/Plugins/Experimental/Mutable/Source/MutableTools/Private/MuT/NodeMeshMorph.cpp
+++ b/Engine/Plugins/Experimental/Mutable/Source/MutableTools/Private/MuT/NodeMeshMorph.cpp
@@ -78,12 +78,12 @@ namespace mu
 	}	
 
 	//---------------------------------------------------------------------------------------------
-	void NodeMeshMorph::AddBoneToDeform(const uint16 BoneId)
+	void NodeMeshMorph::AddBoneToDeform(const FBoneName& BoneId)
 	{
 		m_pD->BonesToDeform.Emplace(BoneId);
 	}
 
-	void NodeMeshMorph::AddPhysicsBodyToDeform(const uint16 BoneId)
+	void NodeMeshMorph::AddPhysicsBodyToDeform(const FBoneName& BoneId)
 	{
 		m_pD->PhysicsToDeform.Emplace(BoneId);
 	}
diff --git a/Engine/Plugins/Experimental/Mutable/Source/MutableTools/Private/MuT/NodeMeshReshape.cpp b/Engine/Plugins/Experimental/Mutable/Source/MutableTools/Private/MuT/NodeMeshReshape.cpp
--- a/Engine/Plugins/Experimental/Mutable/Source/MutableTools/Private/MuT/NodeMeshReshape.cpp
+++ b/Engine/Plugins/Experimental/Mutable/Source/MutableTools/Private/MuT/NodeMeshReshape.cpp
@@ -101,13 +101,13 @@ namespace mu
 	}	
 
 	//---------------------------------------------------------------------------------------------
-	void NodeMeshReshape::AddBoneToDeform(const uint16 BoneId)
+	void NodeMeshReshape::AddBoneToDeform(const FBoneName& BoneId)
 	{
 		m_pD->BonesToDeform.Emplace(BoneId);
 	}
 
 	//---------------------------------------------------------------------------------------------
-	void NodeMeshReshape::AddPhysicsBodyToDeform(const uint16 BoneId)
+	void NodeMeshReshape::AddPhysicsBodyToDeform(const FBoneName& BoneId)
 	{
 		m_pD->PhysicsToDeform.Emplace(BoneId);
 	}
diff --git a/Engine/Plugins/Experimental/Mutable/Source/MutableTools/Private/MuT/NodeMeshReshapePrivate.h b/Engine/Plugins/Experimental/Mutable/Source/MutableTools/Private/MuT/NodeMeshReshapePrivate.h
--- a/Engine/Plugins/Experimental/Mutable/Source/MutableTools/Private/MuT/NodeMeshReshapePrivate.h
+++ b/Engine/Plugins/Experimental/Mutable/Source/MutableTools/Private/MuT/NodeMeshReshapePrivate.h
@@ -6,6 +6,7 @@
 #include "MuT/NodeScalar.h"
 #include "MuT/NodeMeshPrivate.h"
 #include "MuT/AST.h"
+#include "MuR/Skeleton.h"
 
 namespace mu
 {
@@ -30,12 +31,12 @@ namespace mu
 		EVertexColorUsage ColorBChannelUsage = EVertexColorUsage::None;
 		EVertexColorUsage ColorAChannelUsage = EVertexColorUsage::None;
 
-		TArray<uint16> BonesToDeform;
-		TArray<uint16> PhysicsToDeform;
+		TArray<FBoneName> BonesToDeform;
+		TArray<FBoneName> PhysicsToDeform;
         //!
 		void Serialise( OutputArchive& arch ) const
 		{
-            uint32 ver = 11;
+            uint32 ver = 12;
 			arch << ver;
 
 			arch << BaseMesh;
@@ -97,20 +98,30 @@ namespace mu
 				arch >> bDeformAllBones_DEPRECATED;
 			}
 
-			if (ver >= 9)
+			if (ver >= 12)
 			{
 				arch >> BonesToDeform;
 			}
 			else
 			{
-				TArray<string> BonesToDeform_DEPRECATED;
-				arch >> BonesToDeform_DEPRECATED;
+				int32 NumBonesToDeform = 0;
+				if (ver >= 9)
+				{
+					TArray<uint16> BonesToDeform_DEPRECATED;
+					arch >> BonesToDeform_DEPRECATED;
+					NumBonesToDeform = BonesToDeform_DEPRECATED.Num();
+				}
+				else
+				{
+					TArray<string> BonesToDeform_DEPRECATED;
+					arch >> BonesToDeform_DEPRECATED;
+					NumBonesToDeform = BonesToDeform_DEPRECATED.Num();
+				}
 
-				const int32 NumBonesToDeform = BonesToDeform_DEPRECATED.Num();
-				BonesToDeform.Reserve(NumBonesToDeform);
+				BonesToDeform.SetNum(NumBonesToDeform);
 				for (int32 BoneIndex = 0; BoneIndex < NumBonesToDeform; ++BoneIndex)
 				{
-					BonesToDeform.Add(BoneIndex);
+					BonesToDeform[BoneIndex].Id = BoneIndex;
 				}
 			}
 
@@ -125,20 +136,30 @@ namespace mu
 				arch >> bDeformAllPhysics_DEPRECATED;
 			}
 
-			if (ver >= 9)
+			if (ver >= 12)
 			{
 				arch >> PhysicsToDeform;
 			}
-			else if (ver >= 5)
+			else
 			{
-				TArray<string> PhysicsToDeform_DEPRECATED;
-				arch >> PhysicsToDeform_DEPRECATED;
+				int32 NumPhysicsToDeform = 0;
+				if (ver >= 9)
+				{
+					TArray<uint16> PhysicsToDeform_DEPRECATED;
+					arch >> PhysicsToDeform_DEPRECATED;
+					NumPhysicsToDeform = PhysicsToDeform_DEPRECATED.Num();
+				}
+				else if (ver >= 5)
+				{
+					TArray<string> PhysicsToDeform_DEPRECATED;
+					arch >> PhysicsToDeform_DEPRECATED;
+					NumPhysicsToDeform = PhysicsToDeform_DEPRECATED.Num();
+				}
 
-				const int32 NumPhysicsToDeform = PhysicsToDeform_DEPRECATED.Num();
-				PhysicsToDeform.Reserve(NumPhysicsToDeform);
+				PhysicsToDeform.SetNum(NumPhysicsToDeform);
 				for (int32 BoneIndex = 0; BoneIndex < NumPhysicsToDeform; ++BoneIndex)
 				{
-					PhysicsToDeform.Add(BoneIndex);
+					PhysicsToDeform[BoneIndex].Id = BoneIndex;
 				}
 			}
 
diff --git a/Engine/Plugins/Experimental/Mutable/Source/MutableTools/Private/MuT/NodeModifierMeshClipMorphPlane.cpp b/Engine/Plugins/Experimental/Mutable/Source/MutableTools/Private/MuT/NodeModifierMeshClipMorphPlane.cpp
--- a/Engine/Plugins/Experimental/Mutable/Source/MutableTools/Private/MuT/NodeModifierMeshClipMorphPlane.cpp
+++ b/Engine/Plugins/Experimental/Mutable/Source/MutableTools/Private/MuT/NodeModifierMeshClipMorphPlane.cpp
@@ -61,7 +61,7 @@ namespace mu
 	}
 
 	//---------------------------------------------------------------------------------------------
-	void NodeModifierMeshClipMorphPlane::SetVertexSelectionBone(uint16 BoneId, float maxEffectRadius)
+	void NodeModifierMeshClipMorphPlane::SetVertexSelectionBone(const FBoneName& BoneId, float maxEffectRadius)
 	{
 		m_pD->m_vertexSelectionType = Private::VS_BONE_HIERARCHY;
 		m_pD->m_vertexSelectionBone = BoneId;
diff --git a/Engine/Plugins/Experimental/Mutable/Source/MutableTools/Public/MuT/NodeMeshClipMorphPlane.h b/Engine/Plugins/Experimental/Mutable/Source/MutableTools/Public/MuT/NodeMeshClipMorphPlane.h
--- a/Engine/Plugins/Experimental/Mutable/Source/MutableTools/Public/MuT/NodeMeshClipMorphPlane.h
+++ b/Engine/Plugins/Experimental/Mutable/Source/MutableTools/Public/MuT/NodeMeshClipMorphPlane.h
@@ -16,6 +16,7 @@ namespace mu
     typedef Ptr<NodeMeshClipMorphPlane> NodeMeshClipMorphPlanePtr;
     typedef Ptr<const NodeMeshClipMorphPlane> NodeMeshClipMorphPlaneConst;
 
+	struct FBoneName;
 
     //! This node applies a geometric transform represented by a 4x4 matrix to a mesh
 	//! \ingroup model
@@ -60,7 +61,7 @@ namespace mu
 
 		//! Define the root bone of the subhierarchy of the mesh that will be affected.
 		//! Only one of Box or Bone Hierarchy can be used (the last one set)
-		void SetVertexSelectionBone(uint16 BoneId, float maxEffectRadius);
+		void SetVertexSelectionBone(const FBoneName& BoneId, float maxEffectRadius);
 
 		//! Add a tag to the clip morph operation, which will only affect surfaces with the same tag
 		void AddTag(const char* tagName);
diff --git a/Engine/Plugins/Experimental/Mutable/Source/MutableTools/Public/MuT/NodeMeshMorph.h b/Engine/Plugins/Experimental/Mutable/Source/MutableTools/Public/MuT/NodeMeshMorph.h
--- a/Engine/Plugins/Experimental/Mutable/Source/MutableTools/Public/MuT/NodeMeshMorph.h
+++ b/Engine/Plugins/Experimental/Mutable/Source/MutableTools/Public/MuT/NodeMeshMorph.h
@@ -13,7 +13,7 @@ namespace mu
 
 	// Forward definitions
 	class NodeScalar;
-
+	struct FBoneName;
 
 	/** Node that morphs a base mesh with one weighted target. */
 	class MUTABLETOOLS_API NodeMeshMorph : public NodeMesh
@@ -61,10 +61,10 @@ namespace mu
 		void SetReshapePhysicsVolumes(bool);
 
 		/** Sets the number of bones that will be deform */
-		void AddBoneToDeform(const uint16 BoneName);	
+		void AddBoneToDeform(const FBoneName& BoneName);	
 
 		/** Add a Physics Body to deform */
-		void AddPhysicsBodyToDeform(const uint16 BoneName);
+		void AddPhysicsBodyToDeform(const FBoneName& BoneName);
 
         //-----------------------------------------------------------------------------------------
 		// Interface pattern
diff --git a/Engine/Plugins/Experimental/Mutable/Source/MutableTools/Public/MuT/NodeMeshReshape.h b/Engine/Plugins/Experimental/Mutable/Source/MutableTools/Public/MuT/NodeMeshReshape.h
--- a/Engine/Plugins/Experimental/Mutable/Source/MutableTools/Public/MuT/NodeMeshReshape.h
+++ b/Engine/Plugins/Experimental/Mutable/Source/MutableTools/Public/MuT/NodeMeshReshape.h
@@ -9,6 +9,8 @@
 
 namespace mu
 {
+	// Forward declarations
+	struct FBoneName;
 
 	//! Node that morphs a base mesh with one or two weighted targets from a sequence.
 	//! \ingroup model
@@ -62,10 +64,10 @@ namespace mu
 		void SetReshapePhysicsVolumes(bool);
 
 		/** Sets the number of bones that will be deform */
-		void AddBoneToDeform(const uint16 BoneId);
+		void AddBoneToDeform(const FBoneName& BoneId);
 	
 		/** Add a Physics Body to deform */
-		void AddPhysicsBodyToDeform(const uint16 BoneId);
+		void AddPhysicsBodyToDeform(const FBoneName& BoneId);
         
 		//-----------------------------------------------------------------------------------------
 		// Interface pattern
diff --git a/Engine/Plugins/Experimental/Mutable/Source/MutableTools/Public/MuT/NodeModifierMeshClipMorphPlane.h b/Engine/Plugins/Experimental/Mutable/Source/MutableTools/Public/MuT/NodeModifierMeshClipMorphPlane.h
--- a/Engine/Plugins/Experimental/Mutable/Source/MutableTools/Public/MuT/NodeModifierMeshClipMorphPlane.h
+++ b/Engine/Plugins/Experimental/Mutable/Source/MutableTools/Public/MuT/NodeModifierMeshClipMorphPlane.h
@@ -16,6 +16,7 @@ namespace mu
 	typedef Ptr<NodeModifierMeshClipMorphPlane> NodeModifierMeshClipMorphPlanePtr;
 	typedef Ptr<const NodeModifierMeshClipMorphPlane> NodeModifierMeshClipMorphPlanePtrConst;
 
+	struct FBoneName;
 
 	//! This node makes a new component from several meshes and images.
 	//! \ingroup model
@@ -55,7 +56,7 @@ namespace mu
 
 		//! Define the root bone of the subhierarchy of the mesh that will be affected.
 		//! Only one of Box or Bone Hierarchy can be used (the last one set)
-		void SetVertexSelectionBone(uint16 BoneId, float maxEffectRadius);
+		void SetVertexSelectionBone(const FBoneName& BoneId, float maxEffectRadius);
 
 		//-----------------------------------------------------------------------------------------
 		// Interface pattern
