Ensuring that the UCustomizableObjectBulk::BulkFilePrefix is standardized to avoid issues with the streaming platform file layer.

diff --git a/Engine/Plugins/Experimental/Mutable/Source/CustomizableObject/Private/MuCO/CustomizableObject.cpp b/Engine/Plugins/Experimental/Mutable/Source/CustomizableObject/Private/MuCO/CustomizableObject.cpp
--- a/Engine/Plugins/Experimental/Mutable/Source/CustomizableObject/Private/MuCO/CustomizableObject.cpp
+++ b/Engine/Plugins/Experimental/Mutable/Source/CustomizableObject/Private/MuCO/CustomizableObject.cpp
@@ -2025,7 +2025,8 @@ void UCustomizableObjectBulk::PostLoad()
 	UObject::PostLoad();
 
 	const FString OutermostName = GetOutermost()->GetName();
-	const FString PackageFilename = FPackageName::LongPackageNameToFilename(OutermostName);
+	FString PackageFilename = FPackageName::LongPackageNameToFilename(OutermostName);
+	FPaths::MakeStandardFilename(PackageFilename);
 	BulkFilePrefix = PackageFilename;
 }
 
