[Mutable] Realtime morphs mutable mesh data determinism fixes.

- Real time morph data generation has been refactored so it can be done in a working thread.

#rnx

diff --git a/Engine/Plugins/Experimental/Mutable/Source/CustomizableObject/Internal/MuCO/CustomizableObjectPrivate.h b/Engine/Plugins/Experimental/Mutable/Source/CustomizableObject/Internal/MuCO/CustomizableObjectPrivate.h
--- a/Engine/Plugins/Experimental/Mutable/Source/CustomizableObject/Internal/MuCO/CustomizableObjectPrivate.h
+++ b/Engine/Plugins/Experimental/Mutable/Source/CustomizableObject/Internal/MuCO/CustomizableObjectPrivate.h
@@ -87,15 +87,15 @@ struct FCustomizableObjectStreameableResourceId
 		RealTimeMorphTarget   = 2,
 	};
 
-	uint32 Id   : 24;
-	uint32 Type : 8;
+	uint64 Id   : 64 - 8;
+	uint64 Type : 8;
 
 	friend bool operator==(FCustomizableObjectStreameableResourceId A, FCustomizableObjectStreameableResourceId B)
 	{
-		return BitCast<uint32>(A) == BitCast<uint32>(B);
+		return BitCast<uint64>(A) == BitCast<uint64>(B);
 	}
 };
-static_assert(sizeof(FCustomizableObjectStreameableResourceId) == sizeof(uint32));
+static_assert(sizeof(FCustomizableObjectStreameableResourceId) == sizeof(uint64));
 
 
 USTRUCT()
@@ -513,6 +513,26 @@ struct CUSTOMIZABLEOBJECT_API FMutableSkinWeightProfileInfo
 #endif
 };
 
+USTRUCT()
+struct FRealTimeMorphStreamable
+{
+	GENERATED_USTRUCT_BODY()
+	
+	UPROPERTY()
+	TArray<FName> NameResolutionMap;
+
+	UPROPERTY()
+	FMutableStreamableBlock Block;
+
+	friend FArchive& operator<<(FArchive& Ar, FRealTimeMorphStreamable& Elem)
+	{
+		Ar << Elem.NameResolutionMap;
+		Ar << Elem.Block;
+
+		return Ar;
+	}
+};
+
 USTRUCT()
 struct FMorphTargetVertexData
 {
@@ -536,6 +556,7 @@ struct FMorphTargetVertexData
 		return Ar;
 	}
 };
+static_assert(sizeof(FMorphTargetVertexData) == sizeof(FVector3f)*2 + sizeof(uint32)); // Make sure no padding is present.
 template<> struct TCanBulkSerialize<FMorphTargetVertexData> { enum { Value = true }; };
 
 // Referenced materials, skeletons, passthrough textures...
@@ -604,10 +625,7 @@ struct FModelResources
 	TMap<FString, FParameterUIData> StateUIDataMap;
 
 	UPROPERTY()
-	TArray<FName> RealTimeMorphTargetNames;
-
-	UPROPERTY()
-	TArray<FMutableStreamableBlock> RealTimeMorphStreamableBlocks;
+	TMap<uint32, FRealTimeMorphStreamable> RealTimeMorphStreamables;
 
 #if WITH_EDITORONLY_DATA
 	UPROPERTY()
@@ -861,6 +879,6 @@ public:
 	// This is a manual version number for the binary blobs in this asset.
 	// Increasing it invalidates all the previously compiled models.
 	// Warning: If while merging code both versions have changed, take the highest+1.
-	static constexpr int32 CurrentSupportedVersion = 440;
+	static constexpr int32 CurrentSupportedVersion = 441;
 };
 
diff --git a/Engine/Plugins/Experimental/Mutable/Source/CustomizableObject/Internal/MuCO/CustomizableObjectSystemPrivate.h b/Engine/Plugins/Experimental/Mutable/Source/CustomizableObject/Internal/MuCO/CustomizableObjectSystemPrivate.h
--- a/Engine/Plugins/Experimental/Mutable/Source/CustomizableObject/Internal/MuCO/CustomizableObjectSystemPrivate.h
+++ b/Engine/Plugins/Experimental/Mutable/Source/CustomizableObject/Internal/MuCO/CustomizableObjectSystemPrivate.h
@@ -25,6 +25,7 @@
 #include "MuCO/FMutableTaskGraph.h"
 #include "AssetRegistry/AssetData.h"
 #include "ContentStreaming.h"
+#include "Animation/MorphTarget.h"
 
 #include "Tasks/Task.h"
 
@@ -38,7 +39,6 @@ class UCustomizableObjectSystem;
 namespace LowLevelTasks { enum class ETaskPriority : int8; }
 struct FTexturePlatformData;
 
-
 // Split StreamedBulkData into chunks smaller than MUTABLE_STREAMED_DATA_MAXCHUNKSIZE
 #define MUTABLE_STREAMED_DATA_MAXCHUNKSIZE		(512 * 1024 * 1024)
 
@@ -402,8 +402,6 @@ struct FInstanceUpdateData
 	TArray<mu::FBoneName> ActiveBones;
 	TArray<mu::FBoneName> BoneMaps;
 	
-	TMap<uint32, TArray<FMorphTargetVertexData>> MorphTargetsVertexData;
-
 	struct FBone
 	{
 		mu::FBoneName Name;
@@ -412,6 +410,26 @@ struct FInstanceUpdateData
 		bool operator==(const mu::FBoneName& OtherName) const { return Name == OtherName; };
 	};
 
+	
+	struct FMorphTargetMeshData
+	{
+		TArray<FName> NameResolutionMap;
+		TArray<FMorphTargetVertexData> Data;
+	};
+
+	TMap<uint32, FMorphTargetMeshData> RealTimeMorphTargetMeshData;
+
+
+	struct FRealTimeMorphsComponentData
+	{
+		int32 ComponentIndex = INDEX_NONE;
+
+		TArray<FName> RealTimeMorphTargetNames;
+		TArray<TArray<FMorphTargetLODModel>> RealTimeMorphsLODData; 
+	};
+
+	TArray<FRealTimeMorphsComponentData> RealTimeMorphTargets;
+
 	struct FSkeletonData
 	{
 		TArray<uint16> SkeletonIds;
@@ -439,6 +457,7 @@ struct FInstanceUpdateData
 		Images.Empty();
 		Scalars.Empty();
 		Vectors.Empty();
+		RealTimeMorphTargets.Empty();
 		Skeletons.Empty();
 		ExtendedInputPins.Empty();
 	}
diff --git a/Engine/Plugins/Experimental/Mutable/Source/CustomizableObject/Private/MuCO/CustomizableObject.cpp b/Engine/Plugins/Experimental/Mutable/Source/CustomizableObject/Private/MuCO/CustomizableObject.cpp
--- a/Engine/Plugins/Experimental/Mutable/Source/CustomizableObject/Private/MuCO/CustomizableObject.cpp
+++ b/Engine/Plugins/Experimental/Mutable/Source/CustomizableObject/Private/MuCO/CustomizableObject.cpp
@@ -559,8 +559,7 @@ void UCustomizableObjectPrivate::SaveCompiledData(FArchive& MemoryWriter, bool b
 	MemoryWriter << LocalModelResources.ParameterUIDataMap;
 	MemoryWriter << LocalModelResources.StateUIDataMap;
 
-	MemoryWriter << LocalModelResources.RealTimeMorphTargetNames;
-	MemoryWriter << LocalModelResources.RealTimeMorphStreamableBlocks;
+	MemoryWriter << LocalModelResources.RealTimeMorphStreamables;
 
 	MemoryWriter << LocalModelResources.HashToStreamableBlock;
 
@@ -680,8 +679,7 @@ void UCustomizableObjectPrivate::LoadCompiledData(FArchive& MemoryReader, const
 		MemoryReader << LocalModelResource.ParameterUIDataMap;
 		MemoryReader << LocalModelResource.StateUIDataMap;
 
-		MemoryReader << LocalModelResource.RealTimeMorphTargetNames;
-		MemoryReader << LocalModelResource.RealTimeMorphStreamableBlocks;
+		MemoryReader << LocalModelResource.RealTimeMorphStreamables;
 
 		MemoryReader << LocalModelResource.HashToStreamableBlock;
 
@@ -2173,17 +2171,19 @@ void UCustomizableObjectBulk::PrepareBulkData(UCustomizableObject* InOuter, cons
 		uint64 SourceOffset = 0;
 		
 		constexpr bool bGetCooked = false;
-		const TArray<FMutableStreamableBlock>& RealTimeMorphTargetsBlocks = 
-					CustomizableObject->GetPrivate()->GetModelResources(bGetCooked).RealTimeMorphStreamableBlocks;
+
+		TArray<FMutableStreamableBlock> RealTimeMorphTargetsBlocks;
+		
+		const TMap<uint32, FRealTimeMorphStreamable>& RealTimeMorphStreamables = 
+				CustomizableObject->GetPrivate()->GetModelResources(bGetCooked).RealTimeMorphStreamables;
 
 		const int32 NumBlocks = RealTimeMorphTargetsBlocks.Num();
-		for (int32 BlockIndex = 0; BlockIndex < NumBlocks; ++BlockIndex)
+		for (const TPair<uint32, FRealTimeMorphStreamable>& MorphStreamable : RealTimeMorphStreamables)
 		{
-			const FMutableStreamableBlock& StreamableBlock = RealTimeMorphTargetsBlocks[BlockIndex];
-			const uint32 BlockSize = StreamableBlock.Size;
+			const uint32 BlockSize = MorphStreamable.Value.Block.Size;
 			
-			check(SourceOffset == StreamableBlock.Offset);
-			FBlock CurrentBlock = { EDataType::RealTimeMorph, BlockIndex, BlockSize, SourceOffset };
+			check(SourceOffset == MorphStreamable.Value.Block.Offset);
+			FBlock CurrentBlock = { EDataType::RealTimeMorph, MorphStreamable.Key, BlockSize, SourceOffset };
 			
 			if (BlockSize > TargetBulkDataFileBytes)
 			{
@@ -2285,15 +2285,15 @@ void UCustomizableObjectBulk::PrepareBulkData(UCustomizableObject* InOuter, cons
 		else if (CurrentFile.DataType == EDataType::RealTimeMorph)
 		{
 			constexpr bool bGetCooked = true;
-			TArray<FMutableStreamableBlock>& MorphBlocks = 
-					CustomizableObject->GetPrivate()->GetModelResources(bGetCooked).RealTimeMorphStreamableBlocks;
+			TMap<uint32, FRealTimeMorphStreamable>& MorphBlocks = 
+					CustomizableObject->GetPrivate()->GetModelResources(bGetCooked).RealTimeMorphStreamables;
 			// Set it to all streamable blocks
 			uint32 OffsetInFile = 0;
 			for (int32 FileBlockIndex = 0; FileBlockIndex < CurrentFile.Blocks.Num(); ++FileBlockIndex)
 			{
 				FBlock ThisBlock = CurrentFile.Blocks[FileBlockIndex];
 
-				FMutableStreamableBlock& StreamableBlock = MorphBlocks[ThisBlock.Id];
+				FMutableStreamableBlock& StreamableBlock = MorphBlocks[ThisBlock.Id].Block;
 				check(StreamableBlock.Size == ThisBlock.Size);
 				StreamableBlock.FileId = FileId;
 				StreamableBlock.Offset = OffsetInFile;
diff --git a/Engine/Plugins/Experimental/Mutable/Source/CustomizableObject/Private/MuCO/CustomizableObjectInstance.cpp b/Engine/Plugins/Experimental/Mutable/Source/CustomizableObject/Private/MuCO/CustomizableObjectInstance.cpp
--- a/Engine/Plugins/Experimental/Mutable/Source/CustomizableObject/Private/MuCO/CustomizableObjectInstance.cpp
+++ b/Engine/Plugins/Experimental/Mutable/Source/CustomizableObject/Private/MuCO/CustomizableObjectInstance.cpp
@@ -1870,7 +1870,7 @@ bool UCustomizableInstancePrivate::UpdateSkeletalMesh_PostBeginUpdate0(UCustomiz
 					constexpr bool bMarkAsDirty = false;
 					PhysicsAssetResult->SetPreviewMesh(SkeletalMesh, bMarkAsDirty);
 				}
-				}
+			}
 
 			const int32 NumAdditionalPhysicsNum = Component.Mesh->AdditionalPhysicsBodies.Num();
 			for (int32 I = 0; I < NumAdditionalPhysicsNum; ++I)
@@ -3532,217 +3532,45 @@ void UCustomizableInstancePrivate::BuildOrCopyMorphTargetsData(const TSharedRef<
 {
 	MUTABLE_CPUPROFILER_SCOPE(UCustomizableInstancePrivate::BuildOrCopyMorphTargetsData);
 
-	if (!SkeletalMesh || !CustomizableObjectInstance)
+	if (!SkeletalMesh)
 	{
 		return;
 	}
+	
+	FInstanceUpdateData& UpdateData = OperationData->InstanceUpdateData;
 
-	const UCustomizableObject* CustomizableObject = CustomizableObjectInstance->GetCustomizableObject();
-	const TArray<FName>& MorphTargetNames = CustomizableObject->GetPrivate()->GetModelResources().RealTimeMorphTargetNames;
-
-	const TMap<uint32, TArray<FMorphTargetVertexData>>& ResourceIdToVertexDataMap =
-		OperationData->InstanceUpdateData.MorphTargetsVertexData;
+	FInstanceUpdateData::FRealTimeMorphsComponentData* ComponentMorphTargets = 
+		UpdateData.RealTimeMorphTargets.FindByPredicate(
+				[ComponentIndex](const FInstanceUpdateData::FRealTimeMorphsComponentData& Elem) 
+				{ 
+					return Elem.ComponentIndex == ComponentIndex; 
+				});
 
-	if (MorphTargetNames.IsEmpty() || ResourceIdToVertexDataMap.IsEmpty())
+	if (!ComponentMorphTargets)
 	{
 		return;
 	}
 
-	TArray<int32> SectionMorphTargetVertices;
-
-	SectionMorphTargetVertices.SetNumZeroed(MorphTargetNames.Num());
-
-	TArray<uint8> UsedMorphSet;
-	UsedMorphSet.SetNum(MorphTargetNames.Num());
-
-	TArray<TArray<FMorphTargetLODModel>> MorphsData;
-	MorphsData.SetNum(MorphTargetNames.Num());
-
 	TArray<TObjectPtr<UMorphTarget>>& MorphTargets = SkeletalMesh->GetMorphTargets();
 
-	const int32 FirstGeneratedLOD = FMath::Max((int32)OperationData->GetRequestedLODs()[ComponentIndex], OperationData->GetMinLOD());
-	for (int32 LODIndex = FirstGeneratedLOD; LODIndex < OperationData->NumLODsAvailable; ++LODIndex)
+	MorphTargets.Empty(ComponentMorphTargets->RealTimeMorphTargetNames.Num());
+	
+	const int32 NumMorphs = ComponentMorphTargets->RealTimeMorphTargetNames.Num();
+	for (int32 I = 0; I < NumMorphs; ++I)
 	{
-		MUTABLE_CPUPROFILER_SCOPE(BuildMorphTargetsData);
-
-		int32 NumNotFoundLoadedMorphsResources = 0;
-
-		const FInstanceUpdateData::FLOD& LOD = OperationData->InstanceUpdateData.LODs[LODIndex];
-		const FInstanceUpdateData::FComponent& Component = OperationData->InstanceUpdateData.Components[LOD.FirstComponent + ComponentIndex];
-		check(Component.bGenerated);
-		check(Component.Mesh);
-
-		const mu::FMeshBufferSet& MeshSet = Component.Mesh->GetVertexBuffers();
-
-		int32 VertexMorphsInfoIndexBufferIndex, VertexMorphsInfoIndexBufferChannel;
-		MeshSet.FindChannel(mu::MBS_OTHER, 0, &VertexMorphsInfoIndexBufferIndex, &VertexMorphsInfoIndexBufferChannel);
-
-		int32 VertexMorphsCountBufferIndex, VertexMorphsCountBufferChannel;
-		MeshSet.FindChannel(mu::MBS_OTHER, 1, &VertexMorphsCountBufferIndex, &VertexMorphsCountBufferChannel);
-
-		int32 VertexMorphsResourceIdBufferIndex, VertexMorphsResourceIdBufferChannel;
-		MeshSet.FindChannel(mu::MBS_OTHER, 2, &VertexMorphsResourceIdBufferIndex, &VertexMorphsResourceIdBufferChannel);
-
-		if (VertexMorphsInfoIndexBufferIndex < 0 || VertexMorphsCountBufferIndex < 0 || VertexMorphsResourceIdBufferIndex < 0)
+		if (ComponentMorphTargets->RealTimeMorphsLODData[I].IsEmpty())
 		{
 			continue;
 		}
 
-		const int32* const VertexMorphsInfoIndexBuffer = reinterpret_cast<const int32*>(MeshSet.GetBufferData(VertexMorphsInfoIndexBufferIndex));
-		TArrayView<const int32> VertexMorphsInfoIndexView(VertexMorphsInfoIndexBuffer, MeshSet.GetElementCount());
-
-		const uint16* const VertexMorphsCountBuffer = reinterpret_cast<const uint16*>(MeshSet.GetBufferData(VertexMorphsCountBufferIndex));
-		TArrayView<const uint16> VertexMorphsCountView(VertexMorphsCountBuffer, MeshSet.GetElementCount());
-
-		const uint16* const VertexMorphsResourceIdBuffer = reinterpret_cast<const uint16*>(MeshSet.GetBufferData(VertexMorphsResourceIdBufferIndex));
-		TArrayView<const uint16> VertexMorphsResourceIdView(VertexMorphsResourceIdBuffer, MeshSet.GetElementCount());
-
-		const int32 SurfaceCount = Component.Mesh->GetSurfaceCount();
-		for (int32 Section = 0; Section < SurfaceCount; ++Section)
-		{
-			// Reset SectionMorphTargets.
-			for (int32& Elem : SectionMorphTargetVertices)
-			{
-				Elem = 0;
-			}
-
-			int32 FirstVertex, VerticesCount, FirstIndex, IndiciesCount;
-			Component.Mesh->GetSurface(Section, &FirstVertex, &VerticesCount, &FirstIndex, &IndiciesCount, nullptr, nullptr, nullptr);
-
-			for (int32 VertexIdx = FirstVertex; VertexIdx < FirstVertex + VerticesCount;)
-			{
-				// Find a span with the same VertexMorphResourceId to amortise the cost of finding 
-				// in the loaded resources map. It is expected to find large consecutive mesh sections pointing to
-				// the same loaded resource.
-
-				const int32 SpanStart = VertexIdx++;
-				const uint16 CurrentResourceId = VertexMorphsResourceIdView[SpanStart];
-
-				// Vertex with no morphs are marked with TNumericLimits<uint16>::Max(), skip vertex if the case.
-				if (CurrentResourceId == TNumericLimits<uint16>::Max())
-				{
-					continue;
-				}
-
-				for (; VertexIdx < FirstVertex + VerticesCount; ++VertexIdx)
-				{
-					const int32 VertexResourceId = VertexMorphsResourceIdView[VertexIdx];
-					// we can skip vertices with no morph without breaking the span.
-					if (VertexResourceId == TNumericLimits<uint16>::Max())
-					{
-						continue;
-					}
-
-					if (CurrentResourceId != VertexResourceId)
-					{
-						break;
-					}
-				}
-				const int32 SpanEnd = VertexIdx;
-
-				const TArray<FMorphTargetVertexData>* MorphTargetReconstructionData = ResourceIdToVertexDataMap.Find(CurrentResourceId);
-
-				if (!MorphTargetReconstructionData)
-				{
-					++NumNotFoundLoadedMorphsResources;
-					continue;
-				}
-
-				const TArray<FMorphTargetVertexData>& SpanMorphData = *MorphTargetReconstructionData;
-
-
-				// This assumes the number of vertex in an span will be large compared to the number of
-				// morphs. Maybe the allocation could be done in a different pass as an optimization.
-				FMemory::Memzero(UsedMorphSet.GetData(), UsedMorphSet.Num());
-				for (int32 SpanVertexIdx = SpanStart; SpanVertexIdx < SpanEnd; ++SpanVertexIdx)
-				{
-					const int32 MorphCount = VertexMorphsCountView[SpanVertexIdx];
-
-					TArrayView<const FMorphTargetVertexData> MorphsVertexDataView = MakeArrayView(
-						SpanMorphData.GetData() + VertexMorphsInfoIndexView[SpanVertexIdx],
-						MorphCount);
-
-					for (const FMorphTargetVertexData& MorphVertexData : MorphsVertexDataView)
-					{
-						check((uint32)UsedMorphSet.Num() > MorphVertexData.MorphNameIndex);
-						UsedMorphSet[MorphVertexData.MorphNameIndex] = 1;
-					}
-				}
-
-				const int32 NumMorphs = UsedMorphSet.Num();
-				for (int32 MorphIndex = 0; MorphIndex < NumMorphs; ++MorphIndex)
-				{
-					if (!UsedMorphSet[MorphIndex])
-					{
-						continue;
-					}
-
-					if (MorphsData[MorphIndex].IsEmpty())
-					{
-						MorphsData[MorphIndex].SetNum(OperationData->NumLODsAvailable);
-					}
-				}
-
-				for (int32 SpanVertexIdx = SpanStart; SpanVertexIdx < SpanEnd; ++SpanVertexIdx)
-				{
-					const uint16 MorphCount = VertexMorphsCountView[SpanVertexIdx];
-					if (MorphCount == 0)
-					{
-						continue;
-					}
-
-					TArrayView<const FMorphTargetVertexData> MorphsVertexDataView = MakeArrayView(
-						SpanMorphData.GetData() + VertexMorphsInfoIndexView[SpanVertexIdx],
-						MorphCount);
-
-					for (const FMorphTargetVertexData& SourceVertex : MorphsVertexDataView)
-					{
-						FMorphTargetLODModel& DestMorphLODModel = MorphsData[SourceVertex.MorphNameIndex][LODIndex];
-
-						DestMorphLODModel.Vertices.Emplace(
-							FMorphTargetDelta{ SourceVertex.PositionDelta, SourceVertex.TangentZDelta, static_cast<uint32>(SpanVertexIdx) });
-
-						++SectionMorphTargetVertices[SourceVertex.MorphNameIndex];
-					}
-				}
-			}
-
-			const int32 SectionMorphTargetsNum = SectionMorphTargetVertices.Num();
-			for (int32 MorphIdx = 0; MorphIdx < SectionMorphTargetsNum; ++MorphIdx)
-			{
-				if (SectionMorphTargetVertices[MorphIdx] > 0)
-				{
-					FMorphTargetLODModel& MorphTargetLodModel = MorphsData[MorphIdx][LODIndex];
-
-					MorphTargetLodModel.SectionIndices.Add(Section);
-					MorphTargetLodModel.NumVertices += SectionMorphTargetVertices[MorphIdx];
-				}
-			}
-		}
-
-		if (NumNotFoundLoadedMorphsResources > 0)
-		{
-			UE_LOG(LogMutable, Warning, TEXT("Needed realtime morph reconstruction data was not loaded properly. Some realtime morphs may not work correctly."));
-		}
-
-		// Generate the SkeletalMesh data structures. The previous step could be done in an async task.
-		MorphTargets.Empty();
-		const int32 NumMorphs = MorphTargetNames.Num();
-		for (int32 I = 0; I < NumMorphs; ++I)
-		{
-			if (MorphsData[I].IsEmpty())
-			{
-				continue;
-			}
-
-			UMorphTarget* NewMorphTarget = NewObject<UMorphTarget>(SkeletalMesh, MorphTargetNames[I]);
-			NewMorphTarget->BaseSkelMesh = SkeletalMesh;
-			NewMorphTarget->GetMorphLODModels() = MoveTemp(MorphsData[I]);
-			MorphTargets.Add(NewMorphTarget);
-		}
+		UMorphTarget* NewMorphTarget = NewObject<UMorphTarget>(SkeletalMesh, ComponentMorphTargets->RealTimeMorphTargetNames[I]);
+		NewMorphTarget->BaseSkelMesh = SkeletalMesh;
+		NewMorphTarget->GetMorphLODModels() = MoveTemp(ComponentMorphTargets->RealTimeMorphsLODData[I]);
+		MorphTargets.Add(NewMorphTarget);
 	}
 
 	// Copy MorphTargets from the FirstGeneratedLOD to the LODs below
+	const int32 FirstGeneratedLOD = FMath::Max((int32)OperationData->GetRequestedLODs()[ComponentIndex], OperationData->GetMinLOD());
 	for (int32 LODIndex = OperationData->FirstLODAvailable; LODIndex < FirstGeneratedLOD; ++LODIndex)
 	{
 		MUTABLE_CPUPROFILER_SCOPE(CopyMorphTargetsData);
@@ -3751,7 +3579,7 @@ void UCustomizableInstancePrivate::BuildOrCopyMorphTargetsData(const TSharedRef<
 		for (int32 MorphTargetIndex = 0; MorphTargetIndex < NumMorphTargets; ++MorphTargetIndex)
 		{
 			MorphTargets[MorphTargetIndex]->GetMorphLODModels()[LODIndex] =
-				MorphTargets[MorphTargetIndex]->GetMorphLODModels()[FirstGeneratedLOD];
+					MorphTargets[MorphTargetIndex]->GetMorphLODModels()[FirstGeneratedLOD];
 		}
 	}
 
@@ -3772,7 +3600,6 @@ namespace
 
 }
 
-
 void UCustomizableInstancePrivate::BuildOrCopyClothingData(const TSharedRef<FUpdateContextPrivate>&OperationData, USkeletalMesh * SkeletalMesh, const USkeletalMesh* LastUpdateSkeletalMesh, UCustomizableObjectInstance * CustomizableObjectInstance, int32 ComponentIndex)
 {
 	MUTABLE_CPUPROFILER_SCOPE(UCustomizableInstancePrivate::BuildOrCopyClothingData);
@@ -3840,15 +3667,7 @@ void UCustomizableInstancePrivate::BuildOrCopyClothingData(const TSharedRef<FUpd
 				const int32 ClothingDataBufferIndex = [&MeshSet]()
 				{
 					int32 BufferIndex, Channel;
-					MeshSet.FindChannel(mu::MBS_OTHER, 3, &BufferIndex, &Channel);
-
-					if (BufferIndex > 0)
-					{
-						return BufferIndex;
-					}
-
-					MeshSet.FindChannel(mu::MBS_OTHER, 0, &BufferIndex, &Channel);
-
+					MeshSet.FindChannel(mu::MBS_OTHER, 2, &BufferIndex, &Channel);
 					return BufferIndex;
 				}(); // lambda is invoked
 
@@ -4628,7 +4447,6 @@ void UCustomizableInstancePrivate::BuildOrCopyClothingData(const TSharedRef<FUpd
 	SkeletalMesh->SetHasActiveClothingAssets( static_cast<bool>( SectionsWithCloth.Num() ) );
 }
 
-
 bool UCustomizableInstancePrivate::BuildOrCopyRenderData(const TSharedRef<FUpdateContextPrivate>& OperationData, USkeletalMesh* SkeletalMesh, const USkeletalMesh* LastUpdateSkeletalMesh, UCustomizableObjectInstance* Public, int32 ComponentIndex)
 {
 	MUTABLE_CPUPROFILER_SCOPE(UCustomizableInstancePrivate::BuildOrCopyRenderData);
@@ -4942,9 +4760,9 @@ UE::Tasks::FTask UCustomizableInstancePrivate::LoadAdditionalAssetsAndData(
 
 			FCustomizableInstanceComponentData* ComponentData = GetComponentData(Component.Id);
 
-			const TArray<uint32>& StreamedResources = MutableMesh->GetStreamedResources();
+			const TArray<uint64>& StreamedResources = MutableMesh->GetStreamedResources();
 
-			for (uint32 ResourceId : StreamedResources)
+			for (uint64 ResourceId : StreamedResources)
 			{
 				const TArray<FCustomizableObjectStreamedResourceData>& StreamedResourcesData = CustomizableObject->GetPrivate()->GetStreamedResourceData();
 				FCustomizableObjectStreameableResourceId TypedResourceId = BitCast<FCustomizableObjectStreameableResourceId>(ResourceId);	
@@ -4968,8 +4786,10 @@ UE::Tasks::FTask UCustomizableInstancePrivate::LoadAdditionalAssetsAndData(
 				}
 				else if (TypedResourceId.Type == (uint8)FCustomizableObjectStreameableResourceId::EType::RealTimeMorphTarget)
 				{
-					const TArray<FMutableStreamableBlock>& MorphsStremeableBlocks = ModelResources.RealTimeMorphStreamableBlocks;
-					if (MorphsStremeableBlocks.IsValidIndex(TypedResourceId.Id))
+					check(TypedResourceId.Id != 0 && TypedResourceId.Id <= TNumericLimits<uint32>::Max());
+
+					const TMap<uint32, FRealTimeMorphStreamable>& MorphsStreamables = ModelResources.RealTimeMorphStreamables;
+					if (MorphsStreamables.Contains((uint32)TypedResourceId.Id))
 					{
 						RealTimeMorphStreamableBlocksToStream.AddUnique(TypedResourceId.Id);
 					}
@@ -4980,6 +4800,7 @@ UE::Tasks::FTask UCustomizableInstancePrivate::LoadAdditionalAssetsAndData(
 				}
 				else
 				{
+					UE_LOG(LogMutable, Error, TEXT("Unknown streamed resource type found."));
 					check(false);
 				}
 			}
@@ -5153,22 +4974,28 @@ UE::Tasks::FTask UCustomizableInstancePrivate::LoadAdditionalAssetsAndData(
 
 	if (RealTimeMorphStreamableBlocksToStream.Num() && bUpdateMeshes)
 	{
+		
 #if WITH_EDITOR
 		// On editor the data is always loaded, load directly form the ModelResources.
 		MUTABLE_CPUPROFILER_SCOPE(RealTimeMorphStreamingEditor);
 		for (uint32 BlockId : RealTimeMorphStreamableBlocksToStream)
 		{	
-			const FMutableStreamableBlock& Block = ModelResources.RealTimeMorphStreamableBlocks[BlockId]; 
-			
-			const TArray<FMorphTargetVertexData>& SourceData = ModelResources.EditorOnlyMorphTargetReconstructionData;
-			TArray<FMorphTargetVertexData>& DestData = OperationData->InstanceUpdateData.MorphTargetsVertexData.Emplace(BlockId);
+			const FRealTimeMorphStreamable& RealTimeMorphStreamable = ModelResources.RealTimeMorphStreamables[BlockId]; 
+		
+			const FMutableStreamableBlock& Block = RealTimeMorphStreamable.Block;
+
+			FInstanceUpdateData::FMorphTargetMeshData& MeshData = 
+					OperationData->InstanceUpdateData.RealTimeMorphTargetMeshData.FindOrAdd(BlockId);
 			
+			MeshData.NameResolutionMap = RealTimeMorphStreamable.NameResolutionMap;
+
+			const TArray<FMorphTargetVertexData>& SourceData = ModelResources.EditorOnlyMorphTargetReconstructionData;	
 			const uint32 NumElems = Block.Size / sizeof(FMorphTargetVertexData);
 			const uint32 OffsetInElems = Block.Offset / sizeof(FMorphTargetVertexData);
-			DestData.SetNumUninitialized(NumElems);
+			MeshData.Data.SetNumUninitialized(NumElems);
 
 			check(SourceData.Num()*sizeof(FMorphTargetVertexData) >= Block.Offset + Block.Size);
-			FMemory::Memcpy(DestData.GetData(), SourceData.GetData() + OffsetInElems, Block.Size);
+			FMemory::Memcpy(MeshData.Data.GetData(), SourceData.GetData() + OffsetInElems, Block.Size);
 		}
 #else	
 		MUTABLE_CPUPROFILER_SCOPE(RealTimeMorphStreaming);
@@ -5195,20 +5022,23 @@ UE::Tasks::FTask UCustomizableInstancePrivate::LoadAdditionalAssetsAndData(
 			MUTABLE_CPUPROFILER_SCOPE(RealTimeMorphStreamingRequest_Alloc);
 
 			const int32 BlockId = RealTimeMorphStreamableBlocksToStream[I];
-			const FMutableStreamableBlock& Block = ModelResources.RealTimeMorphStreamableBlocks[BlockId]; 
+			const FMutableStreamableBlock& Block = ModelResources.RealTimeMorphStreamables[BlockId].Block; 
 		
-			TArray<FMorphTargetVertexData>& ReadDestData = OperationData->InstanceUpdateData.MorphTargetsVertexData.FindOrAdd(BlockId);
+			FInstanceUpdateData::FMorphTargetMeshData& ReadDestData = 
+					OperationData->InstanceUpdateData.RealTimeMorphTargetMeshData.FindOrAdd(BlockId);
 
 			// Only request blocks once.
-			if (ReadDestData.Num())
+			if (ReadDestData.Data.Num())
 			{
 				continue;
 			}
 
+			ReadDestData.NameResolutionMap = ModelResources.RealTimeMorphStreamables[BlockId].NameResolutionMap;
+
 			check(Block.Size % sizeof(FMorphTargetVertexData) == 0);
 			uint32 NumElems = Block.Size / sizeof(FMorphTargetVertexData);
 
-			ReadDestData.SetNumUninitialized(NumElems);
+			ReadDestData.Data.SetNumUninitialized(NumElems);
 
 			int32 FileHandleIndex = OpenFilesIds.Find(Block.FileId);
 			if (FileHandleIndex == INDEX_NONE && BulkData)
@@ -5223,7 +5053,7 @@ UE::Tasks::FTask UCustomizableInstancePrivate::LoadAdditionalAssetsAndData(
 			{ 
 				Block.Offset,
 				OpenFileHandles[FileHandleIndex].Get(), 
-				MakeArrayView(reinterpret_cast<uint8*>(ReadDestData.GetData()), ReadDestData.Num()*sizeof(FMorphTargetVertexData)),
+				MakeArrayView(reinterpret_cast<uint8*>(ReadDestData.Data.GetData()), ReadDestData.Data.Num()*sizeof(FMorphTargetVertexData)),
 				Block.FileId
 			});
 		}
diff --git a/Engine/Plugins/Experimental/Mutable/Source/CustomizableObject/Private/MuCO/CustomizableObjectSystem.cpp b/Engine/Plugins/Experimental/Mutable/Source/CustomizableObject/Private/MuCO/CustomizableObjectSystem.cpp
--- a/Engine/Plugins/Experimental/Mutable/Source/CustomizableObject/Private/MuCO/CustomizableObjectSystem.cpp
+++ b/Engine/Plugins/Experimental/Mutable/Source/CustomizableObject/Private/MuCO/CustomizableObjectSystem.cpp
@@ -1837,6 +1837,246 @@ namespace impl
 		}
 	}
 
+	void Subtask_Mutable_PrepareRealTimeMorphData(const TSharedRef<FUpdateContextPrivate>& OperationData)
+	{
+		MUTABLE_CPUPROFILER_SCOPE(BuildMorphTargetsData);
+
+		FInstanceUpdateData& UpdateData = OperationData->InstanceUpdateData;
+
+		const TMap<uint32, FInstanceUpdateData::FMorphTargetMeshData>& ResourceIdToMeshDataMap =
+				UpdateData.RealTimeMorphTargetMeshData;
+
+		if (ResourceIdToMeshDataMap.IsEmpty())
+		{
+			return;
+		}
+		
+		int32 NumNotFoundLoadedMorphsResources = 0;
+		
+		for (FInstanceUpdateData::FComponent& Component : OperationData->InstanceUpdateData.Components)
+		{
+			if (!OperationData->InstanceUpdateData.RealTimeMorphTargets.IsValidIndex(Component.Id))
+			{
+				OperationData->InstanceUpdateData.RealTimeMorphTargets.SetNum(Component.Id + 1);
+			}
+
+			FInstanceUpdateData::FRealTimeMorphsComponentData& ComponentMorphTargetsData = 
+					OperationData->InstanceUpdateData.RealTimeMorphTargets[Component.Id];
+
+			ComponentMorphTargetsData.ComponentIndex = Component.Id;
+		}
+
+		const int32 NumComponents = OperationData->InstanceUpdateData.RealTimeMorphTargets.Num();
+		for (int32 ComponentIndex = 0; ComponentIndex < NumComponents; ++ComponentIndex)
+		{
+			FInstanceUpdateData::FRealTimeMorphsComponentData& ComponentMorphTargetsData = 
+					OperationData->InstanceUpdateData.RealTimeMorphTargets[ComponentIndex];
+			
+			if (ComponentMorphTargetsData.ComponentIndex == INDEX_NONE)
+			{
+				continue;
+			}
+
+			struct FMorphTargetMeshData
+			{
+				TArray<int32> NameResolutionMap;
+				TArrayView<const FMorphTargetVertexData> DataView;
+			};
+
+			TArray<FName>& MorphTargetNames = ComponentMorphTargetsData.RealTimeMorphTargetNames;
+			MorphTargetNames.Empty();
+
+			TMap<uint32, FMorphTargetMeshData> MorphTargetMeshData;
+			MorphTargetMeshData.Reserve(ResourceIdToMeshDataMap.Num());
+			
+			for (const TPair<uint32, FInstanceUpdateData::FMorphTargetMeshData>& MorphTargetResource : ResourceIdToMeshDataMap)
+			{
+				FMorphTargetMeshData& MeshData = MorphTargetMeshData.FindOrAdd(MorphTargetResource.Key);
+				MeshData.DataView = MakeArrayView(MorphTargetResource.Value.Data);
+
+				const int32 NumMorphNames = MorphTargetResource.Value.NameResolutionMap.Num();
+				MeshData.NameResolutionMap.SetNumUninitialized(NumMorphNames);
+
+				for (int32 NameIndex = 0; NameIndex < NumMorphNames; ++NameIndex)
+				{
+					const int32 ResolvedNameIndex = MorphTargetNames.AddUnique(MorphTargetResource.Value.NameResolutionMap[NameIndex]);
+					MeshData.NameResolutionMap[NameIndex] = ResolvedNameIndex;
+				}
+			}
+
+			// Allocate Morph data for used morphs.
+			TArray<TArray<FMorphTargetLODModel>>& MorphsData = ComponentMorphTargetsData.RealTimeMorphsLODData;
+			const int32 NumMorphs = MorphTargetNames.Num();
+
+			MorphsData.SetNum(MorphTargetNames.Num());
+			for (int32 MorphIndex = 0; MorphIndex < NumMorphs; ++MorphIndex)
+			{
+				MorphsData[MorphIndex].SetNum(OperationData->NumLODsAvailable);
+			}
+
+			TArray<int32> SectionMorphTargetVerticesCount;
+			SectionMorphTargetVerticesCount.SetNumZeroed(ComponentMorphTargetsData.RealTimeMorphTargetNames.Num());
+
+			const int32 FirstGeneratedLOD = FMath::Max((int32)OperationData->GetRequestedLODs()[ComponentIndex], OperationData->GetMinLOD());
+			for (int32 LODIndex = FirstGeneratedLOD; LODIndex < OperationData->NumLODsAvailable; ++LODIndex)
+			{
+				const FInstanceUpdateData::FLOD& LOD = UpdateData.LODs[LODIndex];
+				const FInstanceUpdateData::FComponent& Component = UpdateData.Components[LOD.FirstComponent + ComponentIndex];
+				check(Component.bGenerated);
+				check(Component.Mesh);
+
+				const mu::FMeshBufferSet& MeshSet = Component.Mesh->GetVertexBuffers();
+
+				int32 VertexMorphsInfoIndexAndCountBufferIndex, VertexMorphsInfoIndexAndCountBufferChannel;
+				MeshSet.FindChannel(mu::MBS_OTHER, 0, &VertexMorphsInfoIndexAndCountBufferIndex, &VertexMorphsInfoIndexAndCountBufferChannel);
+
+				int32 VertexMorphsResourceIdBufferIndex, VertexMorphsResourceIdBufferChannel;
+				MeshSet.FindChannel(mu::MBS_OTHER, 1, &VertexMorphsResourceIdBufferIndex, &VertexMorphsResourceIdBufferChannel);
+
+				if (VertexMorphsInfoIndexAndCountBufferIndex < 0 || VertexMorphsResourceIdBufferIndex < 0)
+				{
+					continue;
+				}
+
+				const uint32* const VertexMorphsInfoIndexAndCountBuffer = reinterpret_cast<const uint32*>(MeshSet.GetBufferData(VertexMorphsInfoIndexAndCountBufferIndex));
+				TArrayView<const uint32> VertexMorphsInfoIndexAndCountView(VertexMorphsInfoIndexAndCountBuffer, MeshSet.GetElementCount());
+
+				const uint32* const VertexMorphsResourceIdBuffer = reinterpret_cast<const uint32*>(MeshSet.GetBufferData(VertexMorphsResourceIdBufferIndex));
+				TArrayView<const uint32> VertexMorphsResourceIdView(VertexMorphsResourceIdBuffer, MeshSet.GetElementCount());
+
+				const int32 SurfaceCount = Component.Mesh->GetSurfaceCount();
+				for (int32 Section = 0; Section < SurfaceCount; ++Section)
+				{
+					// Reset SectionMorphTargets.
+					for (int32& Elem : SectionMorphTargetVerticesCount)
+					{
+						Elem = 0;
+					}
+
+					int32 FirstVertex, VerticesCount, FirstIndex, IndiciesCount;
+					Component.Mesh->GetSurface(Section, &FirstVertex, &VerticesCount, &FirstIndex, &IndiciesCount, nullptr, nullptr, nullptr);
+
+					for (int32 VertexIdx = FirstVertex; VertexIdx < FirstVertex + VerticesCount;)
+					{
+						// Find a span with the same VertexMorphResourceId to amortise the cost of finding 
+						// in the loaded resources map. It is expected to find large consecutive mesh sections pointing to
+						// the same loaded resource.
+
+						const int32 SpanStart = VertexIdx++;
+						const uint32 CurrentResourceId = VertexMorphsResourceIdView[SpanStart];
+
+						// Vertex with no morphs are marked with 0, skip vertex if the case.
+						if (CurrentResourceId == 0)
+						{
+							continue;
+						}
+
+						for (; VertexIdx < FirstVertex + VerticesCount; ++VertexIdx)
+						{
+							const int32 VertexResourceId = VertexMorphsResourceIdView[VertexIdx];
+							// we can skip vertices with no morph without breaking the span.
+							if (VertexResourceId == 0)
+							{
+								continue;
+							}
+
+							if (CurrentResourceId != VertexResourceId)
+							{
+								break;
+							}
+						}
+						const int32 SpanEnd = VertexIdx;
+
+						const FMorphTargetMeshData* MorphTargetReconstructionData = MorphTargetMeshData.Find(CurrentResourceId);
+
+						if (!MorphTargetReconstructionData)
+						{
+							++NumNotFoundLoadedMorphsResources;
+							continue;
+						}
+
+						TArrayView<const FMorphTargetVertexData> SpanMorphData = MorphTargetReconstructionData->DataView;
+
+						for (int32 SpanVertexIdx = SpanStart; SpanVertexIdx < SpanEnd; ++SpanVertexIdx)
+						{
+							const uint32 MorphOffsetAndCount = VertexMorphsInfoIndexAndCountView[SpanVertexIdx];
+							if (MorphOffsetAndCount == 0)
+							{
+								continue;
+							}
+
+							// See encoding in GenerateMutableSourceMesh.cpp.
+							constexpr uint32 Log2MaxNumVerts = 23;
+							
+							TArrayView<const FMorphTargetVertexData> MorphsVertexDataView = MakeArrayView(
+									SpanMorphData.GetData() + (MorphOffsetAndCount & ((1 << Log2MaxNumVerts) - 1)), 
+									MorphOffsetAndCount >> Log2MaxNumVerts);
+
+							for (const FMorphTargetVertexData& SourceVertex : MorphsVertexDataView)
+							{
+								const uint32 ResolvedNameIndex =
+										MorphTargetReconstructionData->NameResolutionMap[SourceVertex.MorphNameIndex];
+
+								FMorphTargetLODModel& DestMorphLODModel = MorphsData[ResolvedNameIndex][LODIndex];
+
+								DestMorphLODModel.Vertices.Emplace(
+										FMorphTargetDelta 
+										{ 
+											SourceVertex.PositionDelta, 
+											SourceVertex.TangentZDelta, 
+											static_cast<uint32>(SpanVertexIdx) 
+										});
+
+								++SectionMorphTargetVerticesCount[ResolvedNameIndex];
+							}
+						}
+					}
+
+					const int32 SectionMorphTargetsNum = SectionMorphTargetVerticesCount.Num();
+					for (int32 MorphIdx = 0; MorphIdx < SectionMorphTargetsNum; ++MorphIdx)
+					{
+						if (SectionMorphTargetVerticesCount[MorphIdx] > 0)
+						{
+							FMorphTargetLODModel& MorphTargetLodModel = MorphsData[MorphIdx][LODIndex];
+
+							MorphTargetLodModel.SectionIndices.Add(Section);
+							MorphTargetLodModel.NumVertices += SectionMorphTargetVerticesCount[MorphIdx];
+						}
+					}
+				}
+			}
+		
+			// Remove empty morph targets;
+			for (int32 MorphIndex = 0; MorphIndex < NumMorphs; ++MorphIndex)
+			{
+				const int32 NumLODs = MorphsData[MorphIndex].Num();
+
+				int32 LODIndex = 0;
+				for (; LODIndex < NumLODs; ++LODIndex)
+				{
+					if (!MorphsData[MorphIndex][LODIndex].Vertices.IsEmpty())
+					{
+						break;
+					}
+				}
+
+				if (LODIndex >= NumLODs)
+				{
+					MorphsData[MorphIndex].Empty();
+				}
+			}
+
+		}
+
+		// Free unneeded data memory.
+		UpdateData.RealTimeMorphTargetMeshData.Empty();
+		
+		if (NumNotFoundLoadedMorphsResources > 0)
+		{
+			UE_LOG(LogMutable, Warning, TEXT("Needed realtime morph reconstruction data was not loaded properly. Some realtime morphs may not work correctly."));
+		}
+	}
+
 
 	/** End of the GetMeshes tasks. */
 	void Task_Mutable_GetMeshes_End(
@@ -1848,7 +2088,6 @@ namespace impl
 
 		// TODO: Not strictly mutable: move to another worker thread task to free mutable access?
 		Subtask_Mutable_PrepareSkeletonData(OperationData);
-
 		if (OperationData->GetCapturedDescriptor().GetBuildParameterRelevancy())
 		{
 			Subtask_Mutable_UpdateParameterRelevancy(OperationData);
@@ -2623,6 +2862,10 @@ namespace impl
 			return;
 		}
 
+		// TODO: This subtask should execute before Convert resources in a worker thread but after 
+		// Loading resources. For now keep it here.
+		Subtask_Mutable_PrepareRealTimeMorphData(OperationData);
+
 		UCustomizableObjectInstance* CustomizableObjectInstance = OperationData->Instance.Get();
 
 		// Actual work
diff --git a/Engine/Plugins/Experimental/Mutable/Source/CustomizableObjectEditor/Private/MuCOE/CustomizableObjectCompiler.cpp b/Engine/Plugins/Experimental/Mutable/Source/CustomizableObjectEditor/Private/MuCOE/CustomizableObjectCompiler.cpp
--- a/Engine/Plugins/Experimental/Mutable/Source/CustomizableObjectEditor/Private/MuCOE/CustomizableObjectCompiler.cpp
+++ b/Engine/Plugins/Experimental/Mutable/Source/CustomizableObjectEditor/Private/MuCOE/CustomizableObjectCompiler.cpp
@@ -1006,35 +1006,30 @@ void FCustomizableObjectCompiler::CompileInternal(UCustomizableObject* Object, c
 		ModelResources.ParameterUIDataMap = MoveTemp(GenerationContext.ParameterUIDataMap);
 		ModelResources.StateUIDataMap = MoveTemp(GenerationContext.StateUIDataMap);
 
-		ModelResources.RealTimeMorphTargetNames = MoveTemp(GenerationContext.RealTimeMorphTargetsNames);
-
-		if (GenerationContext.RealTimeMorphTargetPerMeshData.Num() >= TNumericLimits<uint16>::Max())
-		{
-			UE_LOG(LogMutable, Warning, TEXT("Maximum number of meshes with realtime morph targets reached. Some morphs may not work."));
-		}
-
 		// Create the RealTimeMorphsTargets Blocks from the per mesh Morph data.
 		uint64 RealTimeMorphDataSize = 0;
-		for (const TArray<FMorphTargetVertexData>& VertexDataArray : GenerationContext.RealTimeMorphTargetPerMeshData)
+		for (const TPair<uint32, FRealTimeMorphMeshData>& MeshData : GenerationContext.RealTimeMorphTargetPerMeshData)
 		{
-			RealTimeMorphDataSize += VertexDataArray.Num();
+			RealTimeMorphDataSize += MeshData.Value.Data.Num();
 		}
 		
-		ModelResources.RealTimeMorphStreamableBlocks.Empty(32);
+		ModelResources.RealTimeMorphStreamables.Empty(32);
 		ModelResources.EditorOnlyMorphTargetReconstructionData.Empty(RealTimeMorphDataSize);
 
-		uint64 RealTimeMorphDataOffset = 0;
-		for (const TArray<FMorphTargetVertexData>& VertexDataArray : GenerationContext.RealTimeMorphTargetPerMeshData)
+		uint64 RealTimeMorphDataOffsetInBytes = 0;
+		for (const TPair<uint32, FRealTimeMorphMeshData>& MeshData : GenerationContext.RealTimeMorphTargetPerMeshData)
 		{
-			ModelResources.RealTimeMorphStreamableBlocks.Emplace(FMutableStreamableBlock
-					{
-						uint32(0),
-						(uint32)VertexDataArray.Num()*sizeof(FMorphTargetVertexData), 
-						RealTimeMorphDataOffset, 
-					});
+			const uint32 DataSizeInBytes = (uint32)MeshData.Value.Data.Num()*sizeof(FMorphTargetVertexData); 
+			FRealTimeMorphStreamable& ResourceMeshData = ModelResources.RealTimeMorphStreamables.FindOrAdd(MeshData.Key);
+			
+			check(ResourceMeshData.NameResolutionMap.IsEmpty());
+			check(ResourceMeshData.Block.Size == 0);
+
+			ResourceMeshData.NameResolutionMap = MeshData.Value.NameResolutionMap;
+			ResourceMeshData.Block = FMutableStreamableBlock { uint32(0), DataSizeInBytes, RealTimeMorphDataOffsetInBytes };
 
-			RealTimeMorphDataOffset += VertexDataArray.Num()*sizeof(FMorphTargetVertexData);
-			ModelResources.EditorOnlyMorphTargetReconstructionData.Append(VertexDataArray);
+			RealTimeMorphDataOffsetInBytes += DataSizeInBytes;
+			ModelResources.EditorOnlyMorphTargetReconstructionData.Append(MeshData.Value.Data);
 		}
 		
 		// Clothing	
diff --git a/Engine/Plugins/Experimental/Mutable/Source/CustomizableObjectEditor/Private/MuCOE/GenerateMutableSource/GenerateMutableSource.h b/Engine/Plugins/Experimental/Mutable/Source/CustomizableObjectEditor/Private/MuCOE/GenerateMutableSource/GenerateMutableSource.h
--- a/Engine/Plugins/Experimental/Mutable/Source/CustomizableObjectEditor/Private/MuCOE/GenerateMutableSource/GenerateMutableSource.h
+++ b/Engine/Plugins/Experimental/Mutable/Source/CustomizableObjectEditor/Private/MuCOE/GenerateMutableSource/GenerateMutableSource.h
@@ -271,6 +271,11 @@ struct FPoseBoneData
 	TArray<FTransform> ArrayTransform;
 };
 
+struct FRealTimeMorphMeshData
+{
+	TArray<FName> NameResolutionMap;
+	TArray<FMorphTargetVertexData> Data;
+};
 
 struct FGroupProjectorTempData
 {
@@ -710,8 +715,7 @@ struct FMutableGraphGenerationContext
 	TMap<class UCustomizableObjectNodeMeshClipWithMesh*, TArray<mu::Ptr<mu::NodeModifierMeshClipWithMesh>>> MapClipMeshNodeToMutableClipMeshNodeArray;
 
 	// Data used for MorphTarget reconstruction.
-	TArray<FName> RealTimeMorphTargetsNames;
-	TArray<TArray<FMorphTargetVertexData>> RealTimeMorphTargetPerMeshData;
+	TMap<uint32, FRealTimeMorphMeshData> RealTimeMorphTargetPerMeshData;
 
 	// Data used for Clothing reconstruction.
 	TArray<FCustomizableObjectMeshToMeshVertData> ClothMeshToMeshVertData;
diff --git a/Engine/Plugins/Experimental/Mutable/Source/CustomizableObjectEditor/Private/MuCOE/GenerateMutableSource/GenerateMutableSourceMesh.cpp b/Engine/Plugins/Experimental/Mutable/Source/CustomizableObjectEditor/Private/MuCOE/GenerateMutableSource/GenerateMutableSourceMesh.cpp
--- a/Engine/Plugins/Experimental/Mutable/Source/CustomizableObjectEditor/Private/MuCOE/GenerateMutableSource/GenerateMutableSourceMesh.cpp
+++ b/Engine/Plugins/Experimental/Mutable/Source/CustomizableObjectEditor/Private/MuCOE/GenerateMutableSource/GenerateMutableSourceMesh.cpp
@@ -49,6 +49,8 @@
 #include "PhysicsEngine/SkeletalBodySetup.h"
 #include "Engine/SkinnedAssetCommon.h"
 
+#include "Hash/CityHash.h"
+
 #define LOCTEXT_NAMESPACE "CustomizableObjectEditor"
 
 
@@ -851,11 +853,7 @@ mu::MeshPtr ConvertSkeletalMeshToMutable(const USkeletalMesh* InSkeletalMesh, co
 
 	MutableMesh->GetVertexBuffers().SetElementCount(VertexCount);
 
-	const int32 VertexBuffersCount = 1 +
-		(GenerationContext.Options.bRealTimeMorphTargetsEnabled ? 3 : 0) +
-		(GenerationContext.Options.bClothingEnabled ? 1 : 0);
-
-	MutableMesh->GetVertexBuffers().SetBufferCount(VertexBuffersCount);
+	MutableMesh->GetVertexBuffers().SetBufferCount(1);
 
 	const int32 MaxSectionInfluences = MeshSection.MaxBoneInfluences;
 	const bool bUseUnlimitedInfluences = FGPUBaseSkinVertexFactory::UseUnlimitedBoneInfluences(MaxSectionInfluences, GenerationContext.Options.TargetPlatform);
@@ -1005,10 +1003,9 @@ mu::MeshPtr ConvertSkeletalMeshToMutable(const USkeletalMesh* InSkeletalMesh, co
 
 
 	// TODO: Add Mesh generation flags to not include RT Morph and clothing if not needed.
-	int32 nextBufferIndex = 1;
+	int32 NextBufferIndex = 1;
 	if (GenerationContext.Options.bRealTimeMorphTargetsEnabled)
 	{
-		nextBufferIndex += 3;
 
 		// This call involves resolving every TObjectPtr<UMorphTarget> to a UMorphTarget*, so
 		// cache the result here to avoid calling it repeatedly.
@@ -1117,77 +1114,64 @@ mu::MeshPtr ConvertSkeletalMeshToMutable(const USkeletalMesh* InSkeletalMesh, co
 			}
 		}
 
-		// MorphTarget vertex block offset.
-		{
-			using namespace mu;
-			const int32 ElementSize = sizeof(int32);
-			const int32 ChannelCount = 1;
-			const EMeshBufferSemantic Semantics[ChannelCount] = { MBS_OTHER };
-			const int32 SemanticIndices[ChannelCount] = { 0 };
-			const EMeshBufferFormat Formats[ChannelCount] = { MBF_INT32 };
-			int32 Components[ChannelCount] = { 1 };
-			const int32 Offsets[ChannelCount] = { 0 };
-
-			MutableMesh->GetVertexBuffers().SetBuffer(1, ElementSize, ChannelCount, Semantics, SemanticIndices, Formats, Components, Offsets);
-		}
-
-		// MorphTarget vertex morph count.
+		if (UsedMorphTargets.Num())
 		{
-			using namespace mu;
-			const int32 ElementSize = sizeof(uint16);
-			const int32 ChannelCount = 1;
-			const EMeshBufferSemantic Semantics[ChannelCount] = { MBS_OTHER };
-			const int32 SemanticIndices[ChannelCount] = { 1 };
-			const EMeshBufferFormat Formats[ChannelCount] = { MBF_UINT16 };
-			int32 Components[ChannelCount] = { 1 };
-			const int32 Offsets[ChannelCount] = { 0 };
-
-			MutableMesh->GetVertexBuffers().SetBuffer(2, ElementSize, ChannelCount, Semantics, SemanticIndices, Formats, Components, Offsets);
-		}
+			MutableMesh->GetVertexBuffers().SetBufferCount(NextBufferIndex + 2);
+			// MorphTarget vertex block offset.
+			{
+				using namespace mu;
+				const int32 ElementSize = sizeof(uint32);
+				const int32 ChannelCount = 1;
+				const EMeshBufferSemantic Semantics[ChannelCount] = { MBS_OTHER };
+				const int32 SemanticIndices[ChannelCount] = { 0 };
+				const EMeshBufferFormat Formats[ChannelCount] = { MBF_UINT32 };
+				int32 Components[ChannelCount] = { 1 };
+				const int32 Offsets[ChannelCount] = { 0 };
 
-		// MorphTarget vertex block id.
-		{
-			using namespace mu;
-			const int32 ElementSize = sizeof(uint16);
-			const int32 ChannelCount = 1;
-			const EMeshBufferSemantic Semantics[ChannelCount] = { MBS_OTHER };
-			const int32 SemanticIndices[ChannelCount] = { 2 };
-			const EMeshBufferFormat Formats[ChannelCount] = { MBF_UINT16 };
-			int32 Components[ChannelCount] = { 1 };
-			const int32 Offsets[ChannelCount] = { 0 };
+				MutableMesh->GetVertexBuffers().SetBuffer(NextBufferIndex, ElementSize, ChannelCount, Semantics, SemanticIndices, Formats, Components, Offsets);
+			}
 
-			MutableMesh->GetVertexBuffers().SetBuffer(3, ElementSize, ChannelCount, Semantics, SemanticIndices, Formats, Components, Offsets);
-		}
+			// MorphTarget vertex morph count.
 
+			// MorphTarget vertex block id.
+			{
+				using namespace mu;
+				const int32 ElementSize = sizeof(uint32);
+				const int32 ChannelCount = 1;
+				const EMeshBufferSemantic Semantics[ChannelCount] = { MBS_OTHER };
+				const int32 SemanticIndices[ChannelCount] = { 1 };
+				const EMeshBufferFormat Formats[ChannelCount] = { MBF_UINT32 };
+				int32 Components[ChannelCount] = { 1 };
+				const int32 Offsets[ChannelCount] = { 0 };
 
-		// Setup MorphTarget reconstruction data.
+				MutableMesh->GetVertexBuffers().SetBuffer(NextBufferIndex + 1, ElementSize, ChannelCount, Semantics, SemanticIndices, Formats, Components, Offsets);
+			}
 
-		TArrayView<int32> VertexMorphsOffsetBufferView(reinterpret_cast<int32*>(MutableMesh->GetVertexBuffers().GetBufferData(1)), VertexCount);
-		TArrayView<uint16> VertexMorphsCountBufferView(reinterpret_cast<uint16*>(MutableMesh->GetVertexBuffers().GetBufferData(2)), VertexCount);
-		TArrayView<uint16> VertexMorphsResourceIdBufferView(reinterpret_cast<uint16*>(MutableMesh->GetVertexBuffers().GetBufferData(3)), VertexCount);
-	
-		for (int32& Elem : VertexMorphsOffsetBufferView)
-		{
-			Elem = -1;
-		}
+			// Setup MorphTarget reconstruction data.
+			TArrayView<uint32> VertexMorphsOffsetAndCountBufferView(reinterpret_cast<uint32*>(
+						MutableMesh->GetVertexBuffers().GetBufferData(NextBufferIndex)), VertexCount);
+			TArrayView<uint32> VertexMorphsResourceIdBufferView(reinterpret_cast<uint32*>(
+						MutableMesh->GetVertexBuffers().GetBufferData(NextBufferIndex + 1)), VertexCount);
 
-		for (uint16& Elem : VertexMorphsCountBufferView)
-		{
-			Elem = 0;
-		}
+			for (uint32& Elem : VertexMorphsOffsetAndCountBufferView)
+			{
+				Elem = 0;
+			}
 
-		for (uint16& Elem : VertexMorphsResourceIdBufferView)
-		{
-			Elem = TNumericLimits<uint16>::Max(); 
-		}
+			constexpr uint32 InvalidResourceId = 0;
+			for (uint32& Elem : VertexMorphsResourceIdBufferView)
+			{
+				Elem = InvalidResourceId;
+			}
 
-		if (UsedMorphTargets.Num())
-		{
-			const double StartTime = FPlatformTime::Seconds();
+			MUTABLE_CPUPROFILER_SCOPE(RealTimeMorphTargetProcessing);
 
 			TArray<FMorphTargetVertexData> MorphsMeshData;
 			MorphsMeshData.Reserve(32);
 
+			TArray<FName> NameResolutionMap;
+			NameResolutionMap.Reserve(32);
+
 			TArray<FMorphTargetVertexData> MorphsUsed;
 			for (int32 VertexIdx = VertexStart; VertexIdx < VertexStart + VertexCount && VertexIdx < Vertices.Num(); ++VertexIdx)
 			{
@@ -1221,23 +1205,27 @@ mu::MeshPtr ConvertSkeletalMeshToMutable(const USkeletalMesh* InSkeletalMesh, co
 					const FMorphTargetDelta& VertexFound = MorphLODModels[LODIndex].Vertices[VertexFoundIndex];
 					const FName MorphTargetName = MorphTarget->GetFName();
 
-					TArray<FName>& RealTimeMorphTargetsNames = GenerationContext.RealTimeMorphTargetsNames;
-
-					int32 MorphTargetNameIndex = RealTimeMorphTargetsNames.Find(MorphTargetName);
+					int32 MorphTargetNameIndex = NameResolutionMap.Find(MorphTargetName);
 
 					MorphTargetNameIndex = MorphTargetNameIndex != INDEX_NONE
 							? MorphTargetNameIndex
-							: RealTimeMorphTargetsNames.Emplace(MorphTargetName);
+							: NameResolutionMap.Emplace(MorphTargetName);
 
 					MorphsUsed.Emplace(FMorphTargetVertexData{VertexFound.PositionDelta, VertexFound.TangentZDelta, (uint32)MorphTargetNameIndex});
 				}
 
 				if (MorphsUsed.Num())
 				{
+					// TODO: Those constants should be stored somewhere accessible to the decoding code.
+					constexpr uint32 Log2MaxNumVerts = 23;
+					constexpr uint32 Log2MaxNumMorphs = 32 - Log2MaxNumVerts;
 
-					VertexMorphsOffsetBufferView[VertexIdx - VertexStart] = MorphsMeshData.Num();
-					VertexMorphsCountBufferView[VertexIdx - VertexStart] = (uint16)MorphsUsed.Num();
-					VertexMorphsResourceIdBufferView[VertexIdx - VertexStart] = (uint16)GenerationContext.RealTimeMorphTargetPerMeshData.Num();
+					checkf(MorphsMeshData.Num() < (1 << Log2MaxNumVerts) - 1, TEXT("Maximum number of vertices with morphs per mesh reached."));
+					checkf(MorphsUsed.Num() < (1 << Log2MaxNumMorphs) - 1, TEXT("Maximum number of morphs per vertex reached."));
+
+					uint32 EncodedVertexOffsetAndCount = 
+							(uint32(MorphsMeshData.Num()) & ((1 << Log2MaxNumVerts) - 1)) | (uint32(MorphsUsed.Num()) << Log2MaxNumVerts);
+					VertexMorphsOffsetAndCountBufferView[VertexIdx - VertexStart] = EncodedVertexOffsetAndCount;
 					
 					MorphsMeshData.Append(MorphsUsed);
 				}
@@ -1246,19 +1234,95 @@ mu::MeshPtr ConvertSkeletalMeshToMutable(const USkeletalMesh* InSkeletalMesh, co
 			// Only commit the morph if there is data.
 			if (MorphsMeshData.Num())
 			{
-				if (GenerationContext.RealTimeMorphTargetPerMeshData.Num() < TNumericLimits<uint16>::Max())
+				const uint32 DataHash = CityHash32(
+						reinterpret_cast<const char*>(MorphsMeshData.GetData()), MorphsMeshData.Num() * sizeof(FMorphTargetVertexData));
+				
+				uint32 UniqueHash = DataHash == InvalidResourceId ? InvalidResourceId + 1 : DataHash;
+			
+				const FRealTimeMorphMeshData* FoundHash = GenerationContext.RealTimeMorphTargetPerMeshData.Find(UniqueHash);
+
+				bool bIsDataAlreadyCollected = false;
+				if (FoundHash)
+				{
+					bIsDataAlreadyCollected = 
+							FoundHash->NameResolutionMap == NameResolutionMap &&
+							FoundHash->Data.Num() == MorphsMeshData.Num() &&
+							FMemory::Memcmp(FoundHash->Data.GetData(), MorphsMeshData.GetData(), 
+									MorphsMeshData.Num()*sizeof(FMorphTargetVertexData)) == 0;
+				}
+
+				// NOTE: This way of unique hash generation guarantees all valid values can be used but given its 
+				// sequential nature a cascade of changes can occur if new meshes are added. Not many hash collisions 
+				// are expected so it should not be problematic.
+				if (FoundHash && !bIsDataAlreadyCollected)
+				{
+					uint32 NumTries = 0;
+					for (; NumTries < TNumericLimits<uint32>::Max(); ++NumTries)
+					{
+						FoundHash = GenerationContext.RealTimeMorphTargetPerMeshData.Find(UniqueHash);
+						
+						if (!FoundHash)
+						{
+							break;
+						}
+
+						bIsDataAlreadyCollected = 
+								FoundHash->NameResolutionMap == NameResolutionMap &&
+								FoundHash->Data.Num() == MorphsMeshData.Num() &&
+								FMemory::Memcmp(FoundHash->Data.GetData(), MorphsMeshData.GetData(), 
+										MorphsMeshData.Num()*sizeof(FMorphTargetVertexData)) == 0;
+
+						if (bIsDataAlreadyCollected)
+						{
+							break;
+						}
+
+						UniqueHash = UniqueHash + 1 == InvalidResourceId ? InvalidResourceId + 1 : UniqueHash + 1;
+					}
+
+					if (NumTries == TNumericLimits<uint32>::Max())
+					{
+						UE_LOG(LogMutable, Warning, TEXT("Maximum number of meshes with real time morphs reached, some morphs may not work as expected."));
+						UniqueHash = InvalidResourceId;
+					}
+				}
+
+				if (UniqueHash != InvalidResourceId)
 				{
+					// Add the data block reference to the mesh vertices.
+					for (int32 VertexIdx = VertexStart; VertexIdx < VertexStart + VertexCount && VertexIdx < Vertices.Num(); ++VertexIdx)
+					{
+						// Keep as invalid the vertices with no morphs, this way we can check if a vertex has morphs looking at either
+						// count and offset (Notice that if count is 0 then offset must be 0 as well) or the hash.
+						// This is useful because those two buffers are iterated in different phases when reconstructing the data back after
+						// mesh generation, but we need to know if a vertex has morphs in both phases. 
+						if (VertexMorphsOffsetAndCountBufferView[VertexIdx - VertexStart] != 0)
+						{
+							VertexMorphsResourceIdBufferView[VertexIdx - VertexStart] = UniqueHash;
+						}
+					}
+				
 					FCustomizableObjectStreameableResourceId StreamedMorphResource;
-					StreamedMorphResource.Id = GenerationContext.RealTimeMorphTargetPerMeshData.Num();
+					StreamedMorphResource.Id = UniqueHash;
 					StreamedMorphResource.Type = static_cast<uint8>(FCustomizableObjectStreameableResourceId::EType::RealTimeMorphTarget);
 
-					MutableMesh->AddStreamedResource(BitCast<uint32>(StreamedMorphResource));
-						
-					GenerationContext.RealTimeMorphTargetPerMeshData.Emplace_GetRef() = MoveTemp(MorphsMeshData);
+					MutableMesh->AddStreamedResource(BitCast<uint64>(StreamedMorphResource));
+
+					if (!bIsDataAlreadyCollected)
+					{
+						FRealTimeMorphMeshData& RealTimeMorphMeshData = 
+								GenerationContext.RealTimeMorphTargetPerMeshData.FindOrAdd(UniqueHash);
+
+						check(RealTimeMorphMeshData.NameResolutionMap.IsEmpty());
+						check(RealTimeMorphMeshData.Data.IsEmpty());
+
+						RealTimeMorphMeshData.NameResolutionMap = MoveTemp(NameResolutionMap);
+						RealTimeMorphMeshData.Data = MoveTemp(MorphsMeshData);
+					}
 				}
 			}
 
-			UE_LOG(LogMutable, Verbose, TEXT("Processing morph targets took %.2f ms"), (FPlatformTime::Seconds() - StartTime) * 1000.0);
+			NextBufferIndex += 2;
 		}
 	}
 
@@ -1270,15 +1334,16 @@ mu::MeshPtr ConvertSkeletalMeshToMutable(const USkeletalMesh* InSkeletalMesh, co
 			const int32 ElementSize = sizeof(int32);
 			const int32 ChannelCount = 1;
 			const EMeshBufferSemantic Semantics[ChannelCount] = { MBS_OTHER };
-			const int32 SemanticIndices[ChannelCount] = { GenerationContext.Options.bRealTimeMorphTargetsEnabled ? 3 : 0 };
+			const int32 SemanticIndices[ChannelCount] = { 2 };
 			const EMeshBufferFormat Formats[ChannelCount] = { MBF_INT32 };
 			int32 Components[ChannelCount] = { 1 };
 			const int32 Offsets[ChannelCount] = { 0 };
 
-			MutableMesh->GetVertexBuffers().SetBuffer(nextBufferIndex, ElementSize, ChannelCount, Semantics, SemanticIndices, Formats, Components, Offsets);
+			MutableMesh->GetVertexBuffers().SetBufferCount(NextBufferIndex + 1);
+			MutableMesh->GetVertexBuffers().SetBuffer(NextBufferIndex, ElementSize, ChannelCount, Semantics, SemanticIndices, Formats, Components, Offsets);
 		}
 
-		TArrayView<int32> ClothSectionBufferView(reinterpret_cast<int32*>(MutableMesh->GetVertexBuffers().GetBufferData(nextBufferIndex)), VertexCount);
+		TArrayView<int32> ClothSectionBufferView(reinterpret_cast<int32*>(MutableMesh->GetVertexBuffers().GetBufferData(NextBufferIndex)), VertexCount);
 		for (int32& Elem : ClothSectionBufferView)
 		{
 			Elem = -1;
@@ -1400,7 +1465,7 @@ mu::MeshPtr ConvertSkeletalMeshToMutable(const USkeletalMesh* InSkeletalMesh, co
 			}
 		}
 
-		++nextBufferIndex;
+		NextBufferIndex += 1;
 	}
 
 
@@ -1502,10 +1567,10 @@ mu::MeshPtr ConvertSkeletalMeshToMutable(const USkeletalMesh* InSkeletalMesh, co
 			const int32 Components[ChannelCount] = { 1, MutableBonesPerVertex, MutableBonesPerVertex };
 			const int32 Offsets[ChannelCount] = { 0, sizeof(int32), sizeof(int32) + BoneIndicesSize };
 
-			MutableMesh->GetVertexBuffers().SetBufferCount(nextBufferIndex + 1);
-			MutableMesh->GetVertexBuffers().SetBuffer(nextBufferIndex, ElementSize, ChannelCount, Semantics, SemanticIndices, Formats, Components, Offsets);
-			FMemory::Memcpy(MutableMesh->GetVertexBuffers().GetBufferData(nextBufferIndex), MutSkinWeights.GetData(), VertexCount * SkinWeightProfileVertexSize);
-			++nextBufferIndex;
+			MutableMesh->GetVertexBuffers().SetBufferCount(NextBufferIndex + 1);
+			MutableMesh->GetVertexBuffers().SetBuffer(NextBufferIndex, ElementSize, ChannelCount, Semantics, SemanticIndices, Formats, Components, Offsets);
+			FMemory::Memcpy(MutableMesh->GetVertexBuffers().GetBufferData(NextBufferIndex), MutSkinWeights.GetData(), VertexCount * SkinWeightProfileVertexSize);
+			++NextBufferIndex;
 		}
 	}
 
@@ -2939,9 +3004,7 @@ mu::NodeMeshPtr GenerateMutableSourceMesh(const UEdGraphPin* Pin,
 
 						const int32 ResourceIndex = GenerationContext.AddAssetUserDataToStreamedResources(AssetUserData);
 						if (ResourceIndex >= 0)
-						{
-							check(ResourceIndex < (1 << 24) - 1);
-							
+						{	
 							FCustomizableObjectStreameableResourceId ResourceId;
 							ResourceId.Id = (uint32)ResourceIndex;
 							ResourceId.Type = static_cast<uint8>(FCustomizableObjectStreameableResourceId::EType::AssetUserData);
@@ -3021,7 +3084,7 @@ mu::NodeMeshPtr GenerateMutableSourceMesh(const UEdGraphPin* Pin,
 
 				for (FCustomizableObjectStreameableResourceId ResourceId : StreamedResources)
 				{
-					MutableMesh->AddStreamedResource(BitCast<uint32>(ResourceId));
+					MutableMesh->AddStreamedResource(BitCast<uint64>(ResourceId));
 				}
 
 				AddSocketTagsToMesh(TypedNodeSkel->SkeletalMesh, MutableMesh, GenerationContext);
diff --git a/Engine/Plugins/Experimental/Mutable/Source/CustomizableObjectEditor/Private/MuCOE/GenerateMutableSource/GenerateMutableSourceSurface.cpp b/Engine/Plugins/Experimental/Mutable/Source/CustomizableObjectEditor/Private/MuCOE/GenerateMutableSource/GenerateMutableSourceSurface.cpp
--- a/Engine/Plugins/Experimental/Mutable/Source/CustomizableObjectEditor/Private/MuCOE/GenerateMutableSource/GenerateMutableSourceSurface.cpp
+++ b/Engine/Plugins/Experimental/Mutable/Source/CustomizableObjectEditor/Private/MuCOE/GenerateMutableSource/GenerateMutableSourceSurface.cpp
@@ -136,11 +136,11 @@ void SetSurfaceFormat( FMutableGraphGenerationContext& GenerationContext,
 	{
 		using namespace mu;
 		{
-			const int32 ElementSize = sizeof(int32);
+			const int32 ElementSize = sizeof(uint32);
 			constexpr int32 ChannelCount = 1;
 			const EMeshBufferSemantic Semantics[ChannelCount] = { MBS_OTHER };
 			const int32 SemanticIndices[ChannelCount] = { 0 };
-			const EMeshBufferFormat Formats[ChannelCount] = { MBF_INT32 };
+			const EMeshBufferFormat Formats[ChannelCount] = { MBF_UINT32 };
 			const int32 Components[ChannelCount] = { 1 };
 			const int32 Offsets[ChannelCount] = { 0 };
 
@@ -149,24 +149,11 @@ void SetSurfaceFormat( FMutableGraphGenerationContext& GenerationContext,
 		}
 
 		{
-			const int32 ElementSize = sizeof(uint16);
+			const int32 ElementSize = sizeof(uint32);
 			constexpr int32 ChannelCount = 1;
 			const EMeshBufferSemantic Semantics[ChannelCount] = { MBS_OTHER };
 			const int32 SemanticIndices[ChannelCount] = { 1 };
-			const EMeshBufferFormat Formats[ChannelCount] = { MBF_UINT16 };
-			const int32 Components[ChannelCount] = { 1 };
-			const int32 Offsets[ChannelCount] = { 0 };
-
-			OutVertexBufferFormat.SetBuffer(CurrentVertexBuffer, ElementSize, ChannelCount, Semantics, SemanticIndices, Formats, Components, Offsets);
-			++CurrentVertexBuffer;
-		}
-
-		{
-			const int32 ElementSize = sizeof(uint16);
-			constexpr int32 ChannelCount = 1;
-			const EMeshBufferSemantic Semantics[ChannelCount] = { MBS_OTHER };
-			const int32 SemanticIndices[ChannelCount] = { 2 };
-			const EMeshBufferFormat Formats[ChannelCount] = { MBF_UINT16 };
+			const EMeshBufferFormat Formats[ChannelCount] = { MBF_UINT32 };
 			const int32 Components[ChannelCount] = { 1 };
 			const int32 Offsets[ChannelCount] = { 0 };
 
@@ -182,7 +169,7 @@ void SetSurfaceFormat( FMutableGraphGenerationContext& GenerationContext,
 		const int32 ElementSize = sizeof(int32);
 		constexpr int32 ChannelCount = 1;
 		const EMeshBufferSemantic Semantics[ChannelCount] = { MBS_OTHER };
-		const int32 SemanticIndices[ChannelCount] = { bWithRealTimeMorphs ? 3 : 0 };
+		const int32 SemanticIndices[ChannelCount] = { 2 };
 		const EMeshBufferFormat Formats[ChannelCount] = { MBF_INT32 };
 		const int32 Components[ChannelCount] = { 1 };
 		const int32 Offsets[ChannelCount] = { 0 };
diff --git a/Engine/Plugins/Experimental/Mutable/Source/CustomizableObjectEditor/Private/MuCOE/GenerateMutableSource/GenerateMutableSourceTable.cpp b/Engine/Plugins/Experimental/Mutable/Source/CustomizableObjectEditor/Private/MuCOE/GenerateMutableSource/GenerateMutableSourceTable.cpp
--- a/Engine/Plugins/Experimental/Mutable/Source/CustomizableObjectEditor/Private/MuCOE/GenerateMutableSource/GenerateMutableSourceTable.cpp
+++ b/Engine/Plugins/Experimental/Mutable/Source/CustomizableObjectEditor/Private/MuCOE/GenerateMutableSource/GenerateMutableSourceTable.cpp
@@ -249,7 +249,7 @@ bool FillTableColumn(const UCustomizableObjectNodeTable* TableNode,	mu::TablePtr
 				MeshUniqueTags += GenerateGameplayTag(Tag.ToString());
 			}
 			
-			TArray<int32> StreamedResources;
+			TArray<FCustomizableObjectStreameableResourceId> StreamedResources;
 
 			if (GenerationContext.Object->bEnableAssetUserDataMerge)
 			{
@@ -268,13 +268,11 @@ bool FillTableColumn(const UCustomizableObjectNodeTable* TableNode,	mu::TablePtr
 						
 						if (ResourceIndex >= 0)
 						{
-							check(ResourceIndex < (1 << 24) - 1);
-
 							FCustomizableObjectStreameableResourceId ResourceId;
 							ResourceId.Id = GenerationContext.AddAssetUserDataToStreamedResources(AssetUserData);
 							ResourceId.Type = (uint8)FCustomizableObjectStreameableResourceId::EType::AssetUserData;
 
-							StreamedResources.Add(BitCast<uint32>(ResourceId));
+							StreamedResources.Add(ResourceId);
 						}
 
 						MeshUniqueTags += AssetUserData->GetPathName();
@@ -332,9 +330,9 @@ bool FillTableColumn(const UCustomizableObjectNodeTable* TableNode,	mu::TablePtr
 					AddTagToMutableMeshUnique(*MutableMesh, GenerateGameplayTag(Tag.ToString()));
 				}
 
-				for (int32 ResourceIndex : StreamedResources)
+				for (FCustomizableObjectStreameableResourceId ResourceId : StreamedResources)
 				{
-					MutableMesh->AddStreamedResource(ResourceIndex);
+					MutableMesh->AddStreamedResource(BitCast<uint64>(ResourceId));
 				}
 
 				AddSocketTagsToMesh(SkeletalMesh, MutableMesh, GenerationContext);
diff --git a/Engine/Plugins/Experimental/Mutable/Source/MutableRuntime/Private/MuR/Mesh.cpp b/Engine/Plugins/Experimental/Mutable/Source/MutableRuntime/Private/MuR/Mesh.cpp
--- a/Engine/Plugins/Experimental/Mutable/Source/MutableRuntime/Private/MuR/Mesh.cpp
+++ b/Engine/Plugins/Experimental/Mutable/Source/MutableRuntime/Private/MuR/Mesh.cpp
@@ -531,14 +531,14 @@ void Mesh::SetTag( int tagIndex, const FString& Name )
 
 
 //---------------------------------------------------------------------------------------------
-void Mesh::AddStreamedResource(uint32 ResourceId)
+void Mesh::AddStreamedResource(uint64 ResourceId)
 {
 	StreamedResources.AddUnique(ResourceId);
 }
 
 
 //---------------------------------------------------------------------------------------------
-const TArray<uint32>& Mesh::GetStreamedResources() const
+const TArray<uint64>& Mesh::GetStreamedResources() const
 {
 	return StreamedResources;
 }
diff --git a/Engine/Plugins/Experimental/Mutable/Source/MutableRuntime/Public/MuR/Mesh.h b/Engine/Plugins/Experimental/Mutable/Source/MutableRuntime/Public/MuR/Mesh.h
--- a/Engine/Plugins/Experimental/Mutable/Source/MutableRuntime/Public/MuR/Mesh.h
+++ b/Engine/Plugins/Experimental/Mutable/Source/MutableRuntime/Public/MuR/Mesh.h
@@ -283,10 +283,10 @@ namespace mu
         void SetTag( int tagIndex, const FString& Name );
 
 		//!
-		void AddStreamedResource(uint32 ResourceId);
+		void AddStreamedResource(uint64 ResourceId);
 
 		//!
-		const TArray<uint32>& GetStreamedResources() const;
+		const TArray<uint64>& GetStreamedResources() const;
 
 		//!
 		int32 FindBonePose(const FBoneName& BoneName) const;
@@ -385,7 +385,7 @@ namespace mu
 		TArray<FString> m_tags;
 
 		// Opaque handle to external resources.
-		TArray<uint32> StreamedResources;
+		TArray<uint64> StreamedResources;
 
 		struct FBonePose
 		{
