[Mutable] Fix missing bones due to an error when building the bone hash.

#rnx

diff --git a/Engine/Plugins/Experimental/Mutable/Source/CustomizableObject/Internal/MuCO/CustomizableObjectPrivate.h b/Engine/Plugins/Experimental/Mutable/Source/CustomizableObject/Internal/MuCO/CustomizableObjectPrivate.h
--- a/Engine/Plugins/Experimental/Mutable/Source/CustomizableObject/Internal/MuCO/CustomizableObjectPrivate.h
+++ b/Engine/Plugins/Experimental/Mutable/Source/CustomizableObject/Internal/MuCO/CustomizableObjectPrivate.h
@@ -602,6 +602,9 @@ struct FModelResources
 	UPROPERTY()
 	TArray<FName> MaterialSlotNames;
 
+	UPROPERTY()
+	TMap<FName, uint32> BoneNamesMap;
+
 	/** Bones remapped due to a hash collision at compile time. FName to (remapped)Hash. */
 	UPROPERTY()
 	TArray<FMutableRemappedBone> RemappedBoneNames;
@@ -879,6 +882,6 @@ public:
 	// This is a manual version number for the binary blobs in this asset.
 	// Increasing it invalidates all the previously compiled models.
 	// Warning: If while merging code both versions have changed, take the highest+1.
-	static constexpr int32 CurrentSupportedVersion = 441;
+	static constexpr int32 CurrentSupportedVersion = 442;
 };
 
diff --git a/Engine/Plugins/Experimental/Mutable/Source/CustomizableObject/Private/MuCO/CustomizableObject.cpp b/Engine/Plugins/Experimental/Mutable/Source/CustomizableObject/Private/MuCO/CustomizableObject.cpp
--- a/Engine/Plugins/Experimental/Mutable/Source/CustomizableObject/Private/MuCO/CustomizableObject.cpp
+++ b/Engine/Plugins/Experimental/Mutable/Source/CustomizableObject/Private/MuCO/CustomizableObject.cpp
@@ -550,6 +550,7 @@ void UCustomizableObjectPrivate::SaveCompiledData(FArchive& MemoryWriter, bool b
 	MemoryWriter << LocalModelResources.AnimBpOverridePhysiscAssetsInfo;
 
 	MemoryWriter << LocalModelResources.MaterialSlotNames;
+	MemoryWriter << LocalModelResources.BoneNamesMap;
 	MemoryWriter << LocalModelResources.RemappedBoneNames;
 	MemoryWriter << LocalModelResources.SocketArray;
 
@@ -670,6 +671,7 @@ void UCustomizableObjectPrivate::LoadCompiledData(FArchive& MemoryReader, const
 		MemoryReader << LocalModelResource.AnimBpOverridePhysiscAssetsInfo;
 
 		MemoryReader << LocalModelResource.MaterialSlotNames;
+		MemoryReader << LocalModelResource.BoneNamesMap;
 		MemoryReader << LocalModelResource.RemappedBoneNames;
 		MemoryReader << LocalModelResource.SocketArray;
 
diff --git a/Engine/Plugins/Experimental/Mutable/Source/CustomizableObject/Private/MuCO/CustomizableObjectInstance.cpp b/Engine/Plugins/Experimental/Mutable/Source/CustomizableObject/Private/MuCO/CustomizableObjectInstance.cpp
--- a/Engine/Plugins/Experimental/Mutable/Source/CustomizableObject/Private/MuCO/CustomizableObjectInstance.cpp
+++ b/Engine/Plugins/Experimental/Mutable/Source/CustomizableObject/Private/MuCO/CustomizableObjectInstance.cpp
@@ -3333,10 +3333,10 @@ bool UCustomizableInstancePrivate::BuildSkeletonData(const TSharedRef<FUpdateCon
 				BoneInfo.Key = BoneName;
 				BoneInfo.Value = Index;
 			}
-			else
+			else if(const uint32* Hash = ModelResources.BoneNamesMap.Find(BoneName))
 			{
 				const FString BoneNameString = BoneName.ToString();
-				const mu::FBoneName Bone(CityHash32(reinterpret_cast<const char*>(*BoneNameString), BoneNameString.Len() * sizeof(FString::ElementType)));
+				const mu::FBoneName Bone(*Hash);
 				TPair<FName, uint16>& BoneInfo = BoneInfoMap.Add(Bone);
 				BoneInfo.Key = BoneName;
 				BoneInfo.Value = Index;
diff --git a/Engine/Plugins/Experimental/Mutable/Source/CustomizableObjectEditor/Private/MuCOE/CustomizableObjectCompiler.cpp b/Engine/Plugins/Experimental/Mutable/Source/CustomizableObjectEditor/Private/MuCOE/CustomizableObjectCompiler.cpp
--- a/Engine/Plugins/Experimental/Mutable/Source/CustomizableObjectEditor/Private/MuCOE/CustomizableObjectCompiler.cpp
+++ b/Engine/Plugins/Experimental/Mutable/Source/CustomizableObjectEditor/Private/MuCOE/CustomizableObjectCompiler.cpp
@@ -973,6 +973,14 @@ void FCustomizableObjectCompiler::CompileInternal(UCustomizableObject* Object, c
 		ModelResources.MaterialSlotNames = MoveTemp(GenerationContext.ReferencedMaterialSlotNames);
 		ModelResources.SocketArray = MoveTemp(GenerationContext.SocketArray);
 
+		const int32 NumBones = GenerationContext.UniqueBoneNames.Num();
+		ModelResources.BoneNamesMap.Reserve(NumBones);
+
+		for (auto& It : GenerationContext.UniqueBoneNames)
+		{
+			ModelResources.BoneNamesMap.Add(It.Value, It.Key.Id);
+		}
+
 		for (auto& It : GenerationContext.RemappedBoneNames)
 		{
 			FMutableRemappedBone& RemappedBone = ModelResources.RemappedBoneNames.AddDefaulted_GetRef();
