[mutable] Remove custom serialization code frome the MutableTools module.

#rnx

diff --git a/Engine/Plugins/Experimental/Mutable/Source/CustomizableObjectEditor/Private/MuCOE/SMutableCodeViewer.cpp b/Engine/Plugins/Experimental/Mutable/Source/CustomizableObjectEditor/Private/MuCOE/SMutableCodeViewer.cpp
--- a/Engine/Plugins/Experimental/Mutable/Source/CustomizableObjectEditor/Private/MuCOE/SMutableCodeViewer.cpp
+++ b/Engine/Plugins/Experimental/Mutable/Source/CustomizableObjectEditor/Private/MuCOE/SMutableCodeViewer.cpp
@@ -26,7 +26,6 @@
 #include "MuCOE/UnrealEditorPortabilityHelpers.h"
 #include "MuR/SystemPrivate.h"
 #include "MuT/ErrorLog.h"
-#include "MuT/Streams.h"
 #include "MuT/TypeInfo.h"
 #include "Widgets/SNullWidget.h"
 #include "Widgets/Colors/SColorBlock.h"
@@ -354,60 +353,6 @@ void SMutableCodeViewer::Construct(const FArguments& InArgs, const TSharedPtr<mu
 	ToolbarBuilder.SetLabelVisibility(EVisibility::Visible);
 	ToolbarBuilder.SetStyle(&FAppStyle::Get(), "SlimToolBar");
 
-	ToolbarBuilder.BeginSection("Export");
-
-	// Tree Sizes
-	constexpr float OperationsColumnWidth = 0.90f;
-	constexpr float ExtraDataColumnWidth = 0.10f;
-	
-	// Export
-	ToolbarBuilder.AddToolBarButton(
-		FUIAction(
-			FExecuteAction::CreateLambda([InMutableModel]()
-				{
-					TArray<FString> SaveFilenames;
-					IDesktopPlatform* DesktopPlatform = FDesktopPlatformModule::Get();
-					bool bSave = false;
-					if (DesktopPlatform)
-					{
-						const FString LastExportPath = FEditorDirectories::Get().GetLastDirectory(ELastDirectory::GENERIC_EXPORT);
-						const FString FileTypes = TEXT("Mutable compiled data files|*.mutable_compiled|All files|*.*");
-						bSave = DesktopPlatform->SaveFileDialog(
-							FSlateApplication::Get().FindBestParentWindowHandleForDialogs(nullptr),
-							TEXT("Export Mutable compiled object"),
-							*LastExportPath,
-							TEXT("exported.mutable_compiled"),
-							*FileTypes,
-							EFileDialogFlags::None,
-							SaveFilenames
-						);
-					}
-
-					if (bSave)
-					{
-						const FString SaveFileName = FString(SaveFilenames[0]);
-
-						mu::OutputFileStream Stream(SaveFileName);
-						Stream.Write(MUTABLE_COMPILED_MODEL_FILETAG, 4);
-						mu::OutputArchive Archive(&Stream);
-						mu::Model::Serialise(InMutableModel.Get(), Archive);
-						const uint32_t CodeVersion = MUTABLE_COMPILED_MODEL_CODE_VERSION;
-						Stream.Write((const uint8*)&CodeVersion, sizeof(uint32_t));
-						Stream.Flush();
-
-						FEditorDirectories::Get().SetLastDirectory(ELastDirectory::GENERIC_EXPORT, SaveFileName);
-					}
-				})
-			),
-			NAME_None,
-			LOCTEXT("ExportMutableCode", "Export"),
-			LOCTEXT("ExportMutableCodeTooltip", "Export a debug mutable compiled file."),
-			FSlateIcon(),
-			EUserInterfaceActionType::Button
-		);
-		
-	ToolbarBuilder.EndSection();
-
 	ToolbarBuilder.AddWidget(SNew(STextBlock).Text(FText::FromString(InArgs._DataTag)));
 
 	TSharedRef<SScrollBar> TreeVertScrollBar =
@@ -2866,82 +2811,6 @@ void SMutableCodeViewer::PreviewMutableShape(const mu::FShape* Shape)
 }
 
 
-
-FReply SMutableCodeViewer::OnDragOver(const FGeometry& MyGeometry, const FDragDropEvent& DragDropEvent)
-{
-	if (TSharedPtr<FExternalDragOperation> DragDropOp = DragDropEvent.GetOperationAs<FExternalDragOperation>())
-	{
-		if (DragDropOp->HasFiles())
-		{
-			// For now, only allow a single file.
-			const TArray<FString>& Files = DragDropOp->GetFiles();
-			if (Files.Num() == 1)
-			{
-				const FString DraggedFileExtension = FPaths::GetExtension(Files[0], true);
-				if (DraggedFileExtension == TEXT(".mutable_compiled"))
-				{
-					// Dump source model to a file.
-					mu::InputFileStream stream(Files[0]);
-
-					char MutableSourceTag[4] = {};
-					stream.Read(MutableSourceTag, 4);
-
-					if (!FMemory::Memcmp(MutableSourceTag, MUTABLE_COMPILED_MODEL_FILETAG, 4))
-					{
-						return FReply::Handled();
-					}
-
-					return FReply::Unhandled();
-				}
-			}
-		}
-	}
-
-	return FReply::Unhandled();
-}
-
-
-FReply SMutableCodeViewer::OnDrop(const FGeometry& MyGeometry, const FDragDropEvent& DragDropEvent)
-{
-	SCompoundWidget::OnDrop(MyGeometry, DragDropEvent);
-
-	if (TSharedPtr<FExternalDragOperation> DragDropOp = DragDropEvent.GetOperationAs<FExternalDragOperation>())
-	{
-		if (DragDropOp->HasFiles())
-		{
-			// For now, only allow a single file.
-			const TArray<FString>& Files = DragDropOp->GetFiles();
-			if (Files.Num() == 1)
-			{
-				const FString DraggedFileExtension = FPaths::GetExtension(Files[0], true);
-				if (DraggedFileExtension == TEXT(".mutable_compiled"))
-				{
-					// Read a mutable compiled model file.
-					mu::InputFileStream stream(Files[0]);
-
-					char MutableSourceTag[4] = {};
-					stream.Read(MutableSourceTag, 4);
-
-					if (!FMemory::Memcmp(MutableSourceTag, MUTABLE_COMPILED_MODEL_FILETAG, 4))
-					{
-						mu::InputArchive arch(&stream);
-						TSharedPtr<mu::Model, ESPMode::ThreadSafe> Model = mu::Model::StaticUnserialise(arch);
-						TArray<TSoftObjectPtr<UTexture>> DummyReferencedTextures;
-						SetCurrentModel(Model, DummyReferencedTextures);
-
-						TreeView->RequestTreeRefresh();
-
-						return FReply::Handled();
-					}
-				}
-			}
-		}
-	}
-
-	return FReply::Unhandled();
-}
-
-
 FMutableCodeTreeElement::FMutableCodeTreeElement(int32 InIndexOnTree, const int32& InMutableStateIndex, const TSharedPtr<mu::Model, ESPMode::ThreadSafe>& InModel, mu::OP::ADDRESS InOperation, const FString& InCaption, const FSlateColor InLabelColor, const TSharedPtr<FMutableCodeTreeElement>* InDuplicatedOf)
 {
 	MutableModel = InModel;
diff --git a/Engine/Plugins/Experimental/Mutable/Source/CustomizableObjectEditor/Private/MuCOE/SMutableCodeViewer.h b/Engine/Plugins/Experimental/Mutable/Source/CustomizableObjectEditor/Private/MuCOE/SMutableCodeViewer.h
--- a/Engine/Plugins/Experimental/Mutable/Source/CustomizableObjectEditor/Private/MuCOE/SMutableCodeViewer.h
+++ b/Engine/Plugins/Experimental/Mutable/Source/CustomizableObjectEditor/Private/MuCOE/SMutableCodeViewer.h
@@ -64,8 +64,6 @@ public:
 
 	// SWidget interface
 	virtual void Tick(const FGeometry& AllottedGeometry, const double InCurrentTime, const float InDeltaTime) override;
-	FReply OnDragOver(const FGeometry& MyGeometry, const FDragDropEvent& DragDropEvent) override;
-	FReply OnDrop(const FGeometry& MyGeometry, const FDragDropEvent& DragDropEvent) override;
 
 	// FGCObject interface
 	virtual void AddReferencedObjects(FReferenceCollector& Collector) override;
diff --git a/Engine/Plugins/Experimental/Mutable/Source/CustomizableObjectEditor/Private/MuCOE/SMutableGraphViewer.cpp b/Engine/Plugins/Experimental/Mutable/Source/CustomizableObjectEditor/Private/MuCOE/SMutableGraphViewer.cpp
--- a/Engine/Plugins/Experimental/Mutable/Source/CustomizableObjectEditor/Private/MuCOE/SMutableGraphViewer.cpp
+++ b/Engine/Plugins/Experimental/Mutable/Source/CustomizableObjectEditor/Private/MuCOE/SMutableGraphViewer.cpp
@@ -41,7 +41,6 @@
 #include "MuT/NodeSurfaceEdit.h"
 #include "MuT/NodeSurfaceSwitch.h"
 #include "MuT/NodeSurfaceVariation.h"
-#include "MuT/Streams.h"
 #include "MuT/NodeLOD.h"
 #include "MuT/NodeMeshConstant.h"
 #include "MuT/NodeMeshFormat.h"
@@ -211,50 +210,6 @@ void SMutableGraphViewer::Construct(const FArguments& InArgs, const mu::NodePtr&
 	ToolbarBuilder.SetLabelVisibility(EVisibility::Visible);
 	ToolbarBuilder.SetStyle(&FAppStyle::Get(), "SlimToolBar");
 
-	// Export
-	ToolbarBuilder.AddToolBarButton(
-		FUIAction(
-			FExecuteAction::CreateLambda([InRootNode]()
-				{
-					TArray<FString> SaveFilenames;
-					IDesktopPlatform* DesktopPlatform = FDesktopPlatformModule::Get();
-					bool bSave = false;
-					if (!DesktopPlatform) return;
-
-					FString LastExportPath = FEditorDirectories::Get().GetLastDirectory(ELastDirectory::GENERIC_EXPORT);
-					FString FileTypes = TEXT("Mutable source data files|*.mutable_source|All files|*.*");
-					bSave = DesktopPlatform->SaveFileDialog(
-						FSlateApplication::Get().FindBestParentWindowHandleForDialogs(nullptr),
-						TEXT("Export Mutable object"),
-						*LastExportPath,
-						TEXT("exported.mutable_source"),
-						*FileTypes,
-						EFileDialogFlags::None,
-						SaveFilenames
-					);
-
-					if (!bSave) return;
-
-					// Dump source model to a file.
-					FString SaveFileName = FString(SaveFilenames[0]);
-					mu::OutputFileStream stream(SaveFileName);
-					stream.Write(MUTABLE_SOURCE_MODEL_FILETAG, 4);
-					mu::OutputArchive arch(&stream);
-					mu::Node::Serialise(InRootNode.get(), arch);
-					stream.Flush();
-
-					FEditorDirectories::Get().SetLastDirectory(ELastDirectory::GENERIC_EXPORT, SaveFileName);
-				})
-			),
-			NAME_None,
-			LOCTEXT("ExportMutableGraph", "Export"),
-			LOCTEXT("ExportMutableGraphTooltip", "Export a debug mutable graph file."),
-			FSlateIcon(),
-			EUserInterfaceActionType::Button
-		);
-		
-	ToolbarBuilder.EndSection();
-
 	ToolbarBuilder.AddWidget(SNew(STextBlock).Text(MakeAttributeLambda([this]() { return FText::FromString(DataTag); })));
 
 	ChildSlot
@@ -904,82 +859,6 @@ void SMutableGraphViewer::TreeExpandUnique()
 }
 
 
-FReply SMutableGraphViewer::OnDragOver(const FGeometry& MyGeometry, const FDragDropEvent& DragDropEvent)
-{
-	if (TSharedPtr<FExternalDragOperation> DragDropOp = DragDropEvent.GetOperationAs<FExternalDragOperation>())
-	{
-		if (DragDropOp->HasFiles())
-		{
-			// For now, only allow a single file.
-			const TArray<FString>& Files = DragDropOp->GetFiles();
-			if (Files.Num() == 1)
-			{
-				const FString DraggedFileExtension = FPaths::GetExtension(Files[0], true);
-				if (DraggedFileExtension == TEXT(".mutable_source"))
-				{
-					// Dump source model to a file.
-					mu::InputFileStream stream(Files[0]);
-
-					char MutableSourceTag[4] = {};
-					stream.Read(MutableSourceTag, 4);
-
-					if (!FMemory::Memcmp(MutableSourceTag, MUTABLE_SOURCE_MODEL_FILETAG, 4))
-					{
-						return FReply::Handled();
-					}
-
-					return FReply::Unhandled();
-				}
-			}
-		}
-	}
-
-	return FReply::Unhandled();
-}
-
-
-FReply SMutableGraphViewer::OnDrop(const FGeometry& MyGeometry, const FDragDropEvent& DragDropEvent)
-{
-	SCompoundWidget::OnDrop(MyGeometry, DragDropEvent);
-
-	if (TSharedPtr<FExternalDragOperation> DragDropOp = DragDropEvent.GetOperationAs<FExternalDragOperation>())
-	{
-		if (DragDropOp->HasFiles())
-		{
-			// For now, only allow a single file.
-			const TArray<FString>& Files = DragDropOp->GetFiles();
-			if (Files.Num() == 1)
-			{
-				const FString DraggedFileExtension = FPaths::GetExtension(Files[0], true);
-				if (DraggedFileExtension == TEXT(".mutable_source"))
-				{
-					// Dump source model to a file.
-					mu::InputFileStream stream(Files[0]);
-
-					char MutableSourceTag[4] = {};
-					stream.Read(MutableSourceTag, 4);
-
-					if (!FMemory::Memcmp(MutableSourceTag, MUTABLE_SOURCE_MODEL_FILETAG, 4))
-					{
-						mu::InputArchive arch(&stream);
-						RootNode = mu::Node::StaticUnserialise( arch );
-						DataTag = FString("dropped-file ") + FPaths::GetCleanFilename(Files[0]);
-						RebuildTree();
-
-						return FReply::Handled();
-					}
-
-					return FReply::Unhandled();
-				}
-			}
-		}
-
-		return FReply::Unhandled();
-	}
-
-	return FReply::Unhandled();
-}
-
 #undef LOCTEXT_NAMESPACE 
 
 
diff --git a/Engine/Plugins/Experimental/Mutable/Source/CustomizableObjectEditor/Private/MuCOE/SMutableGraphViewer.h b/Engine/Plugins/Experimental/Mutable/Source/CustomizableObjectEditor/Private/MuCOE/SMutableGraphViewer.h
--- a/Engine/Plugins/Experimental/Mutable/Source/CustomizableObjectEditor/Private/MuCOE/SMutableGraphViewer.h
+++ b/Engine/Plugins/Experimental/Mutable/Source/CustomizableObjectEditor/Private/MuCOE/SMutableGraphViewer.h
@@ -42,10 +42,6 @@ public:
 
 	void Construct(const FArguments& InArgs, const mu::NodePtr& InRootNode);
 
-	// SWidget interface
-	FReply OnDragOver(const FGeometry& MyGeometry, const FDragDropEvent& DragDropEvent) override;
-	FReply OnDrop(const FGeometry& MyGeometry, const FDragDropEvent& DragDropEvent) override;
-
 	// FGCObject interface
 	virtual void AddReferencedObjects(FReferenceCollector& Collector) override;
 	virtual FString GetReferencerName() const override;
diff --git a/Engine/Plugins/Experimental/Mutable/Source/MutableTools/Internal/MuT/AST.h b/Engine/Plugins/Experimental/Mutable/Source/MutableTools/Internal/MuT/AST.h
--- a/Engine/Plugins/Experimental/Mutable/Source/MutableTools/Internal/MuT/AST.h
+++ b/Engine/Plugins/Experimental/Mutable/Source/MutableTools/Internal/MuT/AST.h
@@ -45,6 +45,15 @@ namespace std
 
 }
 
+namespace 
+{
+	// TODO: Replace with UE hashing
+	template <class T> inline void hash_combine(uint64& seed, const T& v) {
+		std::hash<T> hasher;
+		seed ^= hasher(v) + 0x9e3779b9 + (seed << 6) + (seed >> 2);
+	}
+}
+
 
 namespace mu
 {
diff --git a/Engine/Plugins/Experimental/Mutable/Source/MutableTools/Internal/MuT/NodeColourFromScalarsPrivate.h b/Engine/Plugins/Experimental/Mutable/Source/MutableTools/Internal/MuT/NodeColourFromScalarsPrivate.h
--- a/Engine/Plugins/Experimental/Mutable/Source/MutableTools/Internal/MuT/NodeColourFromScalarsPrivate.h
+++ b/Engine/Plugins/Experimental/Mutable/Source/MutableTools/Internal/MuT/NodeColourFromScalarsPrivate.h
@@ -6,7 +6,6 @@
 #include "MuT/NodeColourFromScalars.h"
 #include "MuT/NodeScalar.h"
 #include "MuT/NodeImage.h"
-#include "MuT/AST.h"
 
 
 namespace mu
@@ -24,30 +23,6 @@ namespace mu
 		NodeScalarPtr m_pZ;
 		NodeScalarPtr m_pW;
 
-		//!
-		void Serialise( OutputArchive& arch ) const
-		{
-            uint32_t ver = 0;
-			arch << ver;
-
-			arch << m_pX;
-			arch << m_pY;
-			arch << m_pZ;
-			arch << m_pW;
-		}
-
-		//!
-		void Unserialise( InputArchive& arch )
-		{
-            uint32_t ver;
-			arch >> ver;
-			check(ver==0);
-
-			arch >> m_pX;
-			arch >> m_pY;
-			arch >> m_pZ;
-			arch >> m_pW;
-		}
 	};
 
 
diff --git a/Engine/Plugins/Experimental/Mutable/Source/MutableTools/Internal/MuT/NodeColourParameterPrivate.h b/Engine/Plugins/Experimental/Mutable/Source/MutableTools/Internal/MuT/NodeColourParameterPrivate.h
--- a/Engine/Plugins/Experimental/Mutable/Source/MutableTools/Internal/MuT/NodeColourParameterPrivate.h
+++ b/Engine/Plugins/Experimental/Mutable/Source/MutableTools/Internal/MuT/NodeColourParameterPrivate.h
@@ -23,55 +23,6 @@ namespace mu
 		FString m_uid;
 
         TArray<Ptr<NodeRange>> m_ranges;
-
-		//!
-		void Serialise( OutputArchive& arch ) const
-		{
-            uint32 ver = 4;
-			arch << ver;
-
-			arch << m_defaultValue;
-			arch << m_name;
-            arch << m_uid;
-            arch << m_ranges;
-        }
-
-		//!
-		void Unserialise( InputArchive& arch )
-		{
-            uint32 ver;
-			arch >> ver;
-            check(ver>=2&&ver<=4);
-
-			if (ver <= 3)
-			{
-				FVector3f TempDefaultValue;
-				arch >> TempDefaultValue;
-
-				m_defaultValue = TempDefaultValue;
-				m_defaultValue[3] = 1.f;
-			}
-			else
-			{
-				arch >> m_defaultValue;
-			}
-			
-			if (ver <= 2)
-			{
-				std::string Temp;
-				arch >> Temp;
-				m_name = Temp.c_str();
-				arch >> Temp;
-				m_uid = Temp.c_str();
-			}
-			else
-			{
-				arch >> m_name;
-				arch >> m_uid;
-			}
-			
-			arch >> m_ranges;
-        }
 	};
 
 }
diff --git a/Engine/Plugins/Experimental/Mutable/Source/MutableTools/Internal/MuT/NodeColourSampleImagePrivate.h b/Engine/Plugins/Experimental/Mutable/Source/MutableTools/Internal/MuT/NodeColourSampleImagePrivate.h
--- a/Engine/Plugins/Experimental/Mutable/Source/MutableTools/Internal/MuT/NodeColourSampleImagePrivate.h
+++ b/Engine/Plugins/Experimental/Mutable/Source/MutableTools/Internal/MuT/NodeColourSampleImagePrivate.h
@@ -6,7 +6,6 @@
 #include "MuT/NodeColourSampleImage.h"
 #include "MuT/NodeScalar.h"
 #include "MuT/NodeImage.h"
-#include "MuT/AST.h"
 
 
 namespace mu
@@ -22,29 +21,6 @@ namespace mu
 		NodeImagePtr m_pImage;
 		NodeScalarPtr m_pX;
 		NodeScalarPtr m_pY;
-
-		//!
-		void Serialise( OutputArchive& arch ) const
-		{
-            uint32_t ver = 0;
-			arch << ver;
-
-			arch << m_pImage;
-			arch << m_pX;
-			arch << m_pY;
-		}
-
-		//!
-		void Unserialise( InputArchive& arch )
-		{
-            uint32_t ver;
-			arch >> ver;
-			check(ver==0);
-
-			arch >> m_pImage;
-			arch >> m_pX;
-			arch >> m_pY;
-		}
 	};
 
 
diff --git a/Engine/Plugins/Experimental/Mutable/Source/MutableTools/Internal/MuT/NodeColourSwitchPrivate.h b/Engine/Plugins/Experimental/Mutable/Source/MutableTools/Internal/MuT/NodeColourSwitchPrivate.h
--- a/Engine/Plugins/Experimental/Mutable/Source/MutableTools/Internal/MuT/NodeColourSwitchPrivate.h
+++ b/Engine/Plugins/Experimental/Mutable/Source/MutableTools/Internal/MuT/NodeColourSwitchPrivate.h
@@ -5,7 +5,6 @@
 #include "MuT/NodeColourSwitch.h"
 #include "MuT/NodeColourPrivate.h"
 #include "MuT/NodeScalar.h"
-#include "MuT/AST.h"
 
 
 namespace mu
@@ -21,26 +20,6 @@ namespace mu
 		NodeScalarPtr m_pParameter;
         TArray<NodeColourPtr> m_options;
 
-		//!
-		void Serialise( OutputArchive& arch ) const
-		{
-            uint32_t ver = 0;
-			arch << ver;
-
-			arch << m_pParameter;
-			arch << m_options;
-		}
-
-		//!
-		void Unserialise( InputArchive& arch )
-		{
-            uint32_t ver;
-			arch >> ver;
-			check(ver==0);
-
-			arch >> m_pParameter;
-			arch >> m_options;
-		}
 	};
 
 
diff --git a/Engine/Plugins/Experimental/Mutable/Source/MutableTools/Internal/MuT/NodeComponentEditPrivate.h b/Engine/Plugins/Experimental/Mutable/Source/MutableTools/Internal/MuT/NodeComponentEditPrivate.h
--- a/Engine/Plugins/Experimental/Mutable/Source/MutableTools/Internal/MuT/NodeComponentEditPrivate.h
+++ b/Engine/Plugins/Experimental/Mutable/Source/MutableTools/Internal/MuT/NodeComponentEditPrivate.h
@@ -29,27 +29,6 @@ namespace mu
 
         TArray<NodeSurfacePtr> m_surfaces;
 
-		//!
-		void Serialise( OutputArchive& arch ) const
-		{
-            uint32_t ver = 2;
-			arch << ver;
-
-			arch << m_pParent;
-            arch << m_surfaces;
-		}
-
-		//!
-		void Unserialise( InputArchive& arch )
-		{
-            uint32_t ver;
-			arch >> ver;
-			check(ver == 2);
-
-            arch >> m_pParent;
-            arch >> m_surfaces;
-		}
-
 
 		// NodeComponent::Private interface
 		const NodeComponentNew::Private* GetParentComponentNew() const override;
diff --git a/Engine/Plugins/Experimental/Mutable/Source/MutableTools/Internal/MuT/NodeComponentNewPrivate.h b/Engine/Plugins/Experimental/Mutable/Source/MutableTools/Internal/MuT/NodeComponentNewPrivate.h
--- a/Engine/Plugins/Experimental/Mutable/Source/MutableTools/Internal/MuT/NodeComponentNewPrivate.h
+++ b/Engine/Plugins/Experimental/Mutable/Source/MutableTools/Internal/MuT/NodeComponentNewPrivate.h
@@ -32,44 +32,6 @@ namespace mu
 
         TArray<NodeSurfacePtr> m_surfaces;
 
-		//!
-		void Serialise( OutputArchive& arch ) const
-		{
-            uint32_t ver = 4;
-			
-			arch << ver;
-
-			arch << m_name;
-			arch << m_id;
-            arch << m_surfaces;
-        }
-
-		//!
-		void Unserialise( InputArchive& arch )
-		{
-            uint32_t ver;
-			arch >> ver;
-			check(ver>=2 && ver<=4);
-
-			if (ver <= 3)
-			{
-				std::string Temp;
-				arch>>Temp;
-				m_name = Temp.c_str();
-			}
-			else
-			{
-				arch >> m_name;
-			}
-
-			if (ver >= 3)
-			{
-				arch >> m_id;
-			}
-
-            arch >> m_surfaces;
-		}
-
 		// NodeComponent::Private interface
         const NodeComponentNew::Private* GetParentComponentNew() const override
 		{
diff --git a/Engine/Plugins/Experimental/Mutable/Source/MutableTools/Internal/MuT/NodeImageFormatPrivate.h b/Engine/Plugins/Experimental/Mutable/Source/MutableTools/Internal/MuT/NodeImageFormatPrivate.h
--- a/Engine/Plugins/Experimental/Mutable/Source/MutableTools/Internal/MuT/NodeImageFormatPrivate.h
+++ b/Engine/Plugins/Experimental/Mutable/Source/MutableTools/Internal/MuT/NodeImageFormatPrivate.h
@@ -4,7 +4,6 @@
 
 #include "MuT/NodeImageFormat.h"
 #include "MuT/NodeImagePrivate.h"
-#include "MuT/AST.h"
 #include "MuR/ImagePrivate.h"
 
 
@@ -22,33 +21,6 @@ namespace mu
 		EImageFormat m_formatIfAlpha = EImageFormat::IF_NONE;
         NodeImagePtr m_source;
 
-		//!
-		void Serialise( OutputArchive& arch ) const
-		{
-            uint32_t ver = 1;
-			arch << ver;
-
-            arch << (uint32)m_format;
-            arch << (uint32)m_formatIfAlpha;
-            arch << m_source;
-		}
-
-		//!
-		void Unserialise( InputArchive& arch )
-		{
-            uint32_t ver;
-			arch >> ver;
-            check(ver==1);
-
-			uint32 format;
-			arch >> format;
-			m_format = EImageFormat(format);
-
-			arch >> format;
-			m_formatIfAlpha = EImageFormat(format);
-
-			arch >> m_source;
-		}
 	};
 
 
diff --git a/Engine/Plugins/Experimental/Mutable/Source/MutableTools/Internal/MuT/NodeImageInterpolatePrivate.h b/Engine/Plugins/Experimental/Mutable/Source/MutableTools/Internal/MuT/NodeImageInterpolatePrivate.h
--- a/Engine/Plugins/Experimental/Mutable/Source/MutableTools/Internal/MuT/NodeImageInterpolatePrivate.h
+++ b/Engine/Plugins/Experimental/Mutable/Source/MutableTools/Internal/MuT/NodeImageInterpolatePrivate.h
@@ -18,27 +18,6 @@ namespace mu
 
 		NodeScalarPtr m_pFactor;
 		TArray<NodeImagePtr> m_targets;
-
-		//!
-		void Serialise( OutputArchive& arch ) const
-		{
-            uint32_t ver = 0;
-			arch << ver;
-
-			arch << m_pFactor;
-			arch << m_targets;
-		}
-
-		//!
-		void Unserialise( InputArchive& arch )
-		{
-            uint32_t ver;
-			arch >> ver;
-			check(ver==0);
-
-			arch >> m_pFactor;
-			arch >> m_targets;
-		}
 	};
 
 
diff --git a/Engine/Plugins/Experimental/Mutable/Source/MutableTools/Internal/MuT/NodeImageInvertPrivate.h b/Engine/Plugins/Experimental/Mutable/Source/MutableTools/Internal/MuT/NodeImageInvertPrivate.h
--- a/Engine/Plugins/Experimental/Mutable/Source/MutableTools/Internal/MuT/NodeImageInvertPrivate.h
+++ b/Engine/Plugins/Experimental/Mutable/Source/MutableTools/Internal/MuT/NodeImageInvertPrivate.h
@@ -4,7 +4,6 @@
 
 #include "MuT/NodeImagePrivate.h"
 #include "MuT/NodeImageInvert.h"
-#include "MuT/AST.h"
 
 
 namespace mu
@@ -18,24 +17,5 @@ namespace mu
 		static FNodeType s_type;
 
 		NodeImagePtr m_pBase;
-		
-		//!
-		void Serialise(OutputArchive& arch) const
-		{
-			uint32_t ver = 0;
-			arch << ver;
-
-			arch << m_pBase;
-		}
-
-		//!
-		void Unserialise(InputArchive& arch)
-		{
-			uint32_t ver;
-			arch >> ver;
-			check(ver == 0);
-
-			arch >> m_pBase;
-		}
 	};
 }
diff --git a/Engine/Plugins/Experimental/Mutable/Source/MutableTools/Internal/MuT/NodeImageLayerColourPrivate.h b/Engine/Plugins/Experimental/Mutable/Source/MutableTools/Internal/MuT/NodeImageLayerColourPrivate.h
--- a/Engine/Plugins/Experimental/Mutable/Source/MutableTools/Internal/MuT/NodeImageLayerColourPrivate.h
+++ b/Engine/Plugins/Experimental/Mutable/Source/MutableTools/Internal/MuT/NodeImageLayerColourPrivate.h
@@ -22,36 +22,6 @@ namespace mu
 		NodeImagePtr m_pMask;
 		NodeColourPtr m_pColour;
 		EBlendType m_type;
-
-		//!
-		void Serialise( OutputArchive& arch ) const
-		{
-            uint32_t ver = 1;
-			arch << ver;
-
-			arch << m_pBase;
-			arch << m_pMask;
-			arch << m_pColour;
-			arch << m_type;
-		}
-
-		//!
-		void Unserialise( InputArchive& arch )
-		{
-            uint32_t ver;
-			arch >> ver;
-			check(ver<=1);
-
-			arch >> m_pBase;
-			arch >> m_pMask;
-			arch >> m_pColour;
-			arch >> m_type;
-			if (ver < 1)
-			{
-				m_type = EBlendType( uint32(m_type)+1 );
-			}
-
-		}
 	};
 
 }
diff --git a/Engine/Plugins/Experimental/Mutable/Source/MutableTools/Internal/MuT/NodeImageLayerPrivate.h b/Engine/Plugins/Experimental/Mutable/Source/MutableTools/Internal/MuT/NodeImageLayerPrivate.h
--- a/Engine/Plugins/Experimental/Mutable/Source/MutableTools/Internal/MuT/NodeImageLayerPrivate.h
+++ b/Engine/Plugins/Experimental/Mutable/Source/MutableTools/Internal/MuT/NodeImageLayerPrivate.h
@@ -5,7 +5,6 @@
 #include "MuT/NodeImagePrivate.h"
 #include "MuT/NodeImageLayer.h"
 #include "MuT/NodeColour.h"
-#include "MuT/AST.h"
 
 
 namespace mu
@@ -21,38 +20,6 @@ namespace mu
 		NodeImagePtr m_pMask;
 		NodeImagePtr m_pBlended;
 		EBlendType m_type;
-
-		//!
-		void Serialise( OutputArchive& arch ) const
-		{
-            uint32_t ver = 1;
-			arch << ver;
-
-			arch << m_pBase;
-			arch << m_pMask;
-			arch << m_pBlended;
-            arch << (uint32_t)m_type;
-		}
-
-		//!
-		void Unserialise( InputArchive& arch )
-		{
-            uint32_t ver;
-			arch >> ver;
-			check(ver<=1);
-
-			arch >> m_pBase;
-			arch >> m_pMask;
-			arch >> m_pBlended;
-
-            uint32_t t;
-            arch >> t;
-			if (ver < 1)
-			{
-				++t;
-			}
-            m_type=(EBlendType)t;
-		}
 	};
 
 }
diff --git a/Engine/Plugins/Experimental/Mutable/Source/MutableTools/Internal/MuT/NodeImageMipmapPrivate.h b/Engine/Plugins/Experimental/Mutable/Source/MutableTools/Internal/MuT/NodeImageMipmapPrivate.h
--- a/Engine/Plugins/Experimental/Mutable/Source/MutableTools/Internal/MuT/NodeImageMipmapPrivate.h
+++ b/Engine/Plugins/Experimental/Mutable/Source/MutableTools/Internal/MuT/NodeImageMipmapPrivate.h
@@ -5,7 +5,6 @@
 #include "MuT/NodeImagePrivate.h"
 #include "MuT/NodeImageMipmap.h"
 #include "MuT/NodeScalar.h"
-#include "MuT/AST.h"
 
 #include "MuR/ImagePrivate.h"
 
@@ -23,30 +22,6 @@ namespace mu
 
 		FMipmapGenerationSettings m_settings;
 
-		//!
-		void Serialise( OutputArchive& arch ) const
-		{
-            uint32_t ver = 1;
-			arch << ver;
-
-			arch << m_pSource;
-			arch << m_pFactor;
-
-			arch << m_settings;
-		}
-
-		//!
-		void Unserialise( InputArchive& arch )
-		{
-            uint32_t ver;
-			arch >> ver;
-			check(ver==1);
-
-			arch >> m_pSource;
-			arch >> m_pFactor;
-			
-			arch >> m_settings;
-		}	
 	};
 
 
diff --git a/Engine/Plugins/Experimental/Mutable/Source/MutableTools/Internal/MuT/NodeImageMultiLayerPrivate.h b/Engine/Plugins/Experimental/Mutable/Source/MutableTools/Internal/MuT/NodeImageMultiLayerPrivate.h
--- a/Engine/Plugins/Experimental/Mutable/Source/MutableTools/Internal/MuT/NodeImageMultiLayerPrivate.h
+++ b/Engine/Plugins/Experimental/Mutable/Source/MutableTools/Internal/MuT/NodeImageMultiLayerPrivate.h
@@ -5,7 +5,6 @@
 #include "MuT/NodeImagePrivate.h"
 #include "MuT/NodeImageMultiLayer.h"
 #include "MuT/NodeColour.h"
-#include "MuT/AST.h"
 
 
 namespace mu
@@ -23,40 +22,6 @@ namespace mu
         NodeRangePtr m_pRange;
 		EBlendType m_type;
 
-		//!
-		void Serialise( OutputArchive& arch ) const
-		{
-            uint32_t ver = 1;
-			arch << ver;
-            arch << uint32_t(m_type);
-
-			arch << m_pBase;
-			arch << m_pMask;
-            arch << m_pBlended;
-            arch << m_pRange;
-		}
-
-		//!
-		void Unserialise( InputArchive& arch )
-		{
-            uint32_t ver;
-			arch >> ver;
-			check(ver<=1);
-
-            uint32_t t;
-            arch >> t;
-			if (ver < 1)
-			{
-				++t;
-			}
-            m_type=(EBlendType)t;
-
-            arch >> m_pBase;
-			arch >> m_pMask;
-            arch >> m_pBlended;
-            arch >> m_pRange;
-
-		}
 	};
 
 }
diff --git a/Engine/Plugins/Experimental/Mutable/Source/MutableTools/Internal/MuT/NodeImagePlainColourPrivate.h b/Engine/Plugins/Experimental/Mutable/Source/MutableTools/Internal/MuT/NodeImagePlainColourPrivate.h
--- a/Engine/Plugins/Experimental/Mutable/Source/MutableTools/Internal/MuT/NodeImagePlainColourPrivate.h
+++ b/Engine/Plugins/Experimental/Mutable/Source/MutableTools/Internal/MuT/NodeImagePlainColourPrivate.h
@@ -5,7 +5,6 @@
 #include "MuT/NodeColour.h"
 #include "MuT/NodeImagePrivate.h"
 #include "MuT/NodeImagePlainColour.h"
-#include "MuT/AST.h"
 
 
 namespace mu
@@ -28,41 +27,6 @@ namespace mu
 		int32 m_sizeX, m_sizeY;
 		EImageFormat Format = EImageFormat::IF_RGB_UBYTE;
 
-
-		//!
-		void Serialise( OutputArchive& arch ) const
-		{
-            uint32 ver = 1;
-			arch << ver;
-
-			arch << m_pColour;
-			arch << m_sizeX;
-			arch << m_sizeY;
-			arch << (uint32)Format;
-		}
-
-		//!
-		void Unserialise( InputArchive& arch )
-		{
-            uint32 ver;
-			arch >> ver;
-			check(ver<=1);
-
-			arch >> m_pColour;
-			arch >> m_sizeX;
-			arch >> m_sizeY;
-
-			if (ver >= 1)
-			{
-				uint32 Temp = 0;
-				arch >> Temp;
-				Format = EImageFormat(Temp);
-			}
-			else
-			{
-				Format = EImageFormat::IF_RGB_UBYTE;
-			}
-		}
 	};
 
 
diff --git a/Engine/Plugins/Experimental/Mutable/Source/MutableTools/Internal/MuT/NodeImageProjectPrivate.h b/Engine/Plugins/Experimental/Mutable/Source/MutableTools/Internal/MuT/NodeImageProjectPrivate.h
--- a/Engine/Plugins/Experimental/Mutable/Source/MutableTools/Internal/MuT/NodeImageProjectPrivate.h
+++ b/Engine/Plugins/Experimental/Mutable/Source/MutableTools/Internal/MuT/NodeImageProjectPrivate.h
@@ -33,60 +33,6 @@ namespace mu
 		bool bEnableTextureSeamCorrection = true;
 		ESamplingMethod SamplingMethod = ESamplingMethod::Point;
 		EMinFilterMethod MinFilterMethod = EMinFilterMethod::None;
-
-		//!
-		void Serialise( OutputArchive& arch ) const
-		{
-            uint32 ver = 5;
-			arch << ver;
-
-			arch << m_pProjector;
-			arch << m_pMesh;
-			arch << m_pAngleFadeStart;
-			arch << m_pAngleFadeEnd;
-			arch << m_pImage;
-            arch << m_pMask;
-            arch << m_layout;
-			arch << m_imageSize;
-			arch << bIsRGBFadingEnabled;
-			arch << bIsAlphaFadingEnabled;
-			arch << bEnableTextureSeamCorrection;
-			arch << SamplingMethod;
-			arch << MinFilterMethod;
-		}
-
-		//!
-		void Unserialise( InputArchive& arch )
-		{
-            uint32 ver;
-			arch >> ver;
-            check(ver>=2 && ver<=5);
-
-			arch >> m_pProjector;
-			arch >> m_pMesh;
-			arch >> m_pAngleFadeStart;
-			arch >> m_pAngleFadeEnd;
-			arch >> m_pImage;
-			arch >> m_pMask;
-            arch >> m_layout;
-			arch >> m_imageSize;
-			if (ver >= 3)
-			{
-				arch >> bIsRGBFadingEnabled;
-				arch >> bIsAlphaFadingEnabled;
-			}
-
-			if (ver >= 5)
-			{
-				arch >> bEnableTextureSeamCorrection;
-			}
-
-			if (ver >= 4)
-			{
-				arch >> SamplingMethod;
-				arch >> MinFilterMethod;
-			}
-		}
 	};
 
 
diff --git a/Engine/Plugins/Experimental/Mutable/Source/MutableTools/Internal/MuT/NodeImageResizePrivate.h b/Engine/Plugins/Experimental/Mutable/Source/MutableTools/Internal/MuT/NodeImageResizePrivate.h
--- a/Engine/Plugins/Experimental/Mutable/Source/MutableTools/Internal/MuT/NodeImageResizePrivate.h
+++ b/Engine/Plugins/Experimental/Mutable/Source/MutableTools/Internal/MuT/NodeImageResizePrivate.h
@@ -4,7 +4,6 @@
 
 #include "MuT/NodeImagePrivate.h"
 #include "MuT/NodeImageResize.h"
-#include "MuT/AST.h"
 
 
 namespace mu
@@ -21,30 +20,6 @@ namespace mu
 		bool m_relative = true;
 		float m_sizeX = 0.5f, m_sizeY = 0.5f;
 
-		//!
-		void Serialise( OutputArchive& arch ) const
-		{
-            uint32_t ver = 0;
-			arch << ver;
-
-			arch << m_pBase;
-			arch << m_relative;
-			arch << m_sizeX;
-			arch << m_sizeY;
-		}
-
-		//!
-		void Unserialise( InputArchive& arch )
-		{
-            uint32_t ver;
-			arch >> ver;
-			check(ver==0);
-
-			arch >> m_pBase;
-			arch >> m_relative;
-			arch >> m_sizeX;
-			arch >> m_sizeY;
-		}
 	};
 
 
diff --git a/Engine/Plugins/Experimental/Mutable/Source/MutableTools/Internal/MuT/NodeImageSwitchPrivate.h b/Engine/Plugins/Experimental/Mutable/Source/MutableTools/Internal/MuT/NodeImageSwitchPrivate.h
--- a/Engine/Plugins/Experimental/Mutable/Source/MutableTools/Internal/MuT/NodeImageSwitchPrivate.h
+++ b/Engine/Plugins/Experimental/Mutable/Source/MutableTools/Internal/MuT/NodeImageSwitchPrivate.h
@@ -19,27 +19,6 @@ namespace mu
 
 		NodeScalarPtr m_pParameter;
 		TArray<NodeImagePtr> m_options;
-
-		//!
-		void Serialise( OutputArchive& arch ) const
-		{
-            uint32_t ver = 0;
-			arch << ver;
-
-			arch << m_pParameter;
-			arch << m_options;
-		}
-
-		//!
-		void Unserialise( InputArchive& arch )
-		{
-            uint32_t ver;
-			arch >> ver;
-			check(ver==0);
-
-			arch >> m_pParameter;
-			arch >> m_options;
-		}
 	};
 
 
diff --git a/Engine/Plugins/Experimental/Mutable/Source/MutableTools/Internal/MuT/NodeImageSwizzlePrivate.h b/Engine/Plugins/Experimental/Mutable/Source/MutableTools/Internal/MuT/NodeImageSwizzlePrivate.h
--- a/Engine/Plugins/Experimental/Mutable/Source/MutableTools/Internal/MuT/NodeImageSwizzlePrivate.h
+++ b/Engine/Plugins/Experimental/Mutable/Source/MutableTools/Internal/MuT/NodeImageSwizzlePrivate.h
@@ -20,32 +20,6 @@ namespace mu
 		EImageFormat m_format;
 		TArray<NodeImagePtr> m_sources;
 		TArray<int> m_sourceChannels;
-
-		//!
-		void Serialise( OutputArchive& arch ) const
-		{
-            uint32_t ver = 0;
-			arch << ver;
-
-			arch << (uint32)m_format;
-			arch << m_sources;
-			arch << m_sourceChannels;
-		}
-
-		//!
-		void Unserialise( InputArchive& arch )
-		{
-            uint32_t ver;
-			arch >> ver;
-			check(ver==0);
-
-			uint32 format;
-			arch >> format;
-			m_format = (EImageFormat)format;
-
-			arch >> m_sources;
-			arch >> m_sourceChannels;
-		}
 	};
 
 
diff --git a/Engine/Plugins/Experimental/Mutable/Source/MutableTools/Internal/MuT/NodeLODPrivate.h b/Engine/Plugins/Experimental/Mutable/Source/MutableTools/Internal/MuT/NodeLODPrivate.h
--- a/Engine/Plugins/Experimental/Mutable/Source/MutableTools/Internal/MuT/NodeLODPrivate.h
+++ b/Engine/Plugins/Experimental/Mutable/Source/MutableTools/Internal/MuT/NodeLODPrivate.h
@@ -6,13 +6,10 @@
 #include "MuT/NodeLOD.h"
 #include "MuT/NodeComponentPrivate.h"
 #include "MuT/NodeModifierPrivate.h"
-#include "MuT/AST.h"
-
 
 namespace mu
 {
 
-
 	class NodeLOD::Private : public Node::Private
 	{
 	public:
@@ -21,27 +18,6 @@ namespace mu
 
 		TArray<NodeComponentPtr> m_components;
 		TArray<NodeModifierPtr> m_modifiers;
-
-		//!
-		void Serialise( OutputArchive& arch ) const
-		{
-            uint32_t ver = 1;
-			arch << ver;
-
-			arch << m_components;
-			arch << m_modifiers;
-		}
-
-		//!
-		void Unserialise( InputArchive& arch )
-		{
-            uint32_t ver;
-			arch >> ver;
-            check(ver == 1);
-
-			arch >> m_components;
-			arch >> m_modifiers;
-		}
 	};
 
 }
diff --git a/Engine/Plugins/Experimental/Mutable/Source/MutableTools/Internal/MuT/NodeMeshConstantPrivate.h b/Engine/Plugins/Experimental/Mutable/Source/MutableTools/Internal/MuT/NodeMeshConstantPrivate.h
--- a/Engine/Plugins/Experimental/Mutable/Source/MutableTools/Internal/MuT/NodeMeshConstantPrivate.h
+++ b/Engine/Plugins/Experimental/Mutable/Source/MutableTools/Internal/MuT/NodeMeshConstantPrivate.h
@@ -6,7 +6,6 @@
 
 #include "MuT/NodeMeshConstant.h"
 #include "MuT/NodeLayout.h"
-#include "MuT/AST.h"
 
 
 namespace mu
@@ -22,32 +21,6 @@ namespace mu
 
 		TArray<NodeLayoutPtr> m_layouts;
 
-		//!
-		void Serialise( OutputArchive& arch ) const
-		{
-            uint32_t ver = 0;
-			arch << ver;
-
-			arch << m_pValue;
-			arch << m_layouts;
-		}
-
-		//!
-		void Unserialise( InputArchive& arch )
-		{
-            uint32_t ver;
-			arch >> ver;
-			check(ver==0);
-
-			arch >> m_pValue;
-			arch >> m_layouts;
-
-			//if (m_pValue)
-			//{
-			//	m_pValue->CheckIntegrity();
-			//}
-        }
-
 		// NodeMesh::Private interface
         NodeLayoutPtr GetLayout( int index ) const override;
 
diff --git a/Engine/Plugins/Experimental/Mutable/Source/MutableTools/Internal/MuT/NodeMeshFormatPrivate.h b/Engine/Plugins/Experimental/Mutable/Source/MutableTools/Internal/MuT/NodeMeshFormatPrivate.h
--- a/Engine/Plugins/Experimental/Mutable/Source/MutableTools/Internal/MuT/NodeMeshFormatPrivate.h
+++ b/Engine/Plugins/Experimental/Mutable/Source/MutableTools/Internal/MuT/NodeMeshFormatPrivate.h
@@ -26,42 +26,6 @@ namespace mu
 		/** */
 		bool bOptimizeBuffers = false;
 
-		//!
-		void Serialise( OutputArchive& arch ) const
-		{
-            uint32 ver = 5;
-			arch << ver;
-
-			arch << Source;
-			arch << VertexBuffers;
-			arch << IndexBuffers;
-			arch << FaceBuffers;
-			arch << bOptimizeBuffers;
-		}
-
-		//!
-		void Unserialise( InputArchive& arch )
-		{
-            uint32 ver;
-			arch >> ver;
-            check(ver>=3 && ver<=5);
-
-			arch >> Source;
-			arch >> VertexBuffers;
-			arch >> IndexBuffers;
-			arch >> FaceBuffers;
-			if (ver == 3)
-			{
-				bool bDummy;
-				arch >> bDummy;
-			}
-
-			if (ver >= 5)
-			{
-				arch >> bOptimizeBuffers;
-			}
-		}
-
 		// NodeMesh::Private interface
         Ptr<NodeLayout> GetLayout( int32 index ) const override;
 	};
diff --git a/Engine/Plugins/Experimental/Mutable/Source/MutableTools/Internal/MuT/NodeMeshFragmentPrivate.h b/Engine/Plugins/Experimental/Mutable/Source/MutableTools/Internal/MuT/NodeMeshFragmentPrivate.h
--- a/Engine/Plugins/Experimental/Mutable/Source/MutableTools/Internal/MuT/NodeMeshFragmentPrivate.h
+++ b/Engine/Plugins/Experimental/Mutable/Source/MutableTools/Internal/MuT/NodeMeshFragmentPrivate.h
@@ -5,7 +5,6 @@
 #include "MuT/NodeMeshFragment.h"
 #include "MuT/NodeScalar.h"
 #include "MuT/NodeMeshPrivate.h"
-#include "MuT/AST.h"
 
 
 namespace mu
@@ -29,36 +28,6 @@ namespace mu
 
         FRAGMENT_TYPE m_fragmentType;
 
-		//!
-		void Serialise( OutputArchive& arch ) const
-		{
-            uint32_t ver = 1;
-			arch << ver;
-
-			arch << m_pMesh;
-            arch << (int32_t)m_fragmentType;
-            arch << m_layoutOrGroup;
-			arch << m_blocks;
-
-		}
-
-		//!
-		void Unserialise( InputArchive& arch )
-		{
-            uint32_t ver;
-			arch >> ver;
-            check(ver==1);
-
-			arch >> m_pMesh;
-
-            int32_t temp;
-            arch >> temp;
-            m_fragmentType = (FRAGMENT_TYPE)temp;
-
-            arch >> m_layoutOrGroup;
-			arch >> m_blocks;
-		}
-
 
 		// NodeMesh::Private interface
         NodeLayoutPtr GetLayout( int index ) const override;
diff --git a/Engine/Plugins/Experimental/Mutable/Source/MutableTools/Internal/MuT/NodeMeshMakeMorphPrivate.h b/Engine/Plugins/Experimental/Mutable/Source/MutableTools/Internal/MuT/NodeMeshMakeMorphPrivate.h
--- a/Engine/Plugins/Experimental/Mutable/Source/MutableTools/Internal/MuT/NodeMeshMakeMorphPrivate.h
+++ b/Engine/Plugins/Experimental/Mutable/Source/MutableTools/Internal/MuT/NodeMeshMakeMorphPrivate.h
@@ -4,13 +4,11 @@
 
 #include "MuT/NodeMeshPrivate.h"
 #include "MuT/NodeMeshMakeMorph.h"
-#include "MuT/AST.h"
 
 
 namespace mu
 {
 
-
     class NodeMeshMakeMorph::Private : public NodeMesh::Private
 	{
 	public:
@@ -22,32 +20,6 @@ namespace mu
 		
 		bool bOnlyPositionAndNormal = false;
 
-		//!
-		void Serialise( OutputArchive& arch ) const
-		{
-            uint32_t ver = 1;
-			arch << ver;
-
-			arch << m_pBase;
-			arch << m_pTarget;
-			arch << bOnlyPositionAndNormal;
-		}
-
-		//!
-		void Unserialise( InputArchive& arch )
-		{
-            uint32_t ver;
-			arch >> ver;
-			check(ver <= 1);
-
-			arch >> m_pBase;
-			arch >> m_pTarget;
-
-			if (ver == 1)
-			{
-				arch >> bOnlyPositionAndNormal;
-			}
-		}
 
 		// NodeMesh::Private interface
         NodeLayoutPtr GetLayout( int index ) const override;
diff --git a/Engine/Plugins/Experimental/Mutable/Source/MutableTools/Internal/MuT/NodeMeshMorphPrivate.h b/Engine/Plugins/Experimental/Mutable/Source/MutableTools/Internal/MuT/NodeMeshMorphPrivate.h
--- a/Engine/Plugins/Experimental/Mutable/Source/MutableTools/Internal/MuT/NodeMeshMorphPrivate.h
+++ b/Engine/Plugins/Experimental/Mutable/Source/MutableTools/Internal/MuT/NodeMeshMorphPrivate.h
@@ -5,7 +5,6 @@
 #include "MuT/NodeMeshMorph.h"
 #include "MuT/NodeScalar.h"
 #include "MuT/NodeMeshPrivate.h"
-#include "MuT/AST.h"
 #include "MuR/Skeleton.h"
 
 namespace mu
@@ -28,148 +27,6 @@ namespace mu
 		TArray<FBoneName> BonesToDeform;
 		TArray<FBoneName> PhysicsToDeform;
 
-        //!
-		void Serialise( OutputArchive& arch ) const
-		{
-            uint32_t ver = 8;
-			arch << ver;
-
-			arch << Factor;
-			arch << Base;
-			arch << Morph;
-
-			arch << bReshapeSkeleton;
-			arch << bReshapePhysicsVolumes;
-			arch << BonesToDeform;
-			arch << PhysicsToDeform;
-        }
-
-		//!
-		void Unserialise( InputArchive& arch )
-		{
-            uint32_t ver;
-			arch >> ver;
-			check(ver <= 7);
-
-			arch >> Factor;
-			arch >> Base;
-
-			if (ver < 6)
-			{
-				TArray<Ptr<NodeMesh>> Morphs;
-				arch >> Morphs;
-
-				// The result will not be the same.
-				if (!Morphs.IsEmpty())
-				{
-					Morph = Morphs.Last();
-				}
-			}
-			else
-			{
-				arch >> Morph;
-			}
-
-			if (ver < 6)
-			{
-				bool bDeprecated;
-				arch >> bDeprecated;
-			}
-	
-			if (ver >= 2)
-			{
-				arch >> bReshapeSkeleton;
-				arch >> bReshapePhysicsVolumes;
-	
-				// This repetition is needed, there was a bug where m_reshapePhysicsVolumes was serialized twice in ver 2.
-				if (ver == 2) 
-				{
-					arch >> bReshapePhysicsVolumes;
-				}
-
-				if (ver >= 8)
-				{
-					arch >> BonesToDeform;
-				}
-				else if (ver == 7)
-				{
-					TArray<uint16> BonesToDeform_DEPRECATED;
-					arch >> BonesToDeform_DEPRECATED;
-
-					const int32 NumBonesToDeform = BonesToDeform_DEPRECATED.Num();
-					BonesToDeform.SetNum(NumBonesToDeform);
-					for (int32 Index = 0; Index < NumBonesToDeform; ++Index)
-					{
-						BonesToDeform[Index].Id = BonesToDeform_DEPRECATED[Index];
-					}
-				}
-				else
-				{
-					TArray<string> BonesToDeform_DEPRECATED;
-					arch >> BonesToDeform_DEPRECATED;
-
-					const int32 NumBonesToDeform = BonesToDeform.Num();
-					BonesToDeform.SetNumUninitialized(NumBonesToDeform);
-					for (int32 Index = 0; Index < NumBonesToDeform; ++Index)
-					{
-						BonesToDeform[Index].Id = Index;
-					}
-				}
-			}
-			else
-			{
-				bReshapeSkeleton = false;
-				bReshapePhysicsVolumes = false;
-				BonesToDeform.Empty();
-			}
-
-			if (ver == 3)
-			{
-				bool bDeformAllBones_DEPRECATED;
-				arch >> bDeformAllBones_DEPRECATED;
-			}
-
-			if (ver >= 3 && ver < 5)
-			{
-				bool bDeformAllPhysics_DEPRECATED;
-				arch >> bDeformAllPhysics_DEPRECATED;
-			}
-
-			if (ver >= 8)
-			{
-				arch >> PhysicsToDeform;
-			}
-			else if (ver == 7)
-			{
-				TArray<uint16> PhysicsToDeform_DEPRECATED;
-				arch >> PhysicsToDeform_DEPRECATED;
-
-				const int32 NumPhysicsToDeform = PhysicsToDeform_DEPRECATED.Num();
-				PhysicsToDeform.SetNum(NumPhysicsToDeform);
-				for (int32 Index = 0; Index < NumPhysicsToDeform; ++Index)
-				{
-					PhysicsToDeform[Index].Id = PhysicsToDeform_DEPRECATED[Index];
-				}
-			}
-			else if (ver >= 3)
-			{
-				TArray<string> PhysicsToDeform_DEPRECATED;
-				arch >> PhysicsToDeform_DEPRECATED;
-
-				const int32 NumPhysicsToDeform = PhysicsToDeform.Num();
-				PhysicsToDeform.SetNumUninitialized(NumPhysicsToDeform);
-				for (int32 Index = 0; Index < NumPhysicsToDeform; ++Index)
-				{
-					PhysicsToDeform[Index].Id = Index;
-				}
-			}
-			else
-			{
-				PhysicsToDeform.Empty();
-			}
-
-		}
-
 		// NodeMesh::Private interface
         Ptr<NodeLayout> GetLayout( int index ) const override;
 
diff --git a/Engine/Plugins/Experimental/Mutable/Source/MutableTools/Internal/MuT/NodeMeshTablePrivate.h b/Engine/Plugins/Experimental/Mutable/Source/MutableTools/Internal/MuT/NodeMeshTablePrivate.h
--- a/Engine/Plugins/Experimental/Mutable/Source/MutableTools/Internal/MuT/NodeMeshTablePrivate.h
+++ b/Engine/Plugins/Experimental/Mutable/Source/MutableTools/Internal/MuT/NodeMeshTablePrivate.h
@@ -6,7 +6,6 @@
 #include "MuT/NodeMeshTable.h"
 #include "MuT/NodeLayout.h"
 #include "MuT/TablePrivate.h"
-#include "MuT/AST.h"
 
 
 namespace mu
@@ -27,68 +26,6 @@ namespace mu
 
 		TArray<NodeLayoutPtr> Layouts;
 
-		//!
-		void Serialise( OutputArchive& arch ) const
-		{
-            uint32 ver = 4;
-			arch << ver;
-
-			arch << ParameterName;
-			arch << Table;
-			arch << ColumnName;
-			arch << Layouts;
-			arch << bNoneOption;
-			arch << DefaultRowName;
-		}
-
-		//!
-		void Unserialise( InputArchive& arch )
-		{
-            uint32 ver;
-			arch >> ver;
-            check(ver>=1 && ver<=4);
-
-			if (ver == 1)
-			{
-				std::string Temp;
-				arch >> Temp;
-				ParameterName = Temp.c_str();
-			}
-			else
-			{
-				arch >> ParameterName;
-			}
-
-			arch >> Table;
-
-			if (ver == 1)
-			{
-				std::string Temp;
-				arch >> Temp;
-				ColumnName = Temp.c_str();
-			}
-			else
-			{
-				arch >> ColumnName;
-			}
-			
-			arch >> Layouts;
-
-			if(ver >= 3)
-			{
-				arch >> bNoneOption;
-			}
-			else
-			{
-				bNoneOption = Table->GetPrivate()->bNoneOption_DEPRECATED;
-			}
-
-			if (ver >= 4)
-			{
-				arch >> DefaultRowName;
-			}
-		}
-
 		// NodeMesh::Private interface
         NodeLayoutPtr GetLayout( int index ) const override;
 
diff --git a/Engine/Plugins/Experimental/Mutable/Source/MutableTools/Internal/MuT/NodeModifierMeshClipDeformPrivate.h b/Engine/Plugins/Experimental/Mutable/Source/MutableTools/Internal/MuT/NodeModifierMeshClipDeformPrivate.h
--- a/Engine/Plugins/Experimental/Mutable/Source/MutableTools/Internal/MuT/NodeModifierMeshClipDeformPrivate.h
+++ b/Engine/Plugins/Experimental/Mutable/Source/MutableTools/Internal/MuT/NodeModifierMeshClipDeformPrivate.h
@@ -5,7 +5,6 @@
 #include "MuT/NodeModifierPrivate.h"
 #include "MuT/NodeModifierMeshClipDeform.h"
 #include "MuT/NodeMesh.h"
-#include "MuT/AST.h"
 
 namespace mu
 {
@@ -25,38 +24,6 @@ namespace mu
 		NodeMeshPtr ClipMesh;
 		EShapeBindingMethod BindingMethod = EShapeBindingMethod::ClipDeformClosestProject;
 
-		//!
-		void Serialise( OutputArchive& arch ) const
-		{
-			NodeModifier::Private::Serialise(arch);
-
-            uint32_t ver = 1;
-			arch << ver;
-
-			arch << ClipMesh;
-			arch << BindingMethod;
-		}
-
-		//!
-		void Unserialise( InputArchive& arch )
-		{
-			NodeModifier::Private::Unserialise( arch );
-			
-            uint32_t ver;
-			arch >> ver;
-            check(ver<=1);
-
-			arch >> ClipMesh;
-			if (ver >= 1)
-			{
-				arch >> BindingMethod;
-			}
-			else
-			{
-				BindingMethod = EShapeBindingMethod::ClipDeformClosestProject;
-			}
-	
-		}
-    };
+	};
 
 }
diff --git a/Engine/Plugins/Experimental/Mutable/Source/MutableTools/Internal/MuT/NodeModifierMeshClipMorphPlanePrivate.h b/Engine/Plugins/Experimental/Mutable/Source/MutableTools/Internal/MuT/NodeModifierMeshClipMorphPlanePrivate.h
--- a/Engine/Plugins/Experimental/Mutable/Source/MutableTools/Internal/MuT/NodeModifierMeshClipMorphPlanePrivate.h
+++ b/Engine/Plugins/Experimental/Mutable/Source/MutableTools/Internal/MuT/NodeModifierMeshClipMorphPlanePrivate.h
@@ -4,7 +4,6 @@
 
 #include "MuT/NodeModifierPrivate.h"
 #include "MuT/NodeModifierMeshClipMorphPlane.h"
-#include "MuT/AST.h"
 
 #include "MuR/MutableMath.h"
 #include "MuR/Skeleton.h"
@@ -62,66 +61,6 @@ namespace mu
 		// means no limit.
 		float m_maxEffectRadius;
 
-        //!
-        void Serialise( OutputArchive& arch ) const
-        {
-            NodeModifier::Private::Serialise(arch);
-
-			uint32 ver = 4;
-            arch << ver;
-
-            arch << m_origin;
-            arch << m_normal;
-            arch << m_dist;
-            arch << m_factor;
-            arch << m_radius1;
-            arch << m_radius2;
-            arch << m_rotation;
-            arch << m_vertexSelectionType;
-            arch << m_selectionBoxOrigin;
-            arch << m_selectionBoxRadius;
-            arch << m_vertexSelectionBone;
-			arch << m_maxEffectRadius;
-        }
-
-        //!
-        void Unserialise( InputArchive& arch )
-        {
-            NodeModifier::Private::Unserialise( arch );
-
-            uint32 ver;
-            arch >> ver;
-            check(ver>=2);
-
-            arch >> m_origin;
-            arch >> m_normal;
-            arch >> m_dist;
-            arch >> m_factor;
-            arch >> m_radius1;
-            arch >> m_radius2;
-            arch >> m_rotation;
-            arch >> m_vertexSelectionType;
-            arch >> m_selectionBoxOrigin;
-            arch >> m_selectionBoxRadius;
-			if (ver >= 4)
-			{
-				arch >> m_vertexSelectionBone;
-			}
-			else if (ver == 3)
-			{
-				uint16 BoneId = 0;
-				arch >> BoneId;
-				m_vertexSelectionBone.Id = BoneId;
-			}
-			else
-			{
-				string OldVertexSelectionBone;
-				arch >> OldVertexSelectionBone;
-				m_vertexSelectionBone.Id = 0;
-			}
-			arch >> m_maxEffectRadius;
-        }
-
     };
 
 }
diff --git a/Engine/Plugins/Experimental/Mutable/Source/MutableTools/Internal/MuT/NodeModifierMeshClipWithMeshPrivate.h b/Engine/Plugins/Experimental/Mutable/Source/MutableTools/Internal/MuT/NodeModifierMeshClipWithMeshPrivate.h
--- a/Engine/Plugins/Experimental/Mutable/Source/MutableTools/Internal/MuT/NodeModifierMeshClipWithMeshPrivate.h
+++ b/Engine/Plugins/Experimental/Mutable/Source/MutableTools/Internal/MuT/NodeModifierMeshClipWithMeshPrivate.h
@@ -5,7 +5,6 @@
 #include "MuT/NodeModifierPrivate.h"
 #include "MuT/NodeModifierMeshClipWithMesh.h"
 #include "MuT/NodeMesh.h"
-#include "MuT/AST.h"
 
 #include "MuR/MutableMath.h"
 
@@ -26,29 +25,6 @@ namespace mu
 		//! 
 		Ptr<NodeMesh> ClipMesh;
 
-		//!
-		void Serialise( OutputArchive& arch ) const
-		{
-			NodeModifier::Private::Serialise(arch);
-
-            uint32_t ver = 0;
-			arch << ver;
-
-			arch << ClipMesh;
-		}
-
-		//!
-		void Unserialise( InputArchive& arch )
-		{
-			NodeModifier::Private::Unserialise( arch );
-			
-            uint32_t ver;
-			arch >> ver;
-            check(ver<=0);
-
-			arch >> ClipMesh;
-		}
-
 	};
 
 
diff --git a/Engine/Plugins/Experimental/Mutable/Source/MutableTools/Internal/MuT/NodeModifierMeshClipWithUVMaskPrivate.h b/Engine/Plugins/Experimental/Mutable/Source/MutableTools/Internal/MuT/NodeModifierMeshClipWithUVMaskPrivate.h
--- a/Engine/Plugins/Experimental/Mutable/Source/MutableTools/Internal/MuT/NodeModifierMeshClipWithUVMaskPrivate.h
+++ b/Engine/Plugins/Experimental/Mutable/Source/MutableTools/Internal/MuT/NodeModifierMeshClipWithUVMaskPrivate.h
@@ -5,7 +5,6 @@
 #include "MuT/NodeModifierPrivate.h"
 #include "MuT/NodeModifierMeshClipWithUVMask.h"
 #include "MuT/NodeImage.h"
-#include "MuT/AST.h"
 
 #include "MuR/MutableMath.h"
 
@@ -25,31 +24,6 @@ namespace mu
 		/** Layout index of the UVs to use inthe source mesh to ben clipped with the mask. */
 		uint8 LayoutIndex = 0;
 
-		//!
-		void Serialise( OutputArchive& arch ) const
-		{
-			NodeModifier::Private::Serialise(arch);
-
-            uint32 ver = 0;
-			arch << ver;
-
-			arch << ClipMask;
-			arch << LayoutIndex;
-		}
-
-		//!
-		void Unserialise( InputArchive& arch )
-		{
-			NodeModifier::Private::Unserialise( arch );
-			
-            uint32 ver;
-			arch >> ver;
-            check(ver<=0);
-
-			arch >> ClipMask;
-			arch >> LayoutIndex;
-		}
-
 	};
 
 
diff --git a/Engine/Plugins/Experimental/Mutable/Source/MutableTools/Internal/MuT/NodeModifierPrivate.h b/Engine/Plugins/Experimental/Mutable/Source/MutableTools/Internal/MuT/NodeModifierPrivate.h
--- a/Engine/Plugins/Experimental/Mutable/Source/MutableTools/Internal/MuT/NodeModifierPrivate.h
+++ b/Engine/Plugins/Experimental/Mutable/Source/MutableTools/Internal/MuT/NodeModifierPrivate.h
@@ -5,9 +5,6 @@
 #include "MuT/NodePrivate.h"
 #include "MuT/NodeModifier.h"
 
-#include <string>
-
-
 namespace mu
 {
 
@@ -26,47 +23,6 @@ namespace mu
 		// Wether the modifier has to be applied after the normal node operations or before
 		bool bApplyBeforeNormalOperations = true;
 
-		//!
-		void Serialise(OutputArchive& arch) const
-		{
-            uint32_t ver = 4;
-			arch << ver;
-
-			arch << RequiredTags;
-			arch << MultipleTagsPolicy;
-			arch << bApplyBeforeNormalOperations;
-		}
-
-		//!
-		void Unserialise(InputArchive& arch)
-		{
-            uint32_t ver;
-			arch >> ver;
-			check(ver>=2 && ver<=3);
-
-			if (ver <= 2)
-			{
-				TArray<std::string> Temp;
-				arch >> Temp;
-				RequiredTags.SetNum(Temp.Num());
-				for( int32 i=0; i<Temp.Num(); ++i)
-				{
-					RequiredTags[i] = Temp[i].c_str();
-				}
-			}
-			else
-			{
-				arch >> RequiredTags;
-			}
-
-			if (ver >= 4)
-			{
-				arch >> MultipleTagsPolicy;
-			}
-
-			arch >> bApplyBeforeNormalOperations;
-		}
-
 	};
 
 }
diff --git a/Engine/Plugins/Experimental/Mutable/Source/MutableTools/Internal/MuT/NodeObjectGroupPrivate.h b/Engine/Plugins/Experimental/Mutable/Source/MutableTools/Internal/MuT/NodeObjectGroupPrivate.h
--- a/Engine/Plugins/Experimental/Mutable/Source/MutableTools/Internal/MuT/NodeObjectGroupPrivate.h
+++ b/Engine/Plugins/Experimental/Mutable/Source/MutableTools/Internal/MuT/NodeObjectGroupPrivate.h
@@ -30,50 +30,6 @@ namespace mu
 		TArray<NodeObjectPtr> m_children;
 		int32 DefaultValue;
 
-		//!
-		void Serialise( OutputArchive& arch ) const
-		{
-            uint32_t ver = 3;
-			arch << ver;
-
-			arch << m_type;
-			arch << Name;
-			arch << Uid;
-			arch << m_children;
-			arch << DefaultValue;
-		}
-
-		//!
-		void Unserialise( InputArchive& arch )
-		{
-            uint32_t ver;
-			arch >> ver;
-            check(ver>=1 && ver<=3);
-
-			arch >> m_type;
-			if (ver <= 1)
-			{
-				std::string Temp;
-				arch >> Temp;
-				Name = Temp.c_str();
-				arch >> Temp;
-				Uid = Temp.c_str();
-			}
-			else
-			{
-				arch >> Name;
-				arch >> Uid;
-			}
-
-			arch >> m_children;
-
-			if (ver >= 3)
-			{
-				arch >> DefaultValue;
-			}
-		}
-
-
         // NodeObject::Private interface
         NodeLayoutPtr GetLayout( int lod, int component, int surface, int texture ) const override;
 
diff --git a/Engine/Plugins/Experimental/Mutable/Source/MutableTools/Internal/MuT/NodeObjectNewPrivate.h b/Engine/Plugins/Experimental/Mutable/Source/MutableTools/Internal/MuT/NodeObjectNewPrivate.h
--- a/Engine/Plugins/Experimental/Mutable/Source/MutableTools/Internal/MuT/NodeObjectNewPrivate.h
+++ b/Engine/Plugins/Experimental/Mutable/Source/MutableTools/Internal/MuT/NodeObjectNewPrivate.h
@@ -14,7 +14,6 @@
 namespace mu
 {
 
-
 	class NodeObjectNew::Private : public NodeObject::Private
 	{
 	public:
@@ -32,18 +31,6 @@ namespace mu
 		{
 			Ptr<NodeExtensionData> Node;
 			FString Name;
-
-			void Serialise(OutputArchive& arch) const
-			{
-				arch << Node;
-				arch << Name;
-			}
-
-			void Unserialise(InputArchive& arch)
-			{
-				arch >> Node;
-				arch >> Name;
-			}
 		};
 
 		TArray<NamedExtensionDataNode> m_extensionDataNodes;
@@ -53,65 +40,7 @@ namespace mu
         TArray<FObjectState> m_states;
 
 
-		//!
-		void Serialise( OutputArchive& arch ) const
-		{
-            uint32_t ver = 4;
-			arch << ver;
-
-			arch << m_name;
-			arch << m_uid;
-			arch << m_lods;
-			arch << m_children;
-			arch << m_states;
-			arch << m_extensionDataNodes;
-		}
-
-		//!
-		void Unserialise( InputArchive& arch )
-		{
-            uint32_t ver;
-			arch >> ver;
-            check(ver>=2 && ver<=4);
-
-			if (ver <= 3)
-			{
-				std::string Temp;
-				arch >> Temp;
-				m_name = Temp.c_str();
-				arch >> Temp;
-				m_uid = Temp.c_str();
-			}
-			else
-			{
-				arch >> m_name;
-				arch >> m_uid;
-			}
-
-			arch >> m_lods;
-			arch >> m_children;
-			arch >> m_states;
-
-			if (ver >= 4)
-			{
-				arch >> m_extensionDataNodes;
-			}
-			else if (ver >= 3)
-			{
-				int32 Num = 0;
-				arch >> Num;
-				m_extensionDataNodes.SetNum(Num);
-				for (int i=0;i<Num; ++i)
-				{
-					std::string Temp;
-					arch >> Temp;
-					m_extensionDataNodes[i].Name = Temp.c_str();
-					arch >> m_extensionDataNodes[i].Node;
-				}
-			}
-		}
-
-        //! Return true if the given component is set in any lod of this object.
+		//! Return true if the given component is set in any lod of this object.
         bool HasComponent( const NodeComponent* pComponent ) const;
 
         // NodeObject::Private interface
diff --git a/Engine/Plugins/Experimental/Mutable/Source/MutableTools/Internal/MuT/NodePatchImage.h b/Engine/Plugins/Experimental/Mutable/Source/MutableTools/Internal/MuT/NodePatchImage.h
--- a/Engine/Plugins/Experimental/Mutable/Source/MutableTools/Internal/MuT/NodePatchImage.h
+++ b/Engine/Plugins/Experimental/Mutable/Source/MutableTools/Internal/MuT/NodePatchImage.h
@@ -24,21 +24,12 @@ namespace mu
 
     //! Node that allows to modify an image from an object by blending other images on specific
     //! layout blocks.
-    //! \ingroup model
-    class MUTABLETOOLS_API NodePatchImage : public Node
+	class MUTABLETOOLS_API NodePatchImage : public Node
     {
     public:
 
-        //-----------------------------------------------------------------------------------------
-        // Life cycle
-        //-----------------------------------------------------------------------------------------
-
         NodePatchImage();
 
-		void SerialiseWrapper(OutputArchive& arch) const override;
-		static void Serialise( const NodePatchImage* pNode, OutputArchive& arch );
-        static NodePatchImagePtr StaticUnserialise( InputArchive& arch );
-
 
         //-----------------------------------------------------------------------------------------
         // Node Interface
diff --git a/Engine/Plugins/Experimental/Mutable/Source/MutableTools/Internal/MuT/NodePatchImagePrivate.h b/Engine/Plugins/Experimental/Mutable/Source/MutableTools/Internal/MuT/NodePatchImagePrivate.h
--- a/Engine/Plugins/Experimental/Mutable/Source/MutableTools/Internal/MuT/NodePatchImagePrivate.h
+++ b/Engine/Plugins/Experimental/Mutable/Source/MutableTools/Internal/MuT/NodePatchImagePrivate.h
@@ -26,41 +26,6 @@ namespace mu
 
         // Patch alpha channel as well?
         bool m_applyToAlpha = false;
-
-        //!
-        void Serialise( OutputArchive& arch ) const
-        {
-            uint32_t ver = 3;
-            arch << ver;
-
-            arch << m_pImage;
-            arch << m_pMask;
-            arch << m_blocks;
-            arch << (uint32_t)m_blendType;
-            arch << m_applyToAlpha;
-        }
-
-        //!
-        void Unserialise( InputArchive& arch )
-        {
-            uint32_t ver;
-            arch >> ver;
-            check(ver>=2 && ver<=3);
-
-            arch >> m_pImage;
-            arch >> m_pMask;
-            arch >> m_blocks;
-
-			uint32_t t;
-            arch >> t;
-			if (ver <= 2)
-			{
-				++t;
-			}
-            m_blendType=(EBlendType)t;
-
-            arch >> m_applyToAlpha;
-        }
     };
 
 }
diff --git a/Engine/Plugins/Experimental/Mutable/Source/MutableTools/Internal/MuT/NodePatchMeshPrivate.h b/Engine/Plugins/Experimental/Mutable/Source/MutableTools/Internal/MuT/NodePatchMeshPrivate.h
--- a/Engine/Plugins/Experimental/Mutable/Source/MutableTools/Internal/MuT/NodePatchMeshPrivate.h
+++ b/Engine/Plugins/Experimental/Mutable/Source/MutableTools/Internal/MuT/NodePatchMeshPrivate.h
@@ -4,7 +4,6 @@
 
 #include "MuT/NodeMeshPrivate.h"
 #include "MuT/NodePatchMesh.h"
-#include "MuT/AST.h"
 
 
 namespace mu
@@ -20,27 +19,6 @@ namespace mu
 		NodeMeshPtr m_pRemove;
 		NodeMeshPtr m_pAdd;
 
-		//!
-		void Serialise( OutputArchive& arch ) const
-		{
-            uint32_t ver = 1;
-			arch << ver;
-
-			arch << m_pRemove;
-			arch << m_pAdd;
-		}
-
-		//!
-		void Unserialise( InputArchive& arch )
-		{
-            uint32_t ver;
-			arch >> ver;
-            check( ver==1 );
-
-			arch >> m_pRemove;
-			arch >> m_pAdd;
-		}
-
 	};
 
 
diff --git a/Engine/Plugins/Experimental/Mutable/Source/MutableTools/Internal/MuT/NodePrivate.h b/Engine/Plugins/Experimental/Mutable/Source/MutableTools/Internal/MuT/NodePrivate.h
--- a/Engine/Plugins/Experimental/Mutable/Source/MutableTools/Internal/MuT/NodePrivate.h
+++ b/Engine/Plugins/Experimental/Mutable/Source/MutableTools/Internal/MuT/NodePrivate.h
@@ -41,23 +41,6 @@ namespace mu
         m_pD = nullptr;										\
     }														\
                                                             \
-    void N::Serialise( const N* p, OutputArchive& arch )	\
-    {														\
-        arch << *p->m_pD;									\
-    }														\
-                                                            \
-	void N::SerialiseWrapper( OutputArchive& arch ) const	\
-	{														\
-		N::Serialise(this, arch);							\
-	}														\
-                                                            \
-    Ptr<N> N::StaticUnserialise( InputArchive& arch )		\
-    {														\
-        Ptr<N> pResult = new N();							\
-        arch >> *pResult->m_pD;								\
-        return pResult;										\
-    }														\
-                                                            \
     N::Private* N::GetPrivate() const						\
     {														\
         return m_pD;										\
diff --git a/Engine/Plugins/Experimental/Mutable/Source/MutableTools/Internal/MuT/NodeProjectorPrivate.h b/Engine/Plugins/Experimental/Mutable/Source/MutableTools/Internal/MuT/NodeProjectorPrivate.h
--- a/Engine/Plugins/Experimental/Mutable/Source/MutableTools/Internal/MuT/NodeProjectorPrivate.h
+++ b/Engine/Plugins/Experimental/Mutable/Source/MutableTools/Internal/MuT/NodeProjectorPrivate.h
@@ -31,36 +31,6 @@ namespace mu
 		FVector3f m_up;
 		FVector3f m_scale;
         float m_projectionAngle = 0.0f;
-
-		//!
-		void Serialise( OutputArchive& arch ) const
-		{
-            uint32 ver = 2;
-			arch << ver;
-
-            arch << m_type;
-            arch << m_position;
-            arch << m_direction;
-			arch << m_up;
-			arch << m_scale;
-            arch << m_projectionAngle;
-		}
-
-		//!
-		void Unserialise( InputArchive& arch )
-		{
-            uint32 ver;
-			arch >> ver;
-            check(ver==2);
-
-            arch >> m_type;
-			arch >> m_position;
-			arch >> m_direction;
-			arch >> m_up;
-            arch >> m_scale;
-            arch >> m_projectionAngle;
-        }
-
 	};
 
 
@@ -74,51 +44,6 @@ namespace mu
 		FString m_uid;
 
 		TArray<Ptr<NodeRange>> m_ranges;
-
-		//!
-		void Serialise( OutputArchive& arch ) const
-		{
-			NodeProjectorConstant::Private::Serialise( arch );
-
-            uint32 ver = 3;
-			arch << ver;
-
-			arch << m_name;
-            arch << m_uid;
-            arch << m_ranges;
-        }
-
-		//!
-		void Unserialise( InputArchive& arch )
-		{
-			NodeProjectorConstant::Private::Unserialise( arch );
-
-            uint32 ver;
-			arch >> ver;
-            check(ver<=3);
-
-			if (ver <= 2)
-			{
-				std::string Temp;
-				arch >> Temp;
-				m_name = Temp.c_str();
-				if (ver >= 1)
-				{
-					arch >> Temp;
-					m_uid = Temp.c_str();
-				}
-			}
-			else
-			{
-				arch >> m_name;
-				arch >> m_uid;
-			}
-
-            if (ver>=2)
-            {
-                arch >> m_ranges;
-            }
-        }
 	};
 
 
diff --git a/Engine/Plugins/Experimental/Mutable/Source/MutableTools/Internal/MuT/NodeScalarCurvePrivate.h b/Engine/Plugins/Experimental/Mutable/Source/MutableTools/Internal/MuT/NodeScalarCurvePrivate.h
--- a/Engine/Plugins/Experimental/Mutable/Source/MutableTools/Internal/MuT/NodeScalarCurvePrivate.h
+++ b/Engine/Plugins/Experimental/Mutable/Source/MutableTools/Internal/MuT/NodeScalarCurvePrivate.h
@@ -5,10 +5,6 @@
 #include "MuT/NodePrivate.h"
 #include "MuT/NodeScalarCurve.h"
 #include "MuT/NodeScalar.h"
-#include "MuT/AST.h"
-
-#include "MuR/SerialisationPrivate.h"
-#include "MuR/ParametersPrivate.h"
 
 
 namespace mu
@@ -24,27 +20,6 @@ namespace mu
 
 		NodeScalarPtr m_input_scalar;
 		Curve m_curve;
-
-		//!
-		void Serialise(OutputArchive& arch) const
-		{
-            uint32_t ver = 1;
-			arch << ver;
-
-			arch << m_input_scalar;
-			arch << m_curve;
-		}
-
-		//!
-		void Unserialise(InputArchive& arch)
-		{
-            uint32_t ver;
-			arch >> ver;
-            check(ver == 1);
-
-			arch >> m_input_scalar;
-			arch >> m_curve;
-		}
 	};
 
 
diff --git a/Engine/Plugins/Experimental/Mutable/Source/MutableTools/Internal/MuT/NodeScalarEnumParameterPrivate.h b/Engine/Plugins/Experimental/Mutable/Source/MutableTools/Internal/MuT/NodeScalarEnumParameterPrivate.h
--- a/Engine/Plugins/Experimental/Mutable/Source/MutableTools/Internal/MuT/NodeScalarEnumParameterPrivate.h
+++ b/Engine/Plugins/Experimental/Mutable/Source/MutableTools/Internal/MuT/NodeScalarEnumParameterPrivate.h
@@ -23,89 +23,15 @@ namespace mu
 		FString m_uid;
 		int32 m_defaultValue = 0;
 
-		struct OPTION
+		struct FOption
 		{
 			FString name;
 			float value;
-
-			//!
-			void Serialise( OutputArchive& arch ) const
-			{
-                uint32 ver = 1;
-				arch << ver;
-
-				arch << name;
-				arch << value;
-			}
-
-			//!
-			void Unserialise( InputArchive& arch )
-			{
-                uint32 ver;
-				arch >> ver;
-                check(ver<=1);
-
-				if (ver == 0)
-				{
-					std::string Temp;
-					arch >> Temp;
-					name = Temp.c_str();
-				}
-				else
-				{
-					arch >> name;
-				}
-				arch >> value;
-			}
 		};
 
-		TArray<OPTION> m_options;
+		TArray<FOption> m_options;
 
 		TArray<Ptr<NodeRange>> m_ranges;
-
-		//!
-		void Serialise( OutputArchive& arch ) const
-		{
-            uint32_t ver = 4;
-			arch << ver;
-
-			arch << m_name;
-			arch << m_uid;
-			arch << m_defaultValue;
-            arch << m_options;
-            arch << m_ranges;
-        }
-
-		//!
-		void Unserialise( InputArchive& arch )
-		{
-            uint32_t ver;
-			arch >> ver;
-            check(ver>=2 && ver<=4);
-
-			if (ver <= 3)
-			{
-				std::string Temp;
-				arch >> Temp;
-				m_name = Temp.c_str();
-				arch >> Temp;
-				m_uid = Temp.c_str();
-			}
-			else
-			{
-				arch >> m_name;
-				arch >> m_uid;
-			}
-
-			arch >> m_defaultValue;
-			if (ver <= 2)
-			{
-				int32 Dummy = 0;
-				arch >> Dummy;
-			}
-			arch >> m_options;
-            arch >> m_ranges;
-        }
 	};
 
 }
diff --git a/Engine/Plugins/Experimental/Mutable/Source/MutableTools/Internal/MuT/NodeScalarParameterPrivate.h b/Engine/Plugins/Experimental/Mutable/Source/MutableTools/Internal/MuT/NodeScalarParameterPrivate.h
--- a/Engine/Plugins/Experimental/Mutable/Source/MutableTools/Internal/MuT/NodeScalarParameterPrivate.h
+++ b/Engine/Plugins/Experimental/Mutable/Source/MutableTools/Internal/MuT/NodeScalarParameterPrivate.h
@@ -25,55 +25,6 @@ namespace mu
 		FString m_uid;
 
 		TArray<Ptr<NodeRange>> m_ranges;
-
-		//!
-		void Serialise( OutputArchive& arch ) const
-		{
-            uint32_t ver = 7;
-			arch << ver;
-
-			arch << m_defaultValue;
-			arch << m_name;
-			arch << m_uid;
-            arch << m_ranges;
-        }
-
-		//!
-		void Unserialise( InputArchive& arch )
-		{
-            uint32_t ver;
-			arch >> ver;
-            check(ver>=4 && ver<=7);
-
-			arch >> m_defaultValue;
-			if (ver <= 6)
-			{
-				std::string Temp;
-				arch >> Temp;
-				m_name = Temp.c_str();
-				arch >> Temp;
-				m_uid = Temp.c_str();
-			}
-			else
-			{
-				arch >> m_name;
-				arch >> m_uid;
-			}
-
-			if (ver < 5)
-			{
-				TArray<Ptr<NodeImage>> TempAdditionalImages;
-				arch >> TempAdditionalImages;
-			}
-
-			if (ver <= 5)
-			{
-				int32 Dummy=0;
-				arch >> Dummy;
-			}
-
-            arch >> m_ranges;
-        }
 	};
 
 }
diff --git a/Engine/Plugins/Experimental/Mutable/Source/MutableTools/Internal/MuT/NodeScalarSwitchPrivate.h b/Engine/Plugins/Experimental/Mutable/Source/MutableTools/Internal/MuT/NodeScalarSwitchPrivate.h
--- a/Engine/Plugins/Experimental/Mutable/Source/MutableTools/Internal/MuT/NodeScalarSwitchPrivate.h
+++ b/Engine/Plugins/Experimental/Mutable/Source/MutableTools/Internal/MuT/NodeScalarSwitchPrivate.h
@@ -5,8 +5,6 @@
 #include "MuT/NodeScalarSwitch.h"
 #include "MuT/NodeScalarPrivate.h"
 #include "MuT/NodeScalar.h"
-#include "MuT/AST.h"
-
 
 namespace mu
 {
@@ -20,27 +18,6 @@ namespace mu
 
 		NodeScalarPtr m_pParameter;
         TArray<NodeScalarPtr> m_options;
-
-		//!
-		void Serialise( OutputArchive& arch ) const
-		{
-            uint32_t ver = 0;
-			arch << ver;
-
-			arch << m_pParameter;
-			arch << m_options;
-		}
-
-		//!
-		void Unserialise( InputArchive& arch )
-		{
-            uint32_t ver;
-			arch >> ver;
-			check(ver==0);
-
-			arch >> m_pParameter;
-			arch >> m_options;
-		}
 	};
 
 
diff --git a/Engine/Plugins/Experimental/Mutable/Source/MutableTools/Internal/MuT/NodeSurfaceEditPrivate.h b/Engine/Plugins/Experimental/Mutable/Source/MutableTools/Internal/MuT/NodeSurfaceEditPrivate.h
--- a/Engine/Plugins/Experimental/Mutable/Source/MutableTools/Internal/MuT/NodeSurfaceEditPrivate.h
+++ b/Engine/Plugins/Experimental/Mutable/Source/MutableTools/Internal/MuT/NodeSurfaceEditPrivate.h
@@ -36,19 +36,6 @@ namespace mu
 		{
 			NodeImagePtr m_pExtend;
 			NodePatchImagePtr m_pPatch;
-
-			//!
-			void Serialise( OutputArchive& arch ) const
-			{
-				arch << m_pExtend;
-				arch << m_pPatch;
-			}
-
-			void Unserialise( InputArchive& arch )
-			{
-				arch >> m_pExtend;
-				arch >> m_pPatch;
-			}
 		};
 
 		TArray<FTexture> m_textures;
@@ -58,50 +45,6 @@ namespace mu
 
 		//! Factor of the morph
 		NodeScalarPtr m_pFactor;
-
-		//!
-		void Serialise( OutputArchive& arch ) const
-		{
-            uint32_t ver = 5;
-			arch << ver;
-
-			arch << m_pParent;
-			arch << m_pMesh;
-            arch << m_pMorph;
-            arch << m_morphIsTarget;
-            arch << m_textures;
-            arch << m_tags;
-			arch << m_pFactor;
-        }
-
-		//!
-		void Unserialise( InputArchive& arch )
-		{
-            uint32_t ver;
-			arch >> ver;
-            check( ver>=4 && ver<=5 );
-
-			arch >> m_pParent;
-			arch >> m_pMesh;
-			arch >> m_pMorph;
-            arch >> m_morphIsTarget;
-			arch >> m_textures;
-			if (ver <= 4)
-			{
-				TArray <std::string> Temp;
-				arch >> Temp;
-				m_tags.SetNum(Temp.Num());
-				for (int32 i=0; i<Temp.Num(); ++i)
-				{
-					m_tags[i] = Temp[i].c_str();
-				}
-			}
-			else
-			{
-				arch >> m_tags;
-			}
-			arch >> m_pFactor;
-        }
 	};
 
 }
diff --git a/Engine/Plugins/Experimental/Mutable/Source/MutableTools/Internal/MuT/NodeSurfaceNewPrivate.h b/Engine/Plugins/Experimental/Mutable/Source/MutableTools/Internal/MuT/NodeSurfaceNewPrivate.h
--- a/Engine/Plugins/Experimental/Mutable/Source/MutableTools/Internal/MuT/NodeSurfaceNewPrivate.h
+++ b/Engine/Plugins/Experimental/Mutable/Source/MutableTools/Internal/MuT/NodeSurfaceNewPrivate.h
@@ -29,28 +29,15 @@ namespace mu
         uint32 ExternalId =0;
         int32 SharedSurfaceId =INDEX_NONE;
 
-		struct MESH
+		struct FMeshData
 		{
 			FString m_name;
 			NodeMeshPtr m_pMesh;
-
-			//!
-			void Serialise( OutputArchive& arch ) const
-			{
-				arch << m_name;
-				arch << m_pMesh;
-			}
-
-			void Unserialise( InputArchive& arch )
-			{
-				arch >> m_name;
-				arch >> m_pMesh;
-			}
 		};
 
-		TArray<MESH> m_meshes;
+		TArray<FMeshData> m_meshes;
 
-		struct IMAGE
+		struct FImageData
 		{
 			FString m_name;
 			FString m_materialName;
@@ -59,227 +46,51 @@ namespace mu
 
 			// It could be negative, to indicate no layout.
             int8 m_layoutIndex = 0;
-
-            //!
-			void Serialise(OutputArchive& arch) const
-			{
-				arch << m_name;
-				arch << m_materialName;
-				arch << m_materialParameterName;
-				arch << m_pImage;
-                arch << m_layoutIndex;
-            }
-
-			void Unserialise(InputArchive& arch)
-			{
-				arch >> m_name;
-				arch >> m_materialName;
-				arch >> m_materialParameterName;
-				arch >> m_pImage;
-                arch >> m_layoutIndex;
-            }
         };
 
-		TArray<IMAGE> m_images;
+		TArray<FImageData> m_images;
 
-		struct VECTOR
+		struct FVectorData
 		{
 			FString m_name;
 			NodeColourPtr m_pVector;
-
-			//!
-			void Serialise(OutputArchive& arch) const
-			{
-				arch << m_name;
-				arch << m_pVector;
-			}
-
-			void Unserialise(InputArchive& arch)
-			{
-				arch >> m_name;
-				arch >> m_pVector;
-			}
 		};
 
-		TArray<VECTOR> m_vectors;
+		TArray<FVectorData> m_vectors;
 
-        struct SCALAR
+        struct FScalar
         {
 			FString m_name;
             NodeScalarPtr m_pScalar;
-
-            //!
-            void Serialise(OutputArchive& arch) const
-            {
-                arch << m_name;
-                arch << m_pScalar;
-            }
-
-            void Unserialise(InputArchive& arch)
-            {
-                arch >> m_name;
-                arch >> m_pScalar;
-            }
         };
 
-        TArray<SCALAR> m_scalars;
+        TArray<FScalar> m_scalars;
 
-        struct STRING
+        struct FStringData
         {
 			FString m_name;
             NodeStringPtr m_pString;
-
-            //!
-            void Serialise( OutputArchive& arch ) const
-            {
-                arch << m_name;
-                arch << m_pString;
-            }
-
-            void Unserialise( InputArchive& arch )
-            {
-                arch >> m_name;
-                arch >> m_pString;
-            }
         };
 
-		TArray<STRING> m_strings;
+		TArray<FStringData> m_strings;
 
         //! Tags in this surface
 		TArray<FString> m_tags;
 
 		//! Find an image node index by name or return -1
-		int FindImage( const FString& strName ) const;
+		int32 FindImage( const FString& strName ) const;
 
 		//! Find a mesh node index by name or return -1
-		int FindMesh(const FString& strName) const;
+		int32 FindMesh(const FString& strName) const;
 
         //! Find a vector node index by name or return -1
-        int FindVector(const FString& strName) const;
+        int32 FindVector(const FString& strName) const;
 
         //! Find a scalar node index by name or return -1
-        int FindScalar( const FString& strName ) const;
+        int32 FindScalar( const FString& strName ) const;
 
         //! Find a string node index by name or return -1
-        int FindString( const FString& strName ) const;
-
-        //!
-		void Serialise( OutputArchive& arch ) const
-		{
-            uint32 ver = 9;
-			arch << ver;
-
-			arch << m_name;
-            arch << ExternalId;
-			arch << SharedSurfaceId;
-
-			arch << m_meshes;
-			arch << m_images;
-			arch << m_tags;
-            arch << m_vectors;
-            arch << m_scalars;
-            arch << m_strings;
-        }
-
-        //!
-		void Unserialise( InputArchive& arch )
-		{
-            uint32 ver;
-			arch >> ver;
-            check(ver<=9);
-
-			std::string Temp;
-			if (ver <= 8)
-			{
-				arch >> Temp;
-				m_name = Temp.c_str();
-			}
-			else
-			{
-				arch >> m_name;
-			}
-            arch >> ExternalId;
-
-			if (ver >= 8)
-			{
-				arch >> SharedSurfaceId;
-			}
-
-			if (ver <= 8)
-			{
-				int32 Num = 0;
-
-				arch >> Num;
-				m_meshes.SetNum(Num);
-				for (int32 i=0; i<Num; ++i)
-				{
-					arch >> Temp;
-					m_meshes[i].m_name = Temp.c_str();
-					arch >> m_meshes[i].m_pMesh;
-				}
-
-				arch >> Num;
-				m_images.SetNum(Num);
-				for (int32 i = 0; i < Num; ++i)
-				{
-					arch >> Temp;
-					m_images[i].m_name = Temp.c_str();
-
-					arch >> Temp;
-					m_images[i].m_materialName = Temp.c_str();
-
-					arch >> Temp;
-					m_images[i].m_materialParameterName = Temp.c_str();
-
-					arch >> m_images[i].m_pImage;
-					arch >> m_images[i].m_layoutIndex;
-				}
-
-				arch >> Num;
-				m_tags.SetNum(Num);
-				for (int32 i = 0; i < Num; ++i)
-				{
-					arch >> Temp;
-					m_tags[i] = Temp.c_str();
-				}
-
-				arch >> Num;
-				m_vectors.SetNum(Num);
-				for (int32 i = 0; i < Num; ++i)
-				{
-					arch >> Temp;
-					m_vectors[i].m_name = Temp.c_str();
-					arch >> m_vectors[i].m_pVector;
-				}
-
-				arch >> Num;
-				m_scalars.SetNum(Num);
-				for (int32 i = 0; i < Num; ++i)
-				{
-					arch >> Temp;
-					m_scalars[i].m_name = Temp.c_str();
-					arch >> m_scalars[i].m_pScalar;
-				}
-
-				arch >> Num;
-				m_strings.SetNum(Num);
-				for (int32 i = 0; i < Num; ++i)
-				{
-					arch >> Temp;
-					m_strings[i].m_name = Temp.c_str();
-					arch >> m_strings[i].m_pString;
-				}
-			}
-			else
-			{
-				arch >> m_meshes;
-				arch >> m_images;
-				arch >> m_tags;
-				arch >> m_vectors;
-				arch >> m_scalars;
-				arch >> m_strings;
-			}
-        }
+        int32 FindString( const FString& strName ) const;
     };
 
 
diff --git a/Engine/Plugins/Experimental/Mutable/Source/MutableTools/Internal/MuT/NodeSurfaceSwitchPrivate.h b/Engine/Plugins/Experimental/Mutable/Source/MutableTools/Internal/MuT/NodeSurfaceSwitchPrivate.h
--- a/Engine/Plugins/Experimental/Mutable/Source/MutableTools/Internal/MuT/NodeSurfaceSwitchPrivate.h
+++ b/Engine/Plugins/Experimental/Mutable/Source/MutableTools/Internal/MuT/NodeSurfaceSwitchPrivate.h
@@ -11,7 +11,6 @@
 namespace mu
 {
 
-
 	class NodeSurfaceSwitch::Private : public NodeSurface::Private
 	{
 	public:
@@ -20,28 +19,6 @@ namespace mu
 
 		Ptr<NodeScalar> Parameter;
 		TArray<Ptr<NodeSurface>> Options;
-
-		//!
-		void Serialise( OutputArchive& arch ) const
-		{
-            uint32_t ver = 0;
-			arch << ver;
-
-			arch << Parameter;
-			arch << Options;
-		}
-
-		//!
-		void Unserialise( InputArchive& arch )
-		{
-            uint32_t ver;
-			arch >> ver;
-			check(ver==0);
-
-			arch >> Parameter;
-			arch >> Options;
-		}
-
 	};
 
 
diff --git a/Engine/Plugins/Experimental/Mutable/Source/MutableTools/Internal/MuT/NodeSurfaceVariationPrivate.h b/Engine/Plugins/Experimental/Mutable/Source/MutableTools/Internal/MuT/NodeSurfaceVariationPrivate.h
--- a/Engine/Plugins/Experimental/Mutable/Source/MutableTools/Internal/MuT/NodeSurfaceVariationPrivate.h
+++ b/Engine/Plugins/Experimental/Mutable/Source/MutableTools/Internal/MuT/NodeSurfaceVariationPrivate.h
@@ -28,70 +28,11 @@ namespace mu
 			TArray<NodeSurfacePtr> m_surfaces;
 			TArray<NodeModifierPtr> m_modifiers;
             FString m_tag;
-
-			//!
-			void Serialise( OutputArchive& arch ) const
-			{
-                uint32 ver = 2;
-				arch << ver;
-
-				arch << m_tag;
-                arch << m_surfaces;
-                arch << m_modifiers;
-            }
-
-			void Unserialise( InputArchive& arch )
-			{
-                uint32 ver = 0;
-				arch >> ver;
-                check(ver>=1&&ver<=2);
-
-				if (ver <= 1)
-				{
-					std::string Temp;
-					arch >> Temp;
-					m_tag = Temp.c_str();
-				}
-				else
-				{
-					arch >> m_tag;
-				}
-				arch >> m_surfaces;
-                arch >> m_modifiers;
-			}
 		};
 
         NodeSurfaceVariation::VariationType m_type = NodeSurfaceVariation::VariationType::Tag;
 
 		TArray<FVariation> m_variations;
-
-		//!
-		void Serialise( OutputArchive& arch ) const
-		{
-            uint32 ver = 3;
-			arch << ver;
-
-            arch << uint32_t(m_type);
-            arch << m_defaultSurfaces;
-            arch << m_defaultModifiers;
-            arch << m_variations;
-		}
-
-		//!
-		void Unserialise( InputArchive& arch )
-		{
-            uint32 ver;
-			arch >> ver;
-            check( ver==3 );
-
-            uint32 temp;
-            arch >> temp;
-            m_type = NodeSurfaceVariation::VariationType(temp);
-
-            arch >> m_defaultSurfaces;
-            arch >> m_defaultModifiers;
-            arch >> m_variations;
-		}
 	};
 
 }
diff --git a/Engine/Plugins/Experimental/Mutable/Source/MutableTools/Internal/MuT/Platform.h b/Engine/Plugins/Experimental/Mutable/Source/MutableTools/Internal/MuT/Platform.h
--- a/Engine/Plugins/Experimental/Mutable/Source/MutableTools/Internal/MuT/Platform.h
+++ b/Engine/Plugins/Experimental/Mutable/Source/MutableTools/Internal/MuT/Platform.h
@@ -5,7 +5,6 @@
 #include <vector>
 #include <map>
 #include <set>
-#include <string>
 #include <unordered_map>
 
 
@@ -16,11 +15,6 @@ namespace mu
 	template< typename T >
 	using vector = std::vector<T>;
 
-	template< typename T >
-	using basic_string = std::basic_string<T, std::char_traits<T>>;
-
-	using string = std::basic_string<char, std::char_traits<char> >;
-
 	template< typename K, typename T >
 	using map = std::map< K, T, std::less<K> >;
 
diff --git a/Engine/Plugins/Experimental/Mutable/Source/MutableTools/Internal/MuT/TablePrivate.h b/Engine/Plugins/Experimental/Mutable/Source/MutableTools/Internal/MuT/TablePrivate.h
--- a/Engine/Plugins/Experimental/Mutable/Source/MutableTools/Internal/MuT/TablePrivate.h
+++ b/Engine/Plugins/Experimental/Mutable/Source/MutableTools/Internal/MuT/TablePrivate.h
@@ -6,9 +6,7 @@
 
 #include "MuT/ErrorLogPrivate.h"
 #include "MuR/MutableMath.h"
-#include "MuR/SerialisationPrivate.h"
 
-#include <string>
 
 namespace mu
 {
@@ -55,154 +53,6 @@ namespace mu
 		// Transient value for serialization compatibility
 		bool bNoneOption_DEPRECATED = false;
 
-
-		//!
-		void Serialise( OutputArchive& arch ) const
-		{
-            uint32 ver = 4;
-			arch << ver;
-
-            arch << Columns.Num();
-			for ( int32 c=0; c< Columns.Num(); ++c )
-			{
-				arch << Columns[c].Name;
-				arch << Columns[c].Type;
-			}
-
-            arch << Rows.Num();
-			for (int32 r=0; r< Rows.Num(); ++r )
-			{
-				arch << Rows[r].Id;
-
-				for (int32 c=0; c<Columns.Num(); ++c )
-				{
-					const FTableValue& v = Rows[r].Values[c];
-
-					switch (Columns[c].Type)
-					{
-					case ETableColumnType::Scalar: 	arch << v.Scalar; break;
-					case ETableColumnType::Color: 	arch << v.Color; break;
-					case ETableColumnType::Mesh: 	arch << v.Mesh; break;
-					case ETableColumnType::Image:
-					{
-						Ptr<const Image> Value;
-						if (v.ProxyImage)
-						{
-							Value = v.ProxyImage->Get();
-						}
-						arch << Value;
-						break;
-					}
-					case ETableColumnType::String:	arch << v.String; break;
-					default: check( false);
-					}
-				}
-			}
-		}
-
-		//!
-		void Unserialise( InputArchive& arch )
-		{
-            uint32 ver;
-			arch >> ver;
-			check(ver<=4);
-
-            uint32 ColumnCount;
-			arch >> ColumnCount;
-			Columns.SetNum(ColumnCount);
-            for (uint32 c=0; c<ColumnCount; ++c )
-			{
-				if (ver <= 2)
-				{
-					std::string LegacyString;
-					arch >> LegacyString;
-					Columns[c].Name = LegacyString.c_str();
-				}
-				else
-				{
-					arch >> Columns[c].Name;
-				}
-				arch >> Columns[c].Type;
-			}
-
-            uint32 RowCount;
-			arch >> RowCount;
-			Rows.SetNum( RowCount );
-            for ( uint32 r=0; r<RowCount; ++r )
-			{
-				arch >> Rows[r].Id;
-				Rows[r].Values.SetNum( ColumnCount );
-
-                for (uint32_t c=0; c<ColumnCount; ++c )
-				{
-					FTableValue& v = Rows[r].Values[c];
-
-					switch (Columns[c].Type)
-					{
-					case ETableColumnType::Scalar:	
-					{
-						arch >> v.Scalar;
-						break;
-					}
-					case ETableColumnType::Color:
-					{
-						if (ver <= 1)
-						{
-							FVector3f Value;
-							arch >> Value;
-
-							v.Color = FVector4f(Value[0], Value[1], Value[2], 1.0f);
-						}
-						else
-						{
-							arch >> v.Color;
-						}
-						break;
-					}
-					case ETableColumnType::Mesh:		
-					{
-						arch >> v.Mesh;
-						break;
-					}
-					case ETableColumnType::Image:
-					{
-						// Are we using proxies?
-						v.ProxyImage = arch.NewImageProxy();
-						if (!v.ProxyImage)
-						{
-							// Normal serialisation
-							Ptr<Image> Value;
-							arch >> Value;
-							v.ProxyImage = new ResourceProxyMemory<Image>(Value.get());
-						}
-						break;
-					}
-					case ETableColumnType::String:
-					{
-						if (ver <= 2)
-						{
-							std::string LegacyString;
-							arch >> LegacyString;
-							v.String = LegacyString.c_str();
-						}
-						else
-						{
-							arch >> v.String;
-						}
-						break;
-					}
-					default: check( false);
-					}
-				}
-			}
-
-			if (ver >= 1 && ver <= 3)
-			{
-				arch >> bNoneOption_DEPRECATED;
-			}
-		}
-
-
 		//! Find a row in the table by id. Return -1 if not found.
         int32 FindRow( uint32 id ) const
 		{
diff --git a/Engine/Plugins/Experimental/Mutable/Source/MutableTools/Private/MuT/AST.cpp b/Engine/Plugins/Experimental/Mutable/Source/MutableTools/Private/MuT/AST.cpp
--- a/Engine/Plugins/Experimental/Mutable/Source/MutableTools/Private/MuT/AST.cpp
+++ b/Engine/Plugins/Experimental/Mutable/Source/MutableTools/Private/MuT/AST.cpp
@@ -13,7 +13,6 @@
 #include "MuT/ASTOpMeshRemoveMask.h"
 #include "MuT/ASTOpImageMultiLayer.h"
 #include "MuT/ASTOpParameter.h"
-#include "MuT/StreamsPrivate.h"
 #include "MuT/Platform.h"
 #include "MuT/ErrorLogPrivate.h"
 #include "Trace/Detail/Channel.h"
diff --git a/Engine/Plugins/Experimental/Mutable/Source/MutableTools/Private/MuT/ASTOpAddExtensionData.cpp b/Engine/Plugins/Experimental/Mutable/Source/MutableTools/Private/MuT/ASTOpAddExtensionData.cpp
--- a/Engine/Plugins/Experimental/Mutable/Source/MutableTools/Private/MuT/ASTOpAddExtensionData.cpp
+++ b/Engine/Plugins/Experimental/Mutable/Source/MutableTools/Private/MuT/ASTOpAddExtensionData.cpp
@@ -5,7 +5,6 @@
 #include "HAL/PlatformMath.h"
 #include "MuR/ModelPrivate.h"
 #include "MuR/RefCounted.h"
-#include "MuT/StreamsPrivate.h"
 
 
 namespace mu
diff --git a/Engine/Plugins/Experimental/Mutable/Source/MutableTools/Private/MuT/ASTOpConstantResource.cpp b/Engine/Plugins/Experimental/Mutable/Source/MutableTools/Private/MuT/ASTOpConstantResource.cpp
--- a/Engine/Plugins/Experimental/Mutable/Source/MutableTools/Private/MuT/ASTOpConstantResource.cpp
+++ b/Engine/Plugins/Experimental/Mutable/Source/MutableTools/Private/MuT/ASTOpConstantResource.cpp
@@ -2,10 +2,7 @@
 
 #include "MuT/ASTOpConstantResource.h"
 
-#include "Containers/Array.h"
-#include "HAL/PlatformMath.h"
-#include "Hash/CityHash.h"
-#include "Misc/AssertionMacros.h"
+#include "MuT/CompilerPrivate.h"
 #include "MuR/Layout.h"
 #include "MuR/Mesh.h"
 #include "MuR/ModelPrivate.h"
@@ -14,12 +11,134 @@
 #include "MuR/Serialisation.h"
 #include "MuR/Skeleton.h"
 #include "MuR/Types.h"
-#include "MuT/StreamsPrivate.h"
+#include "Containers/Array.h"
+#include "HAL/PlatformMath.h"
+#include "HAL/PlatformFileManager.h"
+#include "Hash/CityHash.h"
+#include "Misc/AssertionMacros.h"
+#include "GenericPlatform/GenericPlatformFile.h"
 
+#include <inttypes.h> // Required for 64-bit printf macros
 
 namespace mu
 {
 
+	//---------------------------------------------------------------------------------------------
+	//---------------------------------------------------------------------------------------------
+	//---------------------------------------------------------------------------------------------
+	template<class R>
+	class MUTABLETOOLS_API ResourceProxyTempFile : public ResourceProxy<R>
+	{
+	private:
+		Ptr<const R> Resource;
+		FString FileName;
+		uint64 FileSize = 0;
+		FCriticalSection Mutex;
+
+		FProxyFileContext& Options;
+
+	public:
+		ResourceProxyTempFile(const R* resource, FProxyFileContext& InOptions)
+			: Options(InOptions)
+		{
+			IPlatformFile& PlatformFile = FPlatformFileManager::Get().GetPlatformFile();
+
+			if (!resource)
+			{
+				return;
+			}
+
+			FScopeLock Lock(&Mutex);
+
+			OutputMemoryStream stream;
+			OutputArchive arch(&stream);
+			R::Serialise(resource, arch);
+
+			if (stream.GetBufferSize() <= Options.MinProxyFileSize)
+			{
+				Resource = resource;
+			}
+			else
+			{
+				FileSize = stream.GetBufferSize();
+
+				FString Prefix = FPlatformProcess::UserTempDir();
+
+				uint32 PID = FPlatformProcess::GetCurrentProcessId();
+				Prefix += FString::Printf(TEXT("mut.temp.%u"), PID);
+
+				FString FinalTempPath;
+				IFileHandle* ResourceFile = nullptr;
+				uint64 AttemptCount = 0;
+				while (!ResourceFile && AttemptCount < Options.MaxFileCreateAttempts)
+				{
+					uint64 ThisThreadFileIndex = Options.CurrentFileIndex.load();
+					while (!Options.CurrentFileIndex.compare_exchange_strong(ThisThreadFileIndex, ThisThreadFileIndex + 1));
+
+					FinalTempPath = Prefix + FString::Printf(TEXT(".%.16" PRIx64), ThisThreadFileIndex);
+					ResourceFile = PlatformFile.OpenWrite(*FinalTempPath);
+					++AttemptCount;
+				}
+
+				if (!ResourceFile)
+				{
+					UE_LOG(LogMutableCore, Error, TEXT("Failed to create temporary file. Disk full?"));
+					check(false);
+				}
+
+				ResourceFile->Write((const uint8*)stream.GetBuffer(), stream.GetBufferSize());
+				delete ResourceFile;
+
+				FileName = FinalTempPath;
+				Options.FilesWritten++;
+				Options.BytesWritten += stream.GetBufferSize();
+			}
+		}
+
+		~ResourceProxyTempFile()
+		{
+			FScopeLock Lock(&Mutex);
+
+			if (!FileName.IsEmpty())
+			{
+				// Delete temp file
+				FPlatformFileManager::Get().GetPlatformFile().DeleteFile(*FileName);
+				FileName.Empty();
+			}
+		}
+
+		Ptr<const R> Get() override
+		{
+			FScopeLock Lock(&Mutex);
+
+			Ptr<const R> r;
+			if (Resource)
+			{
+				// cached
+				r = Resource;
+			}
+			else if (!FileName.IsEmpty())
+			{
+				TArray<char> buf;
+				buf.SetNumUninitialized(FileSize);
+				IFileHandle* resourceFile = FPlatformFileManager::Get().GetPlatformFile().OpenRead(*FileName);
+				check(resourceFile);
+				resourceFile->Read((uint8*)buf.GetData(), FileSize);
+				delete resourceFile;
+
+				InputMemoryStream stream(buf.GetData(), FileSize);
+				InputArchive arch(&stream);
+				r = R::StaticUnserialise(arch);
+
+				Options.FilesRead++;
+				Options.BytesRead += FileSize;
+			}
+			return r;
+		}
+
+	};
+
+
 	//-------------------------------------------------------------------------------------------------
 	//-------------------------------------------------------------------------------------------------
 	//-------------------------------------------------------------------------------------------------
diff --git a/Engine/Plugins/Experimental/Mutable/Source/MutableTools/Private/MuT/ASTOpLayoutMerge.cpp b/Engine/Plugins/Experimental/Mutable/Source/MutableTools/Private/MuT/ASTOpLayoutMerge.cpp
--- a/Engine/Plugins/Experimental/Mutable/Source/MutableTools/Private/MuT/ASTOpLayoutMerge.cpp
+++ b/Engine/Plugins/Experimental/Mutable/Source/MutableTools/Private/MuT/ASTOpLayoutMerge.cpp
@@ -17,8 +17,6 @@
 #include "MuT/ASTOpMeshRemoveMask.h"
 #include "MuT/ASTOpMeshMorph.h"
 #include "MuT/ASTOpSwitch.h"
-#include "MuT/StreamsPrivate.h"
-
 
 
 namespace mu
diff --git a/Engine/Plugins/Experimental/Mutable/Source/MutableTools/Private/MuT/ASTOpLayoutPack.cpp b/Engine/Plugins/Experimental/Mutable/Source/MutableTools/Private/MuT/ASTOpLayoutPack.cpp
--- a/Engine/Plugins/Experimental/Mutable/Source/MutableTools/Private/MuT/ASTOpLayoutPack.cpp
+++ b/Engine/Plugins/Experimental/Mutable/Source/MutableTools/Private/MuT/ASTOpLayoutPack.cpp
@@ -17,8 +17,6 @@
 #include "MuT/ASTOpMeshRemoveMask.h"
 #include "MuT/ASTOpMeshMorph.h"
 #include "MuT/ASTOpSwitch.h"
-#include "MuT/StreamsPrivate.h"
-
 
 
 namespace mu
diff --git a/Engine/Plugins/Experimental/Mutable/Source/MutableTools/Private/MuT/ASTOpMeshGeometryOperation.cpp b/Engine/Plugins/Experimental/Mutable/Source/MutableTools/Private/MuT/ASTOpMeshGeometryOperation.cpp
--- a/Engine/Plugins/Experimental/Mutable/Source/MutableTools/Private/MuT/ASTOpMeshGeometryOperation.cpp
+++ b/Engine/Plugins/Experimental/Mutable/Source/MutableTools/Private/MuT/ASTOpMeshGeometryOperation.cpp
@@ -5,7 +5,7 @@
 #include "HAL/PlatformMath.h"
 #include "MuR/ModelPrivate.h"
 #include "MuR/RefCounted.h"
-#include "MuT/StreamsPrivate.h"
+
 
 namespace mu
 {
diff --git a/Engine/Plugins/Experimental/Mutable/Source/MutableTools/Private/MuT/ASTOpMeshMorphReshape.cpp b/Engine/Plugins/Experimental/Mutable/Source/MutableTools/Private/MuT/ASTOpMeshMorphReshape.cpp
--- a/Engine/Plugins/Experimental/Mutable/Source/MutableTools/Private/MuT/ASTOpMeshMorphReshape.cpp
+++ b/Engine/Plugins/Experimental/Mutable/Source/MutableTools/Private/MuT/ASTOpMeshMorphReshape.cpp
@@ -6,7 +6,7 @@
 #include "MuR/ModelPrivate.h"
 #include "MuR/RefCounted.h"
 #include "MuR/Types.h"
-#include "MuT/StreamsPrivate.h"
+
 
 namespace mu
 {
diff --git a/Engine/Plugins/Experimental/Mutable/Source/MutableTools/Private/MuT/Compiler.cpp b/Engine/Plugins/Experimental/Mutable/Source/MutableTools/Private/MuT/Compiler.cpp
--- a/Engine/Plugins/Experimental/Mutable/Source/MutableTools/Private/MuT/Compiler.cpp
+++ b/Engine/Plugins/Experimental/Mutable/Source/MutableTools/Private/MuT/Compiler.cpp
@@ -34,7 +34,6 @@
 #include "MuT/Table.h"
 #include "Trace/Detail/Channel.h"
 
-#include <string>		// Required for deserialisation of old data
 #include <inttypes.h>	// Required for 64-bit printf macros
 
 
@@ -53,6 +52,15 @@ namespace mu
     }
 
 
+	//---------------------------------------------------------------------------------------------
+	FProxyFileContext::FProxyFileContext()
+	{
+		uint32 Seed = FPlatformTime::Cycles();
+		FRandomStream RandomStream = FRandomStream((int32)Seed);
+		CurrentFileIndex = RandomStream.GetUnsignedInt();
+	}
+
+
     //---------------------------------------------------------------------------------------------
     //---------------------------------------------------------------------------------------------
     //---------------------------------------------------------------------------------------------
diff --git a/Engine/Plugins/Experimental/Mutable/Source/MutableTools/Private/MuT/CompilerPrivate.h b/Engine/Plugins/Experimental/Mutable/Source/MutableTools/Private/MuT/CompilerPrivate.h
--- a/Engine/Plugins/Experimental/Mutable/Source/MutableTools/Private/MuT/CompilerPrivate.h
+++ b/Engine/Plugins/Experimental/Mutable/Source/MutableTools/Private/MuT/CompilerPrivate.h
@@ -7,12 +7,34 @@
 #include "MuT/AST.h"
 #include "MuT/ErrorLogPrivate.h"
 #include "MuT/NodeObjectPrivate.h"
-#include "MuT/StreamsPrivate.h"
 #include "MuR/Operations.h"
 
 
 namespace mu
 {
+	/** Statistics about the proxy file usage. */
+	struct FProxyFileContext
+	{
+		FProxyFileContext();
+
+		/** Options */
+
+		/** Minimum data size in bytes to dumpt it to the disk. */
+		uint64 MinProxyFileSize = 1024 * 1024;
+
+		/** When creating temporary files, number of retries in case the OS-level call fails. */
+		uint64 MaxFileCreateAttempts = 256;
+
+		/** Statistics */
+		std::atomic<uint64> FilesWritten = 0;
+		std::atomic<uint64> FilesRead = 0;
+		std::atomic<uint64> BytesWritten = 0;
+		std::atomic<uint64> BytesRead = 0;
+
+		/** Internal data. */
+		std::atomic<uint64> CurrentFileIndex = 0;
+	};
+
 		
     //!
     class CompilerOptions::Private
diff --git a/Engine/Plugins/Experimental/Mutable/Source/MutableTools/Private/MuT/Node.cpp b/Engine/Plugins/Experimental/Mutable/Source/MutableTools/Private/MuT/Node.cpp
--- a/Engine/Plugins/Experimental/Mutable/Source/MutableTools/Private/MuT/Node.cpp
+++ b/Engine/Plugins/Experimental/Mutable/Source/MutableTools/Private/MuT/Node.cpp
@@ -2,34 +2,15 @@
 
 
 #include "MuT/Node.h"
-
-#include "Misc/AssertionMacros.h"
-#include "MuR/Serialisation.h"
-#include "MuR/SerialisationPrivate.h"
-#include "MuT/NodeBool.h"
-#include "MuT/NodeColour.h"
-#include "MuT/NodeComponent.h"
-#include "MuT/NodeImage.h"
-#include "MuT/NodeLOD.h"
-#include "MuT/NodeLayout.h"
-#include "MuT/NodeMesh.h"
-#include "MuT/NodeModifier.h"
-#include "MuT/NodeObject.h"
-#include "MuT/NodePatchImage.h"
-#include "MuT/NodePatchMesh.h"
 #include "MuT/NodePrivate.h"
-#include "MuT/NodeProjector.h"
-#include "MuT/NodeRange.h"
-#include "MuT/NodeScalar.h"
-#include "MuT/NodeString.h"
-#include "MuT/NodeSurface.h"
+#include "Misc/AssertionMacros.h"
 
 
 namespace mu
 {
 
 	// Static initialisation
-	static FNodeType s_nodeType = 		FNodeType( "Node", 0 );
+	static FNodeType s_nodeType = FNodeType( "Node", 0 );
 
 	FNodeType::FNodeType()
 	{
@@ -44,76 +25,6 @@ namespace mu
 		m_pParent = pParent;
 	}
 
-
-	void Node::Serialise( const Node* pNode, OutputArchive& arch )
-	{
-        uint32_t ver = 0;
-		arch << ver;
-
-		// Warning: cannot be replaced with the wrapping virtual call because some subclasses are abstract
-		arch << uint32_t(pNode->Type);
-
-		#define SERIALISE_CHILDREN( C, ID ) C::Serialise(static_cast<const C*>(pNode), arch );
-
-		switch (pNode->Type)
-		{
-		case EType::Colour:		SERIALISE_CHILDREN(NodeColour,arch); break;
-		case EType::Component:  SERIALISE_CHILDREN(NodeComponent,arch); break;
-		case EType::Image:		SERIALISE_CHILDREN(NodeImage,arch); break;
-		case EType::Layout:		SERIALISE_CHILDREN(NodeLayout,arch); break;
-		case EType::LOD:		SERIALISE_CHILDREN(NodeLOD,arch); break;
-		case EType::Mesh:		SERIALISE_CHILDREN(NodeMesh,arch); break;
-		case EType::Object:		SERIALISE_CHILDREN(NodeObject,arch); break;
-		case EType::PatchImage: SERIALISE_CHILDREN(NodePatchImage,arch); break;
-		case EType::Scalar:		SERIALISE_CHILDREN(NodeScalar,arch); break;
-		case EType::PatchMesh:	SERIALISE_CHILDREN(NodePatchMesh, arch); break;
-		case EType::Projector:	SERIALISE_CHILDREN(NodeProjector,arch); break;
-		case EType::Surface:	SERIALISE_CHILDREN(NodeSurface,arch); break;
-		case EType::Modifier:	SERIALISE_CHILDREN(NodeModifier,arch); break;
-		case EType::Range:		SERIALISE_CHILDREN(NodeRange,arch); break;
-		case EType::String:		SERIALISE_CHILDREN(NodeString, arch); break;
-		case EType::Bool:		SERIALISE_CHILDREN(NodeBool, arch); break;
-		default: check(false);
-		}
-
-		#undef SERIALISE_CHILDREN
-	}
-
-
-	NodePtr Node::StaticUnserialise( InputArchive& arch )
-	{
-        uint32_t ver;
-		arch >> ver;
-		check( ver == 0 );
-
-        uint32_t id;
-		arch >> id;
-
-		switch (EType(id))
-		{
-		case EType::Colour:		return NodeColour::StaticUnserialise( arch ); break;
-		case EType::Component:  return NodeComponent::StaticUnserialise( arch ); break;
-		case EType::Image:		return NodeImage::StaticUnserialise( arch ); break;
-		case EType::Layout:		return NodeLayout::StaticUnserialise( arch ); break;
-		case EType::LOD:		return NodeLOD::StaticUnserialise( arch ); break;
-		case EType::Mesh:		return NodeMesh::StaticUnserialise( arch ); break;
-		case EType::Object:		return NodeObject::StaticUnserialise( arch ); break;
-		case EType::PatchImage: return NodePatchImage::StaticUnserialise( arch ); break;
-		case EType::Scalar:		return NodeScalar::StaticUnserialise( arch ); break;
-		case EType::PatchMesh:	return NodePatchMesh::StaticUnserialise( arch ); break;
-		case EType::Projector:	return NodeProjector::StaticUnserialise( arch ); break;
-		case EType::Surface:	return NodeSurface::StaticUnserialise( arch ); break;
-		case EType::Modifier:	return NodeModifier::StaticUnserialise( arch ); break;
-		case EType::Range:		return NodeRange::StaticUnserialise( arch ); break;
-		case EType::String:		return NodeString::StaticUnserialise( arch ); break;
-		case EType::Bool:		return NodeBool::StaticUnserialise( arch ); break;
-		default : check(false);
-		}
-
-		return 0;
-	}
-
-
 	const FNodeType* Node::GetType() const
 	{
 		return GetStaticType();
diff --git a/Engine/Plugins/Experimental/Mutable/Source/MutableTools/Private/MuT/NodeBool.cpp b/Engine/Plugins/Experimental/Mutable/Source/MutableTools/Private/MuT/NodeBool.cpp
--- a/Engine/Plugins/Experimental/Mutable/Source/MutableTools/Private/MuT/NodeBool.cpp
+++ b/Engine/Plugins/Experimental/Mutable/Source/MutableTools/Private/MuT/NodeBool.cpp
@@ -4,54 +4,14 @@
 #include "MuT/NodeBool.h"
 
 #include "Misc/AssertionMacros.h"
-#include "MuR/Serialisation.h"
-#include "MuR/SerialisationPrivate.h"
 
 
 namespace mu
 {
 
-	//---------------------------------------------------------------------------------------------
-	// Static initialisation
-	//---------------------------------------------------------------------------------------------
 	static FNodeType s_nodeBoolType = FNodeType( "NodeBool", Node::GetStaticType() );
 
 
-	//---------------------------------------------------------------------------------------------
-	void NodeBool::Serialise( const NodeBool* p, OutputArchive& arch )
-	{
-        uint32 ver = 0;
-		arch << ver;
-
-		arch << uint32_t(p->Type);
-		p->SerialiseWrapper(arch);
-	}
-
-        
-	//---------------------------------------------------------------------------------------------
-	Ptr<NodeBool> NodeBool::StaticUnserialise( InputArchive& arch )
-	{
-        uint32 ver;
-		arch >> ver;
-		check( ver == 0 );
-
-        uint32 id;
-		arch >> id;
-
-		switch (id)
-		{
-		case 0 :  return NodeBoolConstant::StaticUnserialise( arch ); break;
-		case 1 :  return NodeBoolParameter::StaticUnserialise( arch ); break;
-		//case 2 :  return NodeBoolIsNull::StaticUnserialise( arch ); break;
-		case 3 :  return NodeBoolNot::StaticUnserialise( arch ); break;
-		case 4 :  return NodeBoolAnd::StaticUnserialise( arch ); break;
-		default : check(false);
-		}
-
-		return 0;
-	}
-
-
 	//---------------------------------------------------------------------------------------------
 	const FNodeType* NodeBool::GetType() const
 	{
