[mutable] Remove custom serialization code frome the MutableTools module.

diff --git a/Engine/Plugins/Experimental/Mutable/Source/MutableTools/Private/MuT/NodeBoolPrivate.h b/Engine/Plugins/Experimental/Mutable/Source/MutableTools/Private/MuT/NodeBoolPrivate.h
--- a/Engine/Plugins/Experimental/Mutable/Source/MutableTools/Private/MuT/NodeBoolPrivate.h
+++ b/Engine/Plugins/Experimental/Mutable/Source/MutableTools/Private/MuT/NodeBoolPrivate.h
@@ -7,8 +7,6 @@
 #include "MuT/NodeBool.h"
 #include "MuT/NodeImage.h"
 #include "MuT/NodeRange.h"
-#include "MuT/AST.h"
-
 
 namespace mu
 {
@@ -26,25 +24,6 @@ namespace mu
 		static FNodeType s_type;
 
 		bool m_value;
-
-		//!
-		void Serialise( OutputArchive& arch ) const
-		{
-            uint32 ver = 0;
-			arch << ver;
-
-			arch << m_value;
-		}
-
-		//!
-		void Unserialise( InputArchive& arch )
-		{
-            uint32 ver;
-			arch >> ver;
-			check(ver==0);
-
-			arch >> m_value;
-		}
 	};
 
 
@@ -59,43 +38,6 @@ namespace mu
 		FString m_uid;
 
         TArray<Ptr<NodeRange>> m_ranges;
-
-		//!
-		void Serialise( OutputArchive& arch ) const
-		{
-            uint32 ver = 3;
-			arch << ver;
-
-			arch << m_defaultValue;
-			arch << m_name;
-			arch << m_uid;
-            arch << m_ranges;
-        }
-
-		//!
-		void Unserialise( InputArchive& arch )
-		{
-            uint32 ver;
-			arch >> ver;
-            check(ver>=2 && ver<=3);
-
-			arch >> m_defaultValue;
-			if (ver <= 2)
-			{
-				std::string Temp;
-				arch >> Temp;
-				m_name = Temp.c_str();
-				arch >> Temp;
-				m_uid = Temp.c_str();
-			}
-			else
-			{
-				arch >> m_name;
-				arch >> m_uid;
-			}
-			
-			arch >> m_ranges;
-        }
 	};
 
 
@@ -106,25 +48,6 @@ namespace mu
 		static FNodeType s_type;
 
 		Ptr<NodeBool> m_pSource;
-
-		//!
-		void Serialise( OutputArchive& arch ) const
-		{
-            uint32 ver = 0;
-			arch << ver;
-
-			arch << m_pSource;
-		}
-
-		//!
-		void Unserialise( InputArchive& arch )
-		{
-            uint32 ver;
-			arch >> ver;
-            check(ver<=0);
-
-			arch >> m_pSource;
-		}
 	};
 
 
@@ -136,27 +59,6 @@ namespace mu
 
 		Ptr<NodeBool> m_pA;
 		Ptr<NodeBool> m_pB;
-
-		//!
-		void Serialise( OutputArchive& arch ) const
-		{
-            uint32 ver = 0;
-			arch << ver;
-
-			arch << m_pA;
-			arch << m_pB;
-		}
-
-		//!
-		void Unserialise( InputArchive& arch )
-		{
-            uint32 ver;
-			arch >> ver;
-            check(ver<=0);
-
-			arch >> m_pA;
-			arch >> m_pB;
-		}
 	};
 
 }
diff --git a/Engine/Plugins/Experimental/Mutable/Source/MutableTools/Private/MuT/NodeColour.cpp b/Engine/Plugins/Experimental/Mutable/Source/MutableTools/Private/MuT/NodeColour.cpp
--- a/Engine/Plugins/Experimental/Mutable/Source/MutableTools/Private/MuT/NodeColour.cpp
+++ b/Engine/Plugins/Experimental/Mutable/Source/MutableTools/Private/MuT/NodeColour.cpp
@@ -4,72 +4,14 @@
 #include "MuT/NodeColour.h"
 
 #include "Misc/AssertionMacros.h"
-#include "MuR/Serialisation.h"
-#include "MuR/SerialisationPrivate.h"
-#include "MuT/NodeColourArithmeticOperation.h"
-#include "MuT/NodeColourConstant.h"
-#include "MuT/NodeColourFromScalars.h"
-#include "MuT/NodeColourParameter.h"
-#include "MuT/NodeColourSampleImage.h"
-#include "MuT/NodeColourSwitch.h"
-#include "MuT/NodeColourTable.h"
-#include "MuT/NodeColourVariation.h"
 
 
 namespace mu
 {
 
-	//---------------------------------------------------------------------------------------------
-	// Static initialisation
-	//---------------------------------------------------------------------------------------------
 	static FNodeType s_nodeColourType = FNodeType( "NodeColour", Node::GetStaticType() );
 
 
-	//---------------------------------------------------------------------------------------------
-	void NodeColour::Serialise( const NodeColour* p, OutputArchive& arch )
-	{
-        uint32_t ver = 1;
-		arch << ver;
-
-		arch << uint32_t(p->Type);
-		p->SerialiseWrapper(arch);
-    }
-
-        
-	//---------------------------------------------------------------------------------------------
-	NodeColourPtr NodeColour::StaticUnserialise( InputArchive& arch )
-	{
-        uint32_t ver;
-		arch >> ver;
-		check( ver <= 1 );
-
-        uint32_t id;
-		arch >> id;
-
-		switch (id)
-		{
-		case 0 :
-			if (ver == 0)
-			{
-				return NodeColourConstant::OldStaticUnserialise(arch);
-			}
-			return NodeColourConstant::StaticUnserialise(arch);
-			break;
-		case 1 :  return NodeColourParameter::StaticUnserialise( arch ); break;
-		case 2 :  return NodeColourSampleImage::StaticUnserialise( arch ); break;
-		case 3 :  return NodeColourTable::StaticUnserialise( arch ); break;
-		//case 4 :  return NodeColourImageSize::StaticUnserialise( arch ); break; // deprecated
-		case 5 :  return NodeColourFromScalars::StaticUnserialise( arch ); break;
-        case 6 :  return NodeColourArithmeticOperation::StaticUnserialise( arch ); break;
-        case 7 :  return NodeColourSwitch::StaticUnserialise( arch ); break;
-        case 8 :  return NodeColourVariation::StaticUnserialise( arch ); break;
-        default : check(false);
-		}
-
-		return nullptr;
-	}
-
-
 	//---------------------------------------------------------------------------------------------
 	const FNodeType* NodeColour::GetType() const
 	{
diff --git a/Engine/Plugins/Experimental/Mutable/Source/MutableTools/Private/MuT/NodeColourArithmeticOperationPrivate.h b/Engine/Plugins/Experimental/Mutable/Source/MutableTools/Private/MuT/NodeColourArithmeticOperationPrivate.h
--- a/Engine/Plugins/Experimental/Mutable/Source/MutableTools/Private/MuT/NodeColourArithmeticOperationPrivate.h
+++ b/Engine/Plugins/Experimental/Mutable/Source/MutableTools/Private/MuT/NodeColourArithmeticOperationPrivate.h
@@ -5,7 +5,6 @@
 #include "MuT/NodeColourPrivate.h"
 #include "MuT/NodeColourArithmeticOperation.h"
 
-
 namespace mu
 {
 
@@ -18,29 +17,6 @@ namespace mu
 		NodeColourArithmeticOperation::OPERATION m_operation;
 		NodeColourPtr m_pA;
 		NodeColourPtr m_pB;
-
-		//!
-		void Serialise( OutputArchive& arch ) const
-		{
-            uint32_t ver = 0;
-			arch << ver;
-
-			arch << m_operation;
-			arch << m_pA;
-			arch << m_pB;
-		}
-
-		//!
-		void Unserialise( InputArchive& arch )
-		{
-            uint32_t ver;
-			arch >> ver;
-			check(ver==0);
-
-			arch >> m_operation;
-			arch >> m_pA;
-			arch >> m_pB;
-		}
 	};
 
 
diff --git a/Engine/Plugins/Experimental/Mutable/Source/MutableTools/Private/MuT/NodeColourConstant.cpp b/Engine/Plugins/Experimental/Mutable/Source/MutableTools/Private/MuT/NodeColourConstant.cpp
--- a/Engine/Plugins/Experimental/Mutable/Source/MutableTools/Private/MuT/NodeColourConstant.cpp
+++ b/Engine/Plugins/Experimental/Mutable/Source/MutableTools/Private/MuT/NodeColourConstant.cpp
@@ -12,11 +12,7 @@
 namespace mu
 {
 
-	//---------------------------------------------------------------------------------------------
-	// Static initialisation
-	//---------------------------------------------------------------------------------------------
-	FNodeType NodeColourConstant::Private::s_type =
-			FNodeType( "ColourConstant", NodeColour::GetStaticType() );
+	FNodeType NodeColourConstant::Private::s_type = FNodeType( "ColourConstant", NodeColour::GetStaticType() );
 
 
 	//---------------------------------------------------------------------------------------------
@@ -25,16 +21,6 @@ namespace mu
 
 	MUTABLE_IMPLEMENT_NODE( NodeColourConstant, EType::Constant, Node, Node::EType::Colour)
 
-	NodeColourConstantPtr NodeColourConstant::OldStaticUnserialise(InputArchive& arch)
-	{
-		NodeColourConstantPtr pResult = new NodeColourConstant();
-		FVector3f Value;
-
-		arch >> Value;
-		pResult->GetPrivate()->m_value = FVector4f(Value[0], Value[1], Value[2], 1.0f);
-
-		return pResult;
-	}
 
 
 	//---------------------------------------------------------------------------------------------
diff --git a/Engine/Plugins/Experimental/Mutable/Source/MutableTools/Private/MuT/NodeColourConstantPrivate.h b/Engine/Plugins/Experimental/Mutable/Source/MutableTools/Private/MuT/NodeColourConstantPrivate.h
--- a/Engine/Plugins/Experimental/Mutable/Source/MutableTools/Private/MuT/NodeColourConstantPrivate.h
+++ b/Engine/Plugins/Experimental/Mutable/Source/MutableTools/Private/MuT/NodeColourConstantPrivate.h
@@ -4,7 +4,6 @@
 
 #include "MuT/NodeColourConstant.h"
 #include "MuT/NodeColourPrivate.h"
-
 #include "MuR/MutableMath.h"
 
 namespace mu
@@ -17,24 +16,6 @@ namespace mu
 		static FNodeType s_type;
 
 		FVector4f m_value;
-
-		//!
-		void Serialise( OutputArchive& arch ) const
-		{
-			const uint32 Version = 0;
-			arch << Version;
-			arch << m_value;
-		}
-
-		//!
-		void Unserialise( InputArchive& arch )
-		{
-			uint32 Version;
-			arch >> Version;
-			check(Version==0)
-
-			arch >> m_value;
-		}
 	};
 
 }
diff --git a/Engine/Plugins/Experimental/Mutable/Source/MutableTools/Private/MuT/NodeColourTablePrivate.h b/Engine/Plugins/Experimental/Mutable/Source/MutableTools/Private/MuT/NodeColourTablePrivate.h
--- a/Engine/Plugins/Experimental/Mutable/Source/MutableTools/Private/MuT/NodeColourTablePrivate.h
+++ b/Engine/Plugins/Experimental/Mutable/Source/MutableTools/Private/MuT/NodeColourTablePrivate.h
@@ -6,7 +6,6 @@
 #include "MuT/NodePrivate.h"
 #include "MuT/NodeColourTable.h"
 #include "MuT/TablePrivate.h"
-#include "MuT/AST.h"
 
 
 namespace mu
@@ -25,65 +24,6 @@ namespace mu
 		bool bNoneOption = false;
 		FString DefaultRowName;
 
-		//!
-		void Serialise( OutputArchive& arch ) const
-		{
-            uint32_t ver = 3;
-			arch << ver;
-
-			arch << ParameterName;
-			arch << Table;
-			arch << ColumnName;
-			arch << bNoneOption;
-			arch << DefaultRowName;
-		}
-
-		//!
-		void Unserialise( InputArchive& arch )
-		{
-            uint32 ver;
-			arch >> ver;
-			check(ver <= 3);
-
-			if (ver == 0)
-			{
-				std::string Temp;
-				arch >> Temp;
-				ParameterName = Temp.c_str();
-			}
-			else
-			{
-				arch >> ParameterName;
-			}
-
-			arch >> Table;
-
-			if (ver == 0)
-			{
-				std::string Temp;
-				arch >> Temp;
-				ColumnName = Temp.c_str();
-			}
-			else
-			{
-				arch >> ColumnName;
-			}
-
-			if (ver >= 2)
-			{
-				arch >> bNoneOption;
-			}
-			else
-			{
-				bNoneOption = Table->GetPrivate()->bNoneOption_DEPRECATED;
-			}
-
-			if (ver >= 3)
-			{
-				arch >> DefaultRowName;
-			}
-		}
-
 	};
 
 }
diff --git a/Engine/Plugins/Experimental/Mutable/Source/MutableTools/Private/MuT/NodeColourVariationPrivate.h b/Engine/Plugins/Experimental/Mutable/Source/MutableTools/Private/MuT/NodeColourVariationPrivate.h
--- a/Engine/Plugins/Experimental/Mutable/Source/MutableTools/Private/MuT/NodeColourVariationPrivate.h
+++ b/Engine/Plugins/Experimental/Mutable/Source/MutableTools/Private/MuT/NodeColourVariationPrivate.h
@@ -25,59 +25,9 @@ namespace mu
         {
             NodeColourPtr m_colour;
 			FString m_tag;
-
-            //!
-            void Serialise( OutputArchive& arch ) const
-            {
-                uint32 ver = 2;
-                arch << ver;
-
-                arch << m_tag;
-                arch << m_colour;
-            }
-
-            void Unserialise( InputArchive& arch )
-            {
-                uint32 ver = 0;
-                arch >> ver;
-				check(ver >= 1 && ver <= 2);
-
-				if (ver <= 1)
-				{
-					std::string Temp;
-					arch >> Temp;
-					m_tag = Temp.c_str();
-				}
-				else
-				{
-					arch >> m_tag;
-				}
-				arch >> m_colour;
-            }
         };
 
         TArray<FVariation> m_variations;
-
-        //!
-        void Serialise( OutputArchive& arch ) const
-        {
-            uint32_t ver = 1;
-            arch << ver;
-
-            arch << m_defaultColour;
-            arch << m_variations;
-        }
-
-        //!
-        void Unserialise( InputArchive& arch )
-        {
-            uint32_t ver;
-            arch >> ver;
-			check(ver == 1);
-
-            arch >> m_defaultColour;
-            arch >> m_variations;
-        }
     };
 
 } // namespace mu
diff --git a/Engine/Plugins/Experimental/Mutable/Source/MutableTools/Private/MuT/NodeComponent.cpp b/Engine/Plugins/Experimental/Mutable/Source/MutableTools/Private/MuT/NodeComponent.cpp
--- a/Engine/Plugins/Experimental/Mutable/Source/MutableTools/Private/MuT/NodeComponent.cpp
+++ b/Engine/Plugins/Experimental/Mutable/Source/MutableTools/Private/MuT/NodeComponent.cpp
@@ -4,19 +4,11 @@
 #include "MuT/NodeComponent.h"
 
 #include "Misc/AssertionMacros.h"
-#include "MuR/Serialisation.h"
-#include "MuR/SerialisationPrivate.h"
-#include "MuT/NodeComponentEdit.h"
-#include "MuT/NodeComponentNew.h"
 
 
 namespace mu
 {
 
-
-	//---------------------------------------------------------------------------------------------
-	// Static initialisation
-	//---------------------------------------------------------------------------------------------
 	static FNodeType s_nodeComponentType = FNodeType( "NodeComponent", Node::GetStaticType() );
 
 
@@ -35,39 +27,6 @@ namespace mu
 		return &s_nodeComponentType;
 	}
 
-
-	//---------------------------------------------------------------------------------------------
-	void NodeComponent::Serialise( const NodeComponent* p, OutputArchive& arch )
-	{
-        uint32 ver = 0;
-		arch << ver;
-
-		arch << uint32(p->Type);
-		p->SerialiseWrapper(arch);
-	}
-
-        
-	//---------------------------------------------------------------------------------------------
-	NodeComponentPtr NodeComponent::StaticUnserialise( InputArchive& arch )
-	{
-        uint32 ver;
-		arch >> ver;
-		check( ver == 0 );
-
-        uint32 id;
-		arch >> id;
-
-		switch (id)
-		{
-		case 0 :  return NodeComponentNew::StaticUnserialise( arch ); break;
-		case 1 :  return NodeComponentEdit::StaticUnserialise( arch ); break;
-//		case 2 :  return NodeSelectComponent::StaticUnserialise( arch ); break;
-		default : check(false);
-		}
-
-		return 0;
-	}
-
 }
 
 
diff --git a/Engine/Plugins/Experimental/Mutable/Source/MutableTools/Private/MuT/NodeExtensionData.cpp b/Engine/Plugins/Experimental/Mutable/Source/MutableTools/Private/MuT/NodeExtensionData.cpp
--- a/Engine/Plugins/Experimental/Mutable/Source/MutableTools/Private/MuT/NodeExtensionData.cpp
+++ b/Engine/Plugins/Experimental/Mutable/Source/MutableTools/Private/MuT/NodeExtensionData.cpp
@@ -2,57 +2,12 @@
 
 #include "MuT/NodeExtensionData.h"
 
-#include "MuR/Serialisation.h"
-#include "MuR/SerialisationPrivate.h"
-#include "MuT/NodeExtensionDataConstant.h"
-#include "MuT/NodeExtensionDataSwitch.h"
-#include "MuT/NodeExtensionDataVariation.h"
 
 namespace mu
 {
 
-
-	//---------------------------------------------------------------------------------------------
-	// Static initialisation
-	//---------------------------------------------------------------------------------------------
 	static FNodeType s_type = FNodeType("NodeExtensionData", Node::GetStaticType());
 
-	//---------------------------------------------------------------------------------------------
-	//!
-	//---------------------------------------------------------------------------------------------
-
-	//---------------------------------------------------------------------------------------------
-	void NodeExtensionData::Serialise(const NodeExtensionData* Node, OutputArchive& Archive)
-	{
-		uint32_t Ver = 0;
-		Archive << Ver;
-
-		Archive << uint32_t(Node->Type);
-		Node->SerialiseWrapper(Archive);
-	}
-
-
-	//---------------------------------------------------------------------------------------------
-	NodeExtensionDataPtr NodeExtensionData::StaticUnserialise(InputArchive& Archive)
-	{
-		uint32_t Ver;
-		Archive >> Ver;
-		check(Ver == 0);
-
-		uint32_t Id;
-		Archive >> Id;
-
-		switch (EType(Id))
-		{
-			case EType::Constant:		return NodeExtensionDataConstant::StaticUnserialise(Archive); break;
-			case EType::Switch:			return NodeExtensionDataSwitch::StaticUnserialise(Archive); break;
-			case EType::Variation:		return NodeExtensionDataVariation::StaticUnserialise(Archive); break;
-			default: check(false);
-		}
-
-		return 0;
-	}
-
 
 	//---------------------------------------------------------------------------------------------
 	const FNodeType* NodeExtensionData::GetType() const
diff --git a/Engine/Plugins/Experimental/Mutable/Source/MutableTools/Private/MuT/NodeExtensionDataConstantPrivate.h b/Engine/Plugins/Experimental/Mutable/Source/MutableTools/Private/MuT/NodeExtensionDataConstantPrivate.h
--- a/Engine/Plugins/Experimental/Mutable/Source/MutableTools/Private/MuT/NodeExtensionDataConstantPrivate.h
+++ b/Engine/Plugins/Experimental/Mutable/Source/MutableTools/Private/MuT/NodeExtensionDataConstantPrivate.h
@@ -6,7 +6,6 @@
 #include "MuT/NodeExtensionDataConstant.h"
 #include "MuT/NodeLayout.h"
 #include "MuT/NodePrivate.h"
-#include "MuT/AST.h"
 
 
 namespace mu
@@ -20,24 +19,6 @@ public:
 
 	ExtensionDataPtrConst Value;
 	
-	//!
-	void Serialise(OutputArchive& Archive) const
-	{
-		uint32_t Version = 0;
-		Archive << Version;
-
-		Archive << Value;
-	}
-
-	//!
-	void Unserialise(InputArchive& Archive)
-	{
-		uint32_t Version;
-		Archive >> Version;
-		check(Version == 0);
-
-		Archive >> Value;
-	}
 };
 
 }
diff --git a/Engine/Plugins/Experimental/Mutable/Source/MutableTools/Private/MuT/NodeExtensionDataSwitchPrivate.h b/Engine/Plugins/Experimental/Mutable/Source/MutableTools/Private/MuT/NodeExtensionDataSwitchPrivate.h
--- a/Engine/Plugins/Experimental/Mutable/Source/MutableTools/Private/MuT/NodeExtensionDataSwitchPrivate.h
+++ b/Engine/Plugins/Experimental/Mutable/Source/MutableTools/Private/MuT/NodeExtensionDataSwitchPrivate.h
@@ -17,24 +17,5 @@ namespace mu
 
 		NodeScalarPtr Parameter;
 		TArray<NodeExtensionDataPtr> Options;
-
-		void Serialise(OutputArchive& Archive) const
-		{
-			uint32_t Version = 0;
-			Archive << Version;
-
-			Archive << Parameter;
-			Archive << Options;
-		}
-
-		void Unserialise(InputArchive& Archive)
-		{
-			uint32 Version;
-			Archive >> Version;
-			check(Version == 0);
-
-			Archive >> Parameter;
-			Archive >> Options;
-		}
 	};
 }
\ No newline at end of file
diff --git a/Engine/Plugins/Experimental/Mutable/Source/MutableTools/Private/MuT/NodeExtensionDataVariationPrivate.h b/Engine/Plugins/Experimental/Mutable/Source/MutableTools/Private/MuT/NodeExtensionDataVariationPrivate.h
--- a/Engine/Plugins/Experimental/Mutable/Source/MutableTools/Private/MuT/NodeExtensionDataVariationPrivate.h
+++ b/Engine/Plugins/Experimental/Mutable/Source/MutableTools/Private/MuT/NodeExtensionDataVariationPrivate.h
@@ -19,58 +19,8 @@ namespace mu
 		{
 			NodeExtensionDataPtr Value;
 			FString Tag;
-
-			//!
-			void Serialise(OutputArchive& InArch) const
-			{
-				uint32 Ver = 2;
-				InArch << Ver;
-
-				InArch << Tag;
-				InArch << Value;
-			}
-
-			void Unserialise(InputArchive& InArch)
-			{
-				uint32 Ver = 0;
-				InArch >> Ver;
-				check(Ver >= 1 && Ver <= 2);
-
-				if (Ver <= 1)
-				{
-					std::string Temp;
-					InArch >> Temp;
-					Tag = Temp.c_str();
-				}
-				else
-				{
-					InArch >> Tag;
-				}
-				InArch >> Value;
-			}
 		};
 
 		TArray<FVariation> Variations;
-
-		//!
-		void Serialise(OutputArchive& InArch) const
-		{
-			uint32_t Ver = 1;
-			InArch << Ver;
-
-			InArch << DefaultValue;
-			InArch << Variations;
-		}
-
-		//!
-		void Unserialise(InputArchive& InArch)
-		{
-			uint32_t Ver;
-			InArch >> Ver;
-			check(Ver == 1);
-
-			InArch >> DefaultValue;
-			InArch >> Variations;
-		}
 	};
 }
\ No newline at end of file
diff --git a/Engine/Plugins/Experimental/Mutable/Source/MutableTools/Private/MuT/NodeImage.cpp b/Engine/Plugins/Experimental/Mutable/Source/MutableTools/Private/MuT/NodeImage.cpp
--- a/Engine/Plugins/Experimental/Mutable/Source/MutableTools/Private/MuT/NodeImage.cpp
+++ b/Engine/Plugins/Experimental/Mutable/Source/MutableTools/Private/MuT/NodeImage.cpp
@@ -3,33 +3,6 @@
 #include "MuT/NodeImage.h"
 
 #include "Misc/AssertionMacros.h"
-#include "MuR/Serialisation.h"
-#include "MuR/SerialisationPrivate.h"
-#include "MuT/NodeImageBinarise.h"
-#include "MuT/NodeImageColourMap.h"
-#include "MuT/NodeImageConditional.h"
-#include "MuT/NodeImageConstant.h"
-#include "MuT/NodeImageFormat.h"
-#include "MuT/NodeImageGradient.h"
-#include "MuT/NodeImageInterpolate.h"
-#include "MuT/NodeImageInvert.h"
-#include "MuT/NodeImageLayer.h"
-#include "MuT/NodeImageLayerColour.h"
-#include "MuT/NodeImageLuminance.h"
-#include "MuT/NodeImageMipmap.h"
-#include "MuT/NodeImageMultiLayer.h"
-#include "MuT/NodeImageNormalComposite.h"
-#include "MuT/NodeImageParameter.h"
-#include "MuT/NodeImagePlainColour.h"
-#include "MuT/NodeImageProject.h"
-#include "MuT/NodeImageResize.h"
-#include "MuT/NodeImageSaturate.h"
-#include "MuT/NodeImageSwitch.h"
-#include "MuT/NodeImageSwizzle.h"
-#include "MuT/NodeImageTable.h"
-#include "MuT/NodeImageTransform.h"
-#include "MuT/NodeImageVariation.h"
-
 
 namespace mu
 {
@@ -40,67 +13,6 @@ namespace mu
 	static FNodeType s_nodeImageType = 	FNodeType( "NodeImage", Node::GetStaticType() );
 
 
-	//---------------------------------------------------------------------------------------------
-	void NodeImage::Serialise( const NodeImage* p, OutputArchive& arch )
-	{
-        uint32 ver = 0;
-		arch << ver;
-
-		arch << uint32(p->Type);
-		p->SerialiseWrapper(arch);
-    }
-
-        
-	//---------------------------------------------------------------------------------------------
-	NodeImagePtr NodeImage::StaticUnserialise( InputArchive& arch )
-	{
-        uint32 ver;
-		arch >> ver;
-		check( ver == 0 );
-
-        uint32 id;
-		arch >> id;
-
-		switch (id)
-		{
-		case 1 :  return NodeImageConstant::StaticUnserialise( arch ); break;
-		//case 2 :  return NodeImageDifference::StaticUnserialise( arch ); break;
-		//case 3 :  return NodeImageIdentity::StaticUnserialise( arch ); break;
-		case 4 :  return NodeImageInterpolate::StaticUnserialise( arch ); break;
-		case 6 :  return NodeImageSaturate::StaticUnserialise( arch ); break;
-		//case 8 :  return NodeSelectImage::StaticUnserialise( arch ); break;
-		case 9 :  return NodeImageTable::StaticUnserialise( arch ); break;
-		case 10 :  return NodeImageSwizzle::StaticUnserialise( arch ); break;
-		//case 11 :  return NodeImageSelectColour::StaticUnserialise( arch ); break;
-		case 12 :  return NodeImageColourMap::StaticUnserialise( arch ); break;
-		case 13 :  return NodeImageGradient::StaticUnserialise( arch ); break;
-		//case 14 :  return NodeImageVolumeLayer::StaticUnserialise( arch ); break;
-		case 17 :  return NodeImageBinarise::StaticUnserialise( arch ); break;
-		case 18 :  return NodeImageLuminance::StaticUnserialise( arch ); break;
-		case 19 :  return NodeImageLayer::StaticUnserialise( arch ); break;
-		case 20 :  return NodeImageLayerColour::StaticUnserialise( arch ); break;
-		case 21 :  return NodeImageResize::StaticUnserialise( arch ); break;
-		case 22 :  return NodeImagePlainColour::StaticUnserialise( arch ); break;
-		//case 23 :  return NodeImageInterpolate3::StaticUnserialise( arch ); break;
-        case 24 :  return NodeImageProject::StaticUnserialise( arch ); break;
-        case 25 :  return NodeImageMipmap::StaticUnserialise( arch ); break;
-        case 26 :  return NodeImageSwitch::StaticUnserialise( arch ); break;
-		case 27 :  return NodeImageConditional::StaticUnserialise( arch ); break;
-        case 28 :  return NodeImageFormat::StaticUnserialise( arch ); break;
-        case 29 :  return NodeImageParameter::StaticUnserialise( arch ); break;
-        case 30 :  return NodeImageMultiLayer::StaticUnserialise( arch ); break;
-        case 31 :  return NodeImageInvert::StaticUnserialise( arch ); break;
-        case 32 :  return NodeImageVariation::StaticUnserialise( arch ); break;
-        case 33 :  return NodeImageNormalComposite::StaticUnserialise( arch ); break;
-		case 34:  return NodeImageTransform::StaticUnserialise(arch); break;
-		//case 35:  return NodeImageReference::StaticUnserialise(arch); break;
-		default : check(false);
-		}
-
-		return 0;
-	}
-
-
 	//---------------------------------------------------------------------------------------------
 	const FNodeType* NodeImage::GetType() const
 	{
diff --git a/Engine/Plugins/Experimental/Mutable/Source/MutableTools/Private/MuT/NodeImageBinarisePrivate.h b/Engine/Plugins/Experimental/Mutable/Source/MutableTools/Private/MuT/NodeImageBinarisePrivate.h
--- a/Engine/Plugins/Experimental/Mutable/Source/MutableTools/Private/MuT/NodeImageBinarisePrivate.h
+++ b/Engine/Plugins/Experimental/Mutable/Source/MutableTools/Private/MuT/NodeImageBinarisePrivate.h
@@ -1,11 +1,9 @@
 // Copyright Epic Games, Inc. All Rights Reserved.
 
-
 #pragma once
 
 #include "MuT/NodeImagePrivate.h"
 #include "MuT/NodeImageBinarise.h"
-#include "MuT/AST.h"
 #include "MuT/NodeScalar.h"
 
 
@@ -21,27 +19,6 @@ namespace mu
 
 		NodeImagePtr m_pBase;
 		NodeScalarPtr m_pThreshold;
-
-		//!
-		void Serialise( OutputArchive& arch ) const
-		{
-            uint32_t ver = 0;
-			arch << ver;
-
-			arch << m_pBase;
-			arch << m_pThreshold;
-		}
-
-		//!
-		void Unserialise( InputArchive& arch )
-		{
-            uint32_t ver;
-			arch >> ver;
-			check(ver==0);
-
-			arch >> m_pBase;
-			arch >> m_pThreshold;
-		}
 	};
 
 
diff --git a/Engine/Plugins/Experimental/Mutable/Source/MutableTools/Private/MuT/NodeImageColourMapPrivate.h b/Engine/Plugins/Experimental/Mutable/Source/MutableTools/Private/MuT/NodeImageColourMapPrivate.h
--- a/Engine/Plugins/Experimental/Mutable/Source/MutableTools/Private/MuT/NodeImageColourMapPrivate.h
+++ b/Engine/Plugins/Experimental/Mutable/Source/MutableTools/Private/MuT/NodeImageColourMapPrivate.h
@@ -1,17 +1,13 @@
 // Copyright Epic Games, Inc. All Rights Reserved.
 
-
 #pragma once
 
 #include "MuT/NodeImagePrivate.h"
 #include "MuT/NodeImageColourMap.h"
-#include "MuT/AST.h"
-
 
 namespace mu
 {
 
-
 	class NodeImageColourMap::Private : public NodeImage::Private
 	{
 	public:
@@ -22,28 +18,6 @@ namespace mu
 		NodeImagePtr m_pMask;
 		NodeImagePtr m_pMap;
 
-		//!
-		void Serialise( OutputArchive& arch ) const
-		{
-            uint32_t ver = 0;
-			arch << ver;
-
-			arch << m_pBase;
-			arch << m_pMask;
-			arch << m_pMap;
-		}
-
-		//!
-		void Unserialise( InputArchive& arch )
-		{
-            uint32_t ver;
-			arch >> ver;
-			check(ver==0);
-
-			arch >> m_pBase;
-			arch >> m_pMask;
-			arch >> m_pMap;
-		}
 	};
 
 
diff --git a/Engine/Plugins/Experimental/Mutable/Source/MutableTools/Private/MuT/NodeImageConditionalPrivate.h b/Engine/Plugins/Experimental/Mutable/Source/MutableTools/Private/MuT/NodeImageConditionalPrivate.h
--- a/Engine/Plugins/Experimental/Mutable/Source/MutableTools/Private/MuT/NodeImageConditionalPrivate.h
+++ b/Engine/Plugins/Experimental/Mutable/Source/MutableTools/Private/MuT/NodeImageConditionalPrivate.h
@@ -6,8 +6,6 @@
 #include "MuT/NodeImageConditional.h"
 #include "MuT/NodeImagePrivate.h"
 #include "MuT/NodeBool.h"
-#include "MuT/AST.h"
-
 
 namespace mu
 {
@@ -23,28 +21,6 @@ namespace mu
         NodeImagePtr m_true;
         NodeImagePtr m_false;
 
-		//!
-		void Serialise( OutputArchive& arch ) const
-		{
-            uint32_t ver = 0;
-			arch << ver;
-
-            arch << m_parameter;
-            arch << m_true;
-            arch << m_false;
-        }
-
-		//!
-		void Unserialise( InputArchive& arch )
-		{
-            uint32_t ver;
-			arch >> ver;
-			check(ver==0);
-
-            arch >> m_parameter;
-            arch >> m_true;
-            arch >> m_false;
-        }
 	};
 
 
diff --git a/Engine/Plugins/Experimental/Mutable/Source/MutableTools/Private/MuT/NodeImageConstantPrivate.h b/Engine/Plugins/Experimental/Mutable/Source/MutableTools/Private/MuT/NodeImageConstantPrivate.h
--- a/Engine/Plugins/Experimental/Mutable/Source/MutableTools/Private/MuT/NodeImageConstantPrivate.h
+++ b/Engine/Plugins/Experimental/Mutable/Source/MutableTools/Private/MuT/NodeImageConstantPrivate.h
@@ -4,8 +4,6 @@
 
 #include "MuT/NodeImagePrivate.h"
 #include "MuT/NodeImageConstant.h"
-#include "MuT/AST.h"
-#include "MuT/StreamsPrivate.h"
 
 
 namespace mu
@@ -18,38 +16,6 @@ namespace mu
 
         Ptr<ResourceProxy<Image>> m_pProxy;
 
-		//!
-		void Serialise( OutputArchive& arch ) const
-		{
-			uint32 Ver = 0;
-			arch << Ver;
-
-			Ptr<const Image> ActualImage;
-			if (m_pProxy)
-			{
-				ActualImage = m_pProxy->Get();
-			}
-
-			arch << ActualImage;
-		}
-
-		//!
-		void Unserialise( InputArchive& arch )
-		{
-            uint32 Ver;
-			arch >> Ver;
-			check(Ver==0);
-
-
-			m_pProxy = arch.NewImageProxy();
-			if (!m_pProxy)
-			{
-				// Normal serialisation
-				Ptr<Image> ActualImage;
-				arch >> ActualImage;
-				m_pProxy = new ResourceProxyMemory<Image>(ActualImage.get());
-			}
-		}
 	};
 
 }
diff --git a/Engine/Plugins/Experimental/Mutable/Source/MutableTools/Private/MuT/NodeImageGradientPrivate.h b/Engine/Plugins/Experimental/Mutable/Source/MutableTools/Private/MuT/NodeImageGradientPrivate.h
--- a/Engine/Plugins/Experimental/Mutable/Source/MutableTools/Private/MuT/NodeImageGradientPrivate.h
+++ b/Engine/Plugins/Experimental/Mutable/Source/MutableTools/Private/MuT/NodeImageGradientPrivate.h
@@ -5,10 +5,8 @@
 #include "MuT/NodeColour.h"
 #include "MuT/NodeImagePrivate.h"
 #include "MuT/NodeImageGradient.h"
-#include "MuT/AST.h"
 #include "MuR/MutableMath.h"
 
-
 namespace mu
 {
 
@@ -23,28 +21,6 @@ namespace mu
 		NodeColourPtr m_pColour1;
 		FIntVector2 m_size = { 256,1 };		
 
-		//!
-		void Serialise( OutputArchive& arch ) const
-		{
-            uint32 ver = 0;
-			arch << ver;
-
-			arch << m_pColour0;
-			arch << m_pColour1;
-			arch << m_size;
-		}
-
-		//!
-		void Unserialise( InputArchive& arch )
-		{
-            uint32 ver;
-			arch >> ver;
-			check(ver==0);
-
-			arch >> m_pColour0;
-			arch >> m_pColour1;
-			arch >> m_size;
-		}
 	};
 
 
diff --git a/Engine/Plugins/Experimental/Mutable/Source/MutableTools/Private/MuT/NodeImageLuminancePrivate.h b/Engine/Plugins/Experimental/Mutable/Source/MutableTools/Private/MuT/NodeImageLuminancePrivate.h
--- a/Engine/Plugins/Experimental/Mutable/Source/MutableTools/Private/MuT/NodeImageLuminancePrivate.h
+++ b/Engine/Plugins/Experimental/Mutable/Source/MutableTools/Private/MuT/NodeImageLuminancePrivate.h
@@ -16,25 +16,6 @@ namespace mu
 		static FNodeType s_type;
 
 		NodeImagePtr m_pSource;
-
-		//!
-		void Serialise( OutputArchive& arch ) const
-		{
-            uint32_t ver = 0;
-			arch << ver;
-
-			arch << m_pSource;
-		}
-
-		//!
-		void Unserialise( InputArchive& arch )
-		{
-            uint32_t ver;
-			arch >> ver;
-			check(ver==0);
-
-			arch >> m_pSource;
-		}
 	};
 
 
diff --git a/Engine/Plugins/Experimental/Mutable/Source/MutableTools/Private/MuT/NodeImageNormalCompositePrivate.h b/Engine/Plugins/Experimental/Mutable/Source/MutableTools/Private/MuT/NodeImageNormalCompositePrivate.h
--- a/Engine/Plugins/Experimental/Mutable/Source/MutableTools/Private/MuT/NodeImageNormalCompositePrivate.h
+++ b/Engine/Plugins/Experimental/Mutable/Source/MutableTools/Private/MuT/NodeImageNormalCompositePrivate.h
@@ -4,7 +4,6 @@
 
 #include "MuT/NodeImagePrivate.h"
 #include "MuT/NodeImageNormalComposite.h"
-#include "MuT/AST.h"
 
 
 namespace mu
@@ -21,30 +20,6 @@ namespace mu
 		float m_power;
 		ECompositeImageMode m_mode;
 
-		//!
-		void Serialise( OutputArchive& arch ) const
-		{
-            uint32_t ver = 0;
-			arch << ver;
-
-			arch << m_pBase;
-			arch << m_pNormal;
-			arch << m_power;
-			arch << m_mode;
-		}
-
-		//!
-		void Unserialise( InputArchive& arch )
-		{
-            uint32_t ver;
-			arch >> ver;
-			check(ver==0);
-
-			arch >> m_pBase;
-			arch >> m_pNormal;
-			arch >> m_power;
-			arch >> m_mode;
-		}
 	};
 
 }
diff --git a/Engine/Plugins/Experimental/Mutable/Source/MutableTools/Private/MuT/NodeImageParameterPrivate.h b/Engine/Plugins/Experimental/Mutable/Source/MutableTools/Private/MuT/NodeImageParameterPrivate.h
--- a/Engine/Plugins/Experimental/Mutable/Source/MutableTools/Private/MuT/NodeImageParameterPrivate.h
+++ b/Engine/Plugins/Experimental/Mutable/Source/MutableTools/Private/MuT/NodeImageParameterPrivate.h
@@ -21,50 +21,6 @@ namespace mu
 		FString m_uid;
 
 		TArray<Ptr<NodeRange>> m_ranges;
-
-		//!
-		void Serialise( OutputArchive& arch ) const
-		{
-            uint32 ver = 4;
-			arch << ver;
-
-			arch << m_defaultValue;
-			arch << m_name;
-			arch << m_uid;
-			arch << m_ranges;
-		}
-
-		//!
-		void Unserialise( InputArchive& arch )
-		{
-            uint32 ver;
-			arch >> ver;
-            check(ver <= 4);
-
-			if (ver >= 3)
-			{
-				arch >> m_defaultValue;
-			}
-			
-			if (ver <= 3)
-			{
-				std::string Temp;
-				arch >> Temp;
-				m_name = Temp.c_str();
-				arch >> Temp;
-				m_uid = Temp.c_str();
-			}
-			else
-			{
-				arch >> m_name;
-				arch >> m_uid;
-			}
-
-			if (ver >= 2)
-			{
-				arch >> m_ranges;
-			}
-		}
 	};
 
 }
diff --git a/Engine/Plugins/Experimental/Mutable/Source/MutableTools/Private/MuT/NodeImageSaturatePrivate.h b/Engine/Plugins/Experimental/Mutable/Source/MutableTools/Private/MuT/NodeImageSaturatePrivate.h
--- a/Engine/Plugins/Experimental/Mutable/Source/MutableTools/Private/MuT/NodeImageSaturatePrivate.h
+++ b/Engine/Plugins/Experimental/Mutable/Source/MutableTools/Private/MuT/NodeImageSaturatePrivate.h
@@ -5,8 +5,6 @@
 #include "MuT/NodeImagePrivate.h"
 #include "MuT/NodeImageSaturate.h"
 #include "MuT/NodeScalar.h"
-#include "MuT/AST.h"
-
 
 namespace mu
 {
@@ -19,27 +17,6 @@ namespace mu
 
 		NodeImagePtr m_pSource;
 		NodeScalarPtr m_pFactor;
-
-		//!
-		void Serialise( OutputArchive& arch ) const
-		{
-            uint32_t ver = 0;
-			arch << ver;
-
-			arch << m_pSource;
-			arch << m_pFactor;
-		}
-
-		//!
-		void Unserialise( InputArchive& arch )
-		{
-            uint32_t ver;
-			arch >> ver;
-			check(ver==0);
-
-			arch >> m_pSource;
-			arch >> m_pFactor;
-		}
 	};
 
 
diff --git a/Engine/Plugins/Experimental/Mutable/Source/MutableTools/Private/MuT/NodeImageTablePrivate.h b/Engine/Plugins/Experimental/Mutable/Source/MutableTools/Private/MuT/NodeImageTablePrivate.h
--- a/Engine/Plugins/Experimental/Mutable/Source/MutableTools/Private/MuT/NodeImageTablePrivate.h
+++ b/Engine/Plugins/Experimental/Mutable/Source/MutableTools/Private/MuT/NodeImageTablePrivate.h
@@ -5,7 +5,6 @@
 #include "MuT/NodePrivate.h"
 #include "MuT/NodeImageTable.h"
 #include "MuT/TablePrivate.h"
-#include "MuT/AST.h"
 
 
 namespace mu
@@ -26,89 +25,6 @@ namespace mu
 		bool bNoneOption = false;
 		FString DefaultRowName;
 
-		//!
-		void Serialise( OutputArchive& arch ) const
-		{
-            uint32_t ver = 6;
-			arch << ver;
-
-			arch << ParameterName;
-			arch << Table;
-			arch << ColumnName;
-			arch << MaxTextureSize;
-
-			// FImageDesc
-			{
-				arch << ReferenceImageDesc.m_size;
-				arch << uint8(ReferenceImageDesc.m_format);
-				arch << ReferenceImageDesc.m_lods;
-			}
-
-			arch << bNoneOption;
-			arch << DefaultRowName;
-		}
-
-		//!
-		void Unserialise( InputArchive& arch )
-		{
-            uint32_t ver;
-			arch >> ver;
-			check(ver>=1 && ver<= 6);
-
-			if (ver == 1)
-			{
-				std::string Temp;
-				arch >> Temp;
-				ParameterName = Temp.c_str();
-			}
-			else
-			{
-				arch >> ParameterName;
-			}
-
-			arch >> Table;
-
-			if (ver == 1)
-			{
-				std::string Temp;
-				arch >> Temp;
-				ColumnName = Temp.c_str();
-			}
-			else
-			{
-				arch >> ColumnName;
-			}
-
-			if (ver >= 3)
-			{
-				arch >> MaxTextureSize;
-			}
-			else
-			{
-				bNoneOption = Table->GetPrivate()->bNoneOption_DEPRECATED;
-			}
-
-			if (ver >= 4)
-			{
-				arch >> ReferenceImageDesc.m_size;
-				
-				uint8 AuxFromat = 0;
-				arch >> AuxFromat;
-				ReferenceImageDesc.m_format = EImageFormat(AuxFromat);
-				
-				arch >> ReferenceImageDesc.m_lods;
-			}
-			
-			if (ver >= 5)
-			{
-				arch >> bNoneOption;
-			}
-
-			if (ver >= 6)
-			{
-				arch >> DefaultRowName;
-			}
-		}
 	};
 
 }
diff --git a/Engine/Plugins/Experimental/Mutable/Source/MutableTools/Private/MuT/NodeImageTransformPrivate.h b/Engine/Plugins/Experimental/Mutable/Source/MutableTools/Private/MuT/NodeImageTransformPrivate.h
--- a/Engine/Plugins/Experimental/Mutable/Source/MutableTools/Private/MuT/NodeImageTransformPrivate.h
+++ b/Engine/Plugins/Experimental/Mutable/Source/MutableTools/Private/MuT/NodeImageTransformPrivate.h
@@ -6,7 +6,6 @@
 #include "MuR/ImagePrivate.h"
 #include "MuT/NodeImagePrivate.h"
 #include "MuT/NodeImageTransform.h"
-#include "MuT/AST.h"
 #include "MuT/NodeScalar.h"
 
 namespace mu
@@ -34,58 +33,6 @@ namespace mu
 		uint8 UnusedPadding[sizeof(NodeImagePtr) - sizeof(bool)] = {0}; 
 		static_assert(sizeof(NodeImagePtr) - sizeof(bool) >= 1);
 
-
-		//!
-		void Serialise(OutputArchive& arch) const
-		{
-            uint32 ver = 3;
-			arch << ver;
-
-			arch << m_pBase;
-			arch << m_pOffsetX;
-			arch << m_pOffsetY;
-			arch << m_pScaleX;
-			arch << m_pScaleY;
-			arch << m_pRotation;
-			arch << static_cast<uint32>(AddressMode);
-			arch << SizeX;
-			arch << SizeY;
-			arch << bKeepAspectRatio;
-		}
-
-		//!
-		void Unserialise(InputArchive& arch)
-		{
-            uint32 ver;
-			arch >> ver;
-			check(ver <= 3);
-
-			arch >> m_pBase;
-			arch >> m_pOffsetX;
-			arch >> m_pOffsetY;
-			arch >> m_pScaleX;
-			arch >> m_pScaleY;
-			arch >> m_pRotation;
-
-			uint32 AddressModeValue = static_cast<uint32>(EAddressMode::Wrap);
-			if (ver >= 1)
-			{
-				arch >> AddressModeValue;
-			}
-
-			AddressMode = static_cast<EAddressMode>(AddressModeValue);
-
-			if (ver >= 2)
-			{
-				arch >> SizeX;
-				arch >> SizeY;
-			}
-
-			if (ver >= 3)
-			{
-				arch >> bKeepAspectRatio;
-			}
-		}
 	};
 
 
diff --git a/Engine/Plugins/Experimental/Mutable/Source/MutableTools/Private/MuT/NodeImageVariationPrivate.h b/Engine/Plugins/Experimental/Mutable/Source/MutableTools/Private/MuT/NodeImageVariationPrivate.h
--- a/Engine/Plugins/Experimental/Mutable/Source/MutableTools/Private/MuT/NodeImageVariationPrivate.h
+++ b/Engine/Plugins/Experimental/Mutable/Source/MutableTools/Private/MuT/NodeImageVariationPrivate.h
@@ -5,7 +5,6 @@
 #include "MuT/NodePrivate.h"
 #include "MuT/NodeImageVariation.h"
 
-
 namespace mu
 {
 
@@ -25,59 +24,9 @@ namespace mu
         {
             NodeImagePtr m_image;
 			FString m_tag;
-
-            //!
-            void Serialise( OutputArchive& arch ) const
-            {
-                uint32_t ver = 2;
-                arch << ver;
-
-                arch << m_tag;
-                arch << m_image;
-            }
-
-            void Unserialise( InputArchive& arch )
-            {
-                uint32_t ver = 0;
-                arch >> ver;
-				check(ver >= 1 && ver <= 2);
-
-				if (ver <= 1)
-				{
-					std::string Temp;
-					arch >> Temp;
-					m_tag = Temp.c_str();
-				}
-				else
-				{
-					arch >> m_tag;
-				}
-				arch >> m_image;
-            }
         };
 
 		TArray<FVariation> m_variations;
-
-        //!
-        void Serialise( OutputArchive& arch ) const
-        {
-            uint32_t ver = 1;
-            arch << ver;
-
-            arch << m_defaultImage;
-            arch << m_variations;
-        }
-
-        //!
-        void Unserialise( InputArchive& arch )
-        {
-            uint32_t ver;
-            arch >> ver;
-			check(ver == 1);
-
-            arch >> m_defaultImage;
-            arch >> m_variations;
-        }
     };
 
 } // namespace mu
diff --git a/Engine/Plugins/Experimental/Mutable/Source/MutableTools/Private/MuT/NodeLayout.cpp b/Engine/Plugins/Experimental/Mutable/Source/MutableTools/Private/MuT/NodeLayout.cpp
--- a/Engine/Plugins/Experimental/Mutable/Source/MutableTools/Private/MuT/NodeLayout.cpp
+++ b/Engine/Plugins/Experimental/Mutable/Source/MutableTools/Private/MuT/NodeLayout.cpp
@@ -9,64 +9,15 @@
 #include "MuR/MeshBufferSet.h"
 #include "MuR/MeshPrivate.h"
 #include "MuR/MutableMath.h"
-#include "MuR/Serialisation.h"
-#include "MuR/SerialisationPrivate.h"
 #include "MuT/NodeLayoutPrivate.h"
-#include "MuT/NodePrivate.h"
 
 
 namespace mu
 {
 
-	//---------------------------------------------------------------------------------------------
-	//---------------------------------------------------------------------------------------------
-	//---------------------------------------------------------------------------------------------
-	// Static initialisation
-	//---------------------------------------------------------------------------------------------
 	static FNodeType s_nodeLayoutType = FNodeType( "NodeLayout", Node::GetStaticType() );
 
 
-	//---------------------------------------------------------------------------------------------
-	void NodeLayout::Serialise( const NodeLayout* p, OutputArchive& arch )
-	{
-        uint32 ver = 0;
-		arch << ver;
-
-	#define SERIALISE_CHILDREN( C, ID ) \
-		( p->GetType()==C::GetStaticType() )					\
-		{ 														\
-			const C* pTyped = static_cast<const C*>(p);			\
-            arch << (uint32_t)ID;								\
-			C::Serialise( pTyped, arch );						\
-		}														\
-
-		if SERIALISE_CHILDREN( NodeLayoutBlocks				, 0 )
-		//else if SERIALISE_CHILDREN( NodeMeshIdentity		, 1 )
-
-#undef SERIALISE_CHILDREN
-	}
-
-
-	//---------------------------------------------------------------------------------------------
-	NodeLayoutPtr NodeLayout::StaticUnserialise( InputArchive& arch )
-	{
-        uint32 ver;
-		arch >> ver;
-		check( ver == 0 );
-
-        uint32 id;
-		arch >> id;
-
-		switch (id)
-		{
-		case 0 :  return NodeLayoutBlocks::StaticUnserialise( arch ); break;
-		default : check(false);
-		}
-
-		return 0;
-	}
-
-
 	//---------------------------------------------------------------------------------------------
 	const FNodeType* NodeLayout::GetType() const
 	{
@@ -81,14 +32,8 @@ namespace mu
 	}
 
 
-
-	//---------------------------------------------------------------------------------------------
-	//---------------------------------------------------------------------------------------------
-	//---------------------------------------------------------------------------------------------
-	// Static initialisation
 	//---------------------------------------------------------------------------------------------
-	FNodeType NodeLayoutBlocks::Private::s_type =
-			FNodeType( "LayoutBlocks", NodeLayout::GetStaticType() );
+	FNodeType NodeLayoutBlocks::Private::s_type = FNodeType( "LayoutBlocks", NodeLayout::GetStaticType() );
 
 
 	//---------------------------------------------------------------------------------------------
diff --git a/Engine/Plugins/Experimental/Mutable/Source/MutableTools/Private/MuT/NodeLayoutPrivate.h b/Engine/Plugins/Experimental/Mutable/Source/MutableTools/Private/MuT/NodeLayoutPrivate.h
--- a/Engine/Plugins/Experimental/Mutable/Source/MutableTools/Private/MuT/NodeLayoutPrivate.h
+++ b/Engine/Plugins/Experimental/Mutable/Source/MutableTools/Private/MuT/NodeLayoutPrivate.h
@@ -5,7 +5,6 @@
 #include "MuT/NodePrivate.h"
 #include "MuT/NodeLayout.h"
 #include "MuT/NodeImage.h"
-#include "MuT/AST.h"
 
 
 namespace mu
@@ -34,25 +33,6 @@ namespace mu
 
 		LayoutPtr m_pLayout;
 
-		//!
-		void Serialise( OutputArchive& arch ) const
-		{
-            uint32_t ver = 0;
-			arch << ver;
-
-			arch << m_pLayout;
-		}
-
-		//!
-		void Unserialise( InputArchive& arch )
-		{
-            uint32_t ver;
-			arch >> ver;
-			check(ver==0);
-
-			arch >> m_pLayout;
-		}
-
         // NodeLayout::Private interface
         Layout* GetLayout() override
         {
diff --git a/Engine/Plugins/Experimental/Mutable/Source/MutableTools/Private/MuT/NodeMesh.cpp b/Engine/Plugins/Experimental/Mutable/Source/MutableTools/Private/MuT/NodeMesh.cpp
--- a/Engine/Plugins/Experimental/Mutable/Source/MutableTools/Private/MuT/NodeMesh.cpp
+++ b/Engine/Plugins/Experimental/Mutable/Source/MutableTools/Private/MuT/NodeMesh.cpp
@@ -3,25 +3,6 @@
 #include "MuT/NodeMesh.h"
 
 #include "Misc/AssertionMacros.h"
-#include "MuR/Serialisation.h"
-#include "MuR/SerialisationPrivate.h"
-#include "MuT/NodeMeshApplyPose.h"
-#include "MuT/NodeMeshClipDeform.h"
-#include "MuT/NodeMeshClipMorphPlane.h"
-#include "MuT/NodeMeshClipWithMesh.h"
-#include "MuT/NodeMeshConstant.h"
-#include "MuT/NodeMeshFormat.h"
-#include "MuT/NodeMeshFragment.h"
-#include "MuT/NodeMeshGeometryOperation.h"
-#include "MuT/NodeMeshInterpolate.h"
-#include "MuT/NodeMeshMakeMorph.h"
-#include "MuT/NodeMeshMorph.h"
-#include "MuT/NodeMeshReshape.h"
-#include "MuT/NodeMeshSwitch.h"
-#include "MuT/NodeMeshTable.h"
-#include "MuT/NodeMeshTangents.h"
-#include "MuT/NodeMeshTransform.h"
-#include "MuT/NodeMeshVariation.h"
 
 
 namespace mu
@@ -33,53 +14,6 @@ namespace mu
 	static FNodeType s_nodeMeshType = FNodeType( "NodeMesh", Node::GetStaticType() );
 
 
-	//---------------------------------------------------------------------------------------------
-	void NodeMesh::Serialise( const NodeMesh* p, OutputArchive& arch )
-	{
-        uint32 ver = 0;
-		arch << ver;
-
-		arch << uint32(p->Type);
-		p->SerialiseWrapper( arch );
-    }
-
-
-    //---------------------------------------------------------------------------------------------
-	NodeMeshPtr NodeMesh::StaticUnserialise( InputArchive& arch )
-	{
-        uint32 ver;
-		arch >> ver;
-		check( ver == 0 );
-
-        uint32 id;
-		arch >> id;
-
-		switch (EType(id))
-		{
-		case EType::Constant:		return NodeMeshConstant::StaticUnserialise( arch ); break;
-		case EType::Interpolate:	return NodeMeshInterpolate::StaticUnserialise( arch ); break;
-		case EType::Table:			return NodeMeshTable::StaticUnserialise( arch ); break;
-		case EType::Format:			return NodeMeshFormat::StaticUnserialise( arch ); break;
-		case EType::Tangents:		return NodeMeshTangents::StaticUnserialise( arch ); break;
-		case EType::Morph:			return NodeMeshMorph::StaticUnserialise( arch ); break;
-		case EType::MakeMorph:		return NodeMeshMakeMorph::StaticUnserialise( arch ); break;
-		case EType::Switch:			return NodeMeshSwitch::StaticUnserialise( arch ); break;
-        case EType::Fragment:		return NodeMeshFragment::StaticUnserialise( arch ); break;
-		case EType::Transform:		return NodeMeshTransform::StaticUnserialise( arch ); break;
-		case EType::ClipMorphPlane: return NodeMeshClipMorphPlane::StaticUnserialise( arch ); break;
-        case EType::ClipWithMesh:	return NodeMeshClipWithMesh::StaticUnserialise( arch ); break;
-        case EType::ApplyPose:		return NodeMeshApplyPose::StaticUnserialise( arch ); break;
-		case EType::Variation:		return NodeMeshVariation::StaticUnserialise( arch ); break;
-		case EType::GeometryOperation:	return NodeMeshGeometryOperation::StaticUnserialise( arch ); break;
-		case EType::Reshape:			return NodeMeshReshape::StaticUnserialise( arch ); break;
-		case EType::ClipDeform:			return NodeMeshClipDeform::StaticUnserialise( arch ); break;
-		default : check(false);
-		}
-
-		return 0;
-	}
-
-
 	//---------------------------------------------------------------------------------------------
 	const FNodeType* NodeMesh::GetType() const
 	{
diff --git a/Engine/Plugins/Experimental/Mutable/Source/MutableTools/Private/MuT/NodeMeshApplyPosePrivate.h b/Engine/Plugins/Experimental/Mutable/Source/MutableTools/Private/MuT/NodeMeshApplyPosePrivate.h
--- a/Engine/Plugins/Experimental/Mutable/Source/MutableTools/Private/MuT/NodeMeshApplyPosePrivate.h
+++ b/Engine/Plugins/Experimental/Mutable/Source/MutableTools/Private/MuT/NodeMeshApplyPosePrivate.h
@@ -5,7 +5,6 @@
 #include "MuT/NodeMeshApplyPose.h"
 #include "MuT/NodeScalar.h"
 #include "MuT/NodeMeshPrivate.h"
-#include "MuT/AST.h"
 
 namespace mu
 {
@@ -20,28 +19,6 @@ namespace mu
         NodeMeshPtr m_pBase;
         NodeMeshPtr m_pPose;
 
-		//!
-		void Serialise( OutputArchive& arch ) const
-		{
-            uint32_t ver = 0;
-			arch << ver;
-
-            arch << m_pBase;
-            arch << m_pPose;
-		}
-
-		//!
-		void Unserialise( InputArchive& arch )
-		{
-            uint32_t ver;
-			arch >> ver;
-			check(ver==0);
-
-            arch >> m_pBase;
-            arch >> m_pPose;
-        }
-
-
         // NodeMesh::Private interface
         NodeLayoutPtr GetLayout( int index ) const override;
 
diff --git a/Engine/Plugins/Experimental/Mutable/Source/MutableTools/Private/MuT/NodeMeshClipDeformPrivate.h b/Engine/Plugins/Experimental/Mutable/Source/MutableTools/Private/MuT/NodeMeshClipDeformPrivate.h
--- a/Engine/Plugins/Experimental/Mutable/Source/MutableTools/Private/MuT/NodeMeshClipDeformPrivate.h
+++ b/Engine/Plugins/Experimental/Mutable/Source/MutableTools/Private/MuT/NodeMeshClipDeformPrivate.h
@@ -6,7 +6,6 @@
 #include "MuT/NodeScalar.h"
 #include "MuT/NodeMeshPrivate.h"
 #include "MuT/NodeImagePrivate.h"
-#include "MuT/AST.h"
 
 namespace mu
 {
@@ -21,29 +20,6 @@ namespace mu
 		Ptr<NodeMesh> m_pClipShape;
 		Ptr<NodeImage> m_pShapeWeights;
 
-        //!
-		void Serialise( OutputArchive& arch ) const
-		{
-            uint32_t ver = 0;
-			arch << ver;
-
-			arch << m_pBaseMesh;
-			arch << m_pClipShape;
-			arch << m_pShapeWeights;
-        }
-
-		//!
-		void Unserialise( InputArchive& arch )
-		{
-            uint32_t ver;
-			arch >> ver;
-			check(ver == 0);
-
-			arch >> m_pBaseMesh;
-			arch >> m_pClipShape;
-			arch >> m_pShapeWeights;	
-		}
-
 		// NodeMesh::Private interface
         Ptr<NodeLayout> GetLayout( int index ) const override;
 
diff --git a/Engine/Plugins/Experimental/Mutable/Source/MutableTools/Private/MuT/NodeMeshClipMorphPlanePrivate.h b/Engine/Plugins/Experimental/Mutable/Source/MutableTools/Private/MuT/NodeMeshClipMorphPlanePrivate.h
--- a/Engine/Plugins/Experimental/Mutable/Source/MutableTools/Private/MuT/NodeMeshClipMorphPlanePrivate.h
+++ b/Engine/Plugins/Experimental/Mutable/Source/MutableTools/Private/MuT/NodeMeshClipMorphPlanePrivate.h
@@ -4,7 +4,6 @@
 
 #include "MuT/NodeMeshPrivate.h"
 #include "MuT/NodeMeshClipMorphPlane.h"
-#include "MuT/AST.h"
 
 #include "MuR/MutableMath.h"
 #include "MuR/Skeleton.h"
@@ -68,81 +67,6 @@ namespace mu
 		// means no limit.
 		float m_maxEffectRadius;
 
-		//!
-		void Serialise( OutputArchive& arch ) const
-		{
-            uint32 ver = 6;
-			arch << ver;
-
-			arch << m_pSource;
-			arch << m_origin;
-			arch << m_normal;
-			arch << m_dist;
-			arch << m_factor;
-			arch << m_radius1;
-			arch << m_radius2;
-			arch << m_rotation;
-			arch << m_vertexSelectionType;
-			arch << m_selectionBoxOrigin;
-			arch << m_selectionBoxRadius;
-			arch << m_tags;
-			arch << m_vertexSelectionBone;
-			arch << m_maxEffectRadius;
-		}
-
-		//!
-		void Unserialise( InputArchive& arch )
-		{
-            uint32 ver;
-			arch >> ver;
-            check(ver>=3&&ver<=5);
-
-			arch >> m_pSource;
-			arch >> m_origin;
-			arch >> m_normal;
-			arch >> m_dist;
-			arch >> m_factor;
-			arch >> m_radius1;
-			arch >> m_radius2;
-			arch >> m_rotation;
-			arch >> m_vertexSelectionType;
-			arch >> m_selectionBoxOrigin;
-			arch >> m_selectionBoxRadius;
-
-			if (ver >= 5)
-			{
-				arch >> m_tags;
-			}
-			else
-			{
-				TArray<std::string> Temp;
-				arch >> Temp;
-				m_tags.SetNum(Temp.Num());
-				for (int32 i=0;i<Temp.Num();++i)
-				{
-					m_tags[i] = Temp[i].c_str();
-				}
-			}
-
-			if (ver >= 6)
-			{
-				arch >> m_vertexSelectionBone;
-			}
-			else if (ver >= 4)
-			{
-				uint16 BoneId = 0;
-				arch >> BoneId;
-				m_vertexSelectionBone.Id = BoneId;
-			}
-			else
-			{
-				string OldVertexSelectionBone;
-				arch >> OldVertexSelectionBone;
-				m_vertexSelectionBone.Id = 0;
-			}
-			arch >> m_maxEffectRadius;
-		}
-
         // NodeMesh::Private interface
         NodeLayoutPtr GetLayout( int index ) const override;
 
diff --git a/Engine/Plugins/Experimental/Mutable/Source/MutableTools/Private/MuT/NodeMeshClipWithMeshPrivate.h b/Engine/Plugins/Experimental/Mutable/Source/MutableTools/Private/MuT/NodeMeshClipWithMeshPrivate.h
--- a/Engine/Plugins/Experimental/Mutable/Source/MutableTools/Private/MuT/NodeMeshClipWithMeshPrivate.h
+++ b/Engine/Plugins/Experimental/Mutable/Source/MutableTools/Private/MuT/NodeMeshClipWithMeshPrivate.h
@@ -4,7 +4,6 @@
 
 #include "MuT/NodeMeshPrivate.h"
 #include "MuT/NodeMeshClipWithMesh.h"
-#include "MuT/AST.h"
 
 #include "MuR/MutableMath.h"
 
@@ -24,43 +23,6 @@ namespace mu
 
 		TArray<FString> m_tags;
 
-		//!
-		void Serialise( OutputArchive& arch ) const
-		{
-            uint32 ver = 2;
-			arch << ver;
-
-			arch << m_pSource;
-			arch << m_pClipMesh;
-			arch << m_tags;
-		}
-
-		//!
-		void Unserialise( InputArchive& arch )
-		{
-            uint32 ver;
-			arch >> ver;
-			check(ver<=2)
-
-			arch >> m_pSource;
-			arch >> m_pClipMesh;
-
-			if (ver <= 1)
-			{
-				TArray<std::string> Temp;
-				arch >> Temp;
-				m_tags.SetNum(Temp.Num());
-				for (int32 i = 0; i < Temp.Num(); ++i)
-				{
-					m_tags[i] = Temp[i].c_str();
-				}
-			}
-			else
-			{
-				arch >> m_tags;
-			}
-		}
-
 		// NodeMesh::Private interface
         NodeLayoutPtr GetLayout( int index ) const override;
 
diff --git a/Engine/Plugins/Experimental/Mutable/Source/MutableTools/Private/MuT/NodeMeshGeometryOperationPrivate.h b/Engine/Plugins/Experimental/Mutable/Source/MutableTools/Private/MuT/NodeMeshGeometryOperationPrivate.h
--- a/Engine/Plugins/Experimental/Mutable/Source/MutableTools/Private/MuT/NodeMeshGeometryOperationPrivate.h
+++ b/Engine/Plugins/Experimental/Mutable/Source/MutableTools/Private/MuT/NodeMeshGeometryOperationPrivate.h
@@ -5,7 +5,6 @@
 #include "MuT/NodeMeshGeometryOperation.h"
 #include "MuT/NodeScalar.h"
 #include "MuT/NodeMeshPrivate.h"
-#include "MuT/AST.h"
 
 namespace mu
 {
@@ -22,31 +21,6 @@ namespace mu
 		NodeScalarPtr m_pScalarA;
 		NodeScalarPtr m_pScalarB;
 
-        //!
-		void Serialise( OutputArchive& arch ) const
-		{
-            uint32_t ver = 0;
-			arch << ver;
-
-			arch << m_pMeshA;
-			arch << m_pMeshB;
-			arch << m_pScalarA;
-            arch << m_pScalarB;
-        }
-
-		//!
-		void Unserialise( InputArchive& arch )
-		{
-            uint32_t ver;
-			arch >> ver;
-			check(ver==0);
-
-			arch >> m_pMeshA;
-			arch >> m_pMeshB;
-			arch >> m_pScalarA;
-			arch >> m_pScalarB;
-		}
-
 		// NodeMesh::Private interface
         NodeLayoutPtr GetLayout( int index ) const override;
 
diff --git a/Engine/Plugins/Experimental/Mutable/Source/MutableTools/Private/MuT/NodeMeshInterpolatePrivate.h b/Engine/Plugins/Experimental/Mutable/Source/MutableTools/Private/MuT/NodeMeshInterpolatePrivate.h
--- a/Engine/Plugins/Experimental/Mutable/Source/MutableTools/Private/MuT/NodeMeshInterpolatePrivate.h
+++ b/Engine/Plugins/Experimental/Mutable/Source/MutableTools/Private/MuT/NodeMeshInterpolatePrivate.h
@@ -5,7 +5,6 @@
 #include "MuT/NodeMeshInterpolate.h"
 #include "MuT/NodeScalar.h"
 #include "MuT/NodeMeshPrivate.h"
-#include "MuT/AST.h"
 
 
 namespace mu
@@ -39,58 +38,10 @@ namespace mu
 
 			EMeshBufferSemantic semantic;
             int32_t semanticIndex;
-
-			//-------------------------------------------------------------------------------------
-			void Serialise( OutputArchive& arch ) const
-			{
-                const int32_t ver = 0;
-				arch << ver;
-
-				arch << semantic;
-				arch << semanticIndex;
-			}
-
-			//-------------------------------------------------------------------------------------
-			void Unserialise( InputArchive& arch )
-			{
-                int32_t ver = 0;
-				arch >> ver;
-				check(ver == 0);
-
-				arch >> semantic;
-				arch >> semanticIndex;
-			}
 		};
 
 		TArray<CHANNEL> m_channels;
 
-		//!
-		void Serialise( OutputArchive& arch ) const
-		{
-            uint32_t ver = 1;
-			arch << ver;
-
-			arch << m_pFactor;
-			arch << m_targets;
-			arch << m_channels;
-		}
-
-		//!
-		void Unserialise( InputArchive& arch )
-		{
-            uint32_t ver;
-			arch >> ver;
-            check(ver==1);
-
-			arch >> m_pFactor;
-			arch >> m_targets;
-
-			if (ver>=1)
-			{
-				arch >> m_channels;
-			}
-		}
-
 		// NodeMesh::Private interface
         NodeLayoutPtr GetLayout( int index ) const override;
 
diff --git a/Engine/Plugins/Experimental/Mutable/Source/MutableTools/Private/MuT/NodeMeshReshapePrivate.h b/Engine/Plugins/Experimental/Mutable/Source/MutableTools/Private/MuT/NodeMeshReshapePrivate.h
--- a/Engine/Plugins/Experimental/Mutable/Source/MutableTools/Private/MuT/NodeMeshReshapePrivate.h
+++ b/Engine/Plugins/Experimental/Mutable/Source/MutableTools/Private/MuT/NodeMeshReshapePrivate.h
@@ -5,7 +5,6 @@
 #include "MuT/NodeMeshReshape.h"
 #include "MuT/NodeScalar.h"
 #include "MuT/NodeMeshPrivate.h"
-#include "MuT/AST.h"
 #include "MuR/Skeleton.h"
 
 namespace mu
@@ -33,144 +32,6 @@ namespace mu
 
 		TArray<FBoneName> BonesToDeform;
 		TArray<FBoneName> PhysicsToDeform;
-        //!
-		void Serialise( OutputArchive& arch ) const
-		{
-            uint32 ver = 12;
-			arch << ver;
-
-			arch << BaseMesh;
-			arch << BaseShape;
-			arch << TargetShape;
-
-			arch << bReshapeVertices;
-			arch << bApplyLaplacian;
-			
-			arch << bReshapeSkeleton;
-			arch << BonesToDeform;
-
-			arch << bReshapePhysicsVolumes;
-			arch << PhysicsToDeform;
-
-			
-			arch << ColorRChannelUsage;
-			arch << ColorGChannelUsage;
-			arch << ColorBChannelUsage;
-			arch << ColorAChannelUsage;
-        }
-
-		//!
-		void Unserialise( InputArchive& arch )
-		{
-            uint32 ver;
-			arch >> ver;
-			check(ver <= 11);
-
-			arch >> BaseMesh;
-			arch >> BaseShape;
-			arch >> TargetShape;
-			
-			if (ver >= 8)
-			{
-				arch >> bReshapeVertices;
-			}
-
-			if (ver >= 11)
-			{
-				arch >> bApplyLaplacian;
-			}
-
-			arch >> bReshapeSkeleton;
-			if (ver <= 9)
-			{
-				bool bEnableRigidParts_DEPRECATED;
-				arch >> bEnableRigidParts_DEPRECATED;
-
-				if (bEnableRigidParts_DEPRECATED)
-				{
-					ColorRChannelUsage = ColorGChannelUsage = ColorBChannelUsage = ColorAChannelUsage = EVertexColorUsage::ReshapeClusterId;
-				}
-			}
-
-			if (ver <= 5)
-			{
-				bool bDeformAllBones_DEPRECATED;
-				arch >> bDeformAllBones_DEPRECATED;
-			}
-
-			if (ver >= 12)
-			{
-				arch >> BonesToDeform;
-			}
-			else
-			{
-				int32 NumBonesToDeform = 0;
-				if (ver >= 9)
-				{
-					TArray<uint16> BonesToDeform_DEPRECATED;
-					arch >> BonesToDeform_DEPRECATED;
-					NumBonesToDeform = BonesToDeform_DEPRECATED.Num();
-				}
-				else
-				{
-					TArray<string> BonesToDeform_DEPRECATED;
-					arch >> BonesToDeform_DEPRECATED;
-					NumBonesToDeform = BonesToDeform_DEPRECATED.Num();
-				}
-
-				BonesToDeform.SetNum(NumBonesToDeform);
-				for (int32 BoneIndex = 0; BoneIndex < NumBonesToDeform; ++BoneIndex)
-				{
-					BonesToDeform[BoneIndex].Id = BoneIndex;
-				}
-			}
-
-		  	if (ver >= 4)
-		 	{
-				arch >> bReshapePhysicsVolumes;
-		 	}
-
-			if (ver >= 5 && ver < 7)
-			{
-				bool bDeformAllPhysics_DEPRECATED;
-				arch >> bDeformAllPhysics_DEPRECATED;
-			}
-
-			if (ver >= 12)
-			{
-				arch >> PhysicsToDeform;
-			}
-			else
-			{
-				int32 NumPhysicsToDeform = 0;
-				if (ver >= 9)
-				{
-					TArray<uint16> PhysicsToDeform_DEPRECATED;
-					arch >> PhysicsToDeform_DEPRECATED;
-					NumPhysicsToDeform = PhysicsToDeform_DEPRECATED.Num();
-				}
-				else if (ver >= 5)
-				{
-					TArray<string> PhysicsToDeform_DEPRECATED;
-					arch >> PhysicsToDeform_DEPRECATED;
-					NumPhysicsToDeform = PhysicsToDeform_DEPRECATED.Num();
-				}
-
-				PhysicsToDeform.SetNum(NumPhysicsToDeform);
-				for (int32 BoneIndex = 0; BoneIndex < NumPhysicsToDeform; ++BoneIndex)
-				{
-					PhysicsToDeform[BoneIndex].Id = BoneIndex;
-				}
-			}
-
-			if (ver >= 10)
-			{
-				arch >> ColorRChannelUsage;
-				arch >> ColorGChannelUsage;
-				arch >> ColorBChannelUsage;
-				arch >> ColorAChannelUsage;
-			}
-		}
 
 		// NodeMesh::Private interface
         Ptr<NodeLayout> GetLayout( int index ) const override;
diff --git a/Engine/Plugins/Experimental/Mutable/Source/MutableTools/Private/MuT/NodeMeshSwitchPrivate.h b/Engine/Plugins/Experimental/Mutable/Source/MutableTools/Private/MuT/NodeMeshSwitchPrivate.h
--- a/Engine/Plugins/Experimental/Mutable/Source/MutableTools/Private/MuT/NodeMeshSwitchPrivate.h
+++ b/Engine/Plugins/Experimental/Mutable/Source/MutableTools/Private/MuT/NodeMeshSwitchPrivate.h
@@ -5,7 +5,6 @@
 #include "MuT/NodeMeshSwitch.h"
 #include "MuT/NodeScalar.h"
 #include "MuT/NodeMeshPrivate.h"
-#include "MuT/AST.h"
 
 
 namespace mu
@@ -21,28 +20,6 @@ namespace mu
 		NodeScalarPtr m_pParameter;
 		TArray<NodeMeshPtr> m_options;
 
-		//!
-		void Serialise( OutputArchive& arch ) const
-		{
-            uint32_t ver = 0;
-			arch << ver;
-
-			arch << m_pParameter;
-			arch << m_options;
-		}
-
-		//!
-		void Unserialise( InputArchive& arch )
-		{
-            uint32_t ver;
-			arch >> ver;
-			check(ver==0);
-
-			arch >> m_pParameter;
-			arch >> m_options;
-		}
-
-
 		// NodeMesh::Private interface
         NodeLayoutPtr GetLayout( int index ) const override;
 	};
diff --git a/Engine/Plugins/Experimental/Mutable/Source/MutableTools/Private/MuT/NodeMeshTangentsPrivate.h b/Engine/Plugins/Experimental/Mutable/Source/MutableTools/Private/MuT/NodeMeshTangentsPrivate.h
--- a/Engine/Plugins/Experimental/Mutable/Source/MutableTools/Private/MuT/NodeMeshTangentsPrivate.h
+++ b/Engine/Plugins/Experimental/Mutable/Source/MutableTools/Private/MuT/NodeMeshTangentsPrivate.h
@@ -4,7 +4,6 @@
 
 #include "MuT/NodeMeshPrivate.h"
 #include "MuT/NodeMeshTangents.h"
-#include "MuT/AST.h"
 
 
 namespace mu
@@ -19,25 +18,6 @@ namespace mu
 
 		NodeMeshPtr m_pSource;
 
-		//!
-		void Serialise( OutputArchive& arch ) const
-		{
-            uint32_t ver = 0;
-			arch << ver;
-
-			arch << m_pSource;
-		}
-
-		//!
-		void Unserialise( InputArchive& arch )
-		{
-            uint32_t ver;
-			arch >> ver;
-			check(ver==0);
-
-			arch >> m_pSource;
-		}
-
 		// NodeMesh::Private interface
         NodeLayoutPtr GetLayout( int index ) const override;
 
diff --git a/Engine/Plugins/Experimental/Mutable/Source/MutableTools/Private/MuT/NodeMeshTransformPrivate.h b/Engine/Plugins/Experimental/Mutable/Source/MutableTools/Private/MuT/NodeMeshTransformPrivate.h
--- a/Engine/Plugins/Experimental/Mutable/Source/MutableTools/Private/MuT/NodeMeshTransformPrivate.h
+++ b/Engine/Plugins/Experimental/Mutable/Source/MutableTools/Private/MuT/NodeMeshTransformPrivate.h
@@ -4,10 +4,7 @@
 
 #include "MuT/NodeMeshPrivate.h"
 #include "MuT/NodeMeshTransform.h"
-#include "MuT/AST.h"
-
 #include "MuR/MutableMath.h"
-#include "MuR/Serialisation.h"
 
 
 namespace mu
@@ -23,27 +20,6 @@ namespace mu
 		NodeMeshPtr Source;
 		FMatrix44f Transform;
 
-		//!
-		void Serialise( OutputArchive& arch ) const
-		{
-            uint32 ver = 0;
-			arch << ver;
-
-			arch << Source;
-            arch << Transform;
-		}
-
-		//!
-		void Unserialise( InputArchive& arch )
-		{
-            uint32 ver;
-			arch >> ver;
-			check(ver==0);
-
-			arch >> Source;
-            arch >> Transform;
-		}
-
 		// NodeMesh::Private interface
         NodeLayoutPtr GetLayout( int index ) const override;
 
diff --git a/Engine/Plugins/Experimental/Mutable/Source/MutableTools/Private/MuT/NodeMeshVariationPrivate.h b/Engine/Plugins/Experimental/Mutable/Source/MutableTools/Private/MuT/NodeMeshVariationPrivate.h
--- a/Engine/Plugins/Experimental/Mutable/Source/MutableTools/Private/MuT/NodeMeshVariationPrivate.h
+++ b/Engine/Plugins/Experimental/Mutable/Source/MutableTools/Private/MuT/NodeMeshVariationPrivate.h
@@ -5,7 +5,6 @@
 #include "MuT/NodePrivate.h"
 #include "MuT/NodeMeshVariation.h"
 
-
 namespace mu
 {
 
@@ -25,59 +24,9 @@ namespace mu
         {
             NodeMeshPtr m_mesh;
             FString m_tag;
-
-            //!
-            void Serialise( OutputArchive& arch ) const
-            {
-                uint32_t ver = 2;
-                arch << ver;
-
-                arch << m_tag;
-                arch << m_mesh;
-            }
-
-            void Unserialise( InputArchive& arch )
-            {
-                uint32_t ver = 0;
-                arch >> ver;
-				check(ver >= 1 && ver <= 2);
-
-				if (ver <= 1)
-				{
-					std::string Temp;
-					arch >> Temp;
-					m_tag = Temp.c_str();
-				}
-				else
-				{
-					arch >> m_tag;
-				}
-				arch >> m_mesh;
-            }
         };
 
         TArray<FVariation> m_variations;
-
-        //!
-        void Serialise( OutputArchive& arch ) const
-        {
-            uint32_t ver = 1;
-            arch << ver;
-
-            arch << m_defaultMesh;
-            arch << m_variations;
-        }
-
-        //!
-        void Unserialise( InputArchive& arch )
-        {
-            uint32_t ver;
-            arch >> ver;
-			check(ver == 1 );
-
-            arch >> m_defaultMesh;
-            arch >> m_variations;
-        }
     };
 
 } // namespace mu
diff --git a/Engine/Plugins/Experimental/Mutable/Source/MutableTools/Private/MuT/NodeModifier.cpp b/Engine/Plugins/Experimental/Mutable/Source/MutableTools/Private/MuT/NodeModifier.cpp
--- a/Engine/Plugins/Experimental/Mutable/Source/MutableTools/Private/MuT/NodeModifier.cpp
+++ b/Engine/Plugins/Experimental/Mutable/Source/MutableTools/Private/MuT/NodeModifier.cpp
@@ -4,14 +4,7 @@
 #include "MuT/NodeModifier.h"
 
 #include "Misc/AssertionMacros.h"
-#include "MuR/Serialisation.h"
-#include "MuR/SerialisationPrivate.h"
-#include "MuT/NodeModifierMeshClipDeform.h"
-#include "MuT/NodeModifierMeshClipMorphPlane.h"
-#include "MuT/NodeModifierMeshClipWithMesh.h"
-#include "MuT/NodeModifierMeshClipWithUVMask.h"
 #include "MuT/NodeModifierPrivate.h"
-#include "MuT/NodePrivate.h"
 
 
 namespace mu
@@ -32,51 +25,6 @@ namespace mu
 		return &s_nodeModifierType;
     }
 
-
-	void NodeModifier::Serialise( const NodeModifier* p, OutputArchive& arch )
-	{
-        uint32 ver = 0;
-		arch << ver;
-
-	#define SERIALISE_CHILDREN( C, ID ) \
-		( p->GetType()==C::GetStaticType() )					\
-		{ 														\
-			const C* pTyped = static_cast<const C*>(p);			\
-            arch << (uint32)ID;									\
-			C::Serialise( pTyped, arch );						\
-		}														\
-
-		if SERIALISE_CHILDREN(NodeModifierMeshClipMorphPlane, 			0 )
-		else if SERIALISE_CHILDREN(NodeModifierMeshClipWithMesh, 		1 )
-		else if SERIALISE_CHILDREN(NodeModifierMeshClipDeform, 		    2 )
-		else if SERIALISE_CHILDREN(NodeModifierMeshClipWithUVMask, 		3 )
-		else check(false);
-
-#undef SERIALISE_CHILDREN
-    }
-
-        
-	NodeModifierPtr NodeModifier::StaticUnserialise( InputArchive& arch )
-	{
-        uint32 ver;
-		arch >> ver;
-		check( ver == 0 );
-
-        uint32 id;
-		arch >> id;
-
-		switch (id)
-		{
-		case 0 :  return NodeModifierMeshClipMorphPlane::StaticUnserialise( arch ); break;
-		case 1 :  return NodeModifierMeshClipWithMesh::StaticUnserialise( arch ); break;
-		case 2 :  return NodeModifierMeshClipDeform::StaticUnserialise(arch); break;
-		case 3 :  return NodeModifierMeshClipWithUVMask::StaticUnserialise(arch); break;
-		default : check(false);
-		}
-
-		return 0;
-	}
-
 	void NodeModifier::AddTag(const FString& Value)
 	{
 		NodeModifier::Private* pD = static_cast<NodeModifier::Private*>(GetBasePrivate());
diff --git a/Engine/Plugins/Experimental/Mutable/Source/MutableTools/Private/MuT/NodeObject.cpp b/Engine/Plugins/Experimental/Mutable/Source/MutableTools/Private/MuT/NodeObject.cpp
--- a/Engine/Plugins/Experimental/Mutable/Source/MutableTools/Private/MuT/NodeObject.cpp
+++ b/Engine/Plugins/Experimental/Mutable/Source/MutableTools/Private/MuT/NodeObject.cpp
@@ -3,55 +3,14 @@
 #include "MuT/NodeObject.h"
 
 #include "Misc/AssertionMacros.h"
-#include "MuR/Serialisation.h"
-#include "MuR/SerialisationPrivate.h"
-#include "MuT/NodeObjectGroup.h"
-#include "MuT/NodeObjectNew.h"
 
 
 namespace mu
 {
 
-
-	//---------------------------------------------------------------------------------------------
-	// Static initialisation
-	//---------------------------------------------------------------------------------------------
 	FNodeType s_nodeObjectType = FNodeType( "NodeObject", Node::GetStaticType() );
 
 
-	//---------------------------------------------------------------------------------------------
-	void NodeObject::Serialise( const NodeObject* p, OutputArchive& arch )
-	{
-        uint32_t ver = 0;
-		arch << ver;
-
-		arch << uint32_t(p->Type);
-		p->SerialiseWrapper(arch);
-    }
-
-    //---------------------------------------------------------------------------------------------
-	NodeObjectPtr NodeObject::StaticUnserialise( InputArchive& arch )
-	{
-        uint32_t ver;
-		arch >> ver;
-		check( ver == 0 );
-
-        uint32_t id;
-		arch >> id;
-
-		switch (id)
-		{
-		case 0 :  return NodeObjectNew::StaticUnserialise( arch ); break;
-		case 1 :  return NodeObjectGroup::StaticUnserialise( arch ); break;
-        //case 2 :  return NodeObjectState::StaticUnserialise( arch ); break;
-        //case 3 :  return NodeObjectTransform::StaticUnserialise( arch ); break;
-        default : check(false);
-		}
-
-		return 0;
-	}
-
-
 	//---------------------------------------------------------------------------------------------
 	const FNodeType* NodeObject::GetType() const
 	{
diff --git a/Engine/Plugins/Experimental/Mutable/Source/MutableTools/Private/MuT/NodeProjector.cpp b/Engine/Plugins/Experimental/Mutable/Source/MutableTools/Private/MuT/NodeProjector.cpp
--- a/Engine/Plugins/Experimental/Mutable/Source/MutableTools/Private/MuT/NodeProjector.cpp
+++ b/Engine/Plugins/Experimental/Mutable/Source/MutableTools/Private/MuT/NodeProjector.cpp
@@ -3,52 +3,14 @@
 #include "MuT/NodeProjector.h"
 
 #include "Misc/AssertionMacros.h"
-#include "MuR/Serialisation.h"
-#include "MuR/SerialisationPrivate.h"
-#include "MuT/Node.h"
 
 
 namespace mu
 {
 
-	//---------------------------------------------------------------------------------------------
-	// Static initialisation
-	//---------------------------------------------------------------------------------------------
 	static FNodeType s_nodeProjectorType = FNodeType( "NodeProjector", Node::GetStaticType() );
 
 
-	//---------------------------------------------------------------------------------------------
-	void NodeProjector::Serialise( const NodeProjector* p, OutputArchive& arch )
-	{
-        uint32 ver = 0;
-		arch << ver;
-
-		arch << uint32_t(p->Type);
-		p->SerialiseWrapper(arch);
-	}
-
-        
-	//---------------------------------------------------------------------------------------------
-	NodeProjectorPtr NodeProjector::StaticUnserialise( InputArchive& arch )
-	{
-        uint32 ver;
-		arch >> ver;
-		check( ver == 0 );
-
-        uint32 id;
-		arch >> id;
-
-		switch (id)
-		{
-		case 0 :  return NodeProjectorConstant::StaticUnserialise( arch ); break;
-		case 1 :  return NodeProjectorParameter::StaticUnserialise( arch ); break;
-		default : check(false);
-		}
-
-		return 0;
-	}
-
-
 	//---------------------------------------------------------------------------------------------
 	const FNodeType* NodeProjector::GetType() const
 	{
diff --git a/Engine/Plugins/Experimental/Mutable/Source/MutableTools/Private/MuT/NodeRange.cpp b/Engine/Plugins/Experimental/Mutable/Source/MutableTools/Private/MuT/NodeRange.cpp
--- a/Engine/Plugins/Experimental/Mutable/Source/MutableTools/Private/MuT/NodeRange.cpp
+++ b/Engine/Plugins/Experimental/Mutable/Source/MutableTools/Private/MuT/NodeRange.cpp
@@ -4,11 +4,6 @@
 #include "MuT/NodeRange.h"
 
 #include "Misc/AssertionMacros.h"
-#include "MuR/Serialisation.h"
-#include "MuR/SerialisationPrivate.h"
-#include "MuT/Node.h"
-#include "MuT/NodeRangeFromScalar.h"
-
 
 namespace mu
 {
@@ -19,46 +14,6 @@ namespace mu
     static FNodeType s_nodeRangeType = 	FNodeType( "NodeRange", Node::GetStaticType() );
 
 
-	//---------------------------------------------------------------------------------------------
-    void NodeRange::Serialise( const NodeRange* p, OutputArchive& arch )
-	{
-        uint32_t ver = 0;
-		arch << ver;
-
-#define SERIALISE_CHILDREN( C, ID ) \
-		( p->GetType()==C::GetStaticType() )					\
-		{ 														\
-			const C* pTyped = static_cast<const C*>(p);			\
-            arch << (uint32_t)ID;								\
-			C::Serialise( pTyped, arch );						\
-		}														\
-
-        if SERIALISE_CHILDREN( NodeRangeFromScalar          , 0  )
-        else check(false);
-
-#undef SERIALISE_CHILDREN
-    }
-
-    //---------------------------------------------------------------------------------------------
-    NodeRangePtr NodeRange::StaticUnserialise( InputArchive& arch )
-	{
-        uint32 ver;
-		arch >> ver;
-		check( ver == 0 );
-
-        uint32 id;
-		arch >> id;
-
-		switch (id)
-		{
-        case 0 :  return NodeRangeFromScalar::StaticUnserialise( arch ); break;
-        default : check(false);
-		}
-
-        return nullptr;
-	}
-
-
 	//---------------------------------------------------------------------------------------------
     const FNodeType* NodeRange::GetType() const
 	{
diff --git a/Engine/Plugins/Experimental/Mutable/Source/MutableTools/Private/MuT/NodeRangeFromScalarPrivate.h b/Engine/Plugins/Experimental/Mutable/Source/MutableTools/Private/MuT/NodeRangeFromScalarPrivate.h
--- a/Engine/Plugins/Experimental/Mutable/Source/MutableTools/Private/MuT/NodeRangeFromScalarPrivate.h
+++ b/Engine/Plugins/Experimental/Mutable/Source/MutableTools/Private/MuT/NodeRangeFromScalarPrivate.h
@@ -20,37 +20,6 @@ namespace mu
 
         NodeScalarPtr m_pSize;
         FString m_name;
-
-		//!
-		void Serialise( OutputArchive& arch ) const
-		{
-            uint32_t ver = 1;
-			arch << ver;
-
-            arch << m_name;
-            arch << m_pSize;
-		}
-
-		//!
-		void Unserialise( InputArchive& arch )
-		{
-            uint32 ver;
-			arch >> ver;
-			check(ver<=1);
-
-			if (ver == 0)
-			{
-				std::string Temp;
-				arch >> Temp;
-				m_name = Temp.c_str();
-			}
-			else
-			{
-				arch >> m_name;
-			}
-
-            arch >> m_pSize;
-		}
 	};
 
 
diff --git a/Engine/Plugins/Experimental/Mutable/Source/MutableTools/Private/MuT/NodeScalar.cpp b/Engine/Plugins/Experimental/Mutable/Source/MutableTools/Private/MuT/NodeScalar.cpp
--- a/Engine/Plugins/Experimental/Mutable/Source/MutableTools/Private/MuT/NodeScalar.cpp
+++ b/Engine/Plugins/Experimental/Mutable/Source/MutableTools/Private/MuT/NodeScalar.cpp
@@ -4,66 +4,14 @@
 #include "MuT/NodeScalar.h"
 
 #include "Misc/AssertionMacros.h"
-#include "MuR/Serialisation.h"
-#include "MuR/SerialisationPrivate.h"
-#include "MuT/Node.h"
-#include "MuT/NodeScalarArithmeticOperation.h"
-#include "MuT/NodeScalarConstant.h"
-#include "MuT/NodeScalarCurve.h"
-#include "MuT/NodeScalarEnumParameter.h"
-#include "MuT/NodeScalarParameter.h"
-#include "MuT/NodeScalarSwitch.h"
-#include "MuT/NodeScalarTable.h"
-#include "MuT/NodeScalarVariation.h"
 
 
 namespace mu
 {
 
-	//---------------------------------------------------------------------------------------------
-	// Static initialisation
-	//---------------------------------------------------------------------------------------------
 	static FNodeType s_nodeScalarType = FNodeType( "NodeScalar", Node::GetStaticType() );
 
 
-	//---------------------------------------------------------------------------------------------
-	void NodeScalar::Serialise( const NodeScalar* p, OutputArchive& arch )
-	{
-        uint32 ver = 0;
-		arch << ver;
-
-		arch << uint32_t(p->Type);
-		p->SerialiseWrapper(arch);
-    }
-
-
-    //---------------------------------------------------------------------------------------------
-	NodeScalarPtr NodeScalar::StaticUnserialise( InputArchive& arch )
-	{
-        uint32 ver;
-		arch >> ver;
-		check( ver == 0 );
-
-        uint32 id;
-		arch >> id;
-
-		switch (id)
-		{
-		case 0 :  return NodeScalarConstant::StaticUnserialise( arch ); break;
-		case 1 :  return NodeScalarParameter::StaticUnserialise( arch ); break;
-		case 2 :  return NodeScalarEnumParameter::StaticUnserialise( arch ); break;
-        case 3 :  return NodeScalarCurve::StaticUnserialise( arch ); break;
-        case 4 :  return NodeScalarSwitch::StaticUnserialise( arch ); break;
-        case 5 :  return NodeScalarArithmeticOperation::StaticUnserialise( arch ); break;
-        case 6 :  return NodeScalarVariation::StaticUnserialise( arch ); break;
-		case 7 :  return NodeScalarTable::StaticUnserialise( arch ); break;
-        default : check(false);
-		}
-
-		return 0;
-	}
-
-
 	//---------------------------------------------------------------------------------------------
 	const FNodeType* NodeScalar::GetType() const
 	{
diff --git a/Engine/Plugins/Experimental/Mutable/Source/MutableTools/Private/MuT/NodeScalarArithmeticOperationPrivate.h b/Engine/Plugins/Experimental/Mutable/Source/MutableTools/Private/MuT/NodeScalarArithmeticOperationPrivate.h
--- a/Engine/Plugins/Experimental/Mutable/Source/MutableTools/Private/MuT/NodeScalarArithmeticOperationPrivate.h
+++ b/Engine/Plugins/Experimental/Mutable/Source/MutableTools/Private/MuT/NodeScalarArithmeticOperationPrivate.h
@@ -5,7 +5,6 @@
 #include "MuT/NodeScalarPrivate.h"
 #include "MuT/NodeScalarArithmeticOperation.h"
 
-
 namespace mu
 {
     MUTABLE_DEFINE_ENUM_SERIALISABLE(NodeScalarArithmeticOperation::OPERATION)
@@ -20,28 +19,6 @@ namespace mu
         NodeScalarPtr m_pA;
         NodeScalarPtr m_pB;
 
-		//!
-		void Serialise( OutputArchive& arch ) const
-		{
-            uint32_t ver = 0;
-			arch << ver;
-
-			arch << m_operation;
-			arch << m_pA;
-			arch << m_pB;
-		}
-
-		//!
-		void Unserialise( InputArchive& arch )
-		{
-            uint32_t ver;
-			arch >> ver;
-			check(ver==0);
-
-			arch >> m_operation;
-			arch >> m_pA;
-			arch >> m_pB;
-		}
 	};
 
 
diff --git a/Engine/Plugins/Experimental/Mutable/Source/MutableTools/Private/MuT/NodeScalarConstantPrivate.h b/Engine/Plugins/Experimental/Mutable/Source/MutableTools/Private/MuT/NodeScalarConstantPrivate.h
--- a/Engine/Plugins/Experimental/Mutable/Source/MutableTools/Private/MuT/NodeScalarConstantPrivate.h
+++ b/Engine/Plugins/Experimental/Mutable/Source/MutableTools/Private/MuT/NodeScalarConstantPrivate.h
@@ -6,14 +6,10 @@
 
 #include "MuT/NodeScalarConstant.h"
 #include "MuT/NodeImage.h"
-#include "MuT/AST.h"
-#include "MuR/ParametersPrivate.h"
-
 
 namespace mu
 {
 
-
 	class NodeScalarConstant::Private : public NodeScalar::Private
 	{
 	public:
@@ -21,25 +17,6 @@ namespace mu
 		static FNodeType s_type;
 
 		float m_value;
-
-		//!
-		void Serialise( OutputArchive& arch ) const
-		{
-            uint32_t ver = 0;
-			arch << ver;
-
-			arch << m_value;
-		}
-
-		//!
-		void Unserialise( InputArchive& arch )
-		{
-            uint32_t ver;
-			arch >> ver;
-			check(ver==0);
-
-			arch >> m_value;
-		}
 	};
 
 
diff --git a/Engine/Plugins/Experimental/Mutable/Source/MutableTools/Private/MuT/NodeScalarTablePrivate.h b/Engine/Plugins/Experimental/Mutable/Source/MutableTools/Private/MuT/NodeScalarTablePrivate.h
--- a/Engine/Plugins/Experimental/Mutable/Source/MutableTools/Private/MuT/NodeScalarTablePrivate.h
+++ b/Engine/Plugins/Experimental/Mutable/Source/MutableTools/Private/MuT/NodeScalarTablePrivate.h
@@ -5,8 +5,6 @@
 #include "MuT/NodePrivate.h"
 #include "MuT/NodeScalarTable.h"
 #include "MuT/TablePrivate.h"
-#include "MuT/AST.h"
-
 
 namespace mu
 {
@@ -24,65 +22,6 @@ namespace mu
 		bool bNoneOption = false;
 		FString DefaultRowName;
 
-		//!
-		void Serialise( OutputArchive& arch ) const
-		{
-            uint32 ver = 3;
-			arch << ver;
-
-			arch << ParameterName;
-			arch << Table;
-			arch << ColumnName;
-			arch << bNoneOption;
-			arch << DefaultRowName;
-		}
-
-		//!
-		void Unserialise( InputArchive& arch )
-		{
-            uint32 ver;
-			arch >> ver;
-            check(ver <= 3);
-
-			if (ver==0)
-			{
-				std::string Temp;
-				arch >> Temp;
-				ParameterName = Temp.c_str();
-			}
-			else
-			{
-				arch >> ParameterName;
-			}
-
-			arch >> Table;
-
-			if (ver == 0)
-			{
-				std::string Temp;
-				arch >> Temp;
-				ColumnName = Temp.c_str();
-			}
-			else
-			{
-				arch >> ColumnName;
-			}
-
-			if(ver >= 2)
-			{
-				arch >> bNoneOption;
-			}
-			else
-			{
-				bNoneOption = Table->GetPrivate()->bNoneOption_DEPRECATED;
-			}
-
-			if (ver >= 3)
-			{
-				arch >> DefaultRowName;
-			}
-		}
-
 	};
 
 }
diff --git a/Engine/Plugins/Experimental/Mutable/Source/MutableTools/Private/MuT/NodeScalarVariationPrivate.h b/Engine/Plugins/Experimental/Mutable/Source/MutableTools/Private/MuT/NodeScalarVariationPrivate.h
--- a/Engine/Plugins/Experimental/Mutable/Source/MutableTools/Private/MuT/NodeScalarVariationPrivate.h
+++ b/Engine/Plugins/Experimental/Mutable/Source/MutableTools/Private/MuT/NodeScalarVariationPrivate.h
@@ -5,7 +5,6 @@
 #include "MuT/NodePrivate.h"
 #include "MuT/NodeScalarVariation.h"
 
-
 namespace mu
 {
 
@@ -24,60 +23,10 @@ namespace mu
         {
             NodeScalarPtr m_scalar;
             FString m_tag;
-
-            //!
-            void Serialise( OutputArchive& arch ) const
-            {
-                uint32 ver = 2;
-                arch << ver;
-
-                arch << m_tag;
-                arch << m_scalar;
-            }
-
-            void Unserialise( InputArchive& arch )
-            {
-                uint32 ver = 0;
-                arch >> ver;
-				check(ver >= 1 && ver <= 2);
-
-				if (ver <= 1)
-				{
-					std::string Temp;
-					arch >> Temp;
-					m_tag = Temp.c_str();
-				}
-				else
-				{
-					arch >> m_tag;
-				}
-				
-				arch >> m_scalar;
-            }
         };
 
         TArray<FVariation> m_variations;
 
-        //!
-        void Serialise( OutputArchive& arch ) const
-        {
-            uint32 ver = 1;
-            arch << ver;
-
-            arch << m_defaultScalar;
-            arch << m_variations;
-        }
-
-        //!
-        void Unserialise( InputArchive& arch )
-        {
-            uint32 ver;
-            arch >> ver;
-			check(ver == 1);
-
-            arch >> m_defaultScalar;
-            arch >> m_variations;
-        }
     };
 
 } // namespace mu
diff --git a/Engine/Plugins/Experimental/Mutable/Source/MutableTools/Private/MuT/NodeString.cpp b/Engine/Plugins/Experimental/Mutable/Source/MutableTools/Private/MuT/NodeString.cpp
--- a/Engine/Plugins/Experimental/Mutable/Source/MutableTools/Private/MuT/NodeString.cpp
+++ b/Engine/Plugins/Experimental/Mutable/Source/MutableTools/Private/MuT/NodeString.cpp
@@ -4,11 +4,6 @@
 #include "MuT/NodeString.h"
 
 #include "Misc/AssertionMacros.h"
-#include "MuR/Serialisation.h"
-#include "MuR/SerialisationPrivate.h"
-#include "MuT/Node.h"
-#include "MuT/NodeStringConstant.h"
-#include "MuT/NodeStringParameter.h"
 
 
 namespace mu
@@ -20,48 +15,6 @@ namespace mu
 	static FNodeType s_nodeStringType = FNodeType( "NodeString", Node::GetStaticType() );
 
 
-	//---------------------------------------------------------------------------------------------
-	void NodeString::Serialise( const NodeString* p, OutputArchive& arch )
-	{
-        uint32 ver = 0;
-		arch << ver;
-
-#define SERIALISE_CHILDREN( C, ID ) \
-		( p->GetType()==C::GetStaticType() )					\
-		{ 														\
-			const C* pTyped = static_cast<const C*>(p);			\
-            arch << (uint32_t)ID;								\
-			C::Serialise( pTyped, arch );						\
-		}														\
-
-        if SERIALISE_CHILDREN( NodeStringConstant                   , 0 )
-        else if SERIALISE_CHILDREN( NodeStringParameter             , 1 )
-        else check(false);
-
-#undef SERIALISE_CHILDREN
-    }
-
-    //---------------------------------------------------------------------------------------------
-	NodeStringPtr NodeString::StaticUnserialise( InputArchive& arch )
-	{
-        uint32 ver;
-		arch >> ver;
-		check( ver == 0 );
-
-        uint32 id;
-		arch >> id;
-
-		switch (id)
-		{
-		case 0 :  return NodeStringConstant::StaticUnserialise( arch ); break;
-		case 1 :  return NodeStringParameter::StaticUnserialise( arch ); break;
-        default : check(false);
-		}
-
-		return 0;
-	}
-
-
 	//---------------------------------------------------------------------------------------------
 	const FNodeType* NodeString::GetType() const
 	{
diff --git a/Engine/Plugins/Experimental/Mutable/Source/MutableTools/Private/MuT/NodeStringConstantPrivate.h b/Engine/Plugins/Experimental/Mutable/Source/MutableTools/Private/MuT/NodeStringConstantPrivate.h
--- a/Engine/Plugins/Experimental/Mutable/Source/MutableTools/Private/MuT/NodeStringConstantPrivate.h
+++ b/Engine/Plugins/Experimental/Mutable/Source/MutableTools/Private/MuT/NodeStringConstantPrivate.h
@@ -6,9 +6,6 @@
 
 #include "MuT/NodeStringConstant.h"
 #include "MuT/NodeImage.h"
-#include "MuT/AST.h"
-#include "MuR/ParametersPrivate.h"
-
 
 namespace mu
 {
@@ -22,33 +19,6 @@ namespace mu
 
 		FString m_value;
 
-		//!
-		void Serialise( OutputArchive& arch ) const
-		{
-            uint32 ver = 1;
-			arch << ver;
-
-			arch << m_value;
-		}
-
-		//!
-		void Unserialise( InputArchive& arch )
-		{
-            uint32 ver;
-			arch >> ver;
-			check(ver<=1);
-
-			if (ver <= 0)
-			{
-				std::string Temp;
-				arch >> Temp;
-				m_value = Temp.c_str();
-			}
-			else
-			{
-				arch >> m_value;
-			}
-		}
 	};
 
 
diff --git a/Engine/Plugins/Experimental/Mutable/Source/MutableTools/Private/MuT/NodeStringParameterPrivate.h b/Engine/Plugins/Experimental/Mutable/Source/MutableTools/Private/MuT/NodeStringParameterPrivate.h
--- a/Engine/Plugins/Experimental/Mutable/Source/MutableTools/Private/MuT/NodeStringParameterPrivate.h
+++ b/Engine/Plugins/Experimental/Mutable/Source/MutableTools/Private/MuT/NodeStringParameterPrivate.h
@@ -7,13 +7,11 @@
 #include "MuT/NodeStringParameter.h"
 #include "MuT/NodeImage.h"
 #include "MuT/NodeRange.h"
-#include "MuR/ParametersPrivate.h"
 
 
 namespace mu
 {
 
-
 	class NodeStringParameter::Private : public NodeString::Private
 	{
 	public:
@@ -28,51 +26,6 @@ namespace mu
 
 		TArray<Ptr<NodeRange>> m_ranges;
 
-		//!
-		void Serialise( OutputArchive& arch ) const
-		{
-            uint32_t ver = 6;
-			arch << ver;
-
-			arch << m_defaultValue;
-			arch << m_name;
-			arch << m_uid;
-			arch << m_additionalImages;
-            arch << m_ranges;
-        }
-
-		//!
-		void Unserialise( InputArchive& arch )
-		{
-            uint32_t ver;
-			arch >> ver;
-            check(ver>=4 && ver<=6);
-
-			if (ver <= 5)
-			{
-				std::string Temp;
-				arch >> Temp;
-				m_defaultValue = Temp.c_str();
-				arch >> Temp;
-				m_name = Temp.c_str();
-				arch >> Temp;
-				m_uid = Temp.c_str();
-			}
-			else
-			{
-				arch >> m_defaultValue;
-				arch >> m_name;
-				arch >> m_uid;
-			}
-
-			arch >> m_additionalImages;
-			if (ver <= 4)
-			{
-				int32 Dummy=0;
-				arch >> Dummy;
-			}
-            arch >> m_ranges;
-        }
 	};
 
 }
diff --git a/Engine/Plugins/Experimental/Mutable/Source/MutableTools/Private/MuT/NodeSurface.cpp b/Engine/Plugins/Experimental/Mutable/Source/MutableTools/Private/MuT/NodeSurface.cpp
--- a/Engine/Plugins/Experimental/Mutable/Source/MutableTools/Private/MuT/NodeSurface.cpp
+++ b/Engine/Plugins/Experimental/Mutable/Source/MutableTools/Private/MuT/NodeSurface.cpp
@@ -3,13 +3,6 @@
 #include "MuT/NodeSurface.h"
 
 #include "Misc/AssertionMacros.h"
-#include "MuR/Serialisation.h"
-#include "MuR/SerialisationPrivate.h"
-#include "MuT/NodeSurfaceEdit.h"
-#include "MuT/NodeSurfaceNew.h"
-#include "MuT/NodeSurfaceVariation.h"
-#include "MuT/NodeSurfaceSwitch.h"
-
 
 namespace mu
 {
@@ -36,41 +29,6 @@ namespace mu
         return &s_nodeSurfaceType;
 	}
 
-
-	//---------------------------------------------------------------------------------------------
-    void NodeSurface::Serialise( const NodeSurface* p, OutputArchive& arch )
-	{
-        uint32 ver = 0;
-		arch << ver;
-
-		arch << uint32_t(p->Type);
-		p->SerialiseWrapper(arch);
-	}
-
-
-    //---------------------------------------------------------------------------------------------
-    NodeSurfacePtr NodeSurface::StaticUnserialise( InputArchive& arch )
-	{
-        uint32 ver;
-		arch >> ver;
-		check( ver == 0 );
-
-        uint32 id;
-		arch >> id;
-
-		switch (id)
-		{
-        case 0: return NodeSurfaceNew::StaticUnserialise( arch ); break;
-        case 1: return NodeSurfaceEdit::StaticUnserialise( arch ); break;
-		//case 2: return NodeSelectSurface::StaticUnserialise(arch); break;
-		case 3: return NodeSurfaceVariation::StaticUnserialise(arch); break;
-		case 4: return NodeSurfaceSwitch::StaticUnserialise(arch); break;
-		default : check(false);
-		}
-
-		return 0;
-	}
-
 }
 
 
diff --git a/Engine/Plugins/Experimental/Mutable/Source/MutableTools/Private/MuT/Table.cpp b/Engine/Plugins/Experimental/Mutable/Source/MutableTools/Private/MuT/Table.cpp
--- a/Engine/Plugins/Experimental/Mutable/Source/MutableTools/Private/MuT/Table.cpp
+++ b/Engine/Plugins/Experimental/Mutable/Source/MutableTools/Private/MuT/Table.cpp
@@ -34,22 +34,6 @@ namespace mu
 	}
 
 
-	//---------------------------------------------------------------------------------------------
-	void Table::Serialise( const Table* p, OutputArchive& arch )
-	{
-		arch << *p->m_pD;
-	}
-
-
-	//---------------------------------------------------------------------------------------------
-	Ptr<Table> Table::StaticUnserialise( InputArchive& arch )
-	{
-		Ptr<Table> pResult = new Table();
-		arch >> *pResult->m_pD;
-		return pResult;
-	}
-
-
 	Table::Private* Table::GetPrivate() const
 	{
 		return m_pD;
diff --git a/Engine/Plugins/Experimental/Mutable/Source/MutableTools/Public/MuT/Node.h b/Engine/Plugins/Experimental/Mutable/Source/MutableTools/Public/MuT/Node.h
--- a/Engine/Plugins/Experimental/Mutable/Source/MutableTools/Public/MuT/Node.h
+++ b/Engine/Plugins/Experimental/Mutable/Source/MutableTools/Public/MuT/Node.h
@@ -5,14 +5,6 @@
 #include "HAL/Platform.h"
 #include "MuR/RefCounted.h"
 
-namespace mu { class InputArchive; }
-namespace mu { class OutputArchive; }
-
-
-//! This tag is used to identify files containing serialised Node hierarchies. The tag is not added
-//! or checked by the Node serialisation methods, but the involved tools should take care of it.
-#define MUTABLE_SOURCE_MODEL_FILETAG	"amsm"
-
 
 namespace mu
 {
@@ -38,9 +30,7 @@ namespace mu
 	typedef Ptr<const NodeMap> NodeMapPtrConst;
 
 
-	//! Information about the type of a node, to provide some means to the tools to deal generically
-	//! with nodes.
-	//! \ingroup tools
+	/** Information about the type of a node, to provide some means to the tools to deal generically with nodes. */
 	struct FNodeType
 	{
 		FNodeType();
@@ -66,7 +56,6 @@ namespace mu
 	};
 
 
-
     //! %Base class for all graphs used in the source data to define models and transforms.
 	class MUTABLETOOLS_API Node : public RefCounted
 	{
@@ -97,15 +86,6 @@ namespace mu
 			None
 		};
 
-		//-----------------------------------------------------------------------------------------
-		// Life cycle
-		//-----------------------------------------------------------------------------------------
-
-		static void Serialise( const Node* pNode, OutputArchive& arch );
-		virtual void SerialiseWrapper(OutputArchive& arch) const = 0;
-		static Ptr<Node> StaticUnserialise( InputArchive& arch );
-
-
 		//-----------------------------------------------------------------------------------------
 		// Own Interface
 		//-----------------------------------------------------------------------------------------
diff --git a/Engine/Plugins/Experimental/Mutable/Source/MutableTools/Public/MuT/NodeBool.h b/Engine/Plugins/Experimental/Mutable/Source/MutableTools/Public/MuT/NodeBool.h
--- a/Engine/Plugins/Experimental/Mutable/Source/MutableTools/Public/MuT/NodeBool.h
+++ b/Engine/Plugins/Experimental/Mutable/Source/MutableTools/Public/MuT/NodeBool.h
@@ -32,13 +32,6 @@ namespace mu
 			None
 		};
 
-		//-----------------------------------------------------------------------------------------
-		// Life cycle
-		//-----------------------------------------------------------------------------------------
-
-		static void Serialise( const NodeBool* pNode, OutputArchive& arch );
-		static Ptr<NodeBool> StaticUnserialise( InputArchive& arch );
-
 
 		//-----------------------------------------------------------------------------------------
 		// Node Interface
@@ -72,16 +65,8 @@ namespace mu
 	{
 	public:
 
-		//-----------------------------------------------------------------------------------------
-		// Life cycle
-		//-----------------------------------------------------------------------------------------
-
 		NodeBoolConstant();
 
-		static void Serialise( const NodeBoolConstant* pNode, OutputArchive& arch );
-		void SerialiseWrapper(OutputArchive& arch) const override;
-		static Ptr<NodeBoolConstant> StaticUnserialise( InputArchive& arch );
-
 
 		//-----------------------------------------------------------------------------------------
 		// Node Interface
@@ -127,17 +112,8 @@ namespace mu
 	{
 	public:
 
-		//-----------------------------------------------------------------------------------------
-		// Life cycle
-		//-----------------------------------------------------------------------------------------
-
 		NodeBoolParameter();
 
-		void SerialiseWrapper(OutputArchive& arch) const override;
-		static void Serialise( const NodeBoolParameter* pNode, OutputArchive& arch );
-		static Ptr<NodeBoolParameter> StaticUnserialise( InputArchive& arch );
-
-
 		//-----------------------------------------------------------------------------------------
 		// Node Interface
 		//-----------------------------------------------------------------------------------------
@@ -187,17 +163,8 @@ namespace mu
 	{
 	public:
 
-		//-----------------------------------------------------------------------------------------
-		// Life cycle
-		//-----------------------------------------------------------------------------------------
-
 		NodeBoolNot();
 
-		void SerialiseWrapper(OutputArchive& arch) const override;
-		static void Serialise( const NodeBoolNot* pNode, OutputArchive& arch );
-		static Ptr<NodeBoolNot> StaticUnserialise( InputArchive& arch );
-
-
 		//-----------------------------------------------------------------------------------------
 		// Node Interface
 		//-----------------------------------------------------------------------------------------
@@ -240,17 +207,8 @@ namespace mu
 	{
 	public:
 
-		//-----------------------------------------------------------------------------------------
-		// Life cycle
-		//-----------------------------------------------------------------------------------------
-
 		NodeBoolAnd();
 
-		void SerialiseWrapper(OutputArchive& arch) const override;
-		static void Serialise( const NodeBoolAnd* pNode, OutputArchive& arch );
-		static Ptr<NodeBoolAnd> StaticUnserialise( InputArchive& arch );
-
-
 		//-----------------------------------------------------------------------------------------
 		// Node Interface
 		//-----------------------------------------------------------------------------------------
diff --git a/Engine/Plugins/Experimental/Mutable/Source/MutableTools/Public/MuT/NodeColour.h b/Engine/Plugins/Experimental/Mutable/Source/MutableTools/Public/MuT/NodeColour.h
--- a/Engine/Plugins/Experimental/Mutable/Source/MutableTools/Public/MuT/NodeColour.h
+++ b/Engine/Plugins/Experimental/Mutable/Source/MutableTools/Public/MuT/NodeColour.h
@@ -38,13 +38,6 @@ namespace mu
 			None
 		};
 
-        //-----------------------------------------------------------------------------------------
-        // Life cycle
-        //-----------------------------------------------------------------------------------------
-
-        static void Serialise( const NodeColour* pNode, OutputArchive& arch );
-        static NodeColourPtr StaticUnserialise( InputArchive& arch );
-
 
         //-----------------------------------------------------------------------------------------
         // Node Interface
diff --git a/Engine/Plugins/Experimental/Mutable/Source/MutableTools/Public/MuT/NodeColourArithmeticOperation.h b/Engine/Plugins/Experimental/Mutable/Source/MutableTools/Public/MuT/NodeColourArithmeticOperation.h
--- a/Engine/Plugins/Experimental/Mutable/Source/MutableTools/Public/MuT/NodeColourArithmeticOperation.h
+++ b/Engine/Plugins/Experimental/Mutable/Source/MutableTools/Public/MuT/NodeColourArithmeticOperation.h
@@ -4,7 +4,6 @@
 
 #include "MuR/Ptr.h"
 #include "MuR/RefCounted.h"
-#include "MuR/Serialisation.h"
 #include "MuT/Node.h"
 #include "MuT/NodeColour.h"
 
@@ -24,17 +23,8 @@ namespace mu
 	{
 	public:
 
-		//-----------------------------------------------------------------------------------------
-		// Life cycle
-		//-----------------------------------------------------------------------------------------
-
 		NodeColourArithmeticOperation();
 
-		void SerialiseWrapper(OutputArchive& arch) const override;
-		static void Serialise( const NodeColourArithmeticOperation* pNode, OutputArchive& arch );
-		static NodeColourArithmeticOperationPtr StaticUnserialise( InputArchive& arch );
-
-
 		//-----------------------------------------------------------------------------------------
 		// Node Interface
 		//-----------------------------------------------------------------------------------------
@@ -91,7 +81,5 @@ namespace mu
 
 	};
 
-
-	MUTABLE_DEFINE_ENUM_SERIALISABLE(NodeColourArithmeticOperation::OPERATION)
 	
 }
diff --git a/Engine/Plugins/Experimental/Mutable/Source/MutableTools/Public/MuT/NodeColourConstant.h b/Engine/Plugins/Experimental/Mutable/Source/MutableTools/Public/MuT/NodeColourConstant.h
--- a/Engine/Plugins/Experimental/Mutable/Source/MutableTools/Public/MuT/NodeColourConstant.h
+++ b/Engine/Plugins/Experimental/Mutable/Source/MutableTools/Public/MuT/NodeColourConstant.h
@@ -22,17 +22,8 @@ namespace mu
 	{
 	public:
 
-		//-----------------------------------------------------------------------------------------
-		// Life cycle
-		//-----------------------------------------------------------------------------------------
-
 		NodeColourConstant();
 
-		void SerialiseWrapper(OutputArchive& arch) const override;
-		static void Serialise( const NodeColourConstant* pNode, OutputArchive& arch );
-		static NodeColourConstantPtr StaticUnserialise( InputArchive& arch );
-		static NodeColourConstantPtr OldStaticUnserialise( InputArchive& arch );
-
 		//-----------------------------------------------------------------------------------------
 		// Node Interface
 		//-----------------------------------------------------------------------------------------
diff --git a/Engine/Plugins/Experimental/Mutable/Source/MutableTools/Public/MuT/NodeColourFromScalars.h b/Engine/Plugins/Experimental/Mutable/Source/MutableTools/Public/MuT/NodeColourFromScalars.h
--- a/Engine/Plugins/Experimental/Mutable/Source/MutableTools/Public/MuT/NodeColourFromScalars.h
+++ b/Engine/Plugins/Experimental/Mutable/Source/MutableTools/Public/MuT/NodeColourFromScalars.h
@@ -27,17 +27,8 @@ namespace mu
 	{
 	public:
 
-		//-----------------------------------------------------------------------------------------
-		// Life cycle
-		//-----------------------------------------------------------------------------------------
-
 		NodeColourFromScalars();
 
-		void SerialiseWrapper(OutputArchive& arch) const override;
-		static void Serialise( const NodeColourFromScalars* pNode, OutputArchive& arch );
-		static NodeColourFromScalarsPtr StaticUnserialise( InputArchive& arch );
-
-
 		//-----------------------------------------------------------------------------------------
 		// Node Interface
 		//-----------------------------------------------------------------------------------------
