[mutable] Remove custom serialization code frome the MutableTools module.

diff --git a/Engine/Plugins/Experimental/Mutable/Source/MutableTools/Public/MuT/NodeColourParameter.h b/Engine/Plugins/Experimental/Mutable/Source/MutableTools/Public/MuT/NodeColourParameter.h
--- a/Engine/Plugins/Experimental/Mutable/Source/MutableTools/Public/MuT/NodeColourParameter.h
+++ b/Engine/Plugins/Experimental/Mutable/Source/MutableTools/Public/MuT/NodeColourParameter.h
@@ -7,7 +7,6 @@
 #include "MuT/Node.h"
 #include "MuT/NodeColour.h"
 
-
 namespace mu
 {
 
@@ -26,17 +25,8 @@ namespace mu
 	{
 	public:
 
-		//-----------------------------------------------------------------------------------------
-		// Life cycle
-		//-----------------------------------------------------------------------------------------
-
 		NodeColourParameter();
 
-		void SerialiseWrapper(OutputArchive& arch) const override;
-		static void Serialise( const NodeColourParameter* pNode, OutputArchive& arch );
-		static NodeColourParameterPtr StaticUnserialise( InputArchive& arch );
-
-
 		//-----------------------------------------------------------------------------------------
 		// Node Interface
 		//-----------------------------------------------------------------------------------------
diff --git a/Engine/Plugins/Experimental/Mutable/Source/MutableTools/Public/MuT/NodeColourSampleImage.h b/Engine/Plugins/Experimental/Mutable/Source/MutableTools/Public/MuT/NodeColourSampleImage.h
--- a/Engine/Plugins/Experimental/Mutable/Source/MutableTools/Public/MuT/NodeColourSampleImage.h
+++ b/Engine/Plugins/Experimental/Mutable/Source/MutableTools/Public/MuT/NodeColourSampleImage.h
@@ -31,16 +31,8 @@ namespace mu
 	{
 	public:
 
-		//-----------------------------------------------------------------------------------------
-		// Life cycle
-		//-----------------------------------------------------------------------------------------
-
 		NodeColourSampleImage();
 
-		void SerialiseWrapper(OutputArchive& arch) const override;
-		static void Serialise( const NodeColourSampleImage* pNode, OutputArchive& arch );
-		static NodeColourSampleImagePtr StaticUnserialise( InputArchive& arch );
-
 
 		//-----------------------------------------------------------------------------------------
 		// Node Interface
diff --git a/Engine/Plugins/Experimental/Mutable/Source/MutableTools/Public/MuT/NodeColourSwitch.h b/Engine/Plugins/Experimental/Mutable/Source/MutableTools/Public/MuT/NodeColourSwitch.h
--- a/Engine/Plugins/Experimental/Mutable/Source/MutableTools/Public/MuT/NodeColourSwitch.h
+++ b/Engine/Plugins/Experimental/Mutable/Source/MutableTools/Public/MuT/NodeColourSwitch.h
@@ -27,16 +27,8 @@ namespace mu
 	{
 	public:
 
-		//-----------------------------------------------------------------------------------------
-		// Life cycle
-		//-----------------------------------------------------------------------------------------
-
         NodeColourSwitch();
 
-		void SerialiseWrapper(OutputArchive& arch) const override;
-		static void Serialise( const NodeColourSwitch* pNode, OutputArchive& arch );
-        static NodeColourSwitchPtr StaticUnserialise( InputArchive& arch );
-
 
 		//-----------------------------------------------------------------------------------------
 		// Node Interface
diff --git a/Engine/Plugins/Experimental/Mutable/Source/MutableTools/Public/MuT/NodeColourTable.h b/Engine/Plugins/Experimental/Mutable/Source/MutableTools/Public/MuT/NodeColourTable.h
--- a/Engine/Plugins/Experimental/Mutable/Source/MutableTools/Public/MuT/NodeColourTable.h
+++ b/Engine/Plugins/Experimental/Mutable/Source/MutableTools/Public/MuT/NodeColourTable.h
@@ -27,16 +27,8 @@ namespace mu
 	{
 	public:
 
-		//-----------------------------------------------------------------------------------------
-		// Life cycle
-		//-----------------------------------------------------------------------------------------
-
 		NodeColourTable();
 
-		void SerialiseWrapper(OutputArchive& arch) const override;
-		static void Serialise( const NodeColourTable* pNode, OutputArchive& arch );
-		static NodeColourTablePtr StaticUnserialise( InputArchive& arch );
-
 
 		//-----------------------------------------------------------------------------------------
 		// Node Interface
diff --git a/Engine/Plugins/Experimental/Mutable/Source/MutableTools/Public/MuT/NodeColourVariation.h b/Engine/Plugins/Experimental/Mutable/Source/MutableTools/Public/MuT/NodeColourVariation.h
--- a/Engine/Plugins/Experimental/Mutable/Source/MutableTools/Public/MuT/NodeColourVariation.h
+++ b/Engine/Plugins/Experimental/Mutable/Source/MutableTools/Public/MuT/NodeColourVariation.h
@@ -22,16 +22,9 @@ namespace mu
     class MUTABLETOOLS_API NodeColourVariation : public NodeColour
     {
     public:
-        //-----------------------------------------------------------------------------------------
-        // Life cycle
-        //-----------------------------------------------------------------------------------------
 
         NodeColourVariation();
 
-		void SerialiseWrapper(OutputArchive& arch) const override;
-		static void Serialise( const NodeColourVariation* pNode, OutputArchive& arch );
-        static NodeColourVariationPtr StaticUnserialise( InputArchive& arch );
-
         //-----------------------------------------------------------------------------------------
         // Node Interface
         //-----------------------------------------------------------------------------------------
diff --git a/Engine/Plugins/Experimental/Mutable/Source/MutableTools/Public/MuT/NodeComponent.h b/Engine/Plugins/Experimental/Mutable/Source/MutableTools/Public/MuT/NodeComponent.h
--- a/Engine/Plugins/Experimental/Mutable/Source/MutableTools/Public/MuT/NodeComponent.h
+++ b/Engine/Plugins/Experimental/Mutable/Source/MutableTools/Public/MuT/NodeComponent.h
@@ -34,13 +34,6 @@ namespace mu
 			None
 		};
 
-		//-----------------------------------------------------------------------------------------
-		// Life cycle
-		//-----------------------------------------------------------------------------------------
-
-		static void Serialise( const NodeComponent* pNode, OutputArchive& arch );
-		static NodeComponentPtr StaticUnserialise( InputArchive& arch );
-
 		//-----------------------------------------------------------------------------------------
         // Node interface
 		//-----------------------------------------------------------------------------------------
diff --git a/Engine/Plugins/Experimental/Mutable/Source/MutableTools/Public/MuT/NodeComponentEdit.h b/Engine/Plugins/Experimental/Mutable/Source/MutableTools/Public/MuT/NodeComponentEdit.h
--- a/Engine/Plugins/Experimental/Mutable/Source/MutableTools/Public/MuT/NodeComponentEdit.h
+++ b/Engine/Plugins/Experimental/Mutable/Source/MutableTools/Public/MuT/NodeComponentEdit.h
@@ -21,7 +21,6 @@ namespace mu
     typedef Ptr<const NodeSurface> NodeSurfacePtrConst;
 
 
-
 	//! This node modifies a node of the parent object of the object that this node belongs to.
 	//! It allows to extend, cut and morph the parent component's meshes.
 	//! It also allows to patch the parent component's textures.
@@ -30,16 +29,8 @@ namespace mu
 	{
 	public:
 
-		//-----------------------------------------------------------------------------------------
-		// Life cycle
-		//-----------------------------------------------------------------------------------------
-
 		NodeComponentEdit();
 
-		void SerialiseWrapper(OutputArchive& arch) const override;
-		static void Serialise( const NodeComponentEdit* pNode, OutputArchive& arch );
-		static NodeComponentEditPtr StaticUnserialise( InputArchive& arch );
-
 		//-----------------------------------------------------------------------------------------
         // Node interface
 		//-----------------------------------------------------------------------------------------
diff --git a/Engine/Plugins/Experimental/Mutable/Source/MutableTools/Public/MuT/NodeComponentNew.h b/Engine/Plugins/Experimental/Mutable/Source/MutableTools/Public/MuT/NodeComponentNew.h
--- a/Engine/Plugins/Experimental/Mutable/Source/MutableTools/Public/MuT/NodeComponentNew.h
+++ b/Engine/Plugins/Experimental/Mutable/Source/MutableTools/Public/MuT/NodeComponentNew.h
@@ -32,16 +32,8 @@ namespace mu
 	{
 	public:
 
-		//-----------------------------------------------------------------------------------------
-		// Life cycle
-		//-----------------------------------------------------------------------------------------
-
 		NodeComponentNew();
 
-		void SerialiseWrapper(OutputArchive& arch) const override;
-		static void Serialise( const NodeComponentNew* pNode, OutputArchive& arch );
-		static NodeComponentNewPtr StaticUnserialise( InputArchive& arch );
-
 		//-----------------------------------------------------------------------------------------
         // Node interface
 		//-----------------------------------------------------------------------------------------		
diff --git a/Engine/Plugins/Experimental/Mutable/Source/MutableTools/Public/MuT/NodeExtensionData.h b/Engine/Plugins/Experimental/Mutable/Source/MutableTools/Public/MuT/NodeExtensionData.h
--- a/Engine/Plugins/Experimental/Mutable/Source/MutableTools/Public/MuT/NodeExtensionData.h
+++ b/Engine/Plugins/Experimental/Mutable/Source/MutableTools/Public/MuT/NodeExtensionData.h
@@ -11,57 +11,46 @@
 namespace mu
 {
 
-// Forward definitions
-class NodeExtensionData;
-typedef Ptr<NodeExtensionData> NodeExtensionDataPtr;
-typedef Ptr<const NodeExtensionData> NodeExtensionDataPtrConst;
+	// Forward definitions
+	class NodeExtensionData;
+	typedef Ptr<NodeExtensionData> NodeExtensionDataPtr;
+	typedef Ptr<const NodeExtensionData> NodeExtensionDataPtrConst;
 
 
-//! Node that evaluates to an ExtensionData
-//! \ingroup model
-class MUTABLETOOLS_API NodeExtensionData : public Node
-{
-public:
-	// Possible subclasses
-	enum class EType : uint8
+	//! Node that evaluates to an ExtensionData
+	class MUTABLETOOLS_API NodeExtensionData : public Node
 	{
-		Constant = 0,
-		Switch,
-		Variation,
-		None
-	};
-
-	//-----------------------------------------------------------------------------------------
-	// Life cycle
-	//-----------------------------------------------------------------------------------------
+	public:
+		// Possible subclasses
+		enum class EType : uint8
+		{
+			Constant = 0,
+			Switch,
+			Variation,
+			None
+		};
 
-	static void Serialise(const NodeExtensionData* pNode, OutputArchive& arch);
-	static NodeExtensionDataPtr StaticUnserialise(InputArchive& arch);
+		//-----------------------------------------------------------------------------------------
+		// Node Interface
+		//-----------------------------------------------------------------------------------------
+		const FNodeType* GetType() const override;
+		static const FNodeType* GetStaticType();
 
-	//-----------------------------------------------------------------------------------------
-	// Node Interface
-	//-----------------------------------------------------------------------------------------
-	const FNodeType* GetType() const override;
-	static const FNodeType* GetStaticType();
+		inline EType GetExtensionDataNodeType() const { return Type; }
 
-	inline EType GetExtensionDataNodeType() const { return Type; }
 
-	//-----------------------------------------------------------------------------------------
-	// Own Interface
-	//-----------------------------------------------------------------------------------------
+		//-----------------------------------------------------------------------------------------
+		// Interface pattern
+		//-----------------------------------------------------------------------------------------
+		class Private;
 
-	//-----------------------------------------------------------------------------------------
-	// Interface pattern
-	//-----------------------------------------------------------------------------------------
-	class Private;
+	protected:
 
-protected:
+		//! Forbidden. Manage with the Ptr<> template.
+		inline ~NodeExtensionData() {}
 
-	//! Forbidden. Manage with the Ptr<> template.
-	inline ~NodeExtensionData() {}
-
-	EType Type = EType::None;
-};
+		EType Type = EType::None;
+	};
 
 
 
diff --git a/Engine/Plugins/Experimental/Mutable/Source/MutableTools/Public/MuT/NodeExtensionDataConstant.h b/Engine/Plugins/Experimental/Mutable/Source/MutableTools/Public/MuT/NodeExtensionDataConstant.h
--- a/Engine/Plugins/Experimental/Mutable/Source/MutableTools/Public/MuT/NodeExtensionDataConstant.h
+++ b/Engine/Plugins/Experimental/Mutable/Source/MutableTools/Public/MuT/NodeExtensionDataConstant.h
@@ -11,58 +11,50 @@
 namespace mu
 {
 
-// Forward definitions
-class NodeExtensionDataConstant;
-typedef Ptr<NodeExtensionDataConstant> NodeExtensionDataConstantPtr;
-typedef Ptr<const NodeExtensionDataConstant> NodeExtensionDataConstantPtrConst;
-
-//! Node that outputs a constant ExtensionData
-//! \ingroup model
-class MUTABLETOOLS_API NodeExtensionDataConstant : public NodeExtensionData
-{
-public:
-
-	//-----------------------------------------------------------------------------------------
-	// Life cycle
-	//-----------------------------------------------------------------------------------------
-	NodeExtensionDataConstant();
-
-	void SerialiseWrapper(OutputArchive& arch) const override;
-	static void Serialise(const NodeExtensionDataConstant* pNode, OutputArchive& arch);
-	static NodeExtensionDataConstantPtr StaticUnserialise(InputArchive& arch);
-
-
-	//-----------------------------------------------------------------------------------------
-	// Node Interface
-	//-----------------------------------------------------------------------------------------
-	const FNodeType* GetType() const override;
-	static const FNodeType* GetStaticType();
-
-	//-----------------------------------------------------------------------------------------
-	// Own Interface
-	//-----------------------------------------------------------------------------------------
-
-	//! Get the constant ExtensionData that will be returned.
-	ExtensionDataPtrConst GetValue() const;
-
-	//! Set the constant ExtensionData that will be returned.
-	void SetValue(ExtensionDataPtrConst Value);
-
-	//-----------------------------------------------------------------------------------------
-	// Interface pattern
-	//-----------------------------------------------------------------------------------------
-	class Private;
-	Private* GetPrivate() const;
-	Node::Private* GetBasePrivate() const override;
-
-protected:
-	//! Forbidden. Manage with the Ptr<> template.
-	~NodeExtensionDataConstant();
-
-private:
-	Private* m_pD;
-
-};
+	// Forward definitions
+	class NodeExtensionDataConstant;
+	typedef Ptr<NodeExtensionDataConstant> NodeExtensionDataConstantPtr;
+	typedef Ptr<const NodeExtensionDataConstant> NodeExtensionDataConstantPtrConst;
+
+	//! Node that outputs a constant ExtensionData
+	//! \ingroup model
+	class MUTABLETOOLS_API NodeExtensionDataConstant : public NodeExtensionData
+	{
+	public:
+
+		NodeExtensionDataConstant();
+
+		//-----------------------------------------------------------------------------------------
+		// Node Interface
+		//-----------------------------------------------------------------------------------------
+		const FNodeType* GetType() const override;
+		static const FNodeType* GetStaticType();
+
+		//-----------------------------------------------------------------------------------------
+		// Own Interface
+		//-----------------------------------------------------------------------------------------
+
+		//! Get the constant ExtensionData that will be returned.
+		ExtensionDataPtrConst GetValue() const;
+
+		//! Set the constant ExtensionData that will be returned.
+		void SetValue(ExtensionDataPtrConst Value);
+
+		//-----------------------------------------------------------------------------------------
+		// Interface pattern
+		//-----------------------------------------------------------------------------------------
+		class Private;
+		Private* GetPrivate() const;
+		Node::Private* GetBasePrivate() const override;
+
+	protected:
+		//! Forbidden. Manage with the Ptr<> template.
+		~NodeExtensionDataConstant();
+
+	private:
+		Private* m_pD;
+
+	};
 
 
 }
diff --git a/Engine/Plugins/Experimental/Mutable/Source/MutableTools/Public/MuT/NodeExtensionDataSwitch.h b/Engine/Plugins/Experimental/Mutable/Source/MutableTools/Public/MuT/NodeExtensionDataSwitch.h
--- a/Engine/Plugins/Experimental/Mutable/Source/MutableTools/Public/MuT/NodeExtensionDataSwitch.h
+++ b/Engine/Plugins/Experimental/Mutable/Source/MutableTools/Public/MuT/NodeExtensionDataSwitch.h
@@ -22,15 +22,8 @@ namespace mu
 	{
 	public:
 
-		//-----------------------------------------------------------------------------------------
-		// Life cycle
-		//-----------------------------------------------------------------------------------------
 		NodeExtensionDataSwitch();
 
-		void SerialiseWrapper(OutputArchive& arch) const override;
-		static void Serialise(const NodeExtensionDataSwitch* pNode, OutputArchive& arch);
-		static NodeExtensionDataSwitchPtr StaticUnserialise(InputArchive& arch);
-
 		//-----------------------------------------------------------------------------------------
 		// Node Interface
 		//-----------------------------------------------------------------------------------------
diff --git a/Engine/Plugins/Experimental/Mutable/Source/MutableTools/Public/MuT/NodeExtensionDataVariation.h b/Engine/Plugins/Experimental/Mutable/Source/MutableTools/Public/MuT/NodeExtensionDataVariation.h
--- a/Engine/Plugins/Experimental/Mutable/Source/MutableTools/Public/MuT/NodeExtensionDataVariation.h
+++ b/Engine/Plugins/Experimental/Mutable/Source/MutableTools/Public/MuT/NodeExtensionDataVariation.h
@@ -18,15 +18,8 @@ namespace mu
 	{
 	public:
 
-		//-----------------------------------------------------------------------------------------
-		// Life cycle
-		//-----------------------------------------------------------------------------------------
 		NodeExtensionDataVariation();
 
-		void SerialiseWrapper(OutputArchive& arch) const override;
-		static void Serialise(const NodeExtensionDataVariation* pNode, OutputArchive& arch);
-		static NodeExtensionDataVariationPtr StaticUnserialise(InputArchive& arch);
-
 		//-----------------------------------------------------------------------------------------
 		// Node Interface
 		//-----------------------------------------------------------------------------------------
diff --git a/Engine/Plugins/Experimental/Mutable/Source/MutableTools/Public/MuT/NodeImage.h b/Engine/Plugins/Experimental/Mutable/Source/MutableTools/Public/MuT/NodeImage.h
--- a/Engine/Plugins/Experimental/Mutable/Source/MutableTools/Public/MuT/NodeImage.h
+++ b/Engine/Plugins/Experimental/Mutable/Source/MutableTools/Public/MuT/NodeImage.h
@@ -57,13 +57,6 @@ namespace mu
 			None
 		};
 
-		//-----------------------------------------------------------------------------------------
-		// Life cycle
-		//-----------------------------------------------------------------------------------------
-
-		static void Serialise( const NodeImage* pNode, OutputArchive& arch );
-		static NodeImagePtr StaticUnserialise( InputArchive& arch );
-
 
 		//-----------------------------------------------------------------------------------------
 		// Node Interface
diff --git a/Engine/Plugins/Experimental/Mutable/Source/MutableTools/Public/MuT/NodeImageBinarise.h b/Engine/Plugins/Experimental/Mutable/Source/MutableTools/Public/MuT/NodeImageBinarise.h
--- a/Engine/Plugins/Experimental/Mutable/Source/MutableTools/Public/MuT/NodeImageBinarise.h
+++ b/Engine/Plugins/Experimental/Mutable/Source/MutableTools/Public/MuT/NodeImageBinarise.h
@@ -27,17 +27,8 @@ namespace mu
 	{
 	public:
 
-		//-----------------------------------------------------------------------------------------
-		// Life cycle
-		//-----------------------------------------------------------------------------------------
-
 		NodeImageBinarise();
 
-		void SerialiseWrapper(OutputArchive& arch) const override;
-		static void Serialise( const NodeImageBinarise* pNode, OutputArchive& arch );
-		static NodeImageBinarisePtr StaticUnserialise( InputArchive& arch );
-
-
 		//-----------------------------------------------------------------------------------------
 		// Node Interface
 		//-----------------------------------------------------------------------------------------
diff --git a/Engine/Plugins/Experimental/Mutable/Source/MutableTools/Public/MuT/NodeImageColourMap.h b/Engine/Plugins/Experimental/Mutable/Source/MutableTools/Public/MuT/NodeImageColourMap.h
--- a/Engine/Plugins/Experimental/Mutable/Source/MutableTools/Public/MuT/NodeImageColourMap.h
+++ b/Engine/Plugins/Experimental/Mutable/Source/MutableTools/Public/MuT/NodeImageColourMap.h
@@ -16,9 +16,6 @@ namespace mu
 	typedef Ptr<NodeImageColourMap> NodeImageColourMapPtr;
 	typedef Ptr<const NodeImageColourMap> NodeImageColourMapPtrConst;
 
-	class InputArchive;
-	class OutputArchive;
-
 
 	//! This node changes the colours of a selectd part of the image, applying a colour map from
 	//! conteined in another image.
@@ -27,16 +24,8 @@ namespace mu
 	{
 	public:
 
-		//-----------------------------------------------------------------------------------------
-		// Life cycle
-		//-----------------------------------------------------------------------------------------
-
 		NodeImageColourMap();
 
-		void SerialiseWrapper(OutputArchive& arch) const override;
-		static void Serialise( const NodeImageColourMap* pNode, OutputArchive& arch );
-		static NodeImageColourMapPtr StaticUnserialise( InputArchive& arch );
-
 
 		//-----------------------------------------------------------------------------------------
 		// Node Interface
diff --git a/Engine/Plugins/Experimental/Mutable/Source/MutableTools/Public/MuT/NodeImageConditional.h b/Engine/Plugins/Experimental/Mutable/Source/MutableTools/Public/MuT/NodeImageConditional.h
--- a/Engine/Plugins/Experimental/Mutable/Source/MutableTools/Public/MuT/NodeImageConditional.h
+++ b/Engine/Plugins/Experimental/Mutable/Source/MutableTools/Public/MuT/NodeImageConditional.h
@@ -24,17 +24,8 @@ namespace mu
 	{
 	public:
 
-		//-----------------------------------------------------------------------------------------
-		// Life cycle
-		//-----------------------------------------------------------------------------------------
-
         NodeImageConditional();
 
-		void SerialiseWrapper(OutputArchive& arch) const override;
-		static void Serialise( const NodeImageConditional* pNode, OutputArchive& arch );
-        static NodeImageConditionalPtr StaticUnserialise( InputArchive& arch );
-
-
 		//-----------------------------------------------------------------------------------------
 		// Node Interface
 		//-----------------------------------------------------------------------------------------
diff --git a/Engine/Plugins/Experimental/Mutable/Source/MutableTools/Public/MuT/NodeImageConstant.h b/Engine/Plugins/Experimental/Mutable/Source/MutableTools/Public/MuT/NodeImageConstant.h
--- a/Engine/Plugins/Experimental/Mutable/Source/MutableTools/Public/MuT/NodeImageConstant.h
+++ b/Engine/Plugins/Experimental/Mutable/Source/MutableTools/Public/MuT/NodeImageConstant.h
@@ -17,9 +17,6 @@ namespace mu
     typedef Ptr<NodeImageConstant> NodeImageConstantPtr;
     typedef Ptr<const NodeImageConstant> NodeImageConstantPtrConst;
 
-	class InputArchive;
-	class OutputArchive;
-
 
 	/** Node that outputs a constant image.
 	* This node also supports "image references".
@@ -28,16 +25,8 @@ namespace mu
 	{
 	public:
 
-		//-----------------------------------------------------------------------------------------
-		// Life cycle
-		//-----------------------------------------------------------------------------------------
-
 		NodeImageConstant();
 
-		void SerialiseWrapper(OutputArchive& arch) const override;
-		static void Serialise( const NodeImageConstant* pNode, OutputArchive& arch );
-		static NodeImageConstantPtr StaticUnserialise( InputArchive& arch );
-
 		//-----------------------------------------------------------------------------------------
 		// Node Interface
 		//-----------------------------------------------------------------------------------------
diff --git a/Engine/Plugins/Experimental/Mutable/Source/MutableTools/Public/MuT/NodeImageFormat.h b/Engine/Plugins/Experimental/Mutable/Source/MutableTools/Public/MuT/NodeImageFormat.h
--- a/Engine/Plugins/Experimental/Mutable/Source/MutableTools/Public/MuT/NodeImageFormat.h
+++ b/Engine/Plugins/Experimental/Mutable/Source/MutableTools/Public/MuT/NodeImageFormat.h
@@ -17,9 +17,6 @@ namespace mu
 	typedef Ptr<NodeImageFormat> NodeImageFormatPtr;
 	typedef Ptr<const NodeImageFormat> NodeImageFormatPtrConst;
 
-	class InputArchive;
-	class OutputArchive;
-
 
 	//! Node that composes a new image by gathering pixel data from channels in other images.
 	//! \ingroup model
@@ -27,16 +24,8 @@ namespace mu
 	{
 	public:
 
-		//-----------------------------------------------------------------------------------------
-		// Life cycle
-		//-----------------------------------------------------------------------------------------
-
 		NodeImageFormat();
 
-		void SerialiseWrapper(OutputArchive& arch) const override;
-		static void Serialise( const NodeImageFormat* pNode, OutputArchive& arch );
-		static NodeImageFormatPtr StaticUnserialise( InputArchive& arch );
-
 
 		//-----------------------------------------------------------------------------------------
 		// Node Interface
diff --git a/Engine/Plugins/Experimental/Mutable/Source/MutableTools/Public/MuT/NodeImageGradient.h b/Engine/Plugins/Experimental/Mutable/Source/MutableTools/Public/MuT/NodeImageGradient.h
--- a/Engine/Plugins/Experimental/Mutable/Source/MutableTools/Public/MuT/NodeImageGradient.h
+++ b/Engine/Plugins/Experimental/Mutable/Source/MutableTools/Public/MuT/NodeImageGradient.h
@@ -27,16 +27,8 @@ namespace mu
 	{
 	public:
 
-		//-----------------------------------------------------------------------------------------
-		// Life cycle
-		//-----------------------------------------------------------------------------------------
-
 		NodeImageGradient();
 
-		void SerialiseWrapper(OutputArchive& arch) const override;
-		static void Serialise( const NodeImageGradient* pNode, OutputArchive& arch );
-		static NodeImageGradientPtr StaticUnserialise( InputArchive& arch );
-
 
 		//-----------------------------------------------------------------------------------------
 		// Node Interface
diff --git a/Engine/Plugins/Experimental/Mutable/Source/MutableTools/Public/MuT/NodeImageInterpolate.h b/Engine/Plugins/Experimental/Mutable/Source/MutableTools/Public/MuT/NodeImageInterpolate.h
--- a/Engine/Plugins/Experimental/Mutable/Source/MutableTools/Public/MuT/NodeImageInterpolate.h
+++ b/Engine/Plugins/Experimental/Mutable/Source/MutableTools/Public/MuT/NodeImageInterpolate.h
@@ -32,16 +32,8 @@ namespace mu
 	{
 	public:
 
-		//-----------------------------------------------------------------------------------------
-		// Life cycle
-		//-----------------------------------------------------------------------------------------
-
 		NodeImageInterpolate();
 
-		void SerialiseWrapper(OutputArchive& arch) const override;
-		static void Serialise( const NodeImageInterpolate* pNode, OutputArchive& arch );
-		static NodeImageInterpolatePtr StaticUnserialise( InputArchive& arch );
-
 
 		//-----------------------------------------------------------------------------------------
 		// Node Interface
diff --git a/Engine/Plugins/Experimental/Mutable/Source/MutableTools/Public/MuT/NodeImageInvert.h b/Engine/Plugins/Experimental/Mutable/Source/MutableTools/Public/MuT/NodeImageInvert.h
--- a/Engine/Plugins/Experimental/Mutable/Source/MutableTools/Public/MuT/NodeImageInvert.h
+++ b/Engine/Plugins/Experimental/Mutable/Source/MutableTools/Public/MuT/NodeImageInvert.h
@@ -17,22 +17,12 @@ namespace mu
 	typedef Ptr<const NodeImageInvert> NodeImageInvertPtrConst;
 
 	//! Node that inverts the colors of an image, channel by channel
-	//! \ingroup model
-
 	class MUTABLETOOLS_API NodeImageInvert : public NodeImage
 	{
 	public:
 
-		//-----------------------------------------------------------------------------------------
-		// Life cycle
-		//-----------------------------------------------------------------------------------------
-
 		NodeImageInvert();
 
-		static void Serialise(const NodeImageInvert* pNode, OutputArchive& arch);
-		void SerialiseWrapper(OutputArchive& arch) const override;
-		static NodeImageInvertPtr StaticUnserialise(InputArchive& arch);
-
 
 		//-----------------------------------------------------------------------------------------
 		// Node Interface
diff --git a/Engine/Plugins/Experimental/Mutable/Source/MutableTools/Public/MuT/NodeImageLayer.h b/Engine/Plugins/Experimental/Mutable/Source/MutableTools/Public/MuT/NodeImageLayer.h
--- a/Engine/Plugins/Experimental/Mutable/Source/MutableTools/Public/MuT/NodeImageLayer.h
+++ b/Engine/Plugins/Experimental/Mutable/Source/MutableTools/Public/MuT/NodeImageLayer.h
@@ -24,17 +24,8 @@ namespace mu
 	{
 	public:
 
-		//-----------------------------------------------------------------------------------------
-		// Life cycle
-		//-----------------------------------------------------------------------------------------
-
 		NodeImageLayer();
 
-		void SerialiseWrapper(OutputArchive& arch) const override;
-		static void Serialise( const NodeImageLayer* pNode, OutputArchive& arch );
-		static NodeImageLayerPtr StaticUnserialise( InputArchive& arch );
-
-
 		//-----------------------------------------------------------------------------------------
 		// Node Interface
 		//-----------------------------------------------------------------------------------------
diff --git a/Engine/Plugins/Experimental/Mutable/Source/MutableTools/Public/MuT/NodeImageLayerColour.h b/Engine/Plugins/Experimental/Mutable/Source/MutableTools/Public/MuT/NodeImageLayerColour.h
--- a/Engine/Plugins/Experimental/Mutable/Source/MutableTools/Public/MuT/NodeImageLayerColour.h
+++ b/Engine/Plugins/Experimental/Mutable/Source/MutableTools/Public/MuT/NodeImageLayerColour.h
@@ -20,9 +20,6 @@ namespace mu
 	typedef Ptr<NodeImageLayerColour> NodeImageLayerColourPtr;
 	typedef Ptr<const NodeImageLayerColour> NodeImageLayerColourPtrConst;
 
-	class InputArchive;
-	class OutputArchive;
-
 
 	//! This node applies a layer blending effect on a base image using a mask and a colour.
 	//! \ingroup model
@@ -30,17 +27,8 @@ namespace mu
 	{
 	public:
 
-		//-----------------------------------------------------------------------------------------
-		// Life cycle
-		//-----------------------------------------------------------------------------------------
-
 		NodeImageLayerColour();
 
-		void SerialiseWrapper(OutputArchive& arch) const override;
-		static void Serialise( const NodeImageLayerColour* pNode, OutputArchive& arch );
-		static NodeImageLayerColourPtr StaticUnserialise( InputArchive& arch );
-
-
 		//-----------------------------------------------------------------------------------------
 		// Node Interface
 		//-----------------------------------------------------------------------------------------
diff --git a/Engine/Plugins/Experimental/Mutable/Source/MutableTools/Public/MuT/NodeImageLuminance.h b/Engine/Plugins/Experimental/Mutable/Source/MutableTools/Public/MuT/NodeImageLuminance.h
--- a/Engine/Plugins/Experimental/Mutable/Source/MutableTools/Public/MuT/NodeImageLuminance.h
+++ b/Engine/Plugins/Experimental/Mutable/Source/MutableTools/Public/MuT/NodeImageLuminance.h
@@ -16,9 +16,6 @@ namespace mu
 	typedef Ptr<NodeImageLuminance> NodeImageLuminancePtr;
 	typedef Ptr<const NodeImageLuminance> NodeImageLuminancePtrConst;
 
-	class InputArchive;
-	class OutputArchive;
-
 
 	//! Calculate the luminance of an image into a new single-channel image..
 	//! \ingroup model
@@ -26,17 +23,8 @@ namespace mu
 	{
 	public:
 
-		//-----------------------------------------------------------------------------------------
-		// Life cycle
-		//-----------------------------------------------------------------------------------------
-
 		NodeImageLuminance();
 
-		void SerialiseWrapper(OutputArchive& arch) const override;
-		static void Serialise( const NodeImageLuminance* pNode, OutputArchive& arch );
-		static NodeImageLuminancePtr StaticUnserialise( InputArchive& arch );
-
-
 		//-----------------------------------------------------------------------------------------
 		// Node Interface
 		//-----------------------------------------------------------------------------------------
diff --git a/Engine/Plugins/Experimental/Mutable/Source/MutableTools/Public/MuT/NodeImageMipmap.h b/Engine/Plugins/Experimental/Mutable/Source/MutableTools/Public/MuT/NodeImageMipmap.h
--- a/Engine/Plugins/Experimental/Mutable/Source/MutableTools/Public/MuT/NodeImageMipmap.h
+++ b/Engine/Plugins/Experimental/Mutable/Source/MutableTools/Public/MuT/NodeImageMipmap.h
@@ -28,17 +28,8 @@ namespace mu
 	{
 	public:
 
-		//-----------------------------------------------------------------------------------------
-		// Life cycle
-		//-----------------------------------------------------------------------------------------
-
         NodeImageMipmap();
 
-		void SerialiseWrapper(OutputArchive& arch) const override;
-		static void Serialise( const NodeImageMipmap* pNode, OutputArchive& arch );
-        static NodeImageMipmapPtr StaticUnserialise( InputArchive& arch );
-
-
 		//-----------------------------------------------------------------------------------------
 		// Node Interface
 		//-----------------------------------------------------------------------------------------
diff --git a/Engine/Plugins/Experimental/Mutable/Source/MutableTools/Public/MuT/NodeImageMultiLayer.h b/Engine/Plugins/Experimental/Mutable/Source/MutableTools/Public/MuT/NodeImageMultiLayer.h
--- a/Engine/Plugins/Experimental/Mutable/Source/MutableTools/Public/MuT/NodeImageMultiLayer.h
+++ b/Engine/Plugins/Experimental/Mutable/Source/MutableTools/Public/MuT/NodeImageMultiLayer.h
@@ -18,9 +18,6 @@ namespace mu
     typedef Ptr<NodeImageMultiLayer> NodeImageMultiLayerPtr;
     typedef Ptr<const NodeImageMultiLayer> NodeImageMultiLayerPtrConst;
 
-	class InputArchive;
-	class OutputArchive;
-
 
     //! This node applies any numbre of layer blending effect on a base image using a mask and a
     //! blended image. The number of layers depends on a scalar input of this node.
@@ -29,17 +26,8 @@ namespace mu
 	{
 	public:
 
-		//-----------------------------------------------------------------------------------------
-		// Life cycle
-		//-----------------------------------------------------------------------------------------
-
         NodeImageMultiLayer();
 
-		void SerialiseWrapper(OutputArchive& arch) const override;
-		static void Serialise( const NodeImageMultiLayer* pNode, OutputArchive& arch );
-        static NodeImageMultiLayerPtr StaticUnserialise( InputArchive& arch );
-
-
 		//-----------------------------------------------------------------------------------------
 		// Node Interface
 		//-----------------------------------------------------------------------------------------
diff --git a/Engine/Plugins/Experimental/Mutable/Source/MutableTools/Public/MuT/NodeImageNormalComposite.h b/Engine/Plugins/Experimental/Mutable/Source/MutableTools/Public/MuT/NodeImageNormalComposite.h
--- a/Engine/Plugins/Experimental/Mutable/Source/MutableTools/Public/MuT/NodeImageNormalComposite.h
+++ b/Engine/Plugins/Experimental/Mutable/Source/MutableTools/Public/MuT/NodeImageNormalComposite.h
@@ -17,19 +17,12 @@ namespace mu
 	typedef Ptr<const NodeImageNormalComposite> NodeImageNormalCompositePtrConst;
 
 	//! Node that inverts the colors of an image, channel by channel
-	//! \ingroup model
-
 	class MUTABLETOOLS_API NodeImageNormalComposite : public NodeImage
 	{
 	public:
 
 		NodeImageNormalComposite();
 
-		void SerialiseWrapper(OutputArchive& arch) const override;
-		static void Serialise(const NodeImageNormalComposite* pNode, OutputArchive& arch);
-		static NodeImageNormalCompositePtr StaticUnserialise(InputArchive& arch);
-
-
 		//-----------------------------------------------------------------------------------------
 		// Node Interface
 		//-----------------------------------------------------------------------------------------
diff --git a/Engine/Plugins/Experimental/Mutable/Source/MutableTools/Public/MuT/NodeImageParameter.h b/Engine/Plugins/Experimental/Mutable/Source/MutableTools/Public/MuT/NodeImageParameter.h
--- a/Engine/Plugins/Experimental/Mutable/Source/MutableTools/Public/MuT/NodeImageParameter.h
+++ b/Engine/Plugins/Experimental/Mutable/Source/MutableTools/Public/MuT/NodeImageParameter.h
@@ -10,8 +10,6 @@
 
 namespace mu
 {
-	class InputArchive;
-	class OutputArchive;
 	class NodeRange;
 
     //! Node that defines a Image model parameter.
@@ -20,17 +18,8 @@ namespace mu
 	{
 	public:
 
-		//-----------------------------------------------------------------------------------------
-		// Life cycle
-		//-----------------------------------------------------------------------------------------
-
         NodeImageParameter();
 
-		void SerialiseWrapper(OutputArchive& arch) const override;
-		static void Serialise( const NodeImageParameter* pNode, OutputArchive& arch );
-        static Ptr<NodeImageParameter> StaticUnserialise( InputArchive& arch );
-
-
 		//-----------------------------------------------------------------------------------------
 		// Node Interface
 		//-----------------------------------------------------------------------------------------
diff --git a/Engine/Plugins/Experimental/Mutable/Source/MutableTools/Public/MuT/NodeImagePlainColour.h b/Engine/Plugins/Experimental/Mutable/Source/MutableTools/Public/MuT/NodeImagePlainColour.h
--- a/Engine/Plugins/Experimental/Mutable/Source/MutableTools/Public/MuT/NodeImagePlainColour.h
+++ b/Engine/Plugins/Experimental/Mutable/Source/MutableTools/Public/MuT/NodeImagePlainColour.h
@@ -20,9 +20,6 @@ namespace mu
 	typedef Ptr<NodeImagePlainColour> NodeImagePlainColourPtr;
 	typedef Ptr<const NodeImagePlainColour> NodeImagePlainColourPtrConst;
 
-	class InputArchive;
-	class OutputArchive;
-
 
 	//! Node that multiplies the colors of an image, channel by channel.
 	//! \ingroup model
@@ -30,16 +27,8 @@ namespace mu
 	{
 	public:
 
-		//-----------------------------------------------------------------------------------------
-		// Life cycle
-		//-----------------------------------------------------------------------------------------
-
 		NodeImagePlainColour();
 
-		void SerialiseWrapper(OutputArchive& arch) const override;
-		static void Serialise( const NodeImagePlainColour* pNode, OutputArchive& arch );
-		static NodeImagePlainColourPtr StaticUnserialise( InputArchive& arch );
-
 
 		//-----------------------------------------------------------------------------------------
 		// Node Interface
diff --git a/Engine/Plugins/Experimental/Mutable/Source/MutableTools/Public/MuT/NodeImageProject.h b/Engine/Plugins/Experimental/Mutable/Source/MutableTools/Public/MuT/NodeImageProject.h
--- a/Engine/Plugins/Experimental/Mutable/Source/MutableTools/Public/MuT/NodeImageProject.h
+++ b/Engine/Plugins/Experimental/Mutable/Source/MutableTools/Public/MuT/NodeImageProject.h
@@ -3,7 +3,6 @@
 
 #pragma once
 
-#include "Math/UnrealMathSSE.h"
 #include "MuR/Ptr.h"
 #include "MuR/RefCounted.h"
 #include "MuT/Node.h"
@@ -15,12 +14,9 @@ namespace mu
 
 	// Forward definitions
 	class NodeScalar;
-	class NodeColour;
 	class NodeMesh;
 	class NodeProjector;
 	class NodeImageProject;
-	class InputArchive;
-	class OutputArchive;
 
 	//!
 	//! \ingroup model
@@ -28,17 +24,8 @@ namespace mu
 	{
 	public:
 
-		//-----------------------------------------------------------------------------------------
-		// Life cycle
-		//-----------------------------------------------------------------------------------------
-
 		NodeImageProject();
 
-		void SerialiseWrapper(OutputArchive& arch) const override;
-		static void Serialise( const NodeImageProject* pNode, OutputArchive& arch );
-		static Ptr<NodeImageProject> StaticUnserialise( InputArchive& arch );
-
-
 		//-----------------------------------------------------------------------------------------
 		// Node Interface
 		//-----------------------------------------------------------------------------------------
diff --git a/Engine/Plugins/Experimental/Mutable/Source/MutableTools/Public/MuT/NodeImageResize.h b/Engine/Plugins/Experimental/Mutable/Source/MutableTools/Public/MuT/NodeImageResize.h
--- a/Engine/Plugins/Experimental/Mutable/Source/MutableTools/Public/MuT/NodeImageResize.h
+++ b/Engine/Plugins/Experimental/Mutable/Source/MutableTools/Public/MuT/NodeImageResize.h
@@ -20,9 +20,6 @@ namespace mu
 	typedef Ptr<NodeImageResize> NodeImageResizePtr;
 	typedef Ptr<const NodeImageResize> NodeImageResizePtrConst;
 
-	class InputArchive;
-	class OutputArchive;
-
 
 	//! Node that multiplies the colors of an image, channel by channel.
 	//! \ingroup model
@@ -30,16 +27,8 @@ namespace mu
 	{
 	public:
 
-		//-----------------------------------------------------------------------------------------
-		// Life cycle
-		//-----------------------------------------------------------------------------------------
-
 		NodeImageResize();
 
-		void SerialiseWrapper(OutputArchive& arch) const override;
-		static void Serialise( const NodeImageResize* pNode, OutputArchive& arch );
-		static NodeImageResizePtr StaticUnserialise( InputArchive& arch );
-
 
 		//-----------------------------------------------------------------------------------------
 		// Node Interface
diff --git a/Engine/Plugins/Experimental/Mutable/Source/MutableTools/Public/MuT/NodeImageSaturate.h b/Engine/Plugins/Experimental/Mutable/Source/MutableTools/Public/MuT/NodeImageSaturate.h
--- a/Engine/Plugins/Experimental/Mutable/Source/MutableTools/Public/MuT/NodeImageSaturate.h
+++ b/Engine/Plugins/Experimental/Mutable/Source/MutableTools/Public/MuT/NodeImageSaturate.h
@@ -20,29 +20,16 @@ namespace mu
 	typedef Ptr<NodeImageSaturate> NodeImageSaturatePtr;
 	typedef Ptr<const NodeImageSaturate> NodeImageSaturatePtrConst;
 
-	class InputArchive;
-	class OutputArchive;
-
 
 	//! Change the saturation of an image. This node can be used to increase the saturation with a
 	//! factor bigger than 1, or to decrease it or desaturate it completely with a factor smaller
 	//! than 1 or 0.
-	//! \ingroup model
 	class MUTABLETOOLS_API NodeImageSaturate : public NodeImage
 	{
 	public:
 
-		//-----------------------------------------------------------------------------------------
-		// Life cycle
-		//-----------------------------------------------------------------------------------------
-
 		NodeImageSaturate();
 
-		void SerialiseWrapper(OutputArchive& arch) const override;
-		static void Serialise( const NodeImageSaturate* pNode, OutputArchive& arch );
-		static NodeImageSaturatePtr StaticUnserialise( InputArchive& arch );
-
-
 		//-----------------------------------------------------------------------------------------
 		// Node Interface
 		//-----------------------------------------------------------------------------------------
diff --git a/Engine/Plugins/Experimental/Mutable/Source/MutableTools/Public/MuT/NodeImageSwitch.h b/Engine/Plugins/Experimental/Mutable/Source/MutableTools/Public/MuT/NodeImageSwitch.h
--- a/Engine/Plugins/Experimental/Mutable/Source/MutableTools/Public/MuT/NodeImageSwitch.h
+++ b/Engine/Plugins/Experimental/Mutable/Source/MutableTools/Public/MuT/NodeImageSwitch.h
@@ -22,22 +22,12 @@ namespace mu
 
 
 	//! This node selects an output image from a set of input images based on a parameter.
-	//! \ingroup model
 	class MUTABLETOOLS_API NodeImageSwitch : public NodeImage
 	{
 	public:
 
-		//-----------------------------------------------------------------------------------------
-		// Life cycle
-		//-----------------------------------------------------------------------------------------
-
 		NodeImageSwitch();
 
-		void SerialiseWrapper(OutputArchive& arch) const override;
-		static void Serialise( const NodeImageSwitch* pNode, OutputArchive& arch );
-		static NodeImageSwitchPtr StaticUnserialise( InputArchive& arch );
-
-
 		//-----------------------------------------------------------------------------------------
 		// Node Interface
 		//-----------------------------------------------------------------------------------------
diff --git a/Engine/Plugins/Experimental/Mutable/Source/MutableTools/Public/MuT/NodeImageSwizzle.h b/Engine/Plugins/Experimental/Mutable/Source/MutableTools/Public/MuT/NodeImageSwizzle.h
--- a/Engine/Plugins/Experimental/Mutable/Source/MutableTools/Public/MuT/NodeImageSwizzle.h
+++ b/Engine/Plugins/Experimental/Mutable/Source/MutableTools/Public/MuT/NodeImageSwizzle.h
@@ -17,9 +17,6 @@ namespace mu
 	typedef Ptr<NodeImageSwizzle> NodeImageSwizzlePtr;
 	typedef Ptr<const NodeImageSwizzle> NodeImageSwizzlePtrConst;
 
-	class InputArchive;
-	class OutputArchive;
-
 
 	//! Node that composes a new image by gathering pixel data from channels in other images.
 	//! \ingroup model
@@ -27,17 +24,8 @@ namespace mu
 	{
 	public:
 
-		//-----------------------------------------------------------------------------------------
-		// Life cycle
-		//-----------------------------------------------------------------------------------------
-
 		NodeImageSwizzle();
 
-		void SerialiseWrapper(OutputArchive& arch) const override;
-		static void Serialise( const NodeImageSwizzle* pNode, OutputArchive& arch );
-		static NodeImageSwizzlePtr StaticUnserialise( InputArchive& arch );
-
-
 		//-----------------------------------------------------------------------------------------
 		// Node Interface
 		//-----------------------------------------------------------------------------------------
diff --git a/Engine/Plugins/Experimental/Mutable/Source/MutableTools/Public/MuT/NodeImageTable.h b/Engine/Plugins/Experimental/Mutable/Source/MutableTools/Public/MuT/NodeImageTable.h
--- a/Engine/Plugins/Experimental/Mutable/Source/MutableTools/Public/MuT/NodeImageTable.h
+++ b/Engine/Plugins/Experimental/Mutable/Source/MutableTools/Public/MuT/NodeImageTable.h
@@ -20,28 +20,16 @@ namespace mu
 	typedef Ptr<Table> TablePtr;
 	typedef Ptr<const Table> TablePtrConst;
 
-	class InputArchive;
-	class OutputArchive;
 	struct FImageDesc;
 
 
 	//! This node provides the meshes stored in the column of a table.
-	//! \ingroup transform
 	class MUTABLETOOLS_API NodeImageTable : public NodeImage
 	{
 	public:
 
-		//-----------------------------------------------------------------------------------------
-		// Life cycle
-		//-----------------------------------------------------------------------------------------
-
 		NodeImageTable();
 
-		void SerialiseWrapper(OutputArchive& arch) const override;
-		static void Serialise( const NodeImageTable* pNode, OutputArchive& arch );
-		static NodeImageTablePtr StaticUnserialise( InputArchive& arch );
-
-
 		//-----------------------------------------------------------------------------------------
 		// Node Interface
 		//-----------------------------------------------------------------------------------------
diff --git a/Engine/Plugins/Experimental/Mutable/Source/MutableTools/Public/MuT/NodeImageTransform.h b/Engine/Plugins/Experimental/Mutable/Source/MutableTools/Public/MuT/NodeImageTransform.h
--- a/Engine/Plugins/Experimental/Mutable/Source/MutableTools/Public/MuT/NodeImageTransform.h
+++ b/Engine/Plugins/Experimental/Mutable/Source/MutableTools/Public/MuT/NodeImageTransform.h
@@ -19,9 +19,6 @@ namespace mu
 	typedef Ptr<NodeImageTransform> NodeImageTransformPtr;
 	typedef Ptr<const NodeImageTransform> NodeImageTransformPtrConst;
 
-	class InputArchive;
-	class OutputArchive;
-
 	enum class EAddressMode;
 
 	//! Node that multiplies the colors of an image, channel by channel.
@@ -30,17 +27,8 @@ namespace mu
 	{
 	public:
 
-		//-----------------------------------------------------------------------------------------
-		// Life cycle
-		//-----------------------------------------------------------------------------------------
-
 		NodeImageTransform();
 
-		void SerialiseWrapper(OutputArchive& arch) const override;
-		static void Serialise( const NodeImageTransform* pNode, OutputArchive& arch );
-		static NodeImageTransformPtr StaticUnserialise( InputArchive& arch );
-
-
 		//-----------------------------------------------------------------------------------------
 		// Node Interface
 		//-----------------------------------------------------------------------------------------
diff --git a/Engine/Plugins/Experimental/Mutable/Source/MutableTools/Public/MuT/NodeImageVariation.h b/Engine/Plugins/Experimental/Mutable/Source/MutableTools/Public/MuT/NodeImageVariation.h
--- a/Engine/Plugins/Experimental/Mutable/Source/MutableTools/Public/MuT/NodeImageVariation.h
+++ b/Engine/Plugins/Experimental/Mutable/Source/MutableTools/Public/MuT/NodeImageVariation.h
@@ -16,9 +16,6 @@ namespace mu
     typedef Ptr<NodeImageVariation> NodeImageVariationPtr;
     typedef Ptr<const NodeImageVariation> NodeImageVariationPtrConst;
 
-	class InputArchive;
-	class OutputArchive;
-
 
     //!
     //! \ingroup model
@@ -31,10 +28,6 @@ namespace mu
 
         NodeImageVariation();
 
-		void SerialiseWrapper(OutputArchive& arch) const override;
-		static void Serialise( const NodeImageVariation* pNode, OutputArchive& arch );
-        static NodeImageVariationPtr StaticUnserialise( InputArchive& arch );
-
         //-----------------------------------------------------------------------------------------
         // Node Interface
         //-----------------------------------------------------------------------------------------
diff --git a/Engine/Plugins/Experimental/Mutable/Source/MutableTools/Public/MuT/NodeLOD.h b/Engine/Plugins/Experimental/Mutable/Source/MutableTools/Public/MuT/NodeLOD.h
--- a/Engine/Plugins/Experimental/Mutable/Source/MutableTools/Public/MuT/NodeLOD.h
+++ b/Engine/Plugins/Experimental/Mutable/Source/MutableTools/Public/MuT/NodeLOD.h
@@ -1,6 +1,5 @@
 // Copyright Epic Games, Inc. All Rights Reserved.
 
-
 #pragma once
 
 #include "MuR/Ptr.h"
@@ -26,22 +25,12 @@ namespace mu
 
 
 	//! Node that creates a new level of detail by assembling several object components.
-	//! \ingroup model
 	class MUTABLETOOLS_API NodeLOD : public Node
 	{
 	public:
 
-		//-----------------------------------------------------------------------------------------
-		// Life cycle
-		//-----------------------------------------------------------------------------------------
-
 		NodeLOD();
 
-		void SerialiseWrapper(OutputArchive& arch) const override;
-		static void Serialise( const NodeLOD* pNode, OutputArchive& arch );
-		static NodeLODPtr StaticUnserialise( InputArchive& arch );
-
-
 		//-----------------------------------------------------------------------------------------
 		// Node Interface
 		//-----------------------------------------------------------------------------------------
diff --git a/Engine/Plugins/Experimental/Mutable/Source/MutableTools/Public/MuT/NodeLayout.h b/Engine/Plugins/Experimental/Mutable/Source/MutableTools/Public/MuT/NodeLayout.h
--- a/Engine/Plugins/Experimental/Mutable/Source/MutableTools/Public/MuT/NodeLayout.h
+++ b/Engine/Plugins/Experimental/Mutable/Source/MutableTools/Public/MuT/NodeLayout.h
@@ -23,12 +23,9 @@ namespace mu
 	typedef Ptr<NodeLayoutBlocks> NodeLayoutBlocksPtr;
 	typedef Ptr<const NodeLayoutBlocks> NodeLayoutBlocksPtrConst;
 
-	class InputArchive;
-	class OutputArchive;
 
 	//! This node is used to define the texture layout for a texture coordinates channel of a
 	//! constant mesh.
-	//! \ingroup model
 	class MUTABLETOOLS_API NodeLayout : public Node
 	{
 	public:
@@ -41,14 +38,6 @@ namespace mu
 			None
 		};
 
-		//-----------------------------------------------------------------------------------------
-		// Life cycle
-		//-----------------------------------------------------------------------------------------
-
-		static void Serialise( const NodeLayout* pNode, OutputArchive& arch );
-		static NodeLayoutPtr StaticUnserialise( InputArchive& arch );
-
-
 		//-----------------------------------------------------------------------------------------
 		// Node Interface
 		//-----------------------------------------------------------------------------------------
@@ -56,10 +45,6 @@ namespace mu
 		const FNodeType* GetType() const override;
 		static const FNodeType* GetStaticType();
 
-		//-----------------------------------------------------------------------------------------
-		// Own Interface
-		//-----------------------------------------------------------------------------------------
-
 
 		//-----------------------------------------------------------------------------------------
 		// Interface pattern
@@ -80,22 +65,12 @@ namespace mu
 	//! This node is used to define the texture layout for a texture coordinates channel of a
 	//! constant mesh.
 	//! The blocks defined here will be set to the channel of the mesh were this node is connected.
-	//! \ingroup model
 	class MUTABLETOOLS_API NodeLayoutBlocks : public NodeLayout
 	{
 	public:
 
-		//-----------------------------------------------------------------------------------------
-		// Life cycle
-		//-----------------------------------------------------------------------------------------
-
 		NodeLayoutBlocks();
 
-		void SerialiseWrapper(OutputArchive& arch) const override;
-		static void Serialise( const NodeLayoutBlocks* pNode, OutputArchive& arch );
-		static NodeLayoutBlocksPtr StaticUnserialise( InputArchive& arch );
-
-
 		//-----------------------------------------------------------------------------------------
 		// Node Interface
 		//-----------------------------------------------------------------------------------------
diff --git a/Engine/Plugins/Experimental/Mutable/Source/MutableTools/Public/MuT/NodeMesh.h b/Engine/Plugins/Experimental/Mutable/Source/MutableTools/Public/MuT/NodeMesh.h
--- a/Engine/Plugins/Experimental/Mutable/Source/MutableTools/Public/MuT/NodeMesh.h
+++ b/Engine/Plugins/Experimental/Mutable/Source/MutableTools/Public/MuT/NodeMesh.h
@@ -17,7 +17,6 @@ namespace mu
 
 
     //! %Base class of any node that outputs a mesh.
-	//! \ingroup model
 	class MUTABLETOOLS_API NodeMesh : public Node
 	{
 	public:
@@ -47,14 +46,6 @@ namespace mu
 			None
 		};
 
-		//-----------------------------------------------------------------------------------------
-		// Life cycle
-		//-----------------------------------------------------------------------------------------
-
-		static void Serialise( const NodeMesh* pNode, OutputArchive& arch );
-		static NodeMeshPtr StaticUnserialise( InputArchive& arch );
-
-
 		//-----------------------------------------------------------------------------------------
 		// Node Interface
 		//-----------------------------------------------------------------------------------------
diff --git a/Engine/Plugins/Experimental/Mutable/Source/MutableTools/Public/MuT/NodeMeshApplyPose.h b/Engine/Plugins/Experimental/Mutable/Source/MutableTools/Public/MuT/NodeMeshApplyPose.h
--- a/Engine/Plugins/Experimental/Mutable/Source/MutableTools/Public/MuT/NodeMeshApplyPose.h
+++ b/Engine/Plugins/Experimental/Mutable/Source/MutableTools/Public/MuT/NodeMeshApplyPose.h
@@ -13,26 +13,17 @@ namespace mu
 
 	// Forward definitions
     class NodeMeshApplyPose;
-    using NodeMeshApplyPosePtr = Ptr<NodeMeshApplyPose> ;
+    using NodeMeshApplyPosePtr = Ptr<NodeMeshApplyPose>;
     using NodeMeshApplyPosePtrConst = Ptr<const NodeMeshApplyPose>;
 
 
     //! Node that applies a pose to a mesh, baking it into the vertex data
-	//! \ingroup model
     class MUTABLETOOLS_API NodeMeshApplyPose : public NodeMesh
 	{
 	public:
 
-		//-----------------------------------------------------------------------------------------
-		// Life cycle
-		//-----------------------------------------------------------------------------------------
-
         NodeMeshApplyPose();
 
-		void SerialiseWrapper(OutputArchive& arch) const override;
-		static void Serialise( const NodeMeshApplyPose* pNode, OutputArchive& arch );
-        static NodeMeshApplyPosePtr StaticUnserialise( InputArchive& arch );
-
 
 		//-----------------------------------------------------------------------------------------
 		// Node Interface
diff --git a/Engine/Plugins/Experimental/Mutable/Source/MutableTools/Public/MuT/NodeMeshClipDeform.h b/Engine/Plugins/Experimental/Mutable/Source/MutableTools/Public/MuT/NodeMeshClipDeform.h
--- a/Engine/Plugins/Experimental/Mutable/Source/MutableTools/Public/MuT/NodeMeshClipDeform.h
+++ b/Engine/Plugins/Experimental/Mutable/Source/MutableTools/Public/MuT/NodeMeshClipDeform.h
@@ -16,17 +16,8 @@ namespace mu
 	{
 	public:
 
-		//-----------------------------------------------------------------------------------------
-		// Life cycle
-		//-----------------------------------------------------------------------------------------
-
 		NodeMeshClipDeform();
 
-		void SerialiseWrapper(OutputArchive& arch) const override;
-		static void Serialise( const NodeMeshClipDeform* pNode, OutputArchive& arch );
-		static Ptr<NodeMeshClipDeform> StaticUnserialise( InputArchive& arch );
-
-
 		//-----------------------------------------------------------------------------------------
 		// Node Interface
 		//-----------------------------------------------------------------------------------------
diff --git a/Engine/Plugins/Experimental/Mutable/Source/MutableTools/Public/MuT/NodeMeshClipMorphPlane.h b/Engine/Plugins/Experimental/Mutable/Source/MutableTools/Public/MuT/NodeMeshClipMorphPlane.h
--- a/Engine/Plugins/Experimental/Mutable/Source/MutableTools/Public/MuT/NodeMeshClipMorphPlane.h
+++ b/Engine/Plugins/Experimental/Mutable/Source/MutableTools/Public/MuT/NodeMeshClipMorphPlane.h
@@ -19,21 +19,12 @@ namespace mu
 	struct FBoneName;
 
     //! This node applies a geometric transform represented by a 4x4 matrix to a mesh
-	//! \ingroup model
     class MUTABLETOOLS_API NodeMeshClipMorphPlane : public NodeMesh
 	{
 	public:
 
-		//-----------------------------------------------------------------------------------------
-		// Life cycle
-		//-----------------------------------------------------------------------------------------
-
 		NodeMeshClipMorphPlane();
 
-		void SerialiseWrapper(OutputArchive& arch) const override;
-		static void Serialise( const NodeMeshClipMorphPlane* pNode, OutputArchive& arch );
-        static NodeMeshClipMorphPlanePtr StaticUnserialise( InputArchive& arch );
-
 		//-----------------------------------------------------------------------------------------
 		// Node Interface
 		//-----------------------------------------------------------------------------------------
@@ -49,7 +40,6 @@ namespace mu
         NodeMeshPtr GetSource() const;
 		void SetSource( NodeMesh* );
 
-        //! \param 
 		void SetPlane(float centerX, float centerY, float centerZ, float normalX, float normalY, float normalZ);
 		void SetParams(float dist, float factor);
 		void SetMorphEllipse(float radius1, float radius2, float rotation);
diff --git a/Engine/Plugins/Experimental/Mutable/Source/MutableTools/Public/MuT/NodeMeshClipWithMesh.h b/Engine/Plugins/Experimental/Mutable/Source/MutableTools/Public/MuT/NodeMeshClipWithMesh.h
--- a/Engine/Plugins/Experimental/Mutable/Source/MutableTools/Public/MuT/NodeMeshClipWithMesh.h
+++ b/Engine/Plugins/Experimental/Mutable/Source/MutableTools/Public/MuT/NodeMeshClipWithMesh.h
@@ -23,16 +23,8 @@ namespace mu
 	{
 	public:
 
-		//-----------------------------------------------------------------------------------------
-		// Life cycle
-		//-----------------------------------------------------------------------------------------
-
 		NodeMeshClipWithMesh();
 
-		void SerialiseWrapper(OutputArchive& arch) const override;
-		static void Serialise( const NodeMeshClipWithMesh* pNode, OutputArchive& arch );
-        static NodeMeshClipWithMeshPtr StaticUnserialise( InputArchive& arch );
-
 		//-----------------------------------------------------------------------------------------
 		// Node Interface
 		//-----------------------------------------------------------------------------------------
@@ -48,7 +40,6 @@ namespace mu
         NodeMeshPtr GetSource() const;
 		void SetSource(NodeMesh*);
 
-        //! \param 
 		void SetClipMesh(NodeMesh*);
 
 		//! Add a tag to the clip morph operation, which will only affect surfaces with the same tag
diff --git a/Engine/Plugins/Experimental/Mutable/Source/MutableTools/Public/MuT/NodeMeshConstant.h b/Engine/Plugins/Experimental/Mutable/Source/MutableTools/Public/MuT/NodeMeshConstant.h
--- a/Engine/Plugins/Experimental/Mutable/Source/MutableTools/Public/MuT/NodeMeshConstant.h
+++ b/Engine/Plugins/Experimental/Mutable/Source/MutableTools/Public/MuT/NodeMeshConstant.h
@@ -24,22 +24,12 @@ namespace mu
 
 	//! Node that outputs a constant mesh.
 	//! It allows to define the layouts for the texture channels of the constant mesh
-	//! \ingroup model
 	class MUTABLETOOLS_API NodeMeshConstant : public NodeMesh
 	{
 	public:
 
-		//-----------------------------------------------------------------------------------------
-		// Life cycle
-		//-----------------------------------------------------------------------------------------
-
 		NodeMeshConstant();
 
-		void SerialiseWrapper(OutputArchive& arch) const override;
-		static void Serialise( const NodeMeshConstant* pNode, OutputArchive& arch );
-		static NodeMeshConstantPtr StaticUnserialise( InputArchive& arch );
-
-
 		//-----------------------------------------------------------------------------------------
 		// Node Interface
 		//-----------------------------------------------------------------------------------------
diff --git a/Engine/Plugins/Experimental/Mutable/Source/MutableTools/Public/MuT/NodeMeshFormat.h b/Engine/Plugins/Experimental/Mutable/Source/MutableTools/Public/MuT/NodeMeshFormat.h
--- a/Engine/Plugins/Experimental/Mutable/Source/MutableTools/Public/MuT/NodeMeshFormat.h
+++ b/Engine/Plugins/Experimental/Mutable/Source/MutableTools/Public/MuT/NodeMeshFormat.h
@@ -11,26 +11,15 @@ namespace mu
 {
 
 	// Forward definitions
-	class InputArchive;
 	class FMeshBufferSet;
-	class OutputArchive;
 
-	/** This node can change the buffer formats of a mesh vertices, indices and faces.
-	*/
+	/** This node can change the buffer formats of a mesh vertices, indices and faces. */
 	class MUTABLETOOLS_API NodeMeshFormat : public NodeMesh
 	{
 	public:
 
-		//-----------------------------------------------------------------------------------------
-		// Life cycle
-		//-----------------------------------------------------------------------------------------
-
 		NodeMeshFormat();
 
-		void SerialiseWrapper(OutputArchive& arch) const override;
-		static void Serialise( const NodeMeshFormat* pNode, OutputArchive& arch );
-		static Ptr<NodeMeshFormat> StaticUnserialise( InputArchive& arch );
-
 		//-----------------------------------------------------------------------------------------
 		// Node Interface
 		//-----------------------------------------------------------------------------------------
diff --git a/Engine/Plugins/Experimental/Mutable/Source/MutableTools/Public/MuT/NodeMeshFragment.h b/Engine/Plugins/Experimental/Mutable/Source/MutableTools/Public/MuT/NodeMeshFragment.h
--- a/Engine/Plugins/Experimental/Mutable/Source/MutableTools/Public/MuT/NodeMeshFragment.h
+++ b/Engine/Plugins/Experimental/Mutable/Source/MutableTools/Public/MuT/NodeMeshFragment.h
@@ -23,16 +23,8 @@ namespace mu
 	{
 	public:
 
-		//-----------------------------------------------------------------------------------------
-		// Life cycle
-		//-----------------------------------------------------------------------------------------
-
 		NodeMeshFragment();
 
-		void SerialiseWrapper(OutputArchive& arch) const override;
-		static void Serialise( const NodeMeshFragment* pNode, OutputArchive& arch );
-		static NodeMeshFragmentPtr StaticUnserialise( InputArchive& arch );
-
 
 		//-----------------------------------------------------------------------------------------
 		// Node Interface
diff --git a/Engine/Plugins/Experimental/Mutable/Source/MutableTools/Public/MuT/NodeMeshGeometryOperation.h b/Engine/Plugins/Experimental/Mutable/Source/MutableTools/Public/MuT/NodeMeshGeometryOperation.h
--- a/Engine/Plugins/Experimental/Mutable/Source/MutableTools/Public/MuT/NodeMeshGeometryOperation.h
+++ b/Engine/Plugins/Experimental/Mutable/Source/MutableTools/Public/MuT/NodeMeshGeometryOperation.h
@@ -22,22 +22,12 @@ namespace mu
 
 
 	//! Node that morphs a base mesh with one or two weighted targets from a sequence.
-	//! \ingroup model
 	class MUTABLETOOLS_API NodeMeshGeometryOperation : public NodeMesh
 	{
 	public:
 
-		//-----------------------------------------------------------------------------------------
-		// Life cycle
-		//-----------------------------------------------------------------------------------------
-
 		NodeMeshGeometryOperation();
 
-		void SerialiseWrapper(OutputArchive& arch) const override;
-		static void Serialise( const NodeMeshGeometryOperation* pNode, OutputArchive& arch );
-		static NodeMeshGeometryOperationPtr StaticUnserialise( InputArchive& arch );
-
-
 		//-----------------------------------------------------------------------------------------
 		// Node Interface
 		//-----------------------------------------------------------------------------------------
diff --git a/Engine/Plugins/Experimental/Mutable/Source/MutableTools/Public/MuT/NodeMeshInterpolate.h b/Engine/Plugins/Experimental/Mutable/Source/MutableTools/Public/MuT/NodeMeshInterpolate.h
--- a/Engine/Plugins/Experimental/Mutable/Source/MutableTools/Public/MuT/NodeMeshInterpolate.h
+++ b/Engine/Plugins/Experimental/Mutable/Source/MutableTools/Public/MuT/NodeMeshInterpolate.h
@@ -28,21 +28,12 @@ namespace mu
 	//! - if three meshes are set, the first one is 0.0 the second one is 0.5 and the third one 1.0
 	//! regardless of the input slots used: If B and C are set, but not A, B will be at weight 0
 	//! and C will be at weight 1
-	//! \ingroup model
 	class MUTABLETOOLS_API NodeMeshInterpolate : public NodeMesh
 	{
 	public:
 
-		//-----------------------------------------------------------------------------------------
-		// Life cycle
-		//-----------------------------------------------------------------------------------------
-
 		NodeMeshInterpolate();
 
-		void SerialiseWrapper(OutputArchive& arch) const override;
-		static void Serialise( const NodeMeshInterpolate* pNode, OutputArchive& arch );
-		static NodeMeshInterpolatePtr StaticUnserialise( InputArchive& arch );
-
 
 		//-----------------------------------------------------------------------------------------
 		// Node Interface
diff --git a/Engine/Plugins/Experimental/Mutable/Source/MutableTools/Public/MuT/NodeMeshMakeMorph.h b/Engine/Plugins/Experimental/Mutable/Source/MutableTools/Public/MuT/NodeMeshMakeMorph.h
--- a/Engine/Plugins/Experimental/Mutable/Source/MutableTools/Public/MuT/NodeMeshMakeMorph.h
+++ b/Engine/Plugins/Experimental/Mutable/Source/MutableTools/Public/MuT/NodeMeshMakeMorph.h
@@ -18,21 +18,12 @@ namespace mu
 
 
 	//! This node removes from a mesh A all the faces that are also part of a mesh B.
-	//! \ingroup model
     class MUTABLETOOLS_API NodeMeshMakeMorph : public NodeMesh
 	{
 	public:
 
-		//-----------------------------------------------------------------------------------------
-		// Life cycle
-		//-----------------------------------------------------------------------------------------
-
         NodeMeshMakeMorph();
 
-		void SerialiseWrapper(OutputArchive& arch) const override;
-		static void Serialise( const NodeMeshMakeMorph* pNode, OutputArchive& arch );
-        static NodeMeshMakeMorphPtr StaticUnserialise( InputArchive& arch );
-
 		//-----------------------------------------------------------------------------------------
 		// Node Interface
 		//-----------------------------------------------------------------------------------------
diff --git a/Engine/Plugins/Experimental/Mutable/Source/MutableTools/Public/MuT/NodeMeshMorph.h b/Engine/Plugins/Experimental/Mutable/Source/MutableTools/Public/MuT/NodeMeshMorph.h
--- a/Engine/Plugins/Experimental/Mutable/Source/MutableTools/Public/MuT/NodeMeshMorph.h
+++ b/Engine/Plugins/Experimental/Mutable/Source/MutableTools/Public/MuT/NodeMeshMorph.h
@@ -20,17 +20,8 @@ namespace mu
 	{
 	public:
 
-		//-----------------------------------------------------------------------------------------
-		// Life cycle
-		//-----------------------------------------------------------------------------------------
-
 		NodeMeshMorph();
 
-		void SerialiseWrapper(OutputArchive& arch) const override;
-		static void Serialise( const NodeMeshMorph* pNode, OutputArchive& arch );
-		static Ptr<NodeMeshMorph> StaticUnserialise( InputArchive& arch );
-
-
 		//-----------------------------------------------------------------------------------------
 		// Node Interface
 		//-----------------------------------------------------------------------------------------		
diff --git a/Engine/Plugins/Experimental/Mutable/Source/MutableTools/Public/MuT/NodeMeshReshape.h b/Engine/Plugins/Experimental/Mutable/Source/MutableTools/Public/MuT/NodeMeshReshape.h
--- a/Engine/Plugins/Experimental/Mutable/Source/MutableTools/Public/MuT/NodeMeshReshape.h
+++ b/Engine/Plugins/Experimental/Mutable/Source/MutableTools/Public/MuT/NodeMeshReshape.h
@@ -13,22 +13,12 @@ namespace mu
 	struct FBoneName;
 
 	//! Node that morphs a base mesh with one or two weighted targets from a sequence.
-	//! \ingroup model
 	class MUTABLETOOLS_API NodeMeshReshape : public NodeMesh
 	{
 	public:
 
-		//-----------------------------------------------------------------------------------------
-		// Life cycle
-		//-----------------------------------------------------------------------------------------
-
 		NodeMeshReshape();
 
-		void SerialiseWrapper(OutputArchive& arch) const override;
-		static void Serialise( const NodeMeshReshape* pNode, OutputArchive& arch );
-		static Ptr<NodeMeshReshape> StaticUnserialise( InputArchive& arch );
-
-
 		//-----------------------------------------------------------------------------------------
 		// Node Interface
 		//-----------------------------------------------------------------------------------------
diff --git a/Engine/Plugins/Experimental/Mutable/Source/MutableTools/Public/MuT/NodeMeshSwitch.h b/Engine/Plugins/Experimental/Mutable/Source/MutableTools/Public/MuT/NodeMeshSwitch.h
--- a/Engine/Plugins/Experimental/Mutable/Source/MutableTools/Public/MuT/NodeMeshSwitch.h
+++ b/Engine/Plugins/Experimental/Mutable/Source/MutableTools/Public/MuT/NodeMeshSwitch.h
@@ -22,22 +22,12 @@ namespace mu
 
 
 	//! This node selects an output Mesh from a set of input Meshs based on a parameter.
-	//! \ingroup model
 	class MUTABLETOOLS_API NodeMeshSwitch : public NodeMesh
 	{
 	public:
 
-		//-----------------------------------------------------------------------------------------
-		// Life cycle
-		//-----------------------------------------------------------------------------------------
-
 		NodeMeshSwitch();
 
-		void SerialiseWrapper(OutputArchive& arch) const override;
-		static void Serialise( const NodeMeshSwitch* pNode, OutputArchive& arch );
-		static NodeMeshSwitchPtr StaticUnserialise( InputArchive& arch );
-
-
 		//-----------------------------------------------------------------------------------------
 		// Node Interface
 		//-----------------------------------------------------------------------------------------
diff --git a/Engine/Plugins/Experimental/Mutable/Source/MutableTools/Public/MuT/NodeMeshTable.h b/Engine/Plugins/Experimental/Mutable/Source/MutableTools/Public/MuT/NodeMeshTable.h
--- a/Engine/Plugins/Experimental/Mutable/Source/MutableTools/Public/MuT/NodeMeshTable.h
+++ b/Engine/Plugins/Experimental/Mutable/Source/MutableTools/Public/MuT/NodeMeshTable.h
@@ -26,21 +26,12 @@ namespace mu
 
 
 	//! This node provides the meshes stored in the column of a table.
-	//! \ingroup transform
 	class MUTABLETOOLS_API NodeMeshTable : public NodeMesh
 	{
 	public:
 
-		//-----------------------------------------------------------------------------------------
-		// Life cycle
-		//-----------------------------------------------------------------------------------------
-
 		NodeMeshTable();
 
-		void SerialiseWrapper(OutputArchive& arch) const override;
-		static void Serialise( const NodeMeshTable* pNode, OutputArchive& arch );
-		static NodeMeshTablePtr StaticUnserialise( InputArchive& arch );
-
 
 		//-----------------------------------------------------------------------------------------
 		// Node Interface
diff --git a/Engine/Plugins/Experimental/Mutable/Source/MutableTools/Public/MuT/NodeMeshTangents.h b/Engine/Plugins/Experimental/Mutable/Source/MutableTools/Public/MuT/NodeMeshTangents.h
--- a/Engine/Plugins/Experimental/Mutable/Source/MutableTools/Public/MuT/NodeMeshTangents.h
+++ b/Engine/Plugins/Experimental/Mutable/Source/MutableTools/Public/MuT/NodeMeshTangents.h
@@ -18,21 +18,12 @@ namespace mu
 
 
 	//! This node rebuilds the tangents and binormals of the source mesh.
-	//! \ingroup model
 	class MUTABLETOOLS_API NodeMeshTangents : public NodeMesh
 	{
 	public:
 
-		//-----------------------------------------------------------------------------------------
-		// Life cycle
-		//-----------------------------------------------------------------------------------------
-
 		NodeMeshTangents();
 
-		void SerialiseWrapper(OutputArchive& arch) const override;
-		static void Serialise( const NodeMeshTangents* pNode, OutputArchive& arch );
-		static NodeMeshTangentsPtr StaticUnserialise( InputArchive& arch );
-
 		//-----------------------------------------------------------------------------------------
 		// Node Interface
 		//-----------------------------------------------------------------------------------------
diff --git a/Engine/Plugins/Experimental/Mutable/Source/MutableTools/Public/MuT/NodeMeshTransform.h b/Engine/Plugins/Experimental/Mutable/Source/MutableTools/Public/MuT/NodeMeshTransform.h
--- a/Engine/Plugins/Experimental/Mutable/Source/MutableTools/Public/MuT/NodeMeshTransform.h
+++ b/Engine/Plugins/Experimental/Mutable/Source/MutableTools/Public/MuT/NodeMeshTransform.h
@@ -18,21 +18,12 @@ namespace mu
 
 
     //! This node applies a geometric transform represented by a 4x4 matrix to a mesh
-	//! \ingroup model
     class MUTABLETOOLS_API NodeMeshTransform : public NodeMesh
 	{
 	public:
 
-		//-----------------------------------------------------------------------------------------
-		// Life cycle
-		//-----------------------------------------------------------------------------------------
-
         NodeMeshTransform();
 
-		void SerialiseWrapper(OutputArchive& arch) const override;
-		static void Serialise( const NodeMeshTransform* pNode, OutputArchive& arch );
-        static NodeMeshTransformPtr StaticUnserialise( InputArchive& arch );
-
 		//-----------------------------------------------------------------------------------------
 		// Node Interface
 		//-----------------------------------------------------------------------------------------
@@ -48,7 +39,6 @@ namespace mu
         NodeMeshPtr GetSource() const;
 		void SetSource( NodeMesh* );
 
-        //! \param pMat4x4 pointer to 16 floats representing the 4x4 matrix to apply
         void SetTransform(const FMatrix44f&);
 		const FMatrix44f& GetTransform() const;
 
diff --git a/Engine/Plugins/Experimental/Mutable/Source/MutableTools/Public/MuT/NodeMeshVariation.h b/Engine/Plugins/Experimental/Mutable/Source/MutableTools/Public/MuT/NodeMeshVariation.h
--- a/Engine/Plugins/Experimental/Mutable/Source/MutableTools/Public/MuT/NodeMeshVariation.h
+++ b/Engine/Plugins/Experimental/Mutable/Source/MutableTools/Public/MuT/NodeMeshVariation.h
@@ -16,22 +16,14 @@ namespace mu
     typedef Ptr<NodeMeshVariation> NodeMeshVariationPtr;
     typedef Ptr<const NodeMeshVariation> NodeMeshVariationPtrConst;
 
-
     //!
     //! \ingroup model
     class MUTABLETOOLS_API NodeMeshVariation : public NodeMesh
     {
     public:
-        //-----------------------------------------------------------------------------------------
-        // Life cycle
-        //-----------------------------------------------------------------------------------------
 
         NodeMeshVariation();
 
-		void SerialiseWrapper(OutputArchive& arch) const override;
-		static void Serialise( const NodeMeshVariation* pNode, OutputArchive& arch );
-        static NodeMeshVariationPtr StaticUnserialise( InputArchive& arch );
-
         //-----------------------------------------------------------------------------------------
         // Node Interface
         //-----------------------------------------------------------------------------------------
diff --git a/Engine/Plugins/Experimental/Mutable/Source/MutableTools/Public/MuT/NodeModifier.h b/Engine/Plugins/Experimental/Mutable/Source/MutableTools/Public/MuT/NodeModifier.h
--- a/Engine/Plugins/Experimental/Mutable/Source/MutableTools/Public/MuT/NodeModifier.h
+++ b/Engine/Plugins/Experimental/Mutable/Source/MutableTools/Public/MuT/NodeModifier.h
@@ -46,13 +46,6 @@ namespace mu
 			None
 		};
 
-		//-----------------------------------------------------------------------------------------
-		// Life cycle
-		//-----------------------------------------------------------------------------------------
-
-		static void Serialise( const NodeModifier* pNode, OutputArchive& arch );
-		static Ptr<NodeModifier> StaticUnserialise( InputArchive& arch );
-
 		//-----------------------------------------------------------------------------------------
         // Node interface
 		//-----------------------------------------------------------------------------------------
diff --git a/Engine/Plugins/Experimental/Mutable/Source/MutableTools/Public/MuT/NodeModifierMeshClipDeform.h b/Engine/Plugins/Experimental/Mutable/Source/MutableTools/Public/MuT/NodeModifierMeshClipDeform.h
--- a/Engine/Plugins/Experimental/Mutable/Source/MutableTools/Public/MuT/NodeModifierMeshClipDeform.h
+++ b/Engine/Plugins/Experimental/Mutable/Source/MutableTools/Public/MuT/NodeModifierMeshClipDeform.h
@@ -28,16 +28,8 @@ namespace mu
 	{
 	public:
 
-		//-----------------------------------------------------------------------------------------
-		// Life cycle
-		//-----------------------------------------------------------------------------------------
-
 		NodeModifierMeshClipDeform();
 
-		void SerialiseWrapper(OutputArchive& arch) const override;
-		static void Serialise( const NodeModifierMeshClipDeform* pNode, OutputArchive& arch );
-		static NodeModifierMeshClipDeformPtr StaticUnserialise( InputArchive& arch );
-
 		//-----------------------------------------------------------------------------------------
         // Node interface
 		//-----------------------------------------------------------------------------------------
@@ -49,7 +41,6 @@ namespace mu
         // Own interface
 		//-----------------------------------------------------------------------------------------
 
-		//! \param 
 		void SetClipMesh( NodeMesh* InClipMesh);
 		void SetBindingMethod(EShapeBindingMethod BindingMethod);
 	
diff --git a/Engine/Plugins/Experimental/Mutable/Source/MutableTools/Public/MuT/NodeModifierMeshClipMorphPlane.h b/Engine/Plugins/Experimental/Mutable/Source/MutableTools/Public/MuT/NodeModifierMeshClipMorphPlane.h
--- a/Engine/Plugins/Experimental/Mutable/Source/MutableTools/Public/MuT/NodeModifierMeshClipMorphPlane.h
+++ b/Engine/Plugins/Experimental/Mutable/Source/MutableTools/Public/MuT/NodeModifierMeshClipMorphPlane.h
@@ -24,16 +24,8 @@ namespace mu
 	{
 	public:
 
-		//-----------------------------------------------------------------------------------------
-		// Life cycle
-		//-----------------------------------------------------------------------------------------
-
 		NodeModifierMeshClipMorphPlane();
 
-		void SerialiseWrapper(OutputArchive& arch) const override;
-		static void Serialise( const NodeModifierMeshClipMorphPlane* pNode, OutputArchive& arch );
-		static NodeModifierMeshClipMorphPlanePtr StaticUnserialise( InputArchive& arch );
-
 		//-----------------------------------------------------------------------------------------
         // Node interface
 		//-----------------------------------------------------------------------------------------
@@ -45,7 +37,6 @@ namespace mu
         // Own interface
 		//-----------------------------------------------------------------------------------------
 
-		//! \param 
 		void SetPlane(float centerX, float centerY, float centerZ, float normalX, float normalY, float normalZ);
 		void SetParams(float dist, float factor);
 		void SetMorphEllipse(float radius1, float radius2, float rotation);
diff --git a/Engine/Plugins/Experimental/Mutable/Source/MutableTools/Public/MuT/NodeModifierMeshClipWithMesh.h b/Engine/Plugins/Experimental/Mutable/Source/MutableTools/Public/MuT/NodeModifierMeshClipWithMesh.h
--- a/Engine/Plugins/Experimental/Mutable/Source/MutableTools/Public/MuT/NodeModifierMeshClipWithMesh.h
+++ b/Engine/Plugins/Experimental/Mutable/Source/MutableTools/Public/MuT/NodeModifierMeshClipWithMesh.h
@@ -21,16 +21,8 @@ namespace mu
 	{
 	public:
 
-		//-----------------------------------------------------------------------------------------
-		// Life cycle
-		//-----------------------------------------------------------------------------------------
-
 		NodeModifierMeshClipWithMesh();
 
-		void SerialiseWrapper(OutputArchive& arch) const override;
-		static void Serialise( const NodeModifierMeshClipWithMesh* pNode, OutputArchive& arch );
-		static Ptr<NodeModifierMeshClipWithMesh> StaticUnserialise( InputArchive& arch );
-
 		//-----------------------------------------------------------------------------------------
         // Node interface
 		//-----------------------------------------------------------------------------------------
@@ -42,7 +34,6 @@ namespace mu
         // Own interface
 		//-----------------------------------------------------------------------------------------
 
-		//! \param 
 		void SetClipMesh(NodeMesh*);
 
 		//-----------------------------------------------------------------------------------------
diff --git a/Engine/Plugins/Experimental/Mutable/Source/MutableTools/Public/MuT/NodeModifierMeshClipWithUVMask.h b/Engine/Plugins/Experimental/Mutable/Source/MutableTools/Public/MuT/NodeModifierMeshClipWithUVMask.h
--- a/Engine/Plugins/Experimental/Mutable/Source/MutableTools/Public/MuT/NodeModifierMeshClipWithUVMask.h
+++ b/Engine/Plugins/Experimental/Mutable/Source/MutableTools/Public/MuT/NodeModifierMeshClipWithUVMask.h
@@ -20,16 +20,8 @@ namespace mu
 	{
 	public:
 
-		//-----------------------------------------------------------------------------------------
-		// Life cycle
-		//-----------------------------------------------------------------------------------------
-
 		NodeModifierMeshClipWithUVMask();
 
-		void SerialiseWrapper(OutputArchive& arch) const override;
-		static void Serialise( const NodeModifierMeshClipWithUVMask* pNode, OutputArchive& arch );
-		static Ptr<NodeModifierMeshClipWithUVMask> StaticUnserialise( InputArchive& arch );
-
 		//-----------------------------------------------------------------------------------------
         // Node interface
 		//-----------------------------------------------------------------------------------------
diff --git a/Engine/Plugins/Experimental/Mutable/Source/MutableTools/Public/MuT/NodeObject.h b/Engine/Plugins/Experimental/Mutable/Source/MutableTools/Public/MuT/NodeObject.h
--- a/Engine/Plugins/Experimental/Mutable/Source/MutableTools/Public/MuT/NodeObject.h
+++ b/Engine/Plugins/Experimental/Mutable/Source/MutableTools/Public/MuT/NodeObject.h
@@ -32,13 +32,6 @@ namespace mu
 			None
 		};
 
-		//-----------------------------------------------------------------------------------------
-		// Life cycle
-		//-----------------------------------------------------------------------------------------
-
-		static void Serialise( const NodeObject* pNode, OutputArchive& arch );
-		static NodeObjectPtr StaticUnserialise( InputArchive& arch );
-
 
 		//-----------------------------------------------------------------------------------------
 		// Node Interface
diff --git a/Engine/Plugins/Experimental/Mutable/Source/MutableTools/Public/MuT/NodeObjectGroup.h b/Engine/Plugins/Experimental/Mutable/Source/MutableTools/Public/MuT/NodeObjectGroup.h
--- a/Engine/Plugins/Experimental/Mutable/Source/MutableTools/Public/MuT/NodeObjectGroup.h
+++ b/Engine/Plugins/Experimental/Mutable/Source/MutableTools/Public/MuT/NodeObjectGroup.h
@@ -23,16 +23,8 @@ namespace mu
 	{
 	public:
 
-		//-----------------------------------------------------------------------------------------
-		// Life cycle
-		//-----------------------------------------------------------------------------------------
-
 		NodeObjectGroup();
 
-		void SerialiseWrapper(OutputArchive& arch) const override;
-		static void Serialise( const NodeObjectGroup* pNode, OutputArchive& arch );
-		static NodeObjectGroupPtr StaticUnserialise( InputArchive& arch );
-
 
 		//-----------------------------------------------------------------------------------------
 		// Node Interface
diff --git a/Engine/Plugins/Experimental/Mutable/Source/MutableTools/Public/MuT/NodeObjectNew.h b/Engine/Plugins/Experimental/Mutable/Source/MutableTools/Public/MuT/NodeObjectNew.h
--- a/Engine/Plugins/Experimental/Mutable/Source/MutableTools/Public/MuT/NodeObjectNew.h
+++ b/Engine/Plugins/Experimental/Mutable/Source/MutableTools/Public/MuT/NodeObjectNew.h
@@ -14,21 +14,12 @@ namespace mu
 	class NodeLOD;
 
 	//! Node that creates a new object by setting its levels-of-detail and its children.
-	//! \ingroup model
 	class MUTABLETOOLS_API NodeObjectNew : public NodeObject
 	{
 	public:
 
-		//-----------------------------------------------------------------------------------------
-		// Life cycle
-		//-----------------------------------------------------------------------------------------
-
 		NodeObjectNew();
 
-		void SerialiseWrapper(OutputArchive& arch) const override;
-		static void Serialise( const NodeObjectNew* pNode, OutputArchive& arch );
-		static Ptr<NodeObjectNew> StaticUnserialise( InputArchive& arch );
-
 
 		//-----------------------------------------------------------------------------------------
 		// Node Interface
diff --git a/Engine/Plugins/Experimental/Mutable/Source/MutableTools/Public/MuT/NodePatchMesh.h b/Engine/Plugins/Experimental/Mutable/Source/MutableTools/Public/MuT/NodePatchMesh.h
--- a/Engine/Plugins/Experimental/Mutable/Source/MutableTools/Public/MuT/NodePatchMesh.h
+++ b/Engine/Plugins/Experimental/Mutable/Source/MutableTools/Public/MuT/NodePatchMesh.h
@@ -2,7 +2,6 @@
 
 #pragma once
 
-
 #include "MuR/Ptr.h"
 #include "MuR/RefCounted.h"
 #include "MuT/Node.h"
@@ -17,22 +16,12 @@ namespace mu
 
 	class NodeMesh;
 
-	//!
-	//! \ingroup model
 	class MUTABLETOOLS_API NodePatchMesh : public Node
 	{
 	public:
 
-		//-----------------------------------------------------------------------------------------
-		// Life cycle
-		//-----------------------------------------------------------------------------------------
-
 		NodePatchMesh();
 
-		void SerialiseWrapper(OutputArchive& arch) const override;
-		static void Serialise( const NodePatchMesh* pNode, OutputArchive& arch );
-		static NodePatchMeshPtr StaticUnserialise( InputArchive& arch );
-
 		//-----------------------------------------------------------------------------------------
 		// Node Interface
 		//-----------------------------------------------------------------------------------------
diff --git a/Engine/Plugins/Experimental/Mutable/Source/MutableTools/Public/MuT/NodeProjector.h b/Engine/Plugins/Experimental/Mutable/Source/MutableTools/Public/MuT/NodeProjector.h
--- a/Engine/Plugins/Experimental/Mutable/Source/MutableTools/Public/MuT/NodeProjector.h
+++ b/Engine/Plugins/Experimental/Mutable/Source/MutableTools/Public/MuT/NodeProjector.h
@@ -30,7 +30,6 @@ namespace mu
 
 
     //! Base class of any node that outputs a Projector.
-	//! \ingroup model
 	class MUTABLETOOLS_API NodeProjector : public Node
 	{
 	public:
@@ -44,14 +43,6 @@ namespace mu
 			None
 		};
 
-		//-----------------------------------------------------------------------------------------
-		// Life cycle
-		//-----------------------------------------------------------------------------------------
-
-		static void Serialise( const NodeProjector* pNode, OutputArchive& arch );
-		static NodeProjectorPtr StaticUnserialise( InputArchive& arch );
-
-
 		//-----------------------------------------------------------------------------------------
 		// Node Interface
 		//-----------------------------------------------------------------------------------------
@@ -77,21 +68,12 @@ namespace mu
 
 
 	//! This node outputs a predefined Projector value.
-	//! \ingroup model
 	class MUTABLETOOLS_API NodeProjectorConstant : public NodeProjector
 	{
 	public:
 
-		//-----------------------------------------------------------------------------------------
-		// Life cycle
-		//-----------------------------------------------------------------------------------------
-
 		NodeProjectorConstant();
 
-		void SerialiseWrapper(OutputArchive& arch) const override;
-		static void Serialise( const NodeProjectorConstant* pNode, OutputArchive& arch );
-		static NodeProjectorConstantPtr StaticUnserialise( InputArchive& arch );
-
 		//-----------------------------------------------------------------------------------------
 		// Node Interface
 		//-----------------------------------------------------------------------------------------
@@ -134,21 +116,12 @@ namespace mu
 
 
 	//! Node that defines a Projector model parameter.
-	//! \ingroup model
 	class MUTABLETOOLS_API NodeProjectorParameter : public NodeProjector
 	{
 	public:
 
-		//-----------------------------------------------------------------------------------------
-		// Life cycle
-		//-----------------------------------------------------------------------------------------
-
 		NodeProjectorParameter();
 
-		void SerialiseWrapper(OutputArchive& arch) const override;
-		static void Serialise( const NodeProjectorParameter* pNode, OutputArchive& arch );
-		static NodeProjectorParameterPtr StaticUnserialise( InputArchive& arch );
-
 		//-----------------------------------------------------------------------------------------
 		// Node Interface
 		//-----------------------------------------------------------------------------------------
diff --git a/Engine/Plugins/Experimental/Mutable/Source/MutableTools/Public/MuT/NodeRange.h b/Engine/Plugins/Experimental/Mutable/Source/MutableTools/Public/MuT/NodeRange.h
--- a/Engine/Plugins/Experimental/Mutable/Source/MutableTools/Public/MuT/NodeRange.h
+++ b/Engine/Plugins/Experimental/Mutable/Source/MutableTools/Public/MuT/NodeRange.h
@@ -18,7 +18,6 @@ namespace mu
 
 
     //! %Base class of any node that outputs a range.
-	//! \ingroup model
     class MUTABLETOOLS_API NodeRange : public Node
 	{
 	public:
@@ -31,13 +30,6 @@ namespace mu
 			None
 		};
 
-		//-----------------------------------------------------------------------------------------
-		// Life cycle
-		//-----------------------------------------------------------------------------------------
-
-        static void Serialise( const NodeRange* pNode, OutputArchive& arch );
-        static NodeRangePtr StaticUnserialise( InputArchive& arch );
-
 
 		//-----------------------------------------------------------------------------------------
 		// Node Interface
diff --git a/Engine/Plugins/Experimental/Mutable/Source/MutableTools/Public/MuT/NodeRangeFromScalar.h b/Engine/Plugins/Experimental/Mutable/Source/MutableTools/Public/MuT/NodeRangeFromScalar.h
--- a/Engine/Plugins/Experimental/Mutable/Source/MutableTools/Public/MuT/NodeRangeFromScalar.h
+++ b/Engine/Plugins/Experimental/Mutable/Source/MutableTools/Public/MuT/NodeRangeFromScalar.h
@@ -13,8 +13,6 @@ namespace mu
     class NodeScalar;
 
 
-	//! 
-	//! \ingroup model
     class MUTABLETOOLS_API NodeRangeFromScalar : public NodeRange
 	{
 	public:
@@ -25,10 +23,6 @@ namespace mu
 
         NodeRangeFromScalar();
 
-		void SerialiseWrapper(OutputArchive& arch) const override;
-		static void Serialise( const NodeRangeFromScalar* pNode, OutputArchive& arch );
-        static Ptr<NodeRangeFromScalar> StaticUnserialise( InputArchive& arch );
-
 
 		//-----------------------------------------------------------------------------------------
 		// Node Interface
diff --git a/Engine/Plugins/Experimental/Mutable/Source/MutableTools/Public/MuT/NodeScalar.h b/Engine/Plugins/Experimental/Mutable/Source/MutableTools/Public/MuT/NodeScalar.h
--- a/Engine/Plugins/Experimental/Mutable/Source/MutableTools/Public/MuT/NodeScalar.h
+++ b/Engine/Plugins/Experimental/Mutable/Source/MutableTools/Public/MuT/NodeScalar.h
@@ -18,7 +18,6 @@ namespace mu
 
 
     //! %Base class of any node that outputs a scalar value.
-	//! \ingroup model
 	class MUTABLETOOLS_API NodeScalar : public Node
 	{
 	public:
@@ -38,13 +37,6 @@ namespace mu
 			None
 		};
 
-		//-----------------------------------------------------------------------------------------
-		// Life cycle
-		//-----------------------------------------------------------------------------------------
-
-		static void Serialise( const NodeScalar* pNode, OutputArchive& arch );
-		static NodeScalarPtr StaticUnserialise( InputArchive& arch );
-
 
 		//-----------------------------------------------------------------------------------------
 		// Node Interface
diff --git a/Engine/Plugins/Experimental/Mutable/Source/MutableTools/Public/MuT/NodeScalarArithmeticOperation.h b/Engine/Plugins/Experimental/Mutable/Source/MutableTools/Public/MuT/NodeScalarArithmeticOperation.h
--- a/Engine/Plugins/Experimental/Mutable/Source/MutableTools/Public/MuT/NodeScalarArithmeticOperation.h
+++ b/Engine/Plugins/Experimental/Mutable/Source/MutableTools/Public/MuT/NodeScalarArithmeticOperation.h
@@ -23,16 +23,8 @@ namespace mu
 	{
 	public:
 
-		//-----------------------------------------------------------------------------------------
-		// Life cycle
-		//-----------------------------------------------------------------------------------------
-
         NodeScalarArithmeticOperation();
 
-		void SerialiseWrapper(OutputArchive& arch) const override;
-		static void Serialise( const NodeScalarArithmeticOperation* pNode, OutputArchive& arch );
-        static NodeScalarArithmeticOperationPtr StaticUnserialise( InputArchive& arch );
-
 
 		//-----------------------------------------------------------------------------------------
 		// Node Interface
diff --git a/Engine/Plugins/Experimental/Mutable/Source/MutableTools/Public/MuT/NodeScalarConstant.h b/Engine/Plugins/Experimental/Mutable/Source/MutableTools/Public/MuT/NodeScalarConstant.h
--- a/Engine/Plugins/Experimental/Mutable/Source/MutableTools/Public/MuT/NodeScalarConstant.h
+++ b/Engine/Plugins/Experimental/Mutable/Source/MutableTools/Public/MuT/NodeScalarConstant.h
@@ -23,16 +23,8 @@ namespace mu
 	{
 	public:
 
-		//-----------------------------------------------------------------------------------------
-		// Life cycle
-		//-----------------------------------------------------------------------------------------
-
 		NodeScalarConstant();
 
-		void SerialiseWrapper(OutputArchive& arch) const override;
-		static void Serialise( const NodeScalarConstant* pNode, OutputArchive& arch );
-		static NodeScalarConstantPtr StaticUnserialise( InputArchive& arch );
-
 
 		//-----------------------------------------------------------------------------------------
 		// Node Interface
diff --git a/Engine/Plugins/Experimental/Mutable/Source/MutableTools/Public/MuT/NodeScalarCurve.h b/Engine/Plugins/Experimental/Mutable/Source/MutableTools/Public/MuT/NodeScalarCurve.h
--- a/Engine/Plugins/Experimental/Mutable/Source/MutableTools/Public/MuT/NodeScalarCurve.h
+++ b/Engine/Plugins/Experimental/Mutable/Source/MutableTools/Public/MuT/NodeScalarCurve.h
@@ -1,6 +1,5 @@
 // Copyright Epic Games, Inc. All Rights Reserved.
 
-
 #pragma once
 
 #include "MuR/Ptr.h"
@@ -23,16 +22,8 @@ namespace mu
 	{
 	public:
 
-		//-----------------------------------------------------------------------------------------
-		// Life cycle
-		//-----------------------------------------------------------------------------------------
-
         NodeScalarCurve();
 
-		void SerialiseWrapper(OutputArchive& arch) const override;
-		static void Serialise( const NodeScalarCurve* pNode, OutputArchive& arch );
-        static NodeScalarCurvePtr StaticUnserialise( InputArchive& arch );
-
 		//-----------------------------------------------------------------------------------------
 		// Node Interface
 		//-----------------------------------------------------------------------------------------
diff --git a/Engine/Plugins/Experimental/Mutable/Source/MutableTools/Public/MuT/NodeScalarEnumParameter.h b/Engine/Plugins/Experimental/Mutable/Source/MutableTools/Public/MuT/NodeScalarEnumParameter.h
--- a/Engine/Plugins/Experimental/Mutable/Source/MutableTools/Public/MuT/NodeScalarEnumParameter.h
+++ b/Engine/Plugins/Experimental/Mutable/Source/MutableTools/Public/MuT/NodeScalarEnumParameter.h
@@ -27,21 +27,12 @@ namespace mu
 
 
 	//! Node that defines a scalar model parameter to be selected from a set of named values.
-	//! \ingroup model
 	class MUTABLETOOLS_API NodeScalarEnumParameter : public NodeScalar
 	{
 	public:
 
-		//-----------------------------------------------------------------------------------------
-		// Life cycle
-		//-----------------------------------------------------------------------------------------
-
 		NodeScalarEnumParameter();
 
-		void SerialiseWrapper(OutputArchive& arch) const override;
-		static void Serialise( const NodeScalarEnumParameter* pNode, OutputArchive& arch );
-		static NodeScalarEnumParameterPtr StaticUnserialise( InputArchive& arch );
-
 		//-----------------------------------------------------------------------------------------
 		// Node Interface
 		//-----------------------------------------------------------------------------------------
diff --git a/Engine/Plugins/Experimental/Mutable/Source/MutableTools/Public/MuT/NodeScalarParameter.h b/Engine/Plugins/Experimental/Mutable/Source/MutableTools/Public/MuT/NodeScalarParameter.h
--- a/Engine/Plugins/Experimental/Mutable/Source/MutableTools/Public/MuT/NodeScalarParameter.h
+++ b/Engine/Plugins/Experimental/Mutable/Source/MutableTools/Public/MuT/NodeScalarParameter.h
@@ -27,7 +27,6 @@ namespace mu
 
 
 	//! Node that defines a scalar model parameter.
-	//! \ingroup model
 	class MUTABLETOOLS_API NodeScalarParameter : public NodeScalar
 	{
 	public:
@@ -38,10 +37,6 @@ namespace mu
 
 		NodeScalarParameter();
 
-		void SerialiseWrapper(OutputArchive& arch) const override;
-		static void Serialise( const NodeScalarParameter* pNode, OutputArchive& arch );
-		static NodeScalarParameterPtr StaticUnserialise( InputArchive& arch );
-
 		//-----------------------------------------------------------------------------------------
 		// Node Interface
 		//-----------------------------------------------------------------------------------------
