[mutable] Remove custom serialization code frome the MutableTools module.

diff --git a/Engine/Plugins/Experimental/Mutable/Source/MutableTools/Public/MuT/NodeScalarSwitch.h b/Engine/Plugins/Experimental/Mutable/Source/MutableTools/Public/MuT/NodeScalarSwitch.h
--- a/Engine/Plugins/Experimental/Mutable/Source/MutableTools/Public/MuT/NodeScalarSwitch.h
+++ b/Engine/Plugins/Experimental/Mutable/Source/MutableTools/Public/MuT/NodeScalarSwitch.h
@@ -22,21 +22,12 @@ namespace mu
 
 
     //! This node selects an output Scalar from a set of input Scalars based on a parameter.
-	//! \ingroup model
     class MUTABLETOOLS_API NodeScalarSwitch : public NodeScalar
 	{
 	public:
 
-		//-----------------------------------------------------------------------------------------
-		// Life cycle
-		//-----------------------------------------------------------------------------------------
-
         NodeScalarSwitch();
 
-		void SerialiseWrapper(OutputArchive& arch) const override;
-		static void Serialise( const NodeScalarSwitch* pNode, OutputArchive& arch );
-        static NodeScalarSwitchPtr StaticUnserialise( InputArchive& arch );
-
 		//-----------------------------------------------------------------------------------------
 		// Node Interface
 		//-----------------------------------------------------------------------------------------
diff --git a/Engine/Plugins/Experimental/Mutable/Source/MutableTools/Public/MuT/NodeScalarTable.h b/Engine/Plugins/Experimental/Mutable/Source/MutableTools/Public/MuT/NodeScalarTable.h
--- a/Engine/Plugins/Experimental/Mutable/Source/MutableTools/Public/MuT/NodeScalarTable.h
+++ b/Engine/Plugins/Experimental/Mutable/Source/MutableTools/Public/MuT/NodeScalarTable.h
@@ -22,22 +22,12 @@ namespace mu
 
 
 	//! This node provides the meshes stored in the column of a table.
-	//! \ingroup transform
 	class MUTABLETOOLS_API NodeScalarTable : public NodeScalar
 	{
 	public:
 
-		//-----------------------------------------------------------------------------------------
-		// Life cycle
-		//-----------------------------------------------------------------------------------------
-
 		NodeScalarTable();
 
-		void SerialiseWrapper(OutputArchive& arch) const override;
-		static void Serialise( const NodeScalarTable* pNode, OutputArchive& arch );
-		static NodeScalarTablePtr StaticUnserialise( InputArchive& arch );
-
-
 		//-----------------------------------------------------------------------------------------
 		// Node Interface
 		//-----------------------------------------------------------------------------------------
diff --git a/Engine/Plugins/Experimental/Mutable/Source/MutableTools/Public/MuT/NodeScalarVariation.h b/Engine/Plugins/Experimental/Mutable/Source/MutableTools/Public/MuT/NodeScalarVariation.h
--- a/Engine/Plugins/Experimental/Mutable/Source/MutableTools/Public/MuT/NodeScalarVariation.h
+++ b/Engine/Plugins/Experimental/Mutable/Source/MutableTools/Public/MuT/NodeScalarVariation.h
@@ -17,21 +17,12 @@ namespace mu
     typedef Ptr<const NodeScalarVariation> NodeScalarVariationPtrConst;
 
 
-    //!
-    //! \ingroup model
-    class MUTABLETOOLS_API NodeScalarVariation : public NodeScalar
+     class MUTABLETOOLS_API NodeScalarVariation : public NodeScalar
     {
     public:
-        //-----------------------------------------------------------------------------------------
-        // Life cycle
-        //-----------------------------------------------------------------------------------------
-
+ 
         NodeScalarVariation();
 
-		void SerialiseWrapper(OutputArchive& arch) const override;
-		static void Serialise( const NodeScalarVariation* pNode, OutputArchive& arch );
-        static NodeScalarVariationPtr StaticUnserialise( InputArchive& arch );
-
         //-----------------------------------------------------------------------------------------
         // Node Interface
         //-----------------------------------------------------------------------------------------
diff --git a/Engine/Plugins/Experimental/Mutable/Source/MutableTools/Public/MuT/NodeString.h b/Engine/Plugins/Experimental/Mutable/Source/MutableTools/Public/MuT/NodeString.h
--- a/Engine/Plugins/Experimental/Mutable/Source/MutableTools/Public/MuT/NodeString.h
+++ b/Engine/Plugins/Experimental/Mutable/Source/MutableTools/Public/MuT/NodeString.h
@@ -18,7 +18,6 @@ namespace mu
 
 
     //! %Base class of any node that outputs a string value.
-	//! \ingroup model
 	class MUTABLETOOLS_API NodeString : public Node
 	{
 	public:
@@ -32,14 +31,6 @@ namespace mu
 			None
 		};
 
-		//-----------------------------------------------------------------------------------------
-		// Life cycle
-		//-----------------------------------------------------------------------------------------
-
-		static void Serialise( const NodeString* pNode, OutputArchive& arch );
-		static NodeStringPtr StaticUnserialise( InputArchive& arch );
-
-
 		//-----------------------------------------------------------------------------------------
 		// Node Interface
 		//-----------------------------------------------------------------------------------------
diff --git a/Engine/Plugins/Experimental/Mutable/Source/MutableTools/Public/MuT/NodeStringConstant.h b/Engine/Plugins/Experimental/Mutable/Source/MutableTools/Public/MuT/NodeStringConstant.h
--- a/Engine/Plugins/Experimental/Mutable/Source/MutableTools/Public/MuT/NodeStringConstant.h
+++ b/Engine/Plugins/Experimental/Mutable/Source/MutableTools/Public/MuT/NodeStringConstant.h
@@ -23,16 +23,8 @@ namespace mu
 	{
 	public:
 
-		//-----------------------------------------------------------------------------------------
-		// Life cycle
-		//-----------------------------------------------------------------------------------------
-
 		NodeStringConstant();
 
-		void SerialiseWrapper(OutputArchive& arch) const override;
-		static void Serialise( const NodeStringConstant* pNode, OutputArchive& arch );
-		static NodeStringConstantPtr StaticUnserialise( InputArchive& arch );
-
 
 		//-----------------------------------------------------------------------------------------
 		// Node Interface
diff --git a/Engine/Plugins/Experimental/Mutable/Source/MutableTools/Public/MuT/NodeStringParameter.h b/Engine/Plugins/Experimental/Mutable/Source/MutableTools/Public/MuT/NodeStringParameter.h
--- a/Engine/Plugins/Experimental/Mutable/Source/MutableTools/Public/MuT/NodeStringParameter.h
+++ b/Engine/Plugins/Experimental/Mutable/Source/MutableTools/Public/MuT/NodeStringParameter.h
@@ -27,22 +27,12 @@ namespace mu
 
 
 	//! Node that defines a string model parameter.
-	//! \ingroup model
 	class MUTABLETOOLS_API NodeStringParameter : public NodeString
 	{
 	public:
 
-		//-----------------------------------------------------------------------------------------
-		// Life cycle
-		//-----------------------------------------------------------------------------------------
-
 		NodeStringParameter();
 
-		void SerialiseWrapper(OutputArchive& arch) const override;
-		static void Serialise( const NodeStringParameter* pNode, OutputArchive& arch );
-		static NodeStringParameterPtr StaticUnserialise( InputArchive& arch );
-
-
 		//-----------------------------------------------------------------------------------------
 		// Node Interface
 		//-----------------------------------------------------------------------------------------
diff --git a/Engine/Plugins/Experimental/Mutable/Source/MutableTools/Public/MuT/NodeSurface.h b/Engine/Plugins/Experimental/Mutable/Source/MutableTools/Public/MuT/NodeSurface.h
--- a/Engine/Plugins/Experimental/Mutable/Source/MutableTools/Public/MuT/NodeSurface.h
+++ b/Engine/Plugins/Experimental/Mutable/Source/MutableTools/Public/MuT/NodeSurface.h
@@ -17,9 +17,7 @@ namespace mu
     typedef Ptr<const NodeSurface> NodeSurfacePtrConst;
 
 
-
     //! This class is the parent of all nodes that output a Surface.
-	//! \ingroup model
     class MUTABLETOOLS_API NodeSurface : public Node
 	{
 	public:
@@ -35,13 +33,6 @@ namespace mu
 			None
 		};
 
-		//-----------------------------------------------------------------------------------------
-		// Life cycle
-		//-----------------------------------------------------------------------------------------
-
-        static void Serialise( const NodeSurface* pNode, OutputArchive& arch );
-        static NodeSurfacePtr StaticUnserialise( InputArchive& arch );
-
 		//-----------------------------------------------------------------------------------------
 		// Node Interface
 		//-----------------------------------------------------------------------------------------
diff --git a/Engine/Plugins/Experimental/Mutable/Source/MutableTools/Public/MuT/NodeSurfaceEdit.h b/Engine/Plugins/Experimental/Mutable/Source/MutableTools/Public/MuT/NodeSurfaceEdit.h
--- a/Engine/Plugins/Experimental/Mutable/Source/MutableTools/Public/MuT/NodeSurfaceEdit.h
+++ b/Engine/Plugins/Experimental/Mutable/Source/MutableTools/Public/MuT/NodeSurfaceEdit.h
@@ -41,21 +41,12 @@ namespace mu
 	//! This node modifies a node of the parent object of the object that this node belongs to.
     //! It allows to extend, cut and morph the parent Surface's meshes.
     //! It also allows to patch the parent Surface's textures.
-	//! \ingroup model
     class MUTABLETOOLS_API NodeSurfaceEdit : public NodeSurface
 	{
 	public:
 
-		//-----------------------------------------------------------------------------------------
-		// Life cycle
-		//-----------------------------------------------------------------------------------------
-
         NodeSurfaceEdit();
 
-		void SerialiseWrapper(OutputArchive& arch) const override;
-		static void Serialise( const NodeSurfaceEdit* pNode, OutputArchive& arch );
-        static NodeSurfaceEditPtr StaticUnserialise( InputArchive& arch );
-
 		//-----------------------------------------------------------------------------------------
 		// Node Interface
 		//-----------------------------------------------------------------------------------------
diff --git a/Engine/Plugins/Experimental/Mutable/Source/MutableTools/Public/MuT/NodeSurfaceNew.h b/Engine/Plugins/Experimental/Mutable/Source/MutableTools/Public/MuT/NodeSurfaceNew.h
--- a/Engine/Plugins/Experimental/Mutable/Source/MutableTools/Public/MuT/NodeSurfaceNew.h
+++ b/Engine/Plugins/Experimental/Mutable/Source/MutableTools/Public/MuT/NodeSurfaceNew.h
@@ -42,21 +42,12 @@ namespace mu
 
 
     //! This node makes a new Surface from several meshes and images.
-	//! \ingroup model
     class MUTABLETOOLS_API NodeSurfaceNew : public NodeSurface
 	{
 	public:
 
-		//-----------------------------------------------------------------------------------------
-		// Life cycle
-		//-----------------------------------------------------------------------------------------
-
         NodeSurfaceNew();
 
-		void SerialiseWrapper(OutputArchive& arch) const override;
-		static void Serialise( const NodeSurfaceNew* pNode, OutputArchive& arch );
-        static NodeSurfaceNewPtr StaticUnserialise( InputArchive& arch );
-
 		//-----------------------------------------------------------------------------------------
 		// Node Interface
 		//-----------------------------------------------------------------------------------------
diff --git a/Engine/Plugins/Experimental/Mutable/Source/MutableTools/Public/MuT/NodeSurfaceSwitch.h b/Engine/Plugins/Experimental/Mutable/Source/MutableTools/Public/MuT/NodeSurfaceSwitch.h
--- a/Engine/Plugins/Experimental/Mutable/Source/MutableTools/Public/MuT/NodeSurfaceSwitch.h
+++ b/Engine/Plugins/Experimental/Mutable/Source/MutableTools/Public/MuT/NodeSurfaceSwitch.h
@@ -22,21 +22,12 @@ namespace mu
 
 
 	//! This node selects an output Surface from a set of input Surfaces based on a parameter.
-	//! \ingroup model
 	class MUTABLETOOLS_API NodeSurfaceSwitch : public NodeSurface
 	{
 	public:
 
-		//-----------------------------------------------------------------------------------------
-		// Life cycle
-		//-----------------------------------------------------------------------------------------
-
 		NodeSurfaceSwitch();
 
-		void SerialiseWrapper(OutputArchive& arch) const override;
-		static void Serialise( const NodeSurfaceSwitch* pNode, OutputArchive& arch );
-		static NodeSurfaceSwitchPtr StaticUnserialise( InputArchive& arch );
-
 
 		//-----------------------------------------------------------------------------------------
 		// Node Interface
diff --git a/Engine/Plugins/Experimental/Mutable/Source/MutableTools/Public/MuT/NodeSurfaceVariation.h b/Engine/Plugins/Experimental/Mutable/Source/MutableTools/Public/MuT/NodeSurfaceVariation.h
--- a/Engine/Plugins/Experimental/Mutable/Source/MutableTools/Public/MuT/NodeSurfaceVariation.h
+++ b/Engine/Plugins/Experimental/Mutable/Source/MutableTools/Public/MuT/NodeSurfaceVariation.h
@@ -21,21 +21,12 @@ namespace mu
 	//! This node modifies a node of the parent object of the object that this node belongs to.
     //! It allows to extend, cut and morph the parent Surface's meshes.
     //! It also allows to patch the parent Surface's textures.
-	//! \ingroup model
     class MUTABLETOOLS_API NodeSurfaceVariation : public NodeSurface
 	{
 	public:
 
-		//-----------------------------------------------------------------------------------------
-		// Life cycle
-		//-----------------------------------------------------------------------------------------
-
         NodeSurfaceVariation();
 
-		void SerialiseWrapper(OutputArchive& arch) const override;
-		static void Serialise( const NodeSurfaceVariation* pNode, OutputArchive& arch );
-		static NodeSurfaceVariationPtr StaticUnserialise( InputArchive& arch );
-
 		//-----------------------------------------------------------------------------------------
 		// Node Interface
 		//-----------------------------------------------------------------------------------------
diff --git a/Engine/Plugins/Experimental/Mutable/Source/MutableTools/Public/MuT/Table.h b/Engine/Plugins/Experimental/Mutable/Source/MutableTools/Public/MuT/Table.h
--- a/Engine/Plugins/Experimental/Mutable/Source/MutableTools/Public/MuT/Table.h
+++ b/Engine/Plugins/Experimental/Mutable/Source/MutableTools/Public/MuT/Table.h
@@ -16,8 +16,6 @@ namespace mu
 
 	class Mesh;
 	class Image;
-	class InputArchive;
-	class OutputArchive;
 
 	/** Types of the values for the table cells. */
 	enum class ETableColumnType : uint32
@@ -39,14 +37,8 @@ namespace mu
 	{
 	public:
 
-		//-----------------------------------------------------------------------------------------
-		// Life cycle.
-		//-----------------------------------------------------------------------------------------
 		Table();
 
-		static void Serialise( const Table* p, OutputArchive& arch );
-		static Ptr<Table> StaticUnserialise( InputArchive& arch );
-
 		//-----------------------------------------------------------------------------------------
 		// Own interface
 		//-----------------------------------------------------------------------------------------
