[Mutable] Fix mesh generation failure due to missing bones. Fix crashes when accessing an invalid LOD from the LODModel.

#rnx

diff --git a/Engine/Plugins/Experimental/Mutable/Source/CustomizableObject/Classes/MuCO/CustomizableObject.h b/Engine/Plugins/Experimental/Mutable/Source/CustomizableObject/Classes/MuCO/CustomizableObject.h
--- a/Engine/Plugins/Experimental/Mutable/Source/CustomizableObject/Classes/MuCO/CustomizableObject.h
+++ b/Engine/Plugins/Experimental/Mutable/Source/CustomizableObject/Classes/MuCO/CustomizableObject.h
@@ -176,7 +176,7 @@ struct FCompilationOptions
 	bool b16BitBoneWeightsEnabled = false;
 
 	// Used to enable skin weight profiles.
-	bool bSkinWeightProfilesEnabled = true;
+	bool bSkinWeightProfilesEnabled = false;
 
 	// Used to enable physics asset merge.
 	bool bPhysicsAssetMergeEnabled = false;
diff --git a/Engine/Plugins/Experimental/Mutable/Source/CustomizableObject/Internal/MuCO/CustomizableObjectPrivate.h b/Engine/Plugins/Experimental/Mutable/Source/CustomizableObject/Internal/MuCO/CustomizableObjectPrivate.h
--- a/Engine/Plugins/Experimental/Mutable/Source/CustomizableObject/Internal/MuCO/CustomizableObjectPrivate.h
+++ b/Engine/Plugins/Experimental/Mutable/Source/CustomizableObject/Internal/MuCO/CustomizableObjectPrivate.h
@@ -603,11 +603,7 @@ struct FModelResources
 	TArray<FName> MaterialSlotNames;
 
 	UPROPERTY()
-	TMap<FName, uint32> BoneNamesMap;
-
-	/** Bones remapped due to a hash collision at compile time. FName to (remapped)Hash. */
-	UPROPERTY()
-	TArray<FMutableRemappedBone> RemappedBoneNames;
+	TMap<FString, uint32> BoneNamesMap;
 
 	/** Mesh sockets provided by the part skeletal meshes, to be merged in the generated meshes */
 	UPROPERTY()
@@ -882,6 +878,6 @@ public:
 	// This is a manual version number for the binary blobs in this asset.
 	// Increasing it invalidates all the previously compiled models.
 	// Warning: If while merging code both versions have changed, take the highest+1.
-	static constexpr int32 CurrentSupportedVersion = 442;
+	static constexpr int32 CurrentSupportedVersion = 443;
 };
 
diff --git a/Engine/Plugins/Experimental/Mutable/Source/CustomizableObject/Private/MuCO/CustomizableObject.cpp b/Engine/Plugins/Experimental/Mutable/Source/CustomizableObject/Private/MuCO/CustomizableObject.cpp
--- a/Engine/Plugins/Experimental/Mutable/Source/CustomizableObject/Private/MuCO/CustomizableObject.cpp
+++ b/Engine/Plugins/Experimental/Mutable/Source/CustomizableObject/Private/MuCO/CustomizableObject.cpp
@@ -551,7 +551,6 @@ void UCustomizableObjectPrivate::SaveCompiledData(FArchive& MemoryWriter, bool b
 
 	MemoryWriter << LocalModelResources.MaterialSlotNames;
 	MemoryWriter << LocalModelResources.BoneNamesMap;
-	MemoryWriter << LocalModelResources.RemappedBoneNames;
 	MemoryWriter << LocalModelResources.SocketArray;
 
 	MemoryWriter << LocalModelResources.SkinWeightProfilesInfo;
@@ -672,7 +671,6 @@ void UCustomizableObjectPrivate::LoadCompiledData(FArchive& MemoryReader, const
 
 		MemoryReader << LocalModelResource.MaterialSlotNames;
 		MemoryReader << LocalModelResource.BoneNamesMap;
-		MemoryReader << LocalModelResource.RemappedBoneNames;
 		MemoryReader << LocalModelResource.SocketArray;
 
 		MemoryReader << LocalModelResource.SkinWeightProfilesInfo;
diff --git a/Engine/Plugins/Experimental/Mutable/Source/CustomizableObject/Private/MuCO/CustomizableObjectInstance.cpp b/Engine/Plugins/Experimental/Mutable/Source/CustomizableObject/Private/MuCO/CustomizableObjectInstance.cpp
--- a/Engine/Plugins/Experimental/Mutable/Source/CustomizableObject/Private/MuCO/CustomizableObjectInstance.cpp
+++ b/Engine/Plugins/Experimental/Mutable/Source/CustomizableObject/Private/MuCO/CustomizableObjectInstance.cpp
@@ -3327,15 +3327,9 @@ bool UCustomizableInstancePrivate::BuildSkeletonData(const TSharedRef<FUpdateCon
 		for (int32 Index = 0; Index < RawRefBoneCount; ++Index)
 		{
 			const FName BoneName = RawRefBoneInfo[Index].Name;
-			if (const FMutableRemappedBone* RemappedBone = ModelResources.RemappedBoneNames.FindByKey(BoneName))
+			const FString BoneNameString = BoneName.ToString().ToLower();
+			if(const uint32* Hash = ModelResources.BoneNamesMap.Find(BoneNameString))
 			{
-				TPair<FName, uint16>& BoneInfo = BoneInfoMap.Add(mu::FBoneName(RemappedBone->Hash));
-				BoneInfo.Key = BoneName;
-				BoneInfo.Value = Index;
-			}
-			else if(const uint32* Hash = ModelResources.BoneNamesMap.Find(BoneName))
-			{
-				const FString BoneNameString = BoneName.ToString();
 				const mu::FBoneName Bone(*Hash);
 				TPair<FName, uint16>& BoneInfo = BoneInfoMap.Add(Bone);
 				BoneInfo.Key = BoneName;
diff --git a/Engine/Plugins/Experimental/Mutable/Source/CustomizableObjectEditor/Private/MuCOE/CustomizableObjectCompiler.cpp b/Engine/Plugins/Experimental/Mutable/Source/CustomizableObjectEditor/Private/MuCOE/CustomizableObjectCompiler.cpp
--- a/Engine/Plugins/Experimental/Mutable/Source/CustomizableObjectEditor/Private/MuCOE/CustomizableObjectCompiler.cpp
+++ b/Engine/Plugins/Experimental/Mutable/Source/CustomizableObjectEditor/Private/MuCOE/CustomizableObjectCompiler.cpp
@@ -973,7 +973,7 @@ void FCustomizableObjectCompiler::CompileInternal(UCustomizableObject* Object, c
 		ModelResources.MaterialSlotNames = MoveTemp(GenerationContext.ReferencedMaterialSlotNames);
 		ModelResources.SocketArray = MoveTemp(GenerationContext.SocketArray);
 
-		const int32 NumBones = GenerationContext.UniqueBoneNames.Num();
+		const int32 NumBones = GenerationContext.UniqueBoneNames.Num() + GenerationContext.RemappedBoneNames.Num();
 		ModelResources.BoneNamesMap.Reserve(NumBones);
 
 		for (auto& It : GenerationContext.UniqueBoneNames)
@@ -983,9 +983,7 @@ void FCustomizableObjectCompiler::CompileInternal(UCustomizableObject* Object, c
 
 		for (auto& It : GenerationContext.RemappedBoneNames)
 		{
-			FMutableRemappedBone& RemappedBone = ModelResources.RemappedBoneNames.AddDefaulted_GetRef();
-			RemappedBone.Name = It.Key; // BoneName
-			RemappedBone.Hash = It.Value.Id; // mu::FBoneName
+			ModelResources.BoneNamesMap.Add(It.Key, It.Value.Id);
 		}
 
 		ModelResources.SkinWeightProfilesInfo = MoveTemp(GenerationContext.SkinWeightProfilesInfo);
diff --git a/Engine/Plugins/Experimental/Mutable/Source/CustomizableObjectEditor/Private/MuCOE/GenerateMutableSource/GenerateMutableSource.cpp b/Engine/Plugins/Experimental/Mutable/Source/CustomizableObjectEditor/Private/MuCOE/GenerateMutableSource/GenerateMutableSource.cpp
--- a/Engine/Plugins/Experimental/Mutable/Source/CustomizableObjectEditor/Private/MuCOE/GenerateMutableSource/GenerateMutableSource.cpp
+++ b/Engine/Plugins/Experimental/Mutable/Source/CustomizableObjectEditor/Private/MuCOE/GenerateMutableSource/GenerateMutableSource.cpp
@@ -376,14 +376,14 @@ void FMutableGraphGenerationContext::GenerateSharedSurfacesUniqueIds()
 
 bool FMutableGraphGenerationContext::FindBone(const FName& InBoneName, mu::FBoneName& OutBoneName) const
 {
-	const FString BoneNameString = InBoneName.ToString();
+	const FString BoneNameString = InBoneName.ToString().ToLower();
 	OutBoneName.Id = CityHash32(reinterpret_cast<const char*>(*BoneNameString), BoneNameString.Len() * sizeof(FString::ElementType));
 	if (UniqueBoneNames.Contains(OutBoneName))
 	{
 		return true;
 	}
 
-	if (const mu::FBoneName* BoneName = RemappedBoneNames.Find(InBoneName))
+	if (const mu::FBoneName* BoneName = RemappedBoneNames.Find(BoneNameString))
 	{
 		OutBoneName = *BoneName;
 		return true;
@@ -395,15 +395,15 @@ bool FMutableGraphGenerationContext::FindBone(const FName& InBoneName, mu::FBone
 
 mu::FBoneName FMutableGraphGenerationContext::GetBoneUnique(const FName& InBoneName)
 {
-	const FString BoneNameString = InBoneName.ToString();
+	const FString BoneNameString = InBoneName.ToString().ToLower();
 	mu::FBoneName Bone(CityHash32(reinterpret_cast<const char*>(*BoneNameString), BoneNameString.Len() * sizeof(FString::ElementType)));
 
 	bool bNewRemappedBoneName = false;
 
-	FName& BoneName = UniqueBoneNames.FindOrAdd(Bone, InBoneName);
+	FString& BoneName = UniqueBoneNames.FindOrAdd(Bone, BoneNameString);
 	while (BoneName != InBoneName)
 	{
-		if (mu::FBoneName* RemappedBoneName = RemappedBoneNames.Find(InBoneName))
+		if (mu::FBoneName* RemappedBoneName = RemappedBoneNames.Find(BoneNameString))
 		{
 			Bone.Id = RemappedBoneName->Id;
 			break;
@@ -414,12 +414,12 @@ mu::FBoneName FMutableGraphGenerationContext::GetBoneUnique(const FName& InBoneN
 
 		// Increase Id in an attempt to make it unique again.
 		++Bone.Id;
-		BoneName = UniqueBoneNames.FindOrAdd(Bone, InBoneName);
+		BoneName = UniqueBoneNames.FindOrAdd(Bone, BoneNameString);
 	}
 
 	if (bNewRemappedBoneName)
 	{
-		RemappedBoneNames.Add(InBoneName, Bone);
+		RemappedBoneNames.Add(BoneNameString, Bone);
 	}
 
 	return Bone;
diff --git a/Engine/Plugins/Experimental/Mutable/Source/CustomizableObjectEditor/Private/MuCOE/GenerateMutableSource/GenerateMutableSource.h b/Engine/Plugins/Experimental/Mutable/Source/CustomizableObjectEditor/Private/MuCOE/GenerateMutableSource/GenerateMutableSource.h
--- a/Engine/Plugins/Experimental/Mutable/Source/CustomizableObjectEditor/Private/MuCOE/GenerateMutableSource/GenerateMutableSource.h
+++ b/Engine/Plugins/Experimental/Mutable/Source/CustomizableObjectEditor/Private/MuCOE/GenerateMutableSource/GenerateMutableSource.h
@@ -697,8 +697,8 @@ struct FMutableGraphGenerationContext
 	TArray<const USkeleton*> ReferencedSkeletons;
 
 	// Array of unique Bone identifiers. 
-	TMap<mu::FBoneName, FName> UniqueBoneNames;
-	TMap<FName, mu::FBoneName> RemappedBoneNames; // Bone identifiers that had a collision.
+	TMap<mu::FBoneName, FString> UniqueBoneNames;
+	TMap<FString, mu::FBoneName> RemappedBoneNames; // Bone identifiers that had a collision.
 
 	// Used to aviod Nodes with duplicated ids
 	TMap<FGuid, TArray<const UObject*>> NodeIdsMap;
diff --git a/Engine/Plugins/Experimental/Mutable/Source/CustomizableObjectEditor/Private/MuCOE/GenerateMutableSource/GenerateMutableSourceMesh.cpp b/Engine/Plugins/Experimental/Mutable/Source/CustomizableObjectEditor/Private/MuCOE/GenerateMutableSource/GenerateMutableSourceMesh.cpp
--- a/Engine/Plugins/Experimental/Mutable/Source/CustomizableObjectEditor/Private/MuCOE/GenerateMutableSource/GenerateMutableSourceMesh.cpp
+++ b/Engine/Plugins/Experimental/Mutable/Source/CustomizableObjectEditor/Private/MuCOE/GenerateMutableSource/GenerateMutableSourceMesh.cpp
@@ -3063,12 +3063,13 @@ mu::NodeMeshPtr GenerateMutableSourceMesh(const UEdGraphPin* Pin,
 
 				if (GenerationContext.Options.bSkinWeightProfilesEnabled)
 				{
-					FSkeletalMeshModel* ImportModel = TypedNodeSkel->SkeletalMesh->GetImportedModel();
-
-					for (const auto& SkinWeightProfile : ImportModel->LODModels[LODIndex].SkinWeightProfiles)
+					if (ImportedModel && ImportedModel->LODModels.IsValidIndex(LODIndex))
 					{
-						const int32 ProfileId = static_cast<int32>(GenerationContext.GetSkinWeightProfileIdUnique(SkinWeightProfile.Key));
-						MeshData.SkinWeightProfilesSemanticIndices.AddUnique(ProfileId);
+						for (const auto& SkinWeightProfile : ImportedModel->LODModels[LODIndex].SkinWeightProfiles)
+						{
+							const int32 ProfileId = static_cast<int32>(GenerationContext.GetSkinWeightProfileIdUnique(SkinWeightProfile.Key));
+							MeshData.SkinWeightProfilesSemanticIndices.AddUnique(ProfileId);
+						}
 					}
 				}
 
