[mutable] Ensure padding data is zeroed when passing buffers to Mutable to avoid deterministic packaging problems in case those buffers are note reformatted.

Fix a crash in the mutable debugger when no parameters have been created for a preview.

#rnx

diff --git a/Engine/Plugins/Experimental/Mutable/Source/CustomizableObjectEditor/Private/MuCOE/GenerateMutableSource/GenerateMutableSourceMesh.cpp b/Engine/Plugins/Experimental/Mutable/Source/CustomizableObjectEditor/Private/MuCOE/GenerateMutableSource/GenerateMutableSourceMesh.cpp
--- a/Engine/Plugins/Experimental/Mutable/Source/CustomizableObjectEditor/Private/MuCOE/GenerateMutableSource/GenerateMutableSourceMesh.cpp
+++ b/Engine/Plugins/Experimental/Mutable/Source/CustomizableObjectEditor/Private/MuCOE/GenerateMutableSource/GenerateMutableSourceMesh.cpp
@@ -607,6 +607,61 @@ mu::Ptr<mu::PhysicsBody> MakePhysicsBodyFromAsset(FMutableGraphGenerationContext
 	return PhysicsBody;
 }
 
+
+namespace MutablePrivate
+{
+
+	// \TODO: Since we are no longer bulk-copying, we could generate buffers without padding, and only with the relevant information.
+	void CopyBufferClearingPadding(mu::FMeshBufferSet& BufferSet, int32 BufferIndex, const void* InSourceData)
+	{
+		check(BufferIndex < BufferSet.m_buffers.Num());
+
+		int32 ElementCount = BufferSet.m_elementCount;
+		int32 ChannelCount = BufferSet.m_buffers[BufferIndex].m_channels.Num();
+		int32 ElementSize = BufferSet.m_buffers[BufferIndex].m_elementSize;
+
+		uint8* TargetData = BufferSet.GetBufferData(BufferIndex);
+		const uint8* SourceData = reinterpret_cast<const uint8*>(InSourceData);
+
+		for (int32 Element = 0; Element < ElementCount; ++Element)
+		{
+			int32 CurrentOffset = 0;
+			for (int32 ChannelIndex = 0; ChannelIndex < ChannelCount; ++ChannelIndex)
+			{
+				const mu::FMeshBufferChannel& Channel = BufferSet.m_buffers[BufferIndex].m_channels[ChannelIndex];
+
+				int32 ChannelOffset = Channel.m_offset;
+
+				int32 PreviousPadding = ChannelOffset - CurrentOffset;
+				if (PreviousPadding > 0)
+				{
+					FMemory::Memzero(TargetData, PreviousPadding);
+					TargetData += PreviousPadding;
+					SourceData += PreviousPadding;
+					CurrentOffset += PreviousPadding;
+				}
+
+				int32 ChannelSize = Channel.m_componentCount * GetMeshFormatData(Channel.m_format).SizeInBytes;
+				FMemory::Memcpy(TargetData, SourceData, ChannelSize);
+				TargetData += ChannelSize;
+				SourceData += ChannelSize;
+				CurrentOffset += ChannelSize;
+			}
+
+			// Padding at the end?
+			int32 FinalPadding = ElementSize - CurrentOffset;
+			if (FinalPadding > 0)
+			{
+				FMemory::Memzero(TargetData, FinalPadding);
+				TargetData += FinalPadding;
+				SourceData += FinalPadding;
+			}
+		}
+	}
+
+}
+
+
 mu::MeshPtr ConvertSkeletalMeshToMutable(const USkeletalMesh* InSkeletalMesh, const TSoftClassPtr<UAnimInstance>& AnimBp, int32 LODIndexConnected, int32 SectionIndexConnected, int32 LODIndex, int32 SectionIndex, FMutableGraphGenerationContext& GenerationContext, const UCustomizableObjectNode* CurrentNode)
 {
 	if(!InSkeletalMesh)
@@ -884,7 +939,7 @@ mu::MeshPtr ConvertSkeletalMeshToMutable(const USkeletalMesh* InSkeletalMesh, co
 		};
 
 		MutableMesh->GetVertexBuffers().SetBuffer(0, ElementSize, ChannelCount, Semantics, SemanticIndices, Formats, Components, Offsets);
-		FMemory::Memcpy(MutableMesh->GetVertexBuffers().GetBufferData(0), Vertices.GetData() + VertexStart, VertexCount* ElementSize);
+		MutablePrivate::CopyBufferClearingPadding(MutableMesh->GetVertexBuffers(), 0, Vertices.GetData() + VertexStart);
 	}
 	else
 	{
@@ -998,7 +1053,7 @@ mu::MeshPtr ConvertSkeletalMeshToMutable(const USkeletalMesh* InSkeletalMesh, co
 		}
 
 		MutableMesh->GetVertexBuffers().SetBuffer(0, ElementSize, ChannelCount, Semantics, SemanticIndices, Formats, Components, Offsets);
-		FMemory::Memcpy(MutableMesh->GetVertexBuffers().GetBufferData(0), Vertices.GetData() + VertexStart, VertexCount * ElementSize);
+		MutablePrivate::CopyBufferClearingPadding(MutableMesh->GetVertexBuffers(), 0, Vertices.GetData() + VertexStart);
 	}
 
 
@@ -3061,7 +3116,7 @@ mu::NodeMeshPtr GenerateMutableSourceMesh(const UEdGraphPin* Pin,
 					}
 				}
 
-				if (GenerationContext.Options.bSkinWeightProfilesEnabled)
+				if (GenerationContext.Options.bSkinWeightProfilesEnabled && LODIndex>=0)
 				{
 					if (ImportedModel && ImportedModel->LODModels.IsValidIndex(LODIndex))
 					{
diff --git a/Engine/Plugins/Experimental/Mutable/Source/CustomizableObjectEditor/Private/MuCOE/SMutableParametersWidget.cpp b/Engine/Plugins/Experimental/Mutable/Source/CustomizableObjectEditor/Private/MuCOE/SMutableParametersWidget.cpp
--- a/Engine/Plugins/Experimental/Mutable/Source/CustomizableObjectEditor/Private/MuCOE/SMutableParametersWidget.cpp
+++ b/Engine/Plugins/Experimental/Mutable/Source/CustomizableObjectEditor/Private/MuCOE/SMutableParametersWidget.cpp
@@ -317,8 +317,8 @@ void SMutableParametersWidget::GenerateAndAttachParameterSlate(const int32 Param
 						SNew(STextBlock)
 						.Text_Lambda([this, ParamIndex, RangeIndex, OptionNamesAttribute]() -> FText
 						{
-							const int32 Value = MutableParameters->GetIntValue(ParamIndex, RangeIndex);
-							const int32 ValueIndex = MutableParameters->GetIntValueIndex(ParamIndex, Value);
+							const int32 Value = MutableParameters ? MutableParameters->GetIntValue(ParamIndex, RangeIndex) : 0;
+							const int32 ValueIndex = MutableParameters ? MutableParameters->GetIntValueIndex(ParamIndex, Value) : 0;
 
 							return FText::FromString(*OptionNamesAttribute[ValueIndex]);
 						})
