[Mutable] Fix ensure raised with Texture Parameters.

Moved UTexture data accesses from MutableThread to GameThread.

#rnx

diff --git a/Engine/Plugins/Experimental/Mutable/Source/CustomizableObject/Private/MuCO/UnrealMutableImageProvider.cpp b/Engine/Plugins/Experimental/Mutable/Source/CustomizableObject/Private/MuCO/UnrealMutableImageProvider.cpp
--- a/Engine/Plugins/Experimental/Mutable/Source/CustomizableObject/Private/MuCO/UnrealMutableImageProvider.cpp
+++ b/Engine/Plugins/Experimental/Mutable/Source/CustomizableObject/Private/MuCO/UnrealMutableImageProvider.cpp
@@ -499,41 +499,28 @@ mu::FImageDesc FUnrealMutableImageProvider::GetImageDesc(FName Id, uint8 Mipmaps
 			// Easy case where the image was directly provided
 			Result = mu::FImageDesc(ImageInfo.Image->GetSize(), ImageInfo.Image->GetFormat(), ImageInfo.Image->GetLODCount());
 		}
-		else if (UTexture2D* TextureToLoad = ImageInfo.TextureToLoad)
+		else if (ImageInfo.TextureToLoad)
 		{
 			// It's safe to access TextureToLoad because ExternalImagesLock guarantees that the data in GlobalExternalImages is valid,
 			// not being modified by the game thread at the moment and the texture cannot be GCed because of the AddReferencedObjects
 			// in the FUnrealMutableImageProvider
-
-			int32 MipIndex = MipmapsToSkip < TextureToLoad->GetPlatformData()->Mips.Num() ? MipmapsToSkip : TextureToLoad->GetPlatformData()->Mips.Num() - 1;
-
-			// Mips in the mip tail are inlined and can't be streamed, find the smallest mip available.
-			for (; MipIndex > 0; --MipIndex)
-			{
-				if (TextureToLoad->GetPlatformData()->Mips[MipIndex].BulkData.CanLoadFromDisk())
-				{
-					break;
-				}
-			}
-
-			// Texture format and the equivalent mutable format
-			const EPixelFormat Format = TextureToLoad->GetPlatformData()->PixelFormat;
-			const mu::EImageFormat MutableFormat = GetMutablePixelFormat(Format);
-
+			
 			// Check if it's a format we support
-			if (MutableFormat == mu::EImageFormat::IF_NONE)
+			if (ImageInfo.MutableFormat == mu::EImageFormat::IF_NONE)
 			{
-				UE_LOG(LogMutable, Warning, TEXT("Failed to get external image descriptor. Unexpected image format. EImageFormat [%s]."), GetPixelFormatString(Format));
+				UE_LOG(LogMutable, Warning, TEXT("Failed to get external image descriptor. Unexpected image format. EPixelFormat [%s]."), GetPixelFormatString(ImageInfo.Format));
 				return CreateDummyDesc();
 			}
 
+			const int32 MipIndex = FMath::Min3(static_cast<int32>(MipmapsToSkip), ImageInfo.FirstAvailableMip, ImageInfo.NumMips - 1);
+
 			const mu::FImageSize ImageSize = mu::FImageSize(
-					TextureToLoad->GetSizeX() >> MipIndex,
-					TextureToLoad->GetSizeY() >> MipIndex);
+					ImageInfo.SizeX >> MipIndex,
+					ImageInfo.SizeY >> MipIndex);
 
 			const int32 Lods = 1;
 
-			Result = mu::FImageDesc(ImageSize, MutableFormat, Lods);
+			Result = mu::FImageDesc(ImageSize, ImageInfo.MutableFormat, Lods);
 		}
 		else
 		{
@@ -828,3 +815,30 @@ void FUnrealMutableImageProvider::AddReferencedObjects(FReferenceCollector& Coll
 		ExternalImagesLock.Unlock();
 	}
 }
+
+
+FUnrealMutableImageProvider::FUnrealMutableImageInfo::FUnrealMutableImageInfo(const mu::ImagePtr& InImage, UTexture2D* InTextureToLoad) :
+	Image(InImage),
+	TextureToLoad(InTextureToLoad)
+{
+	check(IsInGameThread())
+	
+	FTexturePlatformData* PlatformData = TextureToLoad->GetPlatformData();
+	
+	Format = PlatformData->PixelFormat;
+	MutableFormat = GetMutablePixelFormat(Format);
+
+	NumMips = PlatformData->Mips.Num();
+	
+	// Mips in the mip tail are inlined and can't be streamed, find the smallest mip available.
+	for (FirstAvailableMip = NumMips - 1; FirstAvailableMip > 0; --FirstAvailableMip)
+	{
+		if (PlatformData->Mips[FirstAvailableMip].BulkData.CanLoadFromDisk())
+		{
+			break;
+		}
+	}
+
+	SizeX = TextureToLoad->GetSizeX();
+	SizeY = TextureToLoad->GetSizeY();
+}
diff --git a/Engine/Plugins/Experimental/Mutable/Source/CustomizableObject/Private/MuCO/UnrealMutableImageProvider.h b/Engine/Plugins/Experimental/Mutable/Source/CustomizableObject/Private/MuCO/UnrealMutableImageProvider.h
--- a/Engine/Plugins/Experimental/Mutable/Source/CustomizableObject/Private/MuCO/UnrealMutableImageProvider.h
+++ b/Engine/Plugins/Experimental/Mutable/Source/CustomizableObject/Private/MuCO/UnrealMutableImageProvider.h
@@ -8,6 +8,7 @@
 #include "Containers/Ticker.h"
 #include "Containers/Map.h"
 #include "Containers/Queue.h"
+#include "PixelFormat.h"
 
 #include "Tasks/Task.h"
 
@@ -58,15 +59,20 @@ private:
 	struct FUnrealMutableImageInfo
 	{
 		FUnrealMutableImageInfo() {}
-		FUnrealMutableImageInfo(const mu::ImagePtr& InImage, class UTexture2D* InTextureToLoad) : Image(InImage), TextureToLoad(InTextureToLoad) { }
+		FUnrealMutableImageInfo(const mu::ImagePtr& InImage, UTexture2D* InTextureToLoad);
 
 		mu::ImagePtr Image;
 
 		/** If the above Image has not been loaded in the game thread, the TextureToLoad bulk data will be loaded
-		* from the Mutable thread when it's needed
-		*/
+		* from the Mutable thread when it's needed */
 		TObjectPtr<UTexture2D> TextureToLoad = nullptr;
-
+		EPixelFormat Format = PF_Unknown;
+		mu::EImageFormat MutableFormat = mu::EImageFormat::IF_NONE;
+		int32 NumMips = 0;
+		int32 FirstAvailableMip = -1;
+		int32 SizeX = 0;
+		int32 SizeY = 0;
+		
 		/** true of the reference maintained by the user. */
 		bool ReferencesUser = false;
 		
