[mutable] Changed the name of the generated skeletal meshes so that it references Mutable and the Customizable Object name, to make it easier to identify Mutable meshes.
#rnx
#tests Editor, Client,

diff --git a/Engine/Plugins/Experimental/Mutable/Source/CustomizableObject/Classes/MuCO/CustomizableObjectSkeletalMesh.h b/Engine/Plugins/Experimental/Mutable/Source/CustomizableObject/Classes/MuCO/CustomizableObjectSkeletalMesh.h
--- a/Engine/Plugins/Experimental/Mutable/Source/CustomizableObject/Classes/MuCO/CustomizableObjectSkeletalMesh.h
+++ b/Engine/Plugins/Experimental/Mutable/Source/CustomizableObject/Classes/MuCO/CustomizableObjectSkeletalMesh.h
@@ -21,6 +21,8 @@ class FUpdateContextPrivate;
 
 class FCustomizableObjectMeshStreamIn;
 
+FName GenerateUniqueNameFromCOInstance(const UCustomizableObjectInstance& Instance);
+
 /**
  * CustomizableObjectSkeletalMesh is a class inheriting from USkeletalMesh with the sole purpose of allowing skeletal meshes generated by Mutable
  * to stream-in Mesh LODs. 
diff --git a/Engine/Plugins/Experimental/Mutable/Source/CustomizableObject/Private/MuCO/CustomizableObjectInstance.cpp b/Engine/Plugins/Experimental/Mutable/Source/CustomizableObject/Private/MuCO/CustomizableObjectInstance.cpp
--- a/Engine/Plugins/Experimental/Mutable/Source/CustomizableObject/Private/MuCO/CustomizableObjectInstance.cpp
+++ b/Engine/Plugins/Experimental/Mutable/Source/CustomizableObject/Private/MuCO/CustomizableObjectInstance.cpp
@@ -1832,7 +1832,8 @@ bool UCustomizableInstancePrivate::UpdateSkeletalMesh_PostBeginUpdate0(UCustomiz
 		}
 		else
 		{
-			SkeletalMeshes[Component.Id] = NewObject<USkeletalMesh>(GetTransientPackage(), NAME_None, RF_Transient);
+			FName SkeletalMeshName = GenerateUniqueNameFromCOInstance(*Public);
+			SkeletalMeshes[Component.Id] = NewObject<USkeletalMesh>(GetTransientPackage(), SkeletalMeshName, RF_Transient);
 		}
 
 		USkeletalMesh* SkeletalMesh = SkeletalMeshes[Component.Id];
diff --git a/Engine/Plugins/Experimental/Mutable/Source/CustomizableObject/Private/MuCO/CustomizableObjectSkeletalMesh.cpp b/Engine/Plugins/Experimental/Mutable/Source/CustomizableObject/Private/MuCO/CustomizableObjectSkeletalMesh.cpp
--- a/Engine/Plugins/Experimental/Mutable/Source/CustomizableObject/Private/MuCO/CustomizableObjectSkeletalMesh.cpp
+++ b/Engine/Plugins/Experimental/Mutable/Source/CustomizableObject/Private/MuCO/CustomizableObjectSkeletalMesh.cpp
@@ -9,12 +9,23 @@
 #include "MuCO/CustomizableObjectInstance.h"
 #include "MuCO/CustomizableObjectMeshUpdate.h"
 
+FName GenerateUniqueNameFromCOInstance(const UCustomizableObjectInstance& Instance)
+{
+	FString BaseName = FString("MutableSkMesh-") + Instance.GetCustomizableObject()->GetName();
+	FName SkeletalMeshName = MakeUniqueObjectName(GetTransientPackage(), USkeletalMesh::StaticClass(), FName(BaseName), 
+							                      EUniqueObjectNameOptions::GloballyUnique);
+
+	return SkeletalMeshName;
+}
+
+
 UCustomizableObjectSkeletalMesh* UCustomizableObjectSkeletalMesh::CreateSkeletalMesh(const TSharedRef<FUpdateContextPrivate>& OperationData, const UCustomizableObjectInstance& Instance, const int32 InComponentIndex)
 {
 	UCustomizableObject* CustomizableObject = Instance.GetCustomizableObject();
 	check(CustomizableObject);
 
-	UCustomizableObjectSkeletalMesh* OutSkeletalMesh = NewObject<UCustomizableObjectSkeletalMesh>();
+	FName SkeletalMeshName = GenerateUniqueNameFromCOInstance(Instance);
+	UCustomizableObjectSkeletalMesh* OutSkeletalMesh = NewObject<UCustomizableObjectSkeletalMesh>(GetTransientPackage(), SkeletalMeshName, RF_Transient);
 	
 	// Debug info
 	OutSkeletalMesh->CustomizableObjectPathName = GetNameSafe(CustomizableObject);
