[Mutable] Fix hang in MRQ when compiling COs.

MRQ requires Mutable to be synchronous in the GameThread (BlockTillAllRequestsFinished). Compilations were hanging due to two reasons:
1. We were not forcing synchronous compilations when BlockTillAllRequestsFinished.
2. An already started async compilation was hanging due to not switching to synchronous mode (now the compilar tick has this flag).

- Changed the return meaning of the FCustomizableObjectCompiler::Tick. Now it will return true at the end of the compilation process.
- Removed the Unlock PreloadingReferencerAssetsCallback. This was causing concurrent compilations. Now the object will be unlocked at the end of the compilation.
- Renamed FinishCompilation to FinishCompilationTask to avoid confusions. This function is not called at the end of the full compilation process.
- PreloadingReferencerAssetsCallback is no longer static.
- Removed unused COIEdior Compiler. Was never created.
- Now using the StreamableManager when loading multiple assets synchronously in LoadAdditionalAssetsAndData. Now compilation and update loads UObjects using the same API.

#rnx

diff --git a/Engine/Plugins/Experimental/Mutable/Source/CustomizableObject/Internal/MuCO/CustomizableObjectInstancePrivate.h b/Engine/Plugins/Experimental/Mutable/Source/CustomizableObject/Internal/MuCO/CustomizableObjectInstancePrivate.h
--- a/Engine/Plugins/Experimental/Mutable/Source/CustomizableObject/Internal/MuCO/CustomizableObjectInstancePrivate.h
+++ b/Engine/Plugins/Experimental/Mutable/Source/CustomizableObject/Internal/MuCO/CustomizableObjectInstancePrivate.h
@@ -221,7 +221,7 @@ public:
 
 	/** Returns the task that will be called when all assets and data are loaded, may be already completed if no assets or data needs loading.
 	 * If no StreamableManager is provided, it will load assets synchronously. */
-	UE::Tasks::FTask LoadAdditionalAssetsAndData(const TSharedRef<FUpdateContextPrivate>& OperationData, FStreamableManager* StreamableManager);
+	UE::Tasks::FTask LoadAdditionalAssetsAndData(const TSharedRef<FUpdateContextPrivate>& OperationData, FStreamableManager& StreamableManager, bool bAsync);
 
 	void AdditionalAssetsAsyncLoaded(UCustomizableObjectInstance* Public);
 
diff --git a/Engine/Plugins/Experimental/Mutable/Source/CustomizableObject/Internal/MuCO/CustomizableObjectPrivate.h b/Engine/Plugins/Experimental/Mutable/Source/CustomizableObject/Internal/MuCO/CustomizableObjectPrivate.h
--- a/Engine/Plugins/Experimental/Mutable/Source/CustomizableObject/Internal/MuCO/CustomizableObjectPrivate.h
+++ b/Engine/Plugins/Experimental/Mutable/Source/CustomizableObject/Internal/MuCO/CustomizableObjectPrivate.h
@@ -47,8 +47,9 @@ public:
 	
 	virtual void Compile(UCustomizableObject& Object, const FCompilationOptions& Options, bool bAsync) {};
 
-	virtual bool Tick() { return false; }
-	virtual void ForceFinishCompilation() {};
+	/** Return true if no further ticks are required. In other words, all work has completed. */
+	virtual bool Tick(bool bBlocking) { return true; }
+	virtual void ForceFinishCompilation() {}
 
 	// Return true if this object doesn't reference a parent object.
 	virtual bool IsRootObject(const UCustomizableObject* Object) const { return true; }
@@ -690,7 +691,8 @@ class CUSTOMIZABLEOBJECT_API UCustomizableObjectPrivate : public UObject
 
 public:
 	UCustomizableObjectPrivate();
-	
+
+	/** Must be called after unlocking the CustomizableObject. */
 	void SetModel(const TSharedPtr<mu::Model, ESPMode::ThreadSafe>& Model, const FGuid Identifier);
 	const TSharedPtr<mu::Model, ESPMode::ThreadSafe>& GetModel();
 	TSharedPtr<const mu::Model, ESPMode::ThreadSafe> GetModel() const;
diff --git a/Engine/Plugins/Experimental/Mutable/Source/CustomizableObject/Private/MuCO/CustomizableObjectInstance.cpp b/Engine/Plugins/Experimental/Mutable/Source/CustomizableObject/Private/MuCO/CustomizableObjectInstance.cpp
--- a/Engine/Plugins/Experimental/Mutable/Source/CustomizableObject/Private/MuCO/CustomizableObjectInstance.cpp
+++ b/Engine/Plugins/Experimental/Mutable/Source/CustomizableObject/Private/MuCO/CustomizableObjectInstance.cpp
@@ -4648,7 +4648,7 @@ FAutoConsoleVariableRef CVarMutableHighPriorityLoading(
 	TEXT("If enabled, the request to load additional assets will have high priority."));
 
 UE::Tasks::FTask UCustomizableInstancePrivate::LoadAdditionalAssetsAndData(
-		const TSharedRef<FUpdateContextPrivate>& OperationData, FStreamableManager* StreamableManager)
+		const TSharedRef<FUpdateContextPrivate>& OperationData, FStreamableManager& StreamableManager, bool bAsync)
 {
 	MUTABLE_CPUPROFILER_SCOPE(UCustomizableInstancePrivate::LoadAdditionalAssetsAndDataAsync);
 
@@ -4941,19 +4941,16 @@ UE::Tasks::FTask UCustomizableInstancePrivate::LoadAdditionalAssetsAndData(
 	{	
 		UE::Tasks::FTaskEvent AssetAsyncLoadCompletionEvent = StreamingCompletionEvents.Emplace_GetRef(TEXT("AssetAsyncLoadCompletionEvent"));
 
-		if (StreamableManager)
+		if (bAsync)
 		{
-			StreamingHandle = StreamableManager->RequestAsyncLoad(
+			StreamingHandle = StreamableManager.RequestAsyncLoad(
 					AssetsToStream, 
 					FStreamableDelegate::CreateUObject(this, &UCustomizableInstancePrivate::AdditionalAssetsAsyncLoaded, AssetAsyncLoadCompletionEvent),
 					bEnableHighPriorityLoading ? FStreamableManager::AsyncLoadHighPriority : FStreamableManager::DefaultAsyncLoadPriority);			
 		}
 		else
 		{
-			for (const FSoftObjectPath& Asset : AssetsToStream)
-			{
-				Asset.TryLoad();
-			}
+			StreamableManager.RequestSyncLoad(AssetsToStream);
 			
 			AdditionalAssetsAsyncLoaded(AssetAsyncLoadCompletionEvent);
 		}
diff --git a/Engine/Plugins/Experimental/Mutable/Source/CustomizableObject/Private/MuCO/CustomizableObjectSystem.cpp b/Engine/Plugins/Experimental/Mutable/Source/CustomizableObject/Private/MuCO/CustomizableObjectSystem.cpp
--- a/Engine/Plugins/Experimental/Mutable/Source/CustomizableObject/Private/MuCO/CustomizableObjectSystem.cpp
+++ b/Engine/Plugins/Experimental/Mutable/Source/CustomizableObject/Private/MuCO/CustomizableObjectSystem.cpp
@@ -3076,8 +3076,7 @@ namespace impl
 
 		// Next Task: Load Unreal Assets
 		//-------------------------------------------------------------
-		FStreamableManager* StreamManager = System->GetPrivate()->bBlocking ? nullptr : &System->GetPrivate()->StreamableManager;
-		UE::Tasks::FTask Game_LoadUnrealAssets = ObjectInstancePrivateData->LoadAdditionalAssetsAndData(OperationData, StreamManager);
+		UE::Tasks::FTask Game_LoadUnrealAssets = ObjectInstancePrivateData->LoadAdditionalAssetsAndData(OperationData, System->GetPrivate()->StreamableManager, !System->GetPrivate()->bBlocking);
 
 		// Next-next Task: Convert Resources
 		//-------------------------------------------------------------
@@ -4109,7 +4108,7 @@ void UCustomizableObjectSystem::StartNextRecompile()
 		FCompilationOptions Options = CustomizableObject->CompileOptions;
 		Options.bSilentCompilation = true;
 		check(GetPrivate()->RecompileCustomizableObjectsCompiler != nullptr);
-		GetPrivate()->RecompileCustomizableObjectsCompiler->Compile(*CustomizableObject, Options, true);
+		GetPrivate()->RecompileCustomizableObjectsCompiler->Compile(*CustomizableObject, Options, !GetPrivate()->bBlocking);
 	}
 }
 
@@ -4180,7 +4179,7 @@ void UCustomizableObjectSystem::TickRecompileCustomizableObjects()
 	
 	if (GetPrivate()->RecompileCustomizableObjectsCompiler)
 	{
-		bUpdated = GetPrivate()->RecompileCustomizableObjectsCompiler->Tick() || GetPrivate()->RecompileCustomizableObjectsCompiler->GetCompilationState() == ECustomizableObjectCompilationState::Failed;
+		bUpdated = GetPrivate()->RecompileCustomizableObjectsCompiler->Tick(GetPrivate()->bBlocking) || GetPrivate()->RecompileCustomizableObjectsCompiler->GetCompilationState() == ECustomizableObjectCompilationState::Failed;
 	}
 
 	if (bUpdated)
diff --git a/Engine/Plugins/Experimental/Mutable/Source/CustomizableObjectEditor/Private/MuCOE/CustomizableObjectCompileRunnable.cpp b/Engine/Plugins/Experimental/Mutable/Source/CustomizableObjectEditor/Private/MuCOE/CustomizableObjectCompileRunnable.cpp
--- a/Engine/Plugins/Experimental/Mutable/Source/CustomizableObjectEditor/Private/MuCOE/CustomizableObjectCompileRunnable.cpp
+++ b/Engine/Plugins/Experimental/Mutable/Source/CustomizableObjectEditor/Private/MuCOE/CustomizableObjectCompileRunnable.cpp
@@ -254,11 +254,11 @@ void FCustomizableObjectCompileRunnable::Tick()
 }
 
 
-FCustomizableObjectSaveDDRunnable::FCustomizableObjectSaveDDRunnable(UCustomizableObject* CustomizableObject, const FCompilationOptions& InOptions)
+FCustomizableObjectSaveDDRunnable::FCustomizableObjectSaveDDRunnable(UCustomizableObject* CustomizableObject, const FCompilationOptions& InOptions, TSharedPtr<mu::Model> InModel)
 {
 	MUTABLE_CPUPROFILER_SCOPE(FCustomizableObjectSaveDDRunnable::FCustomizableObjectSaveDDRunnable)
 		
-	Model = CustomizableObject->GetPrivate()->GetModel();
+	Model = InModel;
 	Options = InOptions;
 	
 	CustomizableObjectHeader.InternalVersion = CustomizableObject->GetPrivate()->CurrentSupportedVersion;
diff --git a/Engine/Plugins/Experimental/Mutable/Source/CustomizableObjectEditor/Private/MuCOE/CustomizableObjectCompileRunnable.h b/Engine/Plugins/Experimental/Mutable/Source/CustomizableObjectEditor/Private/MuCOE/CustomizableObjectCompileRunnable.h
--- a/Engine/Plugins/Experimental/Mutable/Source/CustomizableObjectEditor/Private/MuCOE/CustomizableObjectCompileRunnable.h
+++ b/Engine/Plugins/Experimental/Mutable/Source/CustomizableObjectEditor/Private/MuCOE/CustomizableObjectCompileRunnable.h
@@ -88,7 +88,7 @@ class FCustomizableObjectSaveDDRunnable : public FRunnable
 {
 public:
 
-	FCustomizableObjectSaveDDRunnable(class UCustomizableObject* CustomizableObject, const FCompilationOptions& Options);
+	FCustomizableObjectSaveDDRunnable(UCustomizableObject* CustomizableObject, const FCompilationOptions& Options, TSharedPtr<mu::Model> InModel);
 
 	// FRunnable interface
 	uint32 Run() override;
@@ -99,10 +99,8 @@ public:
 
 	const ITargetPlatform* GetTargetPlatform() const;
 
-	//
-	FCompilationOptions Options;
-
 private:
+	FCompilationOptions Options;
 
 	MutableCompiledDataStreamHeader CustomizableObjectHeader;
 
diff --git a/Engine/Plugins/Experimental/Mutable/Source/CustomizableObjectEditor/Private/MuCOE/CustomizableObjectCompiler.cpp b/Engine/Plugins/Experimental/Mutable/Source/CustomizableObjectEditor/Private/MuCOE/CustomizableObjectCompiler.cpp
--- a/Engine/Plugins/Experimental/Mutable/Source/CustomizableObjectEditor/Private/MuCOE/CustomizableObjectCompiler.cpp
+++ b/Engine/Plugins/Experimental/Mutable/Source/CustomizableObjectEditor/Private/MuCOE/CustomizableObjectCompiler.cpp
@@ -56,65 +56,85 @@ FCustomizableObjectCompiler::FCustomizableObjectCompiler() : FCustomizableObject
 
 
 
-bool FCustomizableObjectCompiler::Tick()
+bool FCustomizableObjectCompiler::Tick(bool bBlocking)
 {
-	if (CompileTask)
+	bool bFinished = true;
+
+	if (bBlocking && AsynchronousStreamableHandlePtr)
 	{
-		CompileTask->Tick();
+		bFinished = false;
+		AsynchronousStreamableHandlePtr->CancelHandle();
+		
+		UCustomizableObjectSystem::GetInstance()->GetPrivate()->StreamableManager.RequestSyncLoad(ArrayAssetToStream);
+		PreloadingReferencerAssetsCallback(false);
 	}
 	
-	bool bUpdated = false;
-
-	if (CompileTask.IsValid() && CompileTask->IsCompleted())
+	if (PreloadingReferencerAssets)
+	{
+		bFinished = false;
+	}
+	
+	if (CompileTask)
 	{
-		UE_LOG(LogMutable, Verbose, TEXT("PROFILE: [ %16.8f ] Finishing Compilation task for Object %s."), FPlatformTime::Seconds(), *CurrentObject->GetName());
+		bFinished = false;
+		CompileTask->Tick();
+		
+		if (CompileTask->IsCompleted())
+		{
+			UE_LOG(LogMutable, Verbose, TEXT("PROFILE: [ %16.8f ] Finishing Compilation task for Object %s."), FPlatformTime::Seconds(), *CurrentObject->GetName());
 
-		FinishCompilation();
+			FinishCompilationTask();
 
-		if (SaveDDTask.IsValid())
-		{
-			SaveCODerivedData();
-		}
-		else
-		{
-			bUpdated = true;
-			SetCompilationState(ECustomizableObjectCompilationState::Completed);
+			if (SaveDDTask.IsValid())
+			{
+				SaveCODerivedData();
+			}
+			else
+			{
+				SetCompilationState(ECustomizableObjectCompilationState::Completed);
 
-			RemoveCompileNotification();
+				RemoveCompileNotification();
 
-			NotifyCompilationErrors();
-		}
+				NotifyCompilationErrors();
+			}
 
-		UE_LOG(LogMutable, Verbose, TEXT("PROFILE: [ %16.8f ] Finished Compilation task."), FPlatformTime::Seconds());
-		UE_LOG(LogMutable, Verbose, TEXT("PROFILE: -----------------------------------------------------------"));
+			UE_LOG(LogMutable, Verbose, TEXT("PROFILE: [ %16.8f ] Finished Compilation task."), FPlatformTime::Seconds());
+			UE_LOG(LogMutable, Verbose, TEXT("PROFILE: -----------------------------------------------------------"));
 
-		CleanCachedReferencers();
-		UpdateArrayGCProtect();
+			CleanCachedReferencers();
+			UpdateArrayGCProtect();
 
-		CompilationLogsContainer.ClearMessagesArray();
+			CompilationLogsContainer.ClearMessagesArray();
+		}
 	}
 
-	if (SaveDDTask.IsValid() && SaveDDTask->IsCompleted())
+	if (SaveDDTask)
 	{
-		SetCompilationState(ECustomizableObjectCompilationState::Completed);
-
-		bUpdated = true;
+		bFinished = false;
 
-		FinishSavingDerivedData();
+		if (SaveDDTask->IsCompleted())
+		{
+			SetCompilationState(ECustomizableObjectCompilationState::Completed);
+		
+			FinishSavingDerivedDataTask();
 	
-		CurrentObject->GetPrivate()->PostCompile();
+			CurrentObject->GetPrivate()->PostCompile();
 
-		UE_LOG(LogMutable, Verbose, TEXT("PROFILE: [ %16.8f ] Finished Saving Derived Data task."), FPlatformTime::Seconds());
-		UE_LOG(LogMutable, Verbose, TEXT("PROFILE: -----------------------------------------------------------"));
+			UE_LOG(LogMutable, Verbose, TEXT("PROFILE: [ %16.8f ] Finished Saving Derived Data task."), FPlatformTime::Seconds());
+			UE_LOG(LogMutable, Verbose, TEXT("PROFILE: -----------------------------------------------------------"));
+		}
 	}
+	
 
 	if (CompilationLaunchPending)
 	{
+		bFinished = false;
+
 		CompilationLaunchPending = false;
 		LaunchMutableCompile();
 	}
 
-	return bUpdated;
+	return bFinished;
 }
 
 
@@ -143,27 +163,26 @@ bool FCustomizableObjectCompiler::IsRootObject(const UCustomizableObject* Object
 }
 
 
-void FCustomizableObjectCompiler::PreloadingReferencerAssetsCallback(UCustomizableObject* Object, FCustomizableObjectCompiler* CustomizableObjectCompiler, const FCompilationOptions Options, bool bAsync)
+void FCustomizableObjectCompiler::PreloadingReferencerAssetsCallback(bool bAsync)
 {
 	UE_LOG(LogMutable, Verbose, TEXT("PROFILE: [ %16.8f ] Preload asynchronously assets end."), FPlatformTime::Seconds());
 
-	if (CustomizableObjectCompiler->GetCurrentGAsyncLoadingTimeLimit() != -1.0f)
+	if (CurrentGAsyncLoadingTimeLimit != -1.0f)
 	{
 		static IConsoleVariable* CVAsyncLoadingTimeLimit = IConsoleManager::Get().FindConsoleVariable(TEXT("s.AsyncLoadingTimeLimit"));
 		if (CVAsyncLoadingTimeLimit)
 		{
-			CVAsyncLoadingTimeLimit->Set(CustomizableObjectCompiler->GetCurrentGAsyncLoadingTimeLimit());
+			CVAsyncLoadingTimeLimit->Set(CurrentGAsyncLoadingTimeLimit);
 		}
-		CustomizableObjectCompiler->SetCurrentGAsyncLoadingTimeLimit(-1.0f);
+		CurrentGAsyncLoadingTimeLimit = -1.0f;
 	}
 
-	CustomizableObjectCompiler->UpdateArrayGCProtect();
-	CustomizableObjectCompiler->SetPreloadingReferencerAssets(false);
-	UCustomizableObjectSystem::GetInstance()->UnlockObject(Object);
+	UpdateArrayGCProtect();
+	PreloadingReferencerAssets = false;
 
 	TRACE_END_REGION(UE_MUTABLE_PRELOAD_REGION);
 
-	CustomizableObjectCompiler->CompileInternal(Object, Options, bAsync);
+	CompileInternal(bAsync);
 }
 
 
@@ -231,6 +250,7 @@ void FCustomizableObjectCompiler::Compile(UCustomizableObject& Object, const FCo
 	}
 
 	CurrentObject = &Object;
+	Options = InOptions;
 
 	PreloadingReferencerAssets = true;
 	TRACE_BEGIN_REGION(UE_MUTABLE_PRELOAD_REGION);
@@ -240,8 +260,7 @@ void FCustomizableObjectCompiler::Compile(UCustomizableObject& Object, const FCo
 	TArray<FName> ArrayReferenceNames;
 	AddCachedReferencers(*Object.GetOuter()->GetPathName(), ArrayReferenceNames);
 
-	TArray<FSoftObjectPath> ArrayAssetToStream;
-
+	ArrayAssetToStream.Empty();
 	for (FAssetData& Element : ArrayAssetData)
 	{
 		ArrayAssetToStream.Add(Element.GetSoftObjectPath());
@@ -269,7 +288,7 @@ void FCustomizableObjectCompiler::Compile(UCustomizableObject& Object, const FCo
 				CVAsyncLoadingTimeLimit->Set(MaxConvertToMutableTextureTime * 1000.0f);
 			}
 
-			AsynchronousStreamableHandlePtr = Streamable.RequestAsyncLoad(ArrayAssetToStream, FStreamableDelegate::CreateStatic(PreloadingReferencerAssetsCallback, &Object, this, InOptions, bAsync));
+			AsynchronousStreamableHandlePtr = Streamable.RequestAsyncLoad(ArrayAssetToStream, FStreamableDelegate::CreateRaw(this, &FCustomizableObjectCompiler::PreloadingReferencerAssetsCallback, bAsync));
 			bAssetsLoaded = false;
 		}
 		else
@@ -281,7 +300,7 @@ void FCustomizableObjectCompiler::Compile(UCustomizableObject& Object, const FCo
 
 	if (bAssetsLoaded)
 	{
-		PreloadingReferencerAssetsCallback(&Object, this, InOptions, bAsync);
+		PreloadingReferencerAssetsCallback(bAsync);
 	}
 }
 
@@ -842,13 +861,16 @@ void FCustomizableObjectCompiler::SetCompilationState(ECustomizableObjectCompila
 }
 
 
-void FCustomizableObjectCompiler::CompileInternal(UCustomizableObject* Object, const FCompilationOptions& InOptions, bool bAsync)
+void FCustomizableObjectCompiler::CompileInternal(bool bAsync)
 {
 	MUTABLE_CPUPROFILER_SCOPE(FCustomizableObjectCompiler::Compile)
 	
 	SetCompilationState(ECustomizableObjectCompilationState::Failed);
 
-	if (!Object) return;
+	if (!CurrentObject)
+	{
+		return;
+	}
 
 	if (bAsync && CompileTask.IsValid()) // Don't start compilation if there's a compilation running
 	{
@@ -857,48 +879,43 @@ void FCustomizableObjectCompiler::CompileInternal(UCustomizableObject* Object, c
 	}
 
 	UE_LOG(LogMutable, Log, TEXT("PROFILE: [ %16.8f ] FCustomizableObjectCompiler::Compile start."), FPlatformTime::Seconds());
-
-	CurrentObject = Object;
-
+	
 	// This is redundant but necessary to keep static analysis happy.
-	if (!Object || !CurrentObject) return;
-
-	Options = InOptions;
-	Options.CustomizableObjectNumBoneInfluences = CustomizableObjectNumBoneInfluences;
-	Options.bRealTimeMorphTargetsEnabled = Object->bEnableRealTimeMorphTargets;
-	Options.bClothingEnabled = Object->bEnableClothing;
-	Options.b16BitBoneWeightsEnabled = Object->bEnable16BitBoneWeights;
-	Options.bSkinWeightProfilesEnabled = Object->bEnableAltSkinWeightProfiles;
-	Options.bPhysicsAssetMergeEnabled = Object->bEnablePhysicsAssetMerge;
-	Options.bAnimBpPhysicsManipulationEnabled = Object->bEnableAnimBpPhysicsAssetsManipualtion;
-	Options.bUseDiskCompilation = Object->CompileOptions.bUseDiskCompilation;
-	Options.ImageTiling = Object->CompileOptions.ImageTiling;
-
-	if (!InOptions.bIsCooking && IsRunningCookCommandlet())
-	{
-		UE_LOG(LogMutable, Display, TEXT("Editor compilation suspended for Customizable Object [%s]. Can not compile COs when the cook commandlet is running. "), *Object->GetName());
+	if (!CurrentObject)
+	{
 		return;
 	}
 
-	if (Object->GetPrivate()->IsLocked() || !UCustomizableObjectSystem::GetInstance()->LockObject(Object))
-	{
-		UE_LOG(LogMutable, Display, TEXT("Customizable Object is already being compiled or updated %s. Please wait a few seconds and try again."), *Object->GetName());
+	Options.CustomizableObjectNumBoneInfluences = CustomizableObjectNumBoneInfluences;
+	Options.bRealTimeMorphTargetsEnabled = CurrentObject->bEnableRealTimeMorphTargets;
+	Options.bClothingEnabled = CurrentObject->bEnableClothing;
+	Options.b16BitBoneWeightsEnabled = CurrentObject->bEnable16BitBoneWeights;
+	Options.bSkinWeightProfilesEnabled = CurrentObject->bEnableAltSkinWeightProfiles;
+	Options.bPhysicsAssetMergeEnabled = CurrentObject->bEnablePhysicsAssetMerge;
+	Options.bAnimBpPhysicsManipulationEnabled = CurrentObject->bEnableAnimBpPhysicsAssetsManipualtion;
+	Options.bUseDiskCompilation = CurrentObject->CompileOptions.bUseDiskCompilation;
+	Options.ImageTiling = CurrentObject->CompileOptions.ImageTiling;
+
+	if (!Options.bIsCooking && IsRunningCookCommandlet())
+	{
+		UE_LOG(LogMutable, Display, TEXT("Editor compilation suspended for Customizable Object [%s]. Can not compile COs when the cook commandlet is running. "), *CurrentObject->GetName());
+		UCustomizableObjectSystem::GetInstance()->UnlockObject(CurrentObject);
 		return;
 	}
+	
+	UE_LOG(LogMutable, Display, TEXT("Started Customizable Object Compile %s."), *CurrentObject->GetName());
 
-	UE_LOG(LogMutable, Display, TEXT("Started Customizable Object Compile %s."), *Object->GetName());
-
-	if (InOptions.bIsCooking && InOptions.TargetPlatform)
+	if (Options.bIsCooking && Options.TargetPlatform)
 	{
-		UE_LOG(LogMutable, Display, TEXT("Compiling Customizable Object %s for platform %s."), *Object->GetName(), *InOptions.TargetPlatform->PlatformName());
+		UE_LOG(LogMutable, Display, TEXT("Compiling Customizable Object %s for platform %s."), *CurrentObject->GetName(), *Options.TargetPlatform->PlatformName());
 	}
 
-	if (InOptions.bIsCooking && InOptions.bForceLargeLODBias)
+	if (Options.bIsCooking && Options.bForceLargeLODBias)
 	{
-		UE_LOG(LogMutable, Display, TEXT("Compiling Customizable Object with %d LODBias."), InOptions.DebugBias);
+		UE_LOG(LogMutable, Display, TEXT("Compiling Customizable Object with %d LODBias."), Options.DebugBias);
 	}
 
-	FMutableGraphGenerationContext GenerationContext(Object, this, Options);
+	FMutableGraphGenerationContext GenerationContext(CurrentObject, this, Options);
 	GenerationContext.ParamNamesToSelectedOptions = ParamNamesToSelectedOptions;
 
 	// If we don't have the target platform yet (in editor) we need to get it
@@ -919,7 +936,7 @@ void FCustomizableObjectCompiler::CompileInternal(UCustomizableObject* Object, c
 	// Generate the mutable node expression
 	FText ErrorMessage;
 	bool bIsRootObject = false;
-	mu::NodeObjectPtr MutableRoot = GenerateMutableRoot(Object, GenerationContext, ErrorMessage, bIsRootObject);
+	mu::NodeObjectPtr MutableRoot = GenerateMutableRoot(CurrentObject, GenerationContext, ErrorMessage, bIsRootObject);
 
 	if (!MutableRoot)
 	{
@@ -932,9 +949,9 @@ void FCustomizableObjectCompiler::CompileInternal(UCustomizableObject* Object, c
 			CompilerLog(FText::FromString(TEXT("Failed to generate the mutable node graph. Object not built.")), nullptr);
 		}
 
-		if (Object->GetPrivate()->IsLocked())
+		if (CurrentObject->GetPrivate()->IsLocked())
 		{
-			UCustomizableObjectSystem::GetInstance()->UnlockObject(Object);
+			UCustomizableObjectSystem::GetInstance()->UnlockObject(CurrentObject);
 		}
 
 		RemoveCompileNotification();
@@ -942,7 +959,7 @@ void FCustomizableObjectCompiler::CompileInternal(UCustomizableObject* Object, c
 	else
 	{
 		// Always work with the ModelResources (Editor) when compiling. They'll be copied to the cooked version during PreSave.
-		FModelResources& ModelResources = Object->GetPrivate()->GetModelResources(false);
+		FModelResources& ModelResources = CurrentObject->GetPrivate()->GetModelResources(false);
 		ModelResources = FModelResources();
 		
 		ModelResources.ReferenceSkeletalMeshesData = MoveTemp(GenerationContext.ReferenceSkeletalMeshesData);
@@ -1039,9 +1056,9 @@ void FCustomizableObjectCompiler::CompileInternal(UCustomizableObject* Object, c
 		}
 		
 		// Clothing	
-		Object->ClothMeshToMeshVertData = MoveTemp(GenerationContext.ClothMeshToMeshVertData);
-		Object->ContributingClothingAssetsData = MoveTemp(GenerationContext.ContributingClothingAssetsData);
-		Object->ClothSharedConfigsData.Empty();
+		CurrentObject->ClothMeshToMeshVertData = MoveTemp(GenerationContext.ClothMeshToMeshVertData);
+		CurrentObject->ContributingClothingAssetsData = MoveTemp(GenerationContext.ContributingClothingAssetsData);
+		CurrentObject->ClothSharedConfigsData.Empty();
 
 		// A clothing backend, e.g. Chaos cloth, can use 2 config files, one owned by the asset, and another that is shared 
 		// among all assets in a SkeletalMesh. When merging different assets in a skeletalmesh we need to make sure only one of 
@@ -1055,13 +1072,13 @@ void FCustomizableObjectCompiler::CompileInternal(UCustomizableObject* Object, c
 		};
 		
 		// Find shared configs to be used (One of each type) 
-		for (FCustomizableObjectClothingAssetData& ClothingAssetData : Object->ContributingClothingAssetsData)
+		for (FCustomizableObjectClothingAssetData& ClothingAssetData : CurrentObject->ContributingClothingAssetsData)
 		{
 			 for (FCustomizableObjectClothConfigData& ClothConfigData : ClothingAssetData.ConfigsData)
 			 {
 				  if (IsSharedConfigData(ClothConfigData))
 				  {
-					  FCustomizableObjectClothConfigData* FoundConfig = Object->ClothSharedConfigsData.FindByPredicate(
+					  FCustomizableObjectClothConfigData* FoundConfig = CurrentObject->ClothSharedConfigsData.FindByPredicate(
 						   [Name = ClothConfigData.ConfigName](const FCustomizableObjectClothConfigData& Other)
 						   {
 							   return Name == Other.ConfigName;
@@ -1069,43 +1086,43 @@ void FCustomizableObjectCompiler::CompileInternal(UCustomizableObject* Object, c
 
 					  if (!FoundConfig)
 					  {
-						   Object->ClothSharedConfigsData.AddDefaulted_GetRef() = ClothConfigData;
+						   CurrentObject->ClothSharedConfigsData.AddDefaulted_GetRef() = ClothConfigData;
 					  }
 				  }
 			 }
 		}
 		
 		// Remove shared configs
-		for (FCustomizableObjectClothingAssetData& ClothingAssetData : Object->ContributingClothingAssetsData)
+		for (FCustomizableObjectClothingAssetData& ClothingAssetData : CurrentObject->ContributingClothingAssetsData)
 		{
 			 ClothingAssetData.ConfigsData.RemoveAllSwap(IsSharedConfigData);
 		}
 
-		Object->GetPrivate()->GroupNodeMap = GenerationContext.GroupNodeMap;
+		CurrentObject->GetPrivate()->GroupNodeMap = GenerationContext.GroupNodeMap;
 
 		if (GenerationContext.Options.OptimizationLevel == 0)
 		{
 			// If the optimization level is "none" disable texture streaming, because textures are all referenced
 			// unreal assets and progressive generation is not supported.
-			Object->GetPrivate()->bDisableTextureStreaming = true;
+			CurrentObject->GetPrivate()->bDisableTextureStreaming = true;
 		}
 		else
 		{
-			Object->GetPrivate()->bDisableTextureStreaming = false;
+			CurrentObject->GetPrivate()->bDisableTextureStreaming = false;
 		}
 		
-		Object->GetPrivate()->bIsCompiledWithoutOptimization = GenerationContext.Options.OptimizationLevel < UE_MUTABLE_MAX_OPTIMIZATION;
+		CurrentObject->GetPrivate()->bIsCompiledWithoutOptimization = GenerationContext.Options.OptimizationLevel < UE_MUTABLE_MAX_OPTIMIZATION;
 
-		Object->GetPrivate()->GetAlwaysLoadedExtensionData() = MoveTemp(GenerationContext.AlwaysLoadedExtensionData);
+		CurrentObject->GetPrivate()->GetAlwaysLoadedExtensionData() = MoveTemp(GenerationContext.AlwaysLoadedExtensionData);
 
-		Object->GetPrivate()->GetStreamedExtensionData().Empty(GenerationContext.StreamedExtensionData.Num());
+		CurrentObject->GetPrivate()->GetStreamedExtensionData().Empty(GenerationContext.StreamedExtensionData.Num());
 		for (UCustomizableObjectResourceDataContainer* Container : GenerationContext.StreamedExtensionData)
 		{
-			Object->GetPrivate()->GetStreamedExtensionData().Emplace(Container);
+			CurrentObject->GetPrivate()->GetStreamedExtensionData().Emplace(Container);
 		}
 
 #if WITH_EDITORONLY_DATA
-		Object->GetPrivate()->CustomizableObjectPathMap = GenerationContext.CustomizableObjectPathMap;
+		CurrentObject->GetPrivate()->CustomizableObjectPathMap = GenerationContext.CustomizableObjectPathMap;
 #endif
 
 		ModelResources.NumComponents = GenerationContext.NumMeshComponentsInRoot;
@@ -1113,7 +1130,7 @@ void FCustomizableObjectCompiler::CompileInternal(UCustomizableObject* Object, c
 		ModelResources.NumLODsToStream = GenerationContext.bEnableLODStreaming ? GenerationContext.NumMaxLODsToStream : 0;
 		ModelResources.FirstLODAvailable = GenerationContext.FirstLODAvailable;
 
-		Object->GetPrivate()->GetStreamedResourceData() = MoveTemp(GenerationContext.StreamedResourceData);
+		CurrentObject->GetPrivate()->GetStreamedResourceData() = MoveTemp(GenerationContext.StreamedResourceData);
 
 		// Pass-through textures
 		TArray<FMutableSourceTextureData> NewCompileTimeReferencedTextures;
@@ -1143,7 +1160,7 @@ void FCustomizableObjectCompiler::CompileInternal(UCustomizableObject* Object, c
 			// Notice that, to avoid automatic compilations/warnings, the set of referencing objects set found here must coincide with the set found when loading the
 			// model (discard previous compilations) or when showing PIE warnings.
 			TArray<FName> ReferencingObjectNames;
-			GetReferencingPackages(*Object, ReferencingObjectNames);
+			GetReferencingPackages(*CurrentObject, ReferencingObjectNames);
 
 			for (const FName& ReferencingObjectName : ReferencingObjectNames)
 			{
@@ -1156,8 +1173,8 @@ void FCustomizableObjectCompiler::CompileInternal(UCustomizableObject* Object, c
 			}
 			
 			// Copy final array of participating objects
-			Object->GetPrivate()->ParticipatingObjects = MoveTemp(GenerationContext.ParticipatingObjects);
-			Object->GetPrivate()->DirtyParticipatingObjects.Empty();
+			CurrentObject->GetPrivate()->ParticipatingObjects = MoveTemp(GenerationContext.ParticipatingObjects);
+			CurrentObject->GetPrivate()->DirtyParticipatingObjects.Empty();
 		}
 
 		if (CompileTask.IsValid()) // Don't start compilation if there's a compilation running
@@ -1169,9 +1186,9 @@ void FCustomizableObjectCompiler::CompileInternal(UCustomizableObject* Object, c
 
 		// Lock the object to prevent instance updates while compiling. It's ignored and returns false if it's already locked
 		// Will unlock in the FinishCompilation call.
-		if (!Object->GetPrivate()->IsLocked())
+		if (!CurrentObject->GetPrivate()->IsLocked())
 		{
-			UCustomizableObjectSystem::GetInstance()->LockObject(Object);
+			UCustomizableObjectSystem::GetInstance()->LockObject(CurrentObject);
 		}
 
 		CompileTask = MakeShareable(new FCustomizableObjectCompileRunnable(MutableRoot));
@@ -1183,14 +1200,14 @@ void FCustomizableObjectCompiler::CompileInternal(UCustomizableObject* Object, c
 		{
 			CompileTask->Init();
 			CompileTask->Run();
-			FinishCompilation();
+			FinishCompilationTask();
 
 			if (SaveDDTask.IsValid())
 			{
 				SaveDDTask->Init();
 				SaveDDTask->Run();
-				CurrentObject->GetPrivate()->GetModel()->GetPrivate()->UnloadRoms();
-				FinishSavingDerivedData();
+				Model->GetPrivate()->UnloadRoms();
+				FinishSavingDerivedDataTask();
 			}
 
 			CleanCachedReferencers();
@@ -1229,7 +1246,7 @@ void FCustomizableObjectCompiler::CompileInternal(UCustomizableObject* Object, c
 		//Checking if there is the population plugin
 		if (FModuleManager::Get().IsModuleLoaded("CustomizableObjectPopulation"))
 		{
-			ICustomizableObjectPopulationModule::Get().RecompilePopulations(Object);
+			ICustomizableObjectPopulationModule::Get().RecompilePopulations(CurrentObject);
 		}
 	}
 }
@@ -1303,12 +1320,12 @@ mu::NodePtr FCustomizableObjectCompiler::Export(UCustomizableObject* Object, con
 }
 
 
-void FCustomizableObjectCompiler::FinishCompilation()
+void FCustomizableObjectCompiler::FinishCompilationTask()
 {
 	check(CompileTask.IsValid());
 
 	UpdateCompilerLogData();
-	TSharedPtr<mu::Model, ESPMode::ThreadSafe> Model = CompileTask->Model;
+	Model = CompileTask->Model;
 
 	// Generate a map that using the resource id tells the offset and size of the resource inside the bulk data
 	// At this point it is assumed that all data goes into a single file.
@@ -1333,8 +1350,6 @@ void FCustomizableObjectCompiler::FinishCompilation()
 
 	// Generate ParameterProperties and IntParameterLookUpTable
 	CurrentObject->GetPrivate()->UpdateParameterPropertiesFromModel(Model);
-
-	CurrentObject->GetPrivate()->SetModel(Model, GenerateIdentifier(*CurrentObject));
 	
 	// Reset all instances, as the parameters may need to be rebuilt.
 	for (TObjectIterator<UCustomizableObjectInstance> It; It; ++It)
@@ -1357,22 +1372,25 @@ void FCustomizableObjectCompiler::FinishCompilation()
 	{
 		TRACE_BEGIN_REGION(UE_MUTABLE_SAVEDD_REGION);
 
-		SaveDDTask = MakeShareable(new FCustomizableObjectSaveDDRunnable(CurrentObject, Options));
+		SaveDDTask = MakeShareable(new FCustomizableObjectSaveDDRunnable(CurrentObject, Options, Model));
 	}
 	else
 	{
-		CurrentObject->GetPrivate()->GetModel()->GetPrivate()->UnloadRoms();
+		Model->GetPrivate()->UnloadRoms();
 
 		// when skipping the SaveDerivedData task unlock the object so that instances can be updated
 		UCustomizableObjectSystem::GetInstance()->UnlockObject(CurrentObject);
+
+		CurrentObject->GetPrivate()->SetModel(Model, GenerateIdentifier(*CurrentObject));
 	}
 
 	UE_LOG(LogMutable, Display, TEXT("Finished Customizable Object Compile %s."), *CurrentObject->GetName());
 }
 
-void FCustomizableObjectCompiler::FinishSavingDerivedData()
+
+void FCustomizableObjectCompiler::FinishSavingDerivedDataTask()
 {
-	MUTABLE_CPUPROFILER_SCOPE(FinishSavingDerivedData)
+	MUTABLE_CPUPROFILER_SCOPE(FinishSavingDerivedDataTask)
 
 	check(SaveDDTask.IsValid());
 
@@ -1384,16 +1402,19 @@ void FCustomizableObjectCompiler::FinishSavingDerivedData()
 				SaveDDTask->BulkDataBytes,
 				SaveDDTask->MorphDataBytes);
 	}
+	
 
-	// Order matters
-	SaveDDThread.Reset();
-	SaveDDTask.Reset();
-
-	CurrentObject->GetPrivate()->GetModel()->GetPrivate()->UnloadRoms();
+	Model->GetPrivate()->UnloadRoms();
 	
 	// Unlock the object so that instances can be updated
 	UCustomizableObjectSystem::GetInstance()->UnlockObject(CurrentObject);
 
+	CurrentObject->GetPrivate()->SetModel(Model, GenerateIdentifier(*CurrentObject));
+
+	// Order matters
+	SaveDDThread.Reset();
+	SaveDDTask.Reset();
+	
 	RemoveCompileNotification();
 
 	NotifyCompilationErrors();
diff --git a/Engine/Plugins/Experimental/Mutable/Source/CustomizableObjectEditor/Private/MuCOE/CustomizableObjectCompiler.h b/Engine/Plugins/Experimental/Mutable/Source/CustomizableObjectEditor/Private/MuCOE/CustomizableObjectCompiler.h
--- a/Engine/Plugins/Experimental/Mutable/Source/CustomizableObjectEditor/Private/MuCOE/CustomizableObjectCompiler.h
+++ b/Engine/Plugins/Experimental/Mutable/Source/CustomizableObjectEditor/Private/MuCOE/CustomizableObjectCompiler.h
@@ -31,7 +31,7 @@ public:
 	CUSTOMIZABLEOBJECTEDITOR_API ECustomizableObjectCompilationState GetCompilationState() const override { return State;  };
 	
 	/** Check for pending compilation process. Returns true if an object has been updated. */
-	CUSTOMIZABLEOBJECTEDITOR_API virtual bool Tick() override;
+	CUSTOMIZABLEOBJECTEDITOR_API virtual bool Tick(bool bBlocking) override;
 
 	/** Generate the Mutable Graph from the Unreal Graph. */
 	mu::NodePtr Export(UCustomizableObject* Object, const FCompilationOptions& Options, TArray<TSoftObjectPtr<UTexture>>& OutRuntimeReferencedTextures, TArray<FMutableSourceTextureData>& OutCompilerReferencedTextures);
@@ -40,8 +40,8 @@ public:
 	void CompilerLog(const FText& Message, const UObject* Context = nullptr, const EMessageSeverity::Type MessageSeverity = EMessageSeverity::Warning, const bool bAddBaseObjectInfo = true, const ELoggerSpamBin SpamBin = ELoggerSpamBin::ShowAll);
 	void NotifyCompilationErrors() const;
 
-	void FinishCompilation();
-	void FinishSavingDerivedData();
+	void FinishCompilationTask();
+	void FinishSavingDerivedDataTask();
 	virtual void ForceFinishCompilation() override;
 
 	void AddCompileNotification(const FText& CompilationStep) const;
@@ -62,18 +62,6 @@ public:
 	/** Simply add CO elements from ArrayAssetData to ArrayGCProtect when they've been loaded from ArrayAssetData */
 	void UpdateArrayGCProtect();
 
-	/** Getter of CurrentGAsyncLoadingTimeLimit */
-	float GetCurrentGAsyncLoadingTimeLimit() { return CurrentGAsyncLoadingTimeLimit; }
-
-	/** Setter for CurrentGAsyncLoadingTimeLimit */
-	void SetCurrentGAsyncLoadingTimeLimit(float Value) { CurrentGAsyncLoadingTimeLimit = Value; }
-
-	/** Getter of PreloadingReferencerAssets */
-	bool GetPreloadingReferencerAssets() { return PreloadingReferencerAssets; }
-
-	/** Setter for PreloadingReferencerAssets */
-	void SetPreloadingReferencerAssets(bool Value) { PreloadingReferencerAssets = Value; }
-
 	// This is used to restrict Group nodes to compile only the SelectedOption and discard all the others
 	void AddCompileOnlySelectedOption(const FString& ParamName, const FString& OptionValue);
 	// This clears all restrictions and let's a full compilation happen
@@ -90,14 +78,13 @@ private:
 	// Object containing all error and warning logs raised during compilation.
 	FCompilationMessageCache CompilationLogsContainer;
 	
-private:
 	void SetCompilationState(ECustomizableObjectCompilationState InState);
 
 	ECustomizableObjectCompilationState State = ECustomizableObjectCompilationState::None;
 	
-	void CompileInternal(UCustomizableObject* Object, const FCompilationOptions& Options, bool bAsync = false);
+	void CompileInternal(bool bAsync = false);
 
-	static void PreloadingReferencerAssetsCallback(UCustomizableObject* Object, FCustomizableObjectCompiler* CustomizableObjectCompiler, const FCompilationOptions Options, bool bAsync);
+	void PreloadingReferencerAssetsCallback(bool bAsync);
 	
 	void ProcessChildObjectsRecursively(UCustomizableObject* Object, FMutableGraphGenerationContext &GenerationContext);
 
@@ -125,7 +112,8 @@ private:
 	/** Array where to put the names of the already processed child in ProcessChildObjectsRecursively */
 	TArray<FName> ArrayAlreadyProcessedChild;
 
-
+	TArray<FSoftObjectPath> ArrayAssetToStream;
+	
 	// Will output to Mutable Log the warning and error messages generated during the CO compilation
 	// and update the values of NumWarnings and NumErrors
 	void UpdateCompilerLogData();
@@ -218,6 +206,7 @@ private:
 	TMap<FString, FString> ParamNamesToSelectedOptions;
 
 	/** Pointer to the Asynchronous Preloading process call back */
-	TSharedPtr<struct FStreamableHandle> AsynchronousStreamableHandlePtr;
-	
+	TSharedPtr<FStreamableHandle> AsynchronousStreamableHandlePtr;
+
+	TSharedPtr<mu::Model, ESPMode::ThreadSafe> Model;
 };
diff --git a/Engine/Plugins/Experimental/Mutable/Source/CustomizableObjectEditor/Private/MuCOE/CustomizableObjectEditor.cpp b/Engine/Plugins/Experimental/Mutable/Source/CustomizableObjectEditor/Private/MuCOE/CustomizableObjectEditor.cpp
--- a/Engine/Plugins/Experimental/Mutable/Source/CustomizableObjectEditor/Private/MuCOE/CustomizableObjectEditor.cpp
+++ b/Engine/Plugins/Experimental/Mutable/Source/CustomizableObjectEditor/Private/MuCOE/CustomizableObjectEditor.cpp
@@ -324,12 +324,14 @@ void FCustomizableObjectEditor::InitCustomizableObjectEditor(const EToolkitMode:
 	OnObjectPropertySelectionChanged(NULL);
 	OnInstancePropertySelectionChanged(NULL);
 	FCoreUObjectDelegates::OnObjectModified.AddRaw(this, &FCustomizableObjectEditor::OnObjectModified);
-
+	
 	UCustomizableObjectPrivate* CustomizableObjectPrivate = CustomizableObject->GetPrivate();
 
 	CustomizableObjectPrivate->Status.GetOnStateChangedDelegate().AddRaw(this, &FCustomizableObjectEditor::OnCustomizableObjectStatusChanged);
 	const FCustomizableObjectStatusTypes::EState CurrentStatus = CustomizableObjectPrivate->Status.Get();
 	OnCustomizableObjectStatusChanged(CurrentStatus, CurrentStatus);
+
+	CustomizableObject->GetPostCompileDelegate().AddSP(this, &FCustomizableObjectEditor::OnPostCompile);
 }
 
 
@@ -1776,21 +1778,7 @@ bool FCustomizableObjectEditor::IsTickable() const
 
 void FCustomizableObjectEditor::Tick(float InDeltaTime)
 {
-	const bool bUpdated = Compiler.Tick();
-	
-	if (bUpdated && CustomizableObject->IsCompiled())
-	{
-		if (!PreviewCustomizableSkeletalComponents.IsEmpty() && !PreviewSkeletalMeshComponents.IsEmpty())
-		{
-			CreatePreviewComponents();
-
-			PreviewInstance->UpdateSkeletalMeshAsync(true, true);
-		}
-		else
-		{
-			CreatePreviewInstance();
-		}
-	}
+	Compiler.Tick(false);
 }
 
 
@@ -2141,6 +2129,21 @@ void FCustomizableObjectEditor::FindProperty(const FProperty* Property, const vo
 }
 
 
+void FCustomizableObjectEditor::OnPostCompile()
+{
+	if (!PreviewCustomizableSkeletalComponents.IsEmpty() && !PreviewSkeletalMeshComponents.IsEmpty())
+	{
+		CreatePreviewComponents();
+
+		PreviewInstance->UpdateSkeletalMeshAsync(true, true);
+	}
+	else
+	{
+		CreatePreviewInstance();
+	}
+}
+
+
 void FCustomizableObjectEditor::OnEnterText(const FText& NewText, ETextCommit::Type TextType)
 {
 	if (TextType != ETextCommit::OnEnter)
diff --git a/Engine/Plugins/Experimental/Mutable/Source/CustomizableObjectEditor/Private/MuCOE/CustomizableObjectEditor.h b/Engine/Plugins/Experimental/Mutable/Source/CustomizableObjectEditor/Private/MuCOE/CustomizableObjectEditor.h
--- a/Engine/Plugins/Experimental/Mutable/Source/CustomizableObjectEditor/Private/MuCOE/CustomizableObjectEditor.h
+++ b/Engine/Plugins/Experimental/Mutable/Source/CustomizableObjectEditor/Private/MuCOE/CustomizableObjectEditor.h
@@ -374,7 +374,9 @@ private:
 	  * @param Context UObject Context where this string has been found.
 	  * @param bFound Mark as true if any property has been found. */
 	void FindProperty(const FProperty* Property, const void* Container, const FString& FindString, const UObject& Context, bool& bFound);
-	
+
+	void OnPostCompile();
+
 public:
 	void OnCustomizableObjectStatusChanged(FCustomizableObjectStatus::EState PreviousState, FCustomizableObjectStatus::EState CurrentState);
 
diff --git a/Engine/Plugins/Experimental/Mutable/Source/CustomizableObjectEditor/Private/MuCOE/CustomizableObjectInstanceEditor.cpp b/Engine/Plugins/Experimental/Mutable/Source/CustomizableObjectEditor/Private/MuCOE/CustomizableObjectInstanceEditor.cpp
--- a/Engine/Plugins/Experimental/Mutable/Source/CustomizableObjectEditor/Private/MuCOE/CustomizableObjectInstanceEditor.cpp
+++ b/Engine/Plugins/Experimental/Mutable/Source/CustomizableObjectEditor/Private/MuCOE/CustomizableObjectInstanceEditor.cpp
@@ -19,7 +19,6 @@
 #include "MuCO/CustomizableObjectInstance.h"
 #include "MuCO/CustomizableObjectSystem.h"
 #include "MuCO/CustomizableSkeletalComponent.h"
-#include "MuCOE/CustomizableObjectCompiler.h"
 #include "MuCOE/CustomizableObjectEditorActions.h"
 #include "MuCO/ICustomizableObjectEditorModule.h"
 #include "MuCOE/CustomizableObjectEditorViewportClient.h"
@@ -226,11 +225,7 @@ FCustomizableObjectInstanceEditor::~FCustomizableObjectInstanceEditor()
 	Viewport.Reset();
 
 	FCoreUObjectDelegates::OnObjectModified.Remove(OnObjectModifiedHandle);
-	if (Compiler)
-	{
-		Compiler->ForceFinishCompilation();
-	}
-
+	
 	if (const UCustomizableObject* CustomizableObject = CustomizableObjectInstance->GetCustomizableObject())
 	{
 		CustomizableObject->GetPrivate()->Status.GetOnStateChangedDelegate().RemoveAll(this);		
@@ -330,6 +325,8 @@ void FCustomizableObjectInstanceEditor::InitCustomizableObjectInstanceEditor( co
 			bOnlyRuntimeParameters = false;
 			InCustomizableObjectInstance->GetPrivate()->bShowOnlyRuntimeParameters = false;
 		}
+		
+		CustomizableObject->GetPostCompileDelegate().AddSP(this, &FCustomizableObjectInstanceEditor::OnPostCompile);
 	}
 }
 
@@ -703,26 +700,7 @@ bool FCustomizableObjectInstanceEditor::IsTickable(void) const
 void FCustomizableObjectInstanceEditor::Tick(float InDeltaTime)
 {
 	check(CustomizableObjectInstance);
-	if (Compiler && Compiler->Tick())
-	{
-		if (const UCustomizableObject* CustomizableObject = CustomizableObjectInstance->GetCustomizableObject();
-			CustomizableObject && CustomizableObject->GetPrivate()->GetSource())
-		{
-			if (PreviewCustomizableSkeletalComponents.Num() > 0)
-			{
-				CustomizableObjectInstance->UpdateSkeletalMeshAsync(true, true);
-			}
-			else
-			{
-				CreatePreviewInstance();
-			}
-
-			CustomizableInstanceDetailsView->ForceRefresh();
-		}
-
-		Compiler.Reset();
-	}
-
+	
 	// If we want to show the Relevant/Runtime parameters, we need to refresh the details view to make sure that the scroll bar appears		
 	if (bOnlyRelevantParameters != CustomizableObjectInstance->GetPrivate()->bShowOnlyRelevantParameters)
 	{
@@ -775,9 +753,28 @@ void FCustomizableObjectInstanceEditor::OpenTextureAnalyzerTab()
 }
 
 
+void FCustomizableObjectInstanceEditor::OnPostCompile()
+{
+	const UCustomizableObject* CustomizableObject = CustomizableObjectInstance->GetCustomizableObject();
+	if (CustomizableObject && CustomizableObject->GetPrivate()->GetSource())
+	{
+		if (PreviewCustomizableSkeletalComponents.Num() > 0)
+		{
+			CustomizableObjectInstance->UpdateSkeletalMeshAsync(true, true);
+		}
+		else
+		{
+			CreatePreviewInstance();
+		}
+
+		CustomizableInstanceDetailsView->ForceRefresh();
+	}
+}
+
+
 void FCustomizableObjectInstanceEditor::HideGizmo(const TSharedPtr<ICustomizableObjectInstanceEditor>& Editor,
-	const TSharedPtr<SCustomizableObjectEditorViewportTabBody>& Viewport,
-	const TSharedPtr<IDetailsView>& InstanceDetailsView)
+                                                  const TSharedPtr<SCustomizableObjectEditorViewportTabBody>& Viewport,
+                                                  const TSharedPtr<IDetailsView>& InstanceDetailsView)
 {
 	HideGizmoProjectorParameter(Editor, Viewport, InstanceDetailsView);
 }
diff --git a/Engine/Plugins/Experimental/Mutable/Source/CustomizableObjectEditor/Private/MuCOE/CustomizableObjectInstanceEditor.h b/Engine/Plugins/Experimental/Mutable/Source/CustomizableObjectEditor/Private/MuCOE/CustomizableObjectInstanceEditor.h
--- a/Engine/Plugins/Experimental/Mutable/Source/CustomizableObjectEditor/Private/MuCOE/CustomizableObjectInstanceEditor.h
+++ b/Engine/Plugins/Experimental/Mutable/Source/CustomizableObjectEditor/Private/MuCOE/CustomizableObjectInstanceEditor.h
@@ -221,6 +221,7 @@ private:
 	/** Open the Texture Analyzer tab */
 	void OpenTextureAnalyzerTab();
 
+	void OnPostCompile();
 public:
 	/** Static to reuse code with FCustomizableObjectEditor. */
 	static void HideGizmo(const TSharedPtr<ICustomizableObjectInstanceEditor>& Editor,
@@ -282,9 +283,6 @@ private:
 	/** Scene preview settings widget */
 	TSharedPtr<SWidget> AdvancedPreviewSettingsWidget;
 
-	/** Object compiler */
-	TUniquePtr<FCustomizableObjectCompiler> Compiler;
-
 	/** Pose asset when doing drag and drop of an UPoseAsset to the viewport */
 	UPoseAsset* PoseAsset;
 
