[mutable] Fixed Material node Texture pin modes when connected to an import texture node.
#rnx
#tests Editor

diff --git a/Engine/Plugins/Experimental/Mutable/Source/CustomizableObjectEditor/Private/MuCOE/Nodes/CustomizableObjectNodeMaterial.cpp b/Engine/Plugins/Experimental/Mutable/Source/CustomizableObjectEditor/Private/MuCOE/Nodes/CustomizableObjectNodeMaterial.cpp
--- a/Engine/Plugins/Experimental/Mutable/Source/CustomizableObjectEditor/Private/MuCOE/Nodes/CustomizableObjectNodeMaterial.cpp
+++ b/Engine/Plugins/Experimental/Mutable/Source/CustomizableObjectEditor/Private/MuCOE/Nodes/CustomizableObjectNodeMaterial.cpp
@@ -20,6 +20,7 @@
 #include "PropertyCustomizationHelpers.h"
 #include "Modules/ModuleManager.h"
 #include "MuCOE/CustomizableObjectEditorLogger.h"
+#include "MuCOE/Nodes/CustomizableObjectNodeExternalPin.h"
 
 class SGraphNode;
 class SWidget;
@@ -187,6 +188,24 @@ EPinMode UCustomizableObjectNodeMaterial::GetImagePinMode(const UEdGraphPin& Pin
 
 		return EPinMode::Mutable;
 	}
+	else if (!Pin.LinkedTo.IsEmpty())
+	{
+		// This can happen if an external pin node is not connected to an expose pin node yet
+		if (UEdGraphPin* LinkedToPin = Pin.LinkedTo[0])
+		{
+			if (const UCustomizableObjectNodeExternalPin* ExternalPinNode = Cast<UCustomizableObjectNodeExternalPin>(LinkedToPin->GetOwningNode()))
+			{
+				if (ExternalPinNode->PinType == Schema->PC_Image)
+				{
+					return EPinMode::Mutable;
+				}
+				else if (ExternalPinNode->PinType == Schema->PC_PassThroughImage)
+				{
+					return EPinMode::Passthrough;
+				}
+			}
+		}
+	}
 	
 	switch (GetPinData<UCustomizableObjectNodeMaterialPinDataImage>(Pin).GetPinMode())
 	{
