[mutable] Fix texture format conversion heuristic for recent ASTC changes in some platform.
#rnx

diff --git a/Engine/Plugins/Experimental/Mutable/Source/MutableTools/Private/MuT/UnrealPixelFormatOverride.cpp b/Engine/Plugins/Experimental/Mutable/Source/MutableTools/Private/MuT/UnrealPixelFormatOverride.cpp
--- a/Engine/Plugins/Experimental/Mutable/Source/MutableTools/Private/MuT/UnrealPixelFormatOverride.cpp
+++ b/Engine/Plugins/Experimental/Mutable/Source/MutableTools/Private/MuT/UnrealPixelFormatOverride.cpp
@@ -2,6 +2,7 @@
 
 #include "MuT/UnrealPixelFormatOverride.h"
 
+#include "MuR/MutableRuntimeModule.h"
 #include "Interfaces/ITextureFormatManagerModule.h"
 #include "Interfaces/ITextureFormatModule.h"
 #include "Modules/ModuleManager.h"
@@ -85,7 +86,7 @@ void FillBuildSettingsFromMutableFormat(FTextureBuildSettings& Settings, bool& b
 		break;
 
 	case mu::EImageFormat::IF_ASTC_4x4_RG_LDR:
-		Settings.TextureFormatName = TEXT("ASTC_NormalRG_Precise");
+		Settings.TextureFormatName = TEXT("ASTC_RGB"); // There is no way to get a 4x4 RG in the ASTC compressor from a TextureFormatName
 		Settings.CompressionQuality = 4; // See GetQualityFormat in TextureFormatASTC.cpp
 		bOutHasAlpha = false;
 		break;
@@ -168,12 +169,24 @@ void MutableToImageCore(const mu::Image* InMutable, FImage& CoreImage, int32 LOD
 }
 
 
-void ImageCoreToMutable(const FCompressedImage2D& Compressed, mu::Image* Mutable, int32 LOD)
+bool ImageCoreToMutable(const FCompressedImage2D& Compressed, mu::Image* Mutable, int32 LOD)
 {
 	TArrayView<uint8> MutableView = Mutable->DataStorage.GetLOD(LOD);
 
-	check(Compressed.RawData.Num() == MutableView.Num());
-	FMemory::Memcpy(MutableView.GetData(), Compressed.RawData.GetData(), Compressed.RawData.Num());
+	if (Compressed.RawData.Num() != MutableView.Num())
+	{
+		UE_LOG(LogMutableCore, Error, TEXT("Buffer size mismatch when trying to convert image LOD %d, mutable size is %d and ue size is %d. Mutable is %d x %d format %d and UE is %d x %d format %d."), 
+			LOD, MutableView.Num(), Compressed.RawData.Num(), 
+			Mutable->GetSizeX(), Mutable->GetSizeY(), Mutable->GetFormat(),
+			Compressed.SizeX, Compressed.SizeY, Compressed.PixelFormat
+			);
+
+		return false;
+	}
+
+	SIZE_T Bytes = FMath::Min(SIZE_T(MutableView.Num()),SIZE_T(Compressed.RawData.Num()));
+	FMemory::Memcpy(MutableView.GetData(), Compressed.RawData.GetData(), Bytes);
+	return true;
 }
 
 
@@ -265,19 +278,22 @@ void UnrealPixelFormatFunc(bool& bOutSuccess, int32 Quality, mu::Image* Target,
 		LODCount = 1;
 	}
 
-	for (int32 LOD = FirstLOD; LOD < LODCount; ++LOD)
+	for (int32 LOD = FirstLOD; bOutSuccess && (LOD < LODCount); ++LOD)
 	{
 		FImage SourceUnreal;
 		MutableToImageCore(Source, SourceUnreal, LOD);
 
 		FCompressedImage2D CompressedUnreal;
-		TextureFormat->CompressImage(SourceUnreal, Settings,
+		bOutSuccess = TextureFormat->CompressImage(SourceUnreal, Settings,
 			FIntVector3(SourceUnreal.SizeX, SourceUnreal.SizeY, 1),
 			0, 0, 1,
 			FString(),
 			bHasAlpha, CompressedUnreal);
 
-		ImageCoreToMutable(CompressedUnreal, Target, LOD);
+		if (bOutSuccess)
+		{
+			bOutSuccess = ImageCoreToMutable(CompressedUnreal, Target, LOD);
+		}
 	}
 }
 
diff --git a/Engine/Plugins/Experimental/Mutable/Source/MutableTools/Public/MuT/UnrealPixelFormatOverride.h b/Engine/Plugins/Experimental/Mutable/Source/MutableTools/Public/MuT/UnrealPixelFormatOverride.h
--- a/Engine/Plugins/Experimental/Mutable/Source/MutableTools/Public/MuT/UnrealPixelFormatOverride.h
+++ b/Engine/Plugins/Experimental/Mutable/Source/MutableTools/Public/MuT/UnrealPixelFormatOverride.h
@@ -20,7 +20,7 @@ void FillBuildSettingsFromMutableFormat(FTextureBuildSettings& Settings, bool& b
 void MutableToImageCore(const mu::Image* InMutable, FImage& CoreImage, int32 LOD);
 
 /** */
-void ImageCoreToMutable(const FCompressedImage2D& Compressed, mu::Image* Mutable, int32 LOD);
+bool ImageCoreToMutable(const FCompressedImage2D& Compressed, mu::Image* Mutable, int32 LOD);
 
 /** Convert an Unreal platform pixel format to Mutable. */
 MUTABLETOOLS_API extern mu::EImageFormat UnrealToMutablePixelFormat(EPixelFormat PlatformFormat, bool bHasAlpha);
