[Mutable] Fix crashes when consulting the surface data during mesh generation. Instead of using the BaseSurfaceIndex, we used the InstanceSurfaceIndex for shared surfaces.

#rnx

diff --git a/Engine/Plugins/Experimental/Mutable/Source/CustomizableObject/Private/MuCO/BusyWaits_Deprecated.cpp b/Engine/Plugins/Experimental/Mutable/Source/CustomizableObject/Private/MuCO/BusyWaits_Deprecated.cpp
--- a/Engine/Plugins/Experimental/Mutable/Source/CustomizableObject/Private/MuCO/BusyWaits_Deprecated.cpp
+++ b/Engine/Plugins/Experimental/Mutable/Source/CustomizableObject/Private/MuCO/BusyWaits_Deprecated.cpp
@@ -287,14 +287,14 @@ namespace CustomizableObjectSystem::ImplDeprecated
 
 						// Images
 						Surface.FirstImage = OperationData->InstanceUpdateData.Images.Num();
-						Surface.ImageCount = Instance->GetImageCount(BaseLODIndex, ComponentIndex, InstanceSurfaceIndex);
+						Surface.ImageCount = Instance->GetImageCount(BaseLODIndex, ComponentIndex, BaseSurfaceIndex);
 						for (int32 ImageIndex = 0; ImageIndex < Surface.ImageCount; ++ImageIndex)
 						{
 							MUTABLE_CPUPROFILER_SCOPE(GetImageId);
 
 							OperationData->InstanceUpdateData.Images.Push({});
 							FInstanceUpdateData::FImage& Image = OperationData->InstanceUpdateData.Images.Last();
-							Image.Name = Instance->GetImageName(BaseLODIndex, ComponentIndex, InstanceSurfaceIndex, ImageIndex);
+							Image.Name = Instance->GetImageName(BaseLODIndex, ComponentIndex, BaseSurfaceIndex, ImageIndex);
 							Image.ImageID = Instance->GetImageId(BaseLODIndex, ComponentIndex, BaseSurfaceIndex, ImageIndex);
 							Image.FullImageSizeX = 0;
 							Image.FullImageSizeY = 0;
@@ -340,25 +340,25 @@ namespace CustomizableObjectSystem::ImplDeprecated
 
 						// Vectors
 						Surface.FirstVector = OperationData->InstanceUpdateData.Vectors.Num();
-						Surface.VectorCount = Instance->GetVectorCount(BaseLODIndex, ComponentIndex, InstanceSurfaceIndex);
+						Surface.VectorCount = Instance->GetVectorCount(BaseLODIndex, ComponentIndex, BaseSurfaceIndex);
 						for (int32 VectorIndex = 0; VectorIndex < Surface.VectorCount; ++VectorIndex)
 						{
 							MUTABLE_CPUPROFILER_SCOPE(GetVector);
 							OperationData->InstanceUpdateData.Vectors.Push({});
 							FInstanceUpdateData::FVector& Vector = OperationData->InstanceUpdateData.Vectors.Last();
-							Vector.Name = Instance->GetVectorName(BaseLODIndex, ComponentIndex, InstanceSurfaceIndex, VectorIndex);
-							Vector.Vector = Instance->GetVector(BaseLODIndex, ComponentIndex, InstanceSurfaceIndex, VectorIndex);
+							Vector.Name = Instance->GetVectorName(BaseLODIndex, ComponentIndex, BaseSurfaceIndex, VectorIndex);
+							Vector.Vector = Instance->GetVector(BaseLODIndex, ComponentIndex, BaseSurfaceIndex, VectorIndex);
 						}
 
 						// Scalars
 						Surface.FirstScalar = OperationData->InstanceUpdateData.Scalars.Num();
-						Surface.ScalarCount = Instance->GetScalarCount(BaseLODIndex, ComponentIndex, InstanceSurfaceIndex);
+						Surface.ScalarCount = Instance->GetScalarCount(BaseLODIndex, ComponentIndex, BaseSurfaceIndex);
 						for (int32 ScalarIndex = 0; ScalarIndex < Surface.ScalarCount; ++ScalarIndex)
 						{
 							MUTABLE_CPUPROFILER_SCOPE(GetScalar)
 
-							const FName ScalarName = Instance->GetScalarName(BaseLODIndex, ComponentIndex, InstanceSurfaceIndex, ScalarIndex);
-							const float ScalarValue = Instance->GetScalar(BaseLODIndex, ComponentIndex, InstanceSurfaceIndex, ScalarIndex);
+							const FName ScalarName = Instance->GetScalarName(BaseLODIndex, ComponentIndex, BaseSurfaceIndex, ScalarIndex);
+							const float ScalarValue = Instance->GetScalar(BaseLODIndex, ComponentIndex, BaseSurfaceIndex, ScalarIndex);
 							
 							FString EncodingMaterialIdString = "__MutableMaterialId";
 							
diff --git a/Engine/Plugins/Experimental/Mutable/Source/CustomizableObject/Private/MuCO/CustomizableObjectSystem.cpp b/Engine/Plugins/Experimental/Mutable/Source/CustomizableObject/Private/MuCO/CustomizableObjectSystem.cpp
--- a/Engine/Plugins/Experimental/Mutable/Source/CustomizableObject/Private/MuCO/CustomizableObjectSystem.cpp
+++ b/Engine/Plugins/Experimental/Mutable/Source/CustomizableObject/Private/MuCO/CustomizableObjectSystem.cpp
@@ -2254,25 +2254,25 @@ namespace impl
 
 						// Vectors
 						Surface.FirstVector = OperationData->InstanceUpdateData.Vectors.Num();
-						Surface.VectorCount = MutableInstance->GetVectorCount(BaseLODIndex, ComponentIndex, InstanceSurfaceIndex);
+						Surface.VectorCount = MutableInstance->GetVectorCount(BaseLODIndex, ComponentIndex, BaseSurfaceIndex);
 						for (int32 VectorIndex = 0; VectorIndex < Surface.VectorCount; ++VectorIndex)
 						{
 							MUTABLE_CPUPROFILER_SCOPE(GetVector);
 							OperationData->InstanceUpdateData.Vectors.Push({});
 							FInstanceUpdateData::FVector& Vector = OperationData->InstanceUpdateData.Vectors.Last();
-							Vector.Name = MutableInstance->GetVectorName(BaseLODIndex, ComponentIndex, InstanceSurfaceIndex, VectorIndex);
-							Vector.Vector = MutableInstance->GetVector(BaseLODIndex, ComponentIndex, InstanceSurfaceIndex, VectorIndex);
+							Vector.Name = MutableInstance->GetVectorName(BaseLODIndex, ComponentIndex, BaseSurfaceIndex, VectorIndex);
+							Vector.Vector = MutableInstance->GetVector(BaseLODIndex, ComponentIndex, BaseSurfaceIndex, VectorIndex);
 						}
 
 						// Scalars
 						Surface.FirstScalar = OperationData->InstanceUpdateData.Scalars.Num();
-						Surface.ScalarCount = MutableInstance->GetScalarCount(BaseLODIndex, ComponentIndex, InstanceSurfaceIndex);
+						Surface.ScalarCount = MutableInstance->GetScalarCount(BaseLODIndex, ComponentIndex, BaseSurfaceIndex);
 						for (int32 ScalarIndex = 0; ScalarIndex < Surface.ScalarCount; ++ScalarIndex)
 						{
 							MUTABLE_CPUPROFILER_SCOPE(GetScalar)
 
-							const FName ScalarName = MutableInstance->GetScalarName(BaseLODIndex, ComponentIndex, InstanceSurfaceIndex, ScalarIndex);
-							const float ScalarValue = MutableInstance->GetScalar(BaseLODIndex, ComponentIndex, InstanceSurfaceIndex, ScalarIndex);
+							const FName ScalarName = MutableInstance->GetScalarName(BaseLODIndex, ComponentIndex, BaseSurfaceIndex, ScalarIndex);
+							const float ScalarValue = MutableInstance->GetScalar(BaseLODIndex, ComponentIndex, BaseSurfaceIndex, ScalarIndex);
 							
 							FString EncodingMaterialIdString = "__MutableMaterialId";
 							
@@ -2292,14 +2292,14 @@ namespace impl
 
 						// Images
 						Surface.FirstImage = OperationData->InstanceUpdateData.Images.Num();
-						Surface.ImageCount = MutableInstance->GetImageCount(BaseLODIndex, ComponentIndex, InstanceSurfaceIndex);
+						Surface.ImageCount = MutableInstance->GetImageCount(BaseLODIndex, ComponentIndex, BaseSurfaceIndex);
 						for (int32 ImageIndex = 0; ImageIndex < Surface.ImageCount; ++ImageIndex)
 						{
 							MUTABLE_CPUPROFILER_SCOPE(GetImageId);
 
 							const int32 UpdateDataImageIndex = OperationData->InstanceUpdateData.Images.AddDefaulted();
 							FInstanceUpdateData::FImage& Image = OperationData->InstanceUpdateData.Images.Last();
-							Image.Name = MutableInstance->GetImageName(BaseLODIndex, ComponentIndex, InstanceSurfaceIndex, ImageIndex);
+							Image.Name = MutableInstance->GetImageName(BaseLODIndex, ComponentIndex, BaseSurfaceIndex, ImageIndex);
 							Image.ImageID = MutableInstance->GetImageId(BaseLODIndex, ComponentIndex, BaseSurfaceIndex, ImageIndex);
 							Image.FullImageSizeX = 0;
 							Image.FullImageSizeY = 0;
