[Mutable] Viewport improvements.

- Changed the Viewport Toolbar style.
- Changed the Bake icon.
- Changed the Viewport Text style.
- Added more Viewport Texts:
    - Updating
    - Compiling
    - No Preview
    - Update Error
    - Compile Error
    - ...
- Now Preview Components are recreated after each compilation. This means that we no longer show previous compilation Skeletal Meshes.
- Added more compilation status:
   - Split ECustomizableObjectCompilationState in ECompilationStatePrivate and ECompilationResultPrivate.
   - ECustomizableObjectCompilationState kept only as public API.
- Simplified COE CreatePreviewComponents.

#rnx

diff --git a/Engine/Plugins/Experimental/Mutable/Source/CustomizableObject/Classes/MuCO/CustomizableObject.h b/Engine/Plugins/Experimental/Mutable/Source/CustomizableObject/Classes/MuCO/CustomizableObject.h
--- a/Engine/Plugins/Experimental/Mutable/Source/CustomizableObject/Classes/MuCO/CustomizableObject.h
+++ b/Engine/Plugins/Experimental/Mutable/Source/CustomizableObject/Classes/MuCO/CustomizableObject.h
@@ -28,7 +28,7 @@ class UCustomizableObjectPrivate;
 struct FFrame;
 struct FStreamableHandle;
 template <typename FuncType> class TFunctionRef;
-enum class ECustomizableObjectCompilationState : uint8;
+enum class ECompilationStatePrivate : uint8;
 
 DECLARE_MULTICAST_DELEGATE(FPostCompileDelegate)
  
diff --git a/Engine/Plugins/Experimental/Mutable/Source/CustomizableObject/Internal/MuCO/CustomizableObjectPrivate.h b/Engine/Plugins/Experimental/Mutable/Source/CustomizableObject/Internal/MuCO/CustomizableObjectPrivate.h
--- a/Engine/Plugins/Experimental/Mutable/Source/CustomizableObject/Internal/MuCO/CustomizableObjectPrivate.h
+++ b/Engine/Plugins/Experimental/Mutable/Source/CustomizableObject/Internal/MuCO/CustomizableObjectPrivate.h
@@ -56,8 +56,6 @@ public:
 
 	/** Provides the caller with the warning and error messages produced during compilation */
 	virtual void GetCompilationMessages(TArray<FText>& OutWarningMessages, TArray<FText>& OutErrorMessages) const = 0;
-	
-	virtual ECustomizableObjectCompilationState GetCompilationState() const;
 };
 
 
@@ -183,7 +181,7 @@ struct FCustomizableObjectStatusTypes
 		Loading = 0, // Waiting for PostLoad and Asset Registry to finish.
 		ModelLoaded, // Model loaded correctly.
 		NoModel, // No model (due to no model not found and automatic compilations disabled).
-		// Compiling, // Compiling the CO.
+		// Compiling, // Compiling the CO. Equivalent to UCustomizableObject::IsLocked = true.
 
 		Count,
 	};
@@ -203,17 +201,21 @@ struct FCustomizableObjectStatusTypes
 using FCustomizableObjectStatus = FStateMachine<FCustomizableObjectStatusTypes>;
 
 
-UENUM()
-enum class ECustomizableObjectCompilationState : uint8
+
+enum class ECompilationStatePrivate : uint8
 {
-	//
 	None,
-	// 
 	InProgress,
-	//
-	Completed,
-	//
-	Failed
+	Completed
+};
+
+
+enum class ECompilationResultPrivate : uint8
+{
+	Unknown, // Not compiled yet (compilation may be in progress).
+	Success, // No errors or warnings.
+	Errors, // At least have one error. Can have warnings.
+	Warnings, // Only warnings.
 };
 
 
@@ -857,7 +859,8 @@ public:
 	/** This is a non-user-controlled flag to disable streaming (set at object compilation time, depending on optimization). */
 	bool bDisableTextureStreaming = false;
 	
-	ECustomizableObjectCompilationState CompilationState = ECustomizableObjectCompilationState::None;
+	ECompilationStatePrivate CompilationState = ECompilationStatePrivate::None;
+	ECompilationResultPrivate CompilationResult = ECompilationResultPrivate::Unknown;
 	
 	FPostCompileDelegate PostCompileDelegate;
 
diff --git a/Engine/Plugins/Experimental/Mutable/Source/CustomizableObject/Private/MuCO/CustomizableObject.cpp b/Engine/Plugins/Experimental/Mutable/Source/CustomizableObject/Private/MuCO/CustomizableObject.cpp
--- a/Engine/Plugins/Experimental/Mutable/Source/CustomizableObject/Private/MuCO/CustomizableObject.cpp
+++ b/Engine/Plugins/Experimental/Mutable/Source/CustomizableObject/Private/MuCO/CustomizableObject.cpp
@@ -879,7 +879,7 @@ bool UCustomizableObject::ConditionalAutoCompile()
 	}
 
 	// Don't re-compile objects if they failed to compile. 
-	if (GetPrivate()->CompilationState == ECustomizableObjectCompilationState::Failed)
+	if (GetPrivate()->CompilationResult == ECompilationResultPrivate::Errors)
 	{
 		return false;
 	}
@@ -2012,11 +2012,6 @@ TObjectPtr<UEdGraph>& UCustomizableObjectPrivate::GetSource() const
 //-------------------------------------------------------------------------------------------------
 //-------------------------------------------------------------------------------------------------
 
-ECustomizableObjectCompilationState FCustomizableObjectCompilerBase::GetCompilationState() const
-{
-	return ECustomizableObjectCompilationState::None;
-}
-
 
 void UCustomizableObjectBulk::PostLoad()
 {
diff --git a/Engine/Plugins/Experimental/Mutable/Source/CustomizableObject/Private/MuCO/CustomizableObjectSystem.cpp b/Engine/Plugins/Experimental/Mutable/Source/CustomizableObject/Private/MuCO/CustomizableObjectSystem.cpp
--- a/Engine/Plugins/Experimental/Mutable/Source/CustomizableObject/Private/MuCO/CustomizableObjectSystem.cpp
+++ b/Engine/Plugins/Experimental/Mutable/Source/CustomizableObject/Private/MuCO/CustomizableObjectSystem.cpp
@@ -4179,7 +4179,7 @@ void UCustomizableObjectSystem::TickRecompileCustomizableObjects()
 	
 	if (GetPrivate()->RecompileCustomizableObjectsCompiler)
 	{
-		bUpdated = GetPrivate()->RecompileCustomizableObjectsCompiler->Tick(GetPrivate()->bBlocking) || GetPrivate()->RecompileCustomizableObjectsCompiler->GetCompilationState() == ECustomizableObjectCompilationState::Failed;
+		bUpdated = GetPrivate()->RecompileCustomizableObjectsCompiler->Tick(GetPrivate()->bBlocking);
 	}
 
 	if (bUpdated)
diff --git a/Engine/Plugins/Experimental/Mutable/Source/CustomizableObjectEditor/Private/MuCOE/CustomizableObjectCompiler.cpp b/Engine/Plugins/Experimental/Mutable/Source/CustomizableObjectEditor/Private/MuCOE/CustomizableObjectCompiler.cpp
--- a/Engine/Plugins/Experimental/Mutable/Source/CustomizableObjectEditor/Private/MuCOE/CustomizableObjectCompiler.cpp
+++ b/Engine/Plugins/Experimental/Mutable/Source/CustomizableObjectEditor/Private/MuCOE/CustomizableObjectCompiler.cpp
@@ -91,7 +91,7 @@ bool FCustomizableObjectCompiler::Tick(bool bBlocking)
 			}
 			else
 			{
-				SetCompilationState(ECustomizableObjectCompilationState::Completed);
+			    SetCompilationState(ECompilationStatePrivate::Completed, GetCompilationResult());
 
 				RemoveCompileNotification();
 
@@ -114,7 +114,7 @@ bool FCustomizableObjectCompiler::Tick(bool bBlocking)
 
 		if (SaveDDTask->IsCompleted())
 		{
-			SetCompilationState(ECustomizableObjectCompilationState::Completed);
+		    SetCompilationState(ECompilationStatePrivate::Completed, GetCompilationResult());
 		
 			FinishSavingDerivedDataTask();
 	
@@ -190,11 +190,20 @@ void FCustomizableObjectCompiler::Compile(UCustomizableObject& Object, const FCo
 {
 	TRACE_BEGIN_REGION(UE_MUTABLE_COMPILE_REGION);
 
+	CurrentObject= Object;
+
+	if (CurrentObject->GetPrivate()->CompilationState == ECompilationStatePrivate::InProgress)
+	{
+		UE_LOG(LogMutable, Warning, TEXT("Failed to compile Customizable Object [%s]. Object already being compiled."), *Object.GetName());
+		return;
+	}
+	
+	SetCompilationState(ECompilationStatePrivate::InProgress, ECompilationResultPrivate::Unknown);
+
 	if (!UCustomizableObjectSystem::IsActive())
 	{
 		UE_LOG(LogMutable, Warning, TEXT("Failed to compile Customizable Object [%s]. Mutable is disabled. To enable it set the CVar Mutable.Enabled to true."), *Object.GetName());
-		SetCompilationState(ECustomizableObjectCompilationState::Failed);
-
+		SetCompilationState(ECompilationStatePrivate::Completed, ECompilationResultPrivate::Errors);
 		return;
 	}
 
@@ -202,8 +211,7 @@ void FCustomizableObjectCompiler::Compile(UCustomizableObject& Object, const FCo
 	{
 		UE_LOG(LogMutable, Warning, TEXT("In Customizable Object [%s], the VersionBridge asset [%s] does not implement the required UCustomizableObjectVersionBridgeInterface."), 
 			*Object.GetName(), *Object.VersionBridge.GetName());
-		SetCompilationState(ECustomizableObjectCompilationState::Failed);
-
+		SetCompilationState(ECompilationStatePrivate::Completed, ECompilationResultPrivate::Errors);
 		return;
 	}
 
@@ -230,10 +238,10 @@ void FCustomizableObjectCompiler::Compile(UCustomizableObject& Object, const FCo
 	bool LockResult = System->LockObject(&Object);
 
 	if (!LockResult)
-	{		
+	{
 		UE_LOG(LogMutable, Warning, TEXT("%s"), *Message);
 		FSlateNotificationManager::Get().AddNotification(Info);
-
+		SetCompilationState(ECompilationStatePrivate::Completed, ECompilationResultPrivate::Errors);
 		return;
 	}
 
@@ -844,19 +852,36 @@ float FCustomizableObjectCompiler::ComputeAsyncLoadingTimeLimit()
 	return Value;
 }
 
+
+ECompilationResultPrivate FCustomizableObjectCompiler::GetCompilationResult() const
+{
+	if (CompilationLogsContainer.GetErrorCount())
+	{
+		return ECompilationResultPrivate::Errors;
+	}
+	else if (CompilationLogsContainer.GetWarningCount(true))
+	{
+		return ECompilationResultPrivate::Warnings;
+	}
+	else
+	{
+		return ECompilationResultPrivate::Success;
+	}
+}
+
+
 void FCustomizableObjectCompiler::GetCompilationMessages(TArray<FText>& OutWarningMessages, TArray<FText>& OutErrorMessages) const
 {
 	CompilationLogsContainer.GetMessages(OutWarningMessages,OutErrorMessages);
 }
 
 
-void FCustomizableObjectCompiler::SetCompilationState(ECustomizableObjectCompilationState InState)
+void FCustomizableObjectCompiler::SetCompilationState(ECompilationStatePrivate State, ECompilationResultPrivate Result) const
 {
-	State = InState;
-	
 	if (CurrentObject)
 	{
-		CurrentObject->GetPrivate()->CompilationState = InState;
+		CurrentObject->GetPrivate()->CompilationState = State;
+		CurrentObject->GetPrivate()->CompilationResult = Result;
 	}
 }
 
@@ -865,16 +890,16 @@ void FCustomizableObjectCompiler::CompileInternal(bool bAsync)
 {
 	MUTABLE_CPUPROFILER_SCOPE(FCustomizableObjectCompiler::Compile)
 	
-	SetCompilationState(ECustomizableObjectCompilationState::Failed);
-
 	if (!CurrentObject)
 	{
+		SetCompilationState(ECompilationStatePrivate::Completed, ECompilationResultPrivate::Errors);
 		return;
 	}
 
 	if (bAsync && CompileTask.IsValid()) // Don't start compilation if there's a compilation running
 	{
 		UE_LOG(LogMutable, Log, TEXT("FCustomizableObjectCompiler::Compile An object is already being compiled."));
+		SetCompilationState(ECompilationStatePrivate::Completed, ECompilationResultPrivate::Errors);
 		return;
 	}
 
@@ -883,6 +908,7 @@ void FCustomizableObjectCompiler::CompileInternal(bool bAsync)
 	// This is redundant but necessary to keep static analysis happy.
 	if (!CurrentObject)
 	{
+		SetCompilationState(ECompilationStatePrivate::Completed, ECompilationResultPrivate::Errors);
 		return;
 	}
 
@@ -899,7 +925,8 @@ void FCustomizableObjectCompiler::CompileInternal(bool bAsync)
 	if (!Options.bIsCooking && IsRunningCookCommandlet())
 	{
 		UE_LOG(LogMutable, Display, TEXT("Editor compilation suspended for Customizable Object [%s]. Can not compile COs when the cook commandlet is running. "), *CurrentObject->GetName());
-		UCustomizableObjectSystem::GetInstance()->UnlockObject(CurrentObject);
+		SetCompilationState(ECompilationStatePrivate::Completed, ECompilationResultPrivate::Errors);
+        UCustomizableObjectSystem::GetInstance()->UnlockObject(CurrentObject);
 		return;
 	}
 	
@@ -1194,7 +1221,6 @@ void FCustomizableObjectCompiler::CompileInternal(bool bAsync)
 		CompileTask = MakeShareable(new FCustomizableObjectCompileRunnable(MutableRoot));
 		CompileTask->Options = Options;
 		CompileTask->ReferencedTextures = NewCompileTimeReferencedTextures;
-		SetCompilationState(ECustomizableObjectCompilationState::InProgress);
 
 		if (!bAsync)
 		{
@@ -1215,7 +1241,7 @@ void FCustomizableObjectCompiler::CompileInternal(bool bAsync)
 
 			CurrentObject->GetPrivate()->PostCompile(); 
 
-			SetCompilationState(ECustomizableObjectCompilationState::Completed);
+			SetCompilationState(ECompilationStatePrivate::Completed, GetCompilationResult());
 		}
 		else
 		{
diff --git a/Engine/Plugins/Experimental/Mutable/Source/CustomizableObjectEditor/Private/MuCOE/CustomizableObjectCompiler.h b/Engine/Plugins/Experimental/Mutable/Source/CustomizableObjectEditor/Private/MuCOE/CustomizableObjectCompiler.h
--- a/Engine/Plugins/Experimental/Mutable/Source/CustomizableObjectEditor/Private/MuCOE/CustomizableObjectCompiler.h
+++ b/Engine/Plugins/Experimental/Mutable/Source/CustomizableObjectEditor/Private/MuCOE/CustomizableObjectCompiler.h
@@ -28,8 +28,6 @@ public:
 	// FCustomizableObjectCompilerBase interface
 	CUSTOMIZABLEOBJECTEDITOR_API bool IsRootObject(const class UCustomizableObject* Object) const override;
 	
-	CUSTOMIZABLEOBJECTEDITOR_API ECustomizableObjectCompilationState GetCompilationState() const override { return State;  };
-	
 	/** Check for pending compilation process. Returns true if an object has been updated. */
 	CUSTOMIZABLEOBJECTEDITOR_API virtual bool Tick(bool bBlocking) override;
 
@@ -71,16 +69,14 @@ public:
 	void ForceFinishBeforeStartCompilation(UCustomizableObject* Object);
 
 	// Getter of AsynchronousStreamableHandlePtr
-	TSharedPtr<struct FStreamableHandle> GetAsynchronousStreamableHandlePtr() { return AsynchronousStreamableHandlePtr; }
+	TSharedPtr<FStreamableHandle> GetAsynchronousStreamableHandlePtr() { return AsynchronousStreamableHandlePtr; }
 
 private:
 
 	// Object containing all error and warning logs raised during compilation.
 	FCompilationMessageCache CompilationLogsContainer;
-	
-	void SetCompilationState(ECustomizableObjectCompilationState InState);
 
-	ECustomizableObjectCompilationState State = ECustomizableObjectCompilationState::None;
+	void SetCompilationState(ECompilationStatePrivate State, ECompilationResultPrivate Result) const;
 	
 	void CompileInternal(bool bAsync = false);
 
@@ -179,6 +175,8 @@ private:
 	* @return value in milliseconds to use for AsyncLoadingTimeLimit and MaxConvertToMutableTextureTime */
 	float ComputeAsyncLoadingTimeLimit();
 	
+    ECompilationResultPrivate GetCompilationResult() const;
+	
 public:
 	
 	virtual void GetCompilationMessages(TArray<FText>& OutWarningMessages, TArray<FText>& OutErrorMessages) const override;
diff --git a/Engine/Plugins/Experimental/Mutable/Source/CustomizableObjectEditor/Private/MuCOE/CustomizableObjectEditor.cpp b/Engine/Plugins/Experimental/Mutable/Source/CustomizableObjectEditor/Private/MuCOE/CustomizableObjectEditor.cpp
--- a/Engine/Plugins/Experimental/Mutable/Source/CustomizableObjectEditor/Private/MuCOE/CustomizableObjectEditor.cpp
+++ b/Engine/Plugins/Experimental/Mutable/Source/CustomizableObjectEditor/Private/MuCOE/CustomizableObjectEditor.cpp
@@ -353,41 +353,6 @@ void FCustomizableObjectEditor::SelectNode(const UEdGraphNode* Node)
 }
 
 
-bool FCustomizableObjectEditor::CreatePreviewComponent(int32 ComponentIndex)
-{
-	if (!PreviewCustomizableSkeletalComponents.IsValidIndex(ComponentIndex))
-	{
-		PreviewCustomizableSkeletalComponents.AddZeroed(ComponentIndex - PreviewCustomizableSkeletalComponents.Num() + 1);
-		check(PreviewCustomizableSkeletalComponents.IsValidIndex(ComponentIndex) && !PreviewCustomizableSkeletalComponents.IsValidIndex(ComponentIndex + 1));
-	}
-
-	PreviewCustomizableSkeletalComponents[ComponentIndex] = NewObject<UCustomizableSkeletalComponent>(UCustomizableSkeletalComponent::StaticClass());
-
-	if (PreviewCustomizableSkeletalComponents[ComponentIndex])
-	{
-		PreviewCustomizableSkeletalComponents[ComponentIndex]->CustomizableObjectInstance = PreviewInstance;
-		PreviewCustomizableSkeletalComponents[ComponentIndex]->ComponentIndex = ComponentIndex;
-
-		if (!PreviewSkeletalMeshComponents.IsValidIndex(ComponentIndex))
-		{
-			PreviewSkeletalMeshComponents.AddZeroed(ComponentIndex - PreviewSkeletalMeshComponents.Num() + 1);
-			check(PreviewSkeletalMeshComponents.IsValidIndex(ComponentIndex) && !PreviewSkeletalMeshComponents.IsValidIndex(ComponentIndex + 1));
-		}
-
-		PreviewSkeletalMeshComponents[ComponentIndex] = NewObject<UDebugSkelMeshComponent>(GetTransientPackage(), NAME_None, RF_Transient);
-
-		if (PreviewSkeletalMeshComponents[ComponentIndex])
-		{
-			PreviewCustomizableSkeletalComponents[ComponentIndex]->AttachToComponent(PreviewSkeletalMeshComponents[ComponentIndex], FAttachmentTransformRules::KeepRelativeTransform);
-
-			return true;
-		}
-	}
-
-	return false;
-}
-
-
 void FCustomizableObjectEditor::CreatePreviewInstance()
 {
 	OnUpdatePreviewInstanceWork.Add([this]()
@@ -2131,16 +2096,8 @@ void FCustomizableObjectEditor::FindProperty(const FProperty* Property, const vo
 
 void FCustomizableObjectEditor::OnPostCompile()
 {
-	if (!PreviewCustomizableSkeletalComponents.IsEmpty() && !PreviewSkeletalMeshComponents.IsEmpty())
-	{
-		CreatePreviewComponents();
-
-		PreviewInstance->UpdateSkeletalMeshAsync(true, true);
-	}
-	else
-	{
-		CreatePreviewInstance();
-	}
+	CreatePreviewComponents();
+	PreviewInstance->UpdateSkeletalMeshAsync(true, true);
 }
 
 
@@ -2482,49 +2439,30 @@ void FCustomizableObjectEditor::GetExternalChildObjects(const UCustomizableObjec
 
 void FCustomizableObjectEditor::CreatePreviewComponents()
 {
-	// Getting the number of mesh components from the root node
-	int32 NumMeshComponents = 0;
-
-	if (CustomizableObject->GetPrivate()->GetSource())
+	for (const TObjectPtr<UCustomizableSkeletalComponent>& Component : PreviewCustomizableSkeletalComponents)
 	{
-		TArray<UCustomizableObjectNodeObject*> RootNode;
-		CustomizableObject->GetPrivate()->GetSource()->GetNodesOfClass< UCustomizableObjectNodeObject>(RootNode);
+		Component->DestroyComponent();
+	}
+	PreviewCustomizableSkeletalComponents.Empty();
 
-		for (int32 i = 0; i < RootNode.Num(); ++i)
-		{
-			if (RootNode[i]->bIsBase)
-			{
-				NumMeshComponents = RootNode[i]->NumMeshComponents;
-				break;
-			}
-		}
+	for (const TObjectPtr<UDebugSkelMeshComponent>& Component : PreviewSkeletalMeshComponents)
+	{
+		Component->DestroyComponent();
 	}
+	PreviewSkeletalMeshComponents.Empty();
 
-	if (NumMeshComponents != PreviewSkeletalMeshComponents.Num())
+	for (int32 ComponentIndex = 0; ComponentIndex < CustomizableObject->GetComponentCount(); ++ComponentIndex)
 	{
-		if (NumMeshComponents > PreviewSkeletalMeshComponents.Num())
-		{
-			PreviewSkeletalMeshComponents.AddZeroed(NumMeshComponents - PreviewSkeletalMeshComponents.Num());
-			PreviewCustomizableSkeletalComponents.AddZeroed(NumMeshComponents - PreviewCustomizableSkeletalComponents.Num());
-		}
-		else
-		{
-			PreviewSkeletalMeshComponents.SetNumZeroed(NumMeshComponents, EAllowShrinking::No);
-			PreviewCustomizableSkeletalComponents.SetNumZeroed(NumMeshComponents, EAllowShrinking::No);
-		}
+		UCustomizableSkeletalComponent* PreviewCustomizableSkeletalComponent = NewObject<UCustomizableSkeletalComponent>(UCustomizableSkeletalComponent::StaticClass());
+		UDebugSkelMeshComponent* PreviewSkeletalMeshComponent = NewObject<UDebugSkelMeshComponent>(GetTransientPackage(), NAME_None, RF_Transient);
 
-		for (int32 ComponentIndex = 0; ComponentIndex < PreviewSkeletalMeshComponents.Num(); ++ComponentIndex)
-		{
-			if (!PreviewSkeletalMeshComponents[ComponentIndex])
-			{
-				if (!CreatePreviewComponent(ComponentIndex))
-				{
-					PreviewSkeletalMeshComponents.Empty();
-					PreviewCustomizableSkeletalComponents.Empty();
-					break;
-				}
-			}
-		}
+		PreviewCustomizableSkeletalComponent->bSkipSetReferenceSkeletalMesh = true;
+		PreviewCustomizableSkeletalComponent->CustomizableObjectInstance = PreviewInstance;
+		PreviewCustomizableSkeletalComponent->ComponentIndex = ComponentIndex;
+		PreviewCustomizableSkeletalComponent->AttachToComponent(PreviewSkeletalMeshComponent, FAttachmentTransformRules::KeepRelativeTransform);
+
+		PreviewCustomizableSkeletalComponents.Add(PreviewCustomizableSkeletalComponent);
+		PreviewSkeletalMeshComponents.Add(PreviewSkeletalMeshComponent);
 	}
 }
 
diff --git a/Engine/Plugins/Experimental/Mutable/Source/CustomizableObjectEditor/Private/MuCOE/CustomizableObjectEditor.h b/Engine/Plugins/Experimental/Mutable/Source/CustomizableObjectEditor/Private/MuCOE/CustomizableObjectEditor.h
--- a/Engine/Plugins/Experimental/Mutable/Source/CustomizableObjectEditor/Private/MuCOE/CustomizableObjectEditor.h
+++ b/Engine/Plugins/Experimental/Mutable/Source/CustomizableObjectEditor/Private/MuCOE/CustomizableObjectEditor.h
@@ -285,9 +285,6 @@ private:
 	
 	void CreatePreviewInstance();
 
-	/** Returns true if successful */
-	bool CreatePreviewComponent(int32 ComponentIndex);
-
 	/** Binds commands associated with the Static Mesh Editor. */
 	void BindCommands();
 	
diff --git a/Engine/Plugins/Experimental/Mutable/Source/CustomizableObjectEditor/Private/MuCOE/CustomizableObjectEditorActions.cpp b/Engine/Plugins/Experimental/Mutable/Source/CustomizableObjectEditor/Private/MuCOE/CustomizableObjectEditorActions.cpp
--- a/Engine/Plugins/Experimental/Mutable/Source/CustomizableObjectEditor/Private/MuCOE/CustomizableObjectEditorActions.cpp
+++ b/Engine/Plugins/Experimental/Mutable/Source/CustomizableObjectEditor/Private/MuCOE/CustomizableObjectEditorActions.cpp
@@ -54,7 +54,7 @@ void FCustomizableObjectEditorViewportCommands::RegisterCommands()
 	UI_COMMAND( SetCameraLock, "Camera Lock", "Toggles viewport navigation between orbit and freely moving about.", EUserInterfaceActionType::ToggleButton, FInputChord() );
 	UI_COMMAND( SaveThumbnail, "Save Thumbnail", "Saves the viewpoint position and orientation in the Preview Pane for use as the thumbnail preview in the Content Browser.", EUserInterfaceActionType::Button, FInputChord() );
 
-	UI_COMMAND(BakeInstance, "Bake Instance from Object", "Create baked unreal resources for the current preview instance.", EUserInterfaceActionType::Button, FInputChord());
+	UI_COMMAND(BakeInstance, "Bake Instance", "Create baked unreal resources for the current preview instance.", EUserInterfaceActionType::Button, FInputChord());
 	UI_COMMAND(StateChangeShowData, "Show or hide test results", "Show or hide test results", EUserInterfaceActionType::ToggleButton, FInputChord());
 	UI_COMMAND(StateChangeShowGeometryData, "Show instance geometry data", "Show instance geometry data", EUserInterfaceActionType::ToggleButton, FInputChord());
 }
diff --git a/Engine/Plugins/Experimental/Mutable/Source/CustomizableObjectEditor/Private/MuCOE/CustomizableObjectEditorFunctionLibrary.cpp b/Engine/Plugins/Experimental/Mutable/Source/CustomizableObjectEditor/Private/MuCOE/CustomizableObjectEditorFunctionLibrary.cpp
--- a/Engine/Plugins/Experimental/Mutable/Source/CustomizableObjectEditor/Private/MuCOE/CustomizableObjectEditorFunctionLibrary.cpp
+++ b/Engine/Plugins/Experimental/Mutable/Source/CustomizableObjectEditor/Private/MuCOE/CustomizableObjectEditorFunctionLibrary.cpp
@@ -42,12 +42,14 @@ ECustomizableObjectCompilationState UCustomizableObjectEditorFunctionLibrary::Co
 
 	CustomizableObject->GetOutermost()->SetDirtyFlag(bPackageWasDirty);
 
+	const bool bCompilationSuccess = CustomizableObject->GetPrivate()->CompilationResult == ECompilationResultPrivate::Success ||
+		CustomizableObject->GetPrivate()->CompilationResult == ECompilationResultPrivate::Warnings;
+	
 	const double CurrentTime = FPlatformTime::Seconds();
 	UE_LOG( LogMutable, Display,
 		TEXT("Synchronously Compiled %s %s in %f seconds"),
 		*ObjectPath, 
-		Compiler.GetCompilationState() == ECustomizableObjectCompilationState::Completed ? 
-		TEXT("successfully") : TEXT("unsuccessfully"),
+		bCompilationSuccess ? TEXT("successfully") : TEXT("unsuccessfully"),
 		CurrentTime - StartTime
 	);
 
@@ -56,5 +58,5 @@ ECustomizableObjectCompilationState UCustomizableObjectEditorFunctionLibrary::Co
 		UE_LOG(LogMutable, Warning, TEXT("CO not marked as compiled"));
 	}
 
-	return Compiler.GetCompilationState();
+	return bCompilationSuccess ? ECustomizableObjectCompilationState::Completed : ECustomizableObjectCompilationState::Failed;
 }
diff --git a/Engine/Plugins/Experimental/Mutable/Source/CustomizableObjectEditor/Private/MuCOE/CustomizableObjectEditorStyle.cpp b/Engine/Plugins/Experimental/Mutable/Source/CustomizableObjectEditor/Private/MuCOE/CustomizableObjectEditorStyle.cpp
--- a/Engine/Plugins/Experimental/Mutable/Source/CustomizableObjectEditor/Private/MuCOE/CustomizableObjectEditorStyle.cpp
+++ b/Engine/Plugins/Experimental/Mutable/Source/CustomizableObjectEditor/Private/MuCOE/CustomizableObjectEditorStyle.cpp
@@ -56,9 +56,11 @@ FString FCustomizableObjectEditorStyle::RelativePathToPluginPath(const FString&
 }
 
 
-#define IMAGE_PLUGIN_BRUSH( RelativePath, ... ) FSlateImageBrush( FCustomizableObjectEditorStyle::RelativePathToPluginPath( RelativePath, ".png" ), __VA_ARGS__ )
-#define IMAGE_BRUSH(Style, RelativePath, ... ) FSlateImageBrush( Style->RootToContentDir( RelativePath, TEXT(".png") ), __VA_ARGS__ )
-#define IMAGE_BRUSH_SVG( Style, RelativePath, ... ) FSlateVectorImageBrush( Style->RootToContentDir(RelativePath, TEXT(".svg")), __VA_ARGS__)
+#define IMAGE_PLUGIN_BRUSH(Style, RelativePath, ... ) FSlateImageBrush(Style->RootToContentDir( RelativePath, TEXT(".png")), __VA_ARGS__)
+
+#define IMAGE_BRUSH(Style, RelativePath, ... ) FSlateImageBrush(Style->RootToCoreContentDir( RelativePath, TEXT(".png")), __VA_ARGS__)
+#define IMAGE_BRUSH_SVG( Style, RelativePath, ... ) FSlateVectorImageBrush(Style->RootToCoreContentDir(RelativePath, TEXT(".svg")), __VA_ARGS__)
+
 
 TSharedRef<FSlateStyleSet> FCustomizableObjectEditorStyle::Create()
 {
@@ -69,8 +71,9 @@ TSharedRef<FSlateStyleSet> FCustomizableObjectEditorStyle::Create()
 	const FVector2D Icon64x64(64.0f, 64.0f);
 
 	TSharedRef< FSlateStyleSet > Style = MakeShareable(new FSlateStyleSet("CustomizableObjectEditorStyle"));
-	Style->SetContentRoot(FPaths::EngineContentDir() / TEXT("Editor/Slate"));
-
+	Style->SetCoreContentRoot(FPaths::EngineContentDir() / TEXT("Editor/Slate"));
+	Style->SetContentRoot(FPaths::EnginePluginsDir() / TEXT("Experimental/Mutable/Content"));
+	
 	Style->Set("CustomizableObjectEditor.Save", new IMAGE_BRUSH_SVG(Style, "Starship/Common/SaveCurrent", Icon40x40));
 	Style->Set("CustomizableObjectEditor.Save.Small", new IMAGE_BRUSH_SVG(Style, "Starship/Common/SaveCurrent", Icon20x20));
 	Style->Set("CustomizableObjectEditor.Compile", new IMAGE_BRUSH_SVG(Style, "Starship/MainToolbar/compile", Icon40x40));
@@ -125,7 +128,7 @@ TSharedRef<FSlateStyleSet> FCustomizableObjectEditorStyle::Create()
 	Style->Set("CustomizableObjectEditorViewport.SetShowGrid", new IMAGE_BRUSH_SVG(Style, "Starship/Common/Grid", Icon20x20));
 	Style->Set("CustomizableObjectEditorViewport.SetShowSky", new IMAGE_BRUSH_SVG(Style, "Starship/Common/Reflections", Icon20x20));
 	Style->Set("CustomizableObjectEditorViewport.SetDrawUVs", new IMAGE_BRUSH_SVG(Style, "Starship/Common/SetDrawUVs", Icon20x20));
-	Style->Set("CustomizableObjectEditorViewport.BakeInstance", new IMAGE_BRUSH_SVG(Style, "Starship/MainToolbar/compile", Icon20x20));
+	Style->Set("CustomizableObjectEditorViewport.BakeInstance", new IMAGE_PLUGIN_BRUSH(Style, "Icons/BakerIcon", Icon40x40));
 
 	Style->Set("COEditorViewportLODCommands.TranslateMode", new IMAGE_BRUSH_SVG(Style, "Starship/EditorViewport/translate", Icon20x20));
 	Style->Set("COEditorViewportLODCommands.RotateMode", new IMAGE_BRUSH_SVG(Style, "Starship/EditorViewport/rotate", Icon20x20));
@@ -151,13 +154,7 @@ TSharedRef<FSlateStyleSet> FCustomizableObjectEditorStyle::Create()
 	Style->Set("CustomizableObjectDebugger.GenerateMutableGraph.Small", new IMAGE_BRUSH_SVG(Style, "Starship/Common/blueprint", Icon20x20));
 	Style->Set("CustomizableObjectDebugger.CompileMutableCode", new IMAGE_BRUSH_SVG(Style, "Starship/MainToolbar/Compile", Icon40x40));
 	Style->Set("CustomizableObjectDebugger.CompileMutableCode.Small", new IMAGE_BRUSH_SVG(Style, "Starship/MainToolbar/Compile", Icon20x20));
-
-	//Style->Set("ClassIcon.CustomizableObject", new IMAGE_PLUGIN_BRUSH("Icons/Icon_CustomObject_16x16", Icon16x16));
-	//Style->Set("ClassThumbnail.CustomizableObject", new IMAGE_PLUGIN_BRUSH("Icons/Icon_CustomObject_64x64.png", Icon64x64));
-	//Style->Set("ClassIcon.CustomizableObjectInstance", new IMAGE_PLUGIN_BRUSH("Icons/Icon_CustomObjectInstance_16x16.png", Icon16x16));
-	//Style->Set("ClassThumbnail.CustomizableObjectInstance", new IMAGE_PLUGIN_BRUSH("Icons/Icon_CustomObjectInstance_64x64.png", Icon64x64));
-
-
+	
 	FTableRowStyle PerformanceReportRowStyle = FAppStyle::Get().GetWidgetStyle<FTableRowStyle>("TableView.Row");
 	PerformanceReportRowStyle.SetEvenRowBackgroundBrush(*FAppStyle::Get().GetBrush("Brushes.AccentBrown"));
 	PerformanceReportRowStyle.SetOddRowBackgroundBrush(*FAppStyle::Get().GetBrush("Brushes.AccentGray"));
diff --git a/Engine/Plugins/Experimental/Mutable/Source/CustomizableObjectEditor/Private/MuCOE/CustomizableObjectInstanceEditor.cpp b/Engine/Plugins/Experimental/Mutable/Source/CustomizableObjectEditor/Private/MuCOE/CustomizableObjectInstanceEditor.cpp
--- a/Engine/Plugins/Experimental/Mutable/Source/CustomizableObjectEditor/Private/MuCOE/CustomizableObjectInstanceEditor.cpp
+++ b/Engine/Plugins/Experimental/Mutable/Source/CustomizableObjectEditor/Private/MuCOE/CustomizableObjectInstanceEditor.cpp
@@ -351,30 +351,13 @@ void FCustomizableObjectInstanceEditor::CreatePreviewInstance()
 		return;
 	}
 
+	CreatePreviewComponents();
+	
 	// Bind update delegate
 	HelperCallback = NewObject<UUpdateClassWrapperClass>();
 	HelperCallback->Delegate.BindSP(this, &FCustomizableObjectInstanceEditor::OnUpdatePreviewInstance);
 	CustomizableObjectInstance->UpdatedDelegate.AddDynamic(HelperCallback, &UUpdateClassWrapperClass::DelegatedCallback);
 
-	// Create a SkeletalMeshComponent for each component in the CO
-	const int32 NumMeshComponents = CustomizableObjectInstance->GetNumComponents();
-	PreviewSkeletalMeshComponents.SetNum(NumMeshComponents);
-	PreviewCustomizableSkeletalComponents.SetNum(NumMeshComponents);
-
-	for (int32 ComponentIndex = 0; ComponentIndex < NumMeshComponents; ++ComponentIndex)
-	{
-		PreviewCustomizableSkeletalComponents[ComponentIndex] = NewObject<UCustomizableSkeletalComponent>();
-		check(PreviewCustomizableSkeletalComponents[ComponentIndex]);
-
-		PreviewCustomizableSkeletalComponents[ComponentIndex]->CustomizableObjectInstance = CustomizableObjectInstance;
-		PreviewCustomizableSkeletalComponents[ComponentIndex]->ComponentIndex = ComponentIndex;
-
-		PreviewSkeletalMeshComponents[ComponentIndex] = NewObject<UDebugSkelMeshComponent>();
-		check(PreviewSkeletalMeshComponents[ComponentIndex]);
-
-		PreviewCustomizableSkeletalComponents[ComponentIndex]->AttachToComponent(PreviewSkeletalMeshComponents[ComponentIndex], FAttachmentTransformRules::KeepRelativeTransform);
-	}
-
 	CustomizableObjectInstance->UpdateSkeletalMeshAsync(true, true);
 	Viewport->SetPreviewComponents(PreviewSkeletalMeshComponents);
 }
@@ -755,20 +738,8 @@ void FCustomizableObjectInstanceEditor::OpenTextureAnalyzerTab()
 
 void FCustomizableObjectInstanceEditor::OnPostCompile()
 {
-	const UCustomizableObject* CustomizableObject = CustomizableObjectInstance->GetCustomizableObject();
-	if (CustomizableObject && CustomizableObject->GetPrivate()->GetSource())
-	{
-		if (PreviewCustomizableSkeletalComponents.Num() > 0)
-		{
-			CustomizableObjectInstance->UpdateSkeletalMeshAsync(true, true);
-		}
-		else
-		{
-			CreatePreviewInstance();
-		}
-
-		CustomizableInstanceDetailsView->ForceRefresh();
-	}
+	CreatePreviewComponents();
+	CustomizableObjectInstance->UpdateSkeletalMeshAsync(true, true);
 }
 
 
@@ -1030,6 +1001,39 @@ void FCustomizableObjectInstanceEditor::HideGizmoProjectorParameter(const TShare
 }
 
 
+void FCustomizableObjectInstanceEditor::CreatePreviewComponents()
+{
+	const UCustomizableObject* Object = CustomizableObjectInstance->GetCustomizableObject();
+	check(Object);
+	
+	for (const TObjectPtr<UCustomizableSkeletalComponent>& Component : PreviewCustomizableSkeletalComponents)
+	{
+		Component->DestroyComponent();
+	}
+	PreviewCustomizableSkeletalComponents.Empty();
+
+	for (const TObjectPtr<UDebugSkelMeshComponent>& Component : PreviewSkeletalMeshComponents)
+	{
+		Component->DestroyComponent();
+	}
+	PreviewSkeletalMeshComponents.Empty();
+
+	for (int32 ComponentIndex = 0; ComponentIndex < Object->GetComponentCount(); ++ComponentIndex)
+	{
+		UCustomizableSkeletalComponent* PreviewCustomizableSkeletalComponent = NewObject<UCustomizableSkeletalComponent>(UCustomizableSkeletalComponent::StaticClass());
+		UDebugSkelMeshComponent* PreviewSkeletalMeshComponent = NewObject<UDebugSkelMeshComponent>(GetTransientPackage(), NAME_None, RF_Transient);
+
+		PreviewCustomizableSkeletalComponent->bSkipSetReferenceSkeletalMesh = true;
+		PreviewCustomizableSkeletalComponent->CustomizableObjectInstance = CustomizableObjectInstance;
+		PreviewCustomizableSkeletalComponent->ComponentIndex = ComponentIndex;
+		PreviewCustomizableSkeletalComponent->AttachToComponent(PreviewSkeletalMeshComponent, FAttachmentTransformRules::KeepRelativeTransform);
+
+		PreviewCustomizableSkeletalComponents.Add(PreviewCustomizableSkeletalComponent);
+		PreviewSkeletalMeshComponents.Add(PreviewSkeletalMeshComponent);
+	}
+}
+
+
 TSharedRef<SDockTab> FCustomizableObjectInstanceEditor::SpawnTab_TextureAnalyzer(const FSpawnTabArgs& Args)
 {
 	check(Args.GetTabId() == TextureAnalyzerTabId);
diff --git a/Engine/Plugins/Experimental/Mutable/Source/CustomizableObjectEditor/Private/MuCOE/CustomizableObjectInstanceEditor.h b/Engine/Plugins/Experimental/Mutable/Source/CustomizableObjectEditor/Private/MuCOE/CustomizableObjectInstanceEditor.h
--- a/Engine/Plugins/Experimental/Mutable/Source/CustomizableObjectEditor/Private/MuCOE/CustomizableObjectInstanceEditor.h
+++ b/Engine/Plugins/Experimental/Mutable/Source/CustomizableObjectEditor/Private/MuCOE/CustomizableObjectInstanceEditor.h
@@ -221,7 +221,6 @@ private:
 	/** Open the Texture Analyzer tab */
 	void OpenTextureAnalyzerTab();
 
-	void OnPostCompile();
 public:
 	/** Static to reuse code with FCustomizableObjectEditor. */
 	static void HideGizmo(const TSharedPtr<ICustomizableObjectInstanceEditor>& Editor,
@@ -240,6 +239,10 @@ public:
 		const TSharedPtr<IDetailsView>& InstanceDetailsView);
 	
 private:
+	void CreatePreviewComponents();
+	
+	void OnPostCompile();
+	
 	/** The currently viewed object. */
 	TObjectPtr<UCustomizableObjectInstance> CustomizableObjectInstance;
 	TArray<TObjectPtr<UCustomizableSkeletalComponent>> PreviewCustomizableSkeletalComponents;
diff --git a/Engine/Plugins/Experimental/Mutable/Source/CustomizableObjectEditor/Private/MuCOE/SCustomizableObjectEditorViewport.cpp b/Engine/Plugins/Experimental/Mutable/Source/CustomizableObjectEditor/Private/MuCOE/SCustomizableObjectEditorViewport.cpp
--- a/Engine/Plugins/Experimental/Mutable/Source/CustomizableObjectEditor/Private/MuCOE/SCustomizableObjectEditorViewport.cpp
+++ b/Engine/Plugins/Experimental/Mutable/Source/CustomizableObjectEditor/Private/MuCOE/SCustomizableObjectEditorViewport.cpp
@@ -9,8 +9,8 @@
 #include "Framework/Commands/UICommandList.h"
 #include "Framework/MultiBox/MultiBoxBuilder.h"
 #include "MuCO/CustomizableObjectInstance.h"
+#include "MuCO/CustomizableObjectInstancePrivate.h"
 #include "MuCOE/CustomizableObjectEditorActions.h"
-#include "MuCOE/CustomizableObjectEditorUtilities.h"
 #include "MuCOE/CustomizableObjectEditorViewportClient.h"
 #include "MuCOE/CustomizableObjectEditorViewportLODCommands.h"
 #include "MuCOE/CustomizableObjectEditorViewportMenuCommands.h"
@@ -21,6 +21,9 @@
 #include "MuCOE/UnrealEditorPortabilityHelpers.h"
 #include "Settings/LevelEditorViewportSettings.h"
 #include "UnrealEdGlobals.h"
+#include "AssetRegistry/AssetRegistryModule.h"
+#include "MuCO/CustomizableObjectPrivate.h"
+#include "MuCO/CustomizableObjectSystem.h"
 #include "Widgets/Input/STextComboBox.h"
 
 class UCustomizableObject;
@@ -30,11 +33,18 @@ struct FCustomizableObjectProjector;
 struct FGeometry;
 
 
-void SCustomizableObjectEditorViewport::Construct(const FArguments& InArgs, const FCustomizableObjectEditorViewportRequiredArgs& InRequiredArgs)
-{
-	PreviewScenePtr = InRequiredArgs.PreviewScene;
-	TabBodyPtr = InRequiredArgs.TabBody;
+#define LOCTEXT_NAMESPACE "CustomizableObjectEditorViewportToolBar"
+
 
+void SCustomizableObjectEditorViewport::Construct(const FArguments& InArgs,
+	const TWeakPtr<FCustomizableObjectPreviewScene>& PreviewScene,
+	const TWeakPtr<SCustomizableObjectEditorViewportTabBody>& TabBody,
+	const TWeakPtr<ICustomizableObjectInstanceEditor>& Editor)
+{
+	PreviewScenePtr = PreviewScene;
+	TabBodyPtr = TabBody;
+	WeakEditor = Editor;
+	
 	SEditorViewport::Construct(
 		SEditorViewport::FArguments()
 		.IsEnabled(FSlateApplication::Get().GetNormalExecutionAttribute())
@@ -46,7 +56,7 @@ void SCustomizableObjectEditorViewport::Construct(const FArguments& InArgs, cons
 
 TSharedRef<FEditorViewportClient> SCustomizableObjectEditorViewport::MakeEditorViewportClient()
 {
-	LevelViewportClient = MakeShareable(new FCustomizableObjectEditorViewportClient(TabBodyPtr.Pin()->CustomizableObjectEditorPtr, PreviewScenePtr.Pin().Get(), SharedThis(this)));
+	LevelViewportClient = MakeShareable(new FCustomizableObjectEditorViewportClient(TabBodyPtr.Pin()->WeakEditor, PreviewScenePtr.Pin().Get(), SharedThis(this)));
 
 	LevelViewportClient->ViewportType = LVT_Perspective;
 	LevelViewportClient->bSetListenerPosition = false;
@@ -65,12 +75,117 @@ TSharedPtr<FSceneViewport>& SCustomizableObjectEditorViewport::GetSceneViewport(
 }
 
 
-TSharedPtr<SWidget> SCustomizableObjectEditorViewport::MakeViewportToolbar()
+void SCustomizableObjectEditorViewport::PopulateViewportOverlays(TSharedRef<SOverlay> Overlay)
+{
+	FTextBlockStyle NormalTextStyle = FAppStyle::GetWidgetStyle<FTextBlockStyle>("NormalText");
+	FTextBlockStyle CompileOverlayText = FTextBlockStyle(NormalTextStyle).SetFontSize(18);
+	
+	Overlay->AddSlot()
+	.VAlign(VAlign_Top)
+	[
+		SNew(SCustomizableObjectEditorViewportToolBar, TabBodyPtr.Pin(), SharedThis(this)).Cursor(EMouseCursor::Default)
+	];
+	Overlay->AddSlot()
+	.VAlign(VAlign_Center)
+	.HAlign(HAlign_Center)
+	[
+		SNew(STextBlock)
+		.Visibility(this, &SCustomizableObjectEditorViewport::GetShowCompileErrorOverlay)
+		.Text_Raw(this, &SCustomizableObjectEditorViewport::GetCompileErrorOverlayText)
+		.TextStyle(&CompileOverlayText)
+		.ColorAndOpacity(FLinearColor::White)
+		.ShadowOffset(FVector2D(1.5, 1.5))
+		.ShadowColorAndOpacity(FLinearColor(0, 0, 0, 0.9f))
+	];
+}
+
+
+EVisibility SCustomizableObjectEditorViewport::GetShowCompileErrorOverlay() const
 {
-	return SNew(SCustomizableObjectEditorViewportToolBar, TabBodyPtr.Pin(), SharedThis(this))
-		.Cursor(EMouseCursor::Default);
+	return GetCompileErrorOverlayText().IsEmpty() ? EVisibility::Hidden : EVisibility::Visible;
 }
 
+
+FText SCustomizableObjectEditorViewport::GetCompileErrorOverlayText() const
+{
+	const TSharedPtr<ICustomizableObjectInstanceEditor> Editor = WeakEditor.Pin();
+	if (!Editor)
+	{
+		return {};
+	}
+
+	const UObject* EditingObject = (*Editor->GetObjectsCurrentlyBeingEdited())[0];
+
+	const UCustomizableObject* Object = nullptr;
+	const UCustomizableObjectInstance* Instance = nullptr;
+
+	if (auto* CastObject = Cast<UCustomizableObject>(EditingObject))
+	{
+		Object = CastObject;
+		Instance = Editor->GetPreviewInstance();
+	}
+	else if (auto* CastInstance = Cast<UCustomizableObjectInstance>(EditingObject))
+	{
+		Object = CastInstance->GetCustomizableObject();
+		Instance = CastInstance;
+	}
+	else
+	{
+		check(false);
+	}
+
+	if (!Object)
+	{
+		return LOCTEXT("NoPreviewInstance", "No Customizable Object");
+	}
+
+	const FAssetRegistryModule& AssetRegistryModule = FModuleManager::LoadModuleChecked<FAssetRegistryModule>("AssetRegistry");
+	if (AssetRegistryModule.Get().IsLoadingAssets())
+	{
+		return LOCTEXT("LoadingAssetRegistry", "Loading Asset Registry...");
+	}
+	
+	if (Object->GetPrivate()->Status.Get() == FCustomizableObjectStatusTypes::EState::Loading)
+	{
+		return LOCTEXT("Loading", "Loading...");
+	}
+	
+	if (Object->GetPrivate()->CompilationState == ECompilationStatePrivate::InProgress)
+	{
+		return LOCTEXT("Compiling", "Compiling...");
+	}
+	
+	if (!Instance) // Only happens if Mutable has compilations disabled and the CO was not compiled.
+	{
+		return LOCTEXT("EmptyPreview", "Empty Preview");
+	}
+	
+	const UCustomizableObjectSystem* System = UCustomizableObjectSystem::GetInstanceChecked();
+	
+	if (System->IsUpdating(Instance))
+	{
+		return LOCTEXT("Updating", "Updating...");
+	}
+
+	if (Object->GetPrivate()->CompilationResult == ECompilationResultPrivate::Errors) // Compilation errors have more priority than Update errors
+	{
+		return LOCTEXT("ErrorCompiling", "Error Compiling");
+	}
+
+	if (Instance->GetPrivate()->SkeletalMeshStatus == ESkeletalMeshStatus::Error)
+	{
+		return LOCTEXT("ErrorUpdating", "Error Updating");
+	}
+
+	if (!Instance->HasAnySkeletalMesh())
+	{
+		return LOCTEXT("EmptyPreview", "Empty Preview");
+	}
+	
+	return {};
+}
+
+
 void SCustomizableObjectEditorViewport::OnUndoRedo()
 {
 	LevelViewportClient->Invalidate();
@@ -84,37 +199,26 @@ void SCustomizableObjectEditorViewportTabBody::Construct(const FArguments& InArg
 {
 	UICommandList = MakeShareable(new FUICommandList);
 
-	CustomizableObjectEditorPtr = InArgs._CustomizableObjectEditor;
+	WeakEditor = InArgs._CustomizableObjectEditor;
 
 	FCustomizableObjectEditorViewportMenuCommands::Register();
-	//FAnimViewportShowCommands::Register();
 	FCustomizableObjectEditorViewportLODCommands::Register();
-	//FAnimViewportPlaybackCommands::Register();
 
 	FPreviewScene::ConstructionValues SceneConstructValues;
 	SceneConstructValues.bShouldSimulatePhysics = true;
 
 	PreviewScenePtr = MakeShareable(new FCustomizableObjectPreviewScene(SceneConstructValues));
-	//PreviewScenePtr->GetWorld()->WorldType = EWorldType::Editor; // Needed for custom depth pass for screenshot with no background
-
-	FCustomizableObjectEditorViewportRequiredArgs ViewportArgs(PreviewScenePtr.ToSharedRef(),SharedThis(this));
-	//(InSkeletonTree, InPreviewScene, SharedThis(this), InAssetEditorToolkit, InOnUndoRedo);
-
-	ViewportWidget = SNew(SCustomizableObjectEditorViewport, ViewportArgs);
-
-
+	
+	ViewportWidget = SNew(SCustomizableObjectEditorViewport, PreviewScenePtr.ToSharedRef(), SharedThis(this), WeakEditor);
 
 	this->ChildSlot
 		[
 			SNew(SVerticalBox)
-
 			+ SVerticalBox::Slot()
-			.Padding(2.0f)
 			.AutoHeight()
 			[
 				BuildToolBar()
 			]
-
 			+SVerticalBox::Slot()
 			.FillHeight(1)
 			[
@@ -123,9 +227,7 @@ void SCustomizableObjectEditorViewportTabBody::Construct(const FArguments& InArg
 		];
 
 	LevelViewportClient = StaticCastSharedPtr<FCustomizableObjectEditorViewportClient>(ViewportWidget->GetViewportClient());
-
-	//PreviewSkeletalMeshComponent = 0;
-
+	
 	BindCommands();
 }
 
@@ -573,7 +675,7 @@ int32 SCustomizableObjectEditorViewportTabBody::GetViewportCameraSpeed()
 
 TSharedRef<SWidget> SCustomizableObjectEditorViewportTabBody::BuildToolBar()
 {
-	FToolBarBuilder CommandToolbarBuilder(UICommandList, FMultiBoxCustomization::None);
+	FSlimHorizontalToolBarBuilder CommandToolbarBuilder(UICommandList, FMultiBoxCustomization::None);
 	{
 		CommandToolbarBuilder.AddToolBarButton(FCustomizableObjectEditorViewportCommands::Get().SetShowGrid);
 		CommandToolbarBuilder.AddToolBarButton(FCustomizableObjectEditorViewportCommands::Get().SetShowSky);
@@ -584,12 +686,11 @@ TSharedRef<SWidget> SCustomizableObjectEditorViewportTabBody::BuildToolBar()
 
 		CommandToolbarBuilder.AddComboButton(
 			FUIAction(),
-			FOnGetContent::CreateSP(
-				this,
-				&SCustomizableObjectEditorViewportTabBody::GenerateUVMaterialOptionsMenuContent),
+			FOnGetContent::CreateSP(this, &SCustomizableObjectEditorViewportTabBody::GenerateUVMaterialOptionsMenuContent),
 			TAttribute<FText>(),
 			TAttribute<FText>(),
-			FSlateIcon());
+			FSlateIcon(),
+			true);
 	}
 	CommandToolbarBuilder.EndSection();
 
@@ -598,21 +699,7 @@ TSharedRef<SWidget> SCustomizableObjectEditorViewportTabBody::BuildToolBar()
 	CommandToolbarBuilder.AddToolBarButton(FCustomizableObjectEditorViewportCommands::Get().BakeInstance);
 	CommandToolbarBuilder.EndSection();
 
-
-	return
-		SNew(SHorizontalBox)
-
-		+ SHorizontalBox::Slot()
-		.Padding(4, 0)
-		[
-			SNew(SBorder)
-			.Padding(0)
-		.BorderImage(UE_MUTABLE_GET_BRUSH("NoBorder"))
-		.IsEnabled(FSlateApplication::Get().GetNormalExecutionAttribute())
-		[
-			CommandToolbarBuilder.MakeWidget()
-		]
-		];
+	return CommandToolbarBuilder.MakeWidget();
 }
 
 
@@ -970,3 +1057,6 @@ FLinearColor SCustomizableObjectEditorViewportTabBody::GetViewportBackgroundColo
 {
 	return LevelViewportClient->GetBackgroundColor();
 }
+
+
+#undef LOCTEXT_NAMESPACE
\ No newline at end of file
diff --git a/Engine/Plugins/Experimental/Mutable/Source/CustomizableObjectEditor/Private/MuCOE/SCustomizableObjectEditorViewport.h b/Engine/Plugins/Experimental/Mutable/Source/CustomizableObjectEditor/Private/MuCOE/SCustomizableObjectEditorViewport.h
--- a/Engine/Plugins/Experimental/Mutable/Source/CustomizableObjectEditor/Private/MuCOE/SCustomizableObjectEditorViewport.h
+++ b/Engine/Plugins/Experimental/Mutable/Source/CustomizableObjectEditor/Private/MuCOE/SCustomizableObjectEditorViewport.h
@@ -6,6 +6,8 @@
 #include "MuCOE/CustomizableObjectEditorViewportClient.h"
 #include "UObject/GCObject.h"
 
+class FCustomizableObjectPreviewScene;
+class SCustomizableObjectEditorViewportTabBody;
 enum ERotationGridMode : int;
 enum EViewModeIndex : int;
 namespace EAnimationMode { enum Type : int; }
@@ -31,22 +33,6 @@ struct FGeometry;
 enum class ECustomizableObjectProjectorType : uint8;
 
 
-struct FCustomizableObjectEditorViewportRequiredArgs
-{
-	FCustomizableObjectEditorViewportRequiredArgs(
-		const TSharedRef<class FCustomizableObjectPreviewScene>& InPreviewScene,
-		TSharedRef<class SCustomizableObjectEditorViewportTabBody> InTabBody 
-	)
-		: PreviewScene(InPreviewScene)
-		, TabBody(InTabBody)
-	{}
-
-	TSharedRef<FCustomizableObjectPreviewScene> PreviewScene;
-
-	TSharedRef<SCustomizableObjectEditorViewportTabBody> TabBody;
-};
-
-
 //////////////////////////////////////////////////////////////////////////
 // SCustomizableObjectEditorViewport: the pure viewport widget
 
@@ -56,19 +42,26 @@ public:
 	SLATE_BEGIN_ARGS(SCustomizableObjectEditorViewport)	{}
 	SLATE_END_ARGS()
 
-	void Construct(const FArguments& InArgs, const FCustomizableObjectEditorViewportRequiredArgs& InRequiredArgs);
-
-	TSharedPtr<FSceneViewport>& GetSceneViewport();
+	void Construct(const FArguments& InArgs,
+		const TWeakPtr<FCustomizableObjectPreviewScene>& PreviewScene,
+		const TWeakPtr<SCustomizableObjectEditorViewportTabBody>& TabBody,
+		const TWeakPtr<ICustomizableObjectInstanceEditor>& Editor);
 
-protected:
 	// SEditorViewport interface
 	virtual TSharedRef<FEditorViewportClient> MakeEditorViewportClient() override;
-	virtual TSharedPtr<SWidget> MakeViewportToolbar() override;
+	virtual void PopulateViewportOverlays(TSharedRef<SOverlay> Overlay) override;
 	// End of SEditorViewport interface
 
+	TSharedPtr<FSceneViewport>& GetSceneViewport();
+
+private:
 	/**  Handle undo/redo by refreshing the viewport */
 	void OnUndoRedo();
 
+	EVisibility GetShowCompileErrorOverlay() const;
+
+	FText GetCompileErrorOverlayText() const;
+	
 	// Viewport client
 	TSharedPtr<FCustomizableObjectEditorViewportClient> LevelViewportClient;
 
@@ -79,6 +72,9 @@ protected:
 	TWeakPtr<FCustomizableObjectPreviewScene> PreviewScenePtr;
 
 	TWeakPtr<SCustomizableObjectViewportToolBar> ToolbarPtr;
+
+	/** Weak pointer of the open editor */
+	TWeakPtr<ICustomizableObjectInstanceEditor> WeakEditor;
 };
 
 
@@ -198,7 +194,7 @@ public:
 	TSharedPtr<FCustomizableObjectEditorViewportClient> GetViewportClient() const;
 
 	// Pointer back to the editor tool that owns us.
-	TWeakPtr<ICustomizableObjectInstanceEditor> CustomizableObjectEditorPtr;
+	TWeakPtr<ICustomizableObjectInstanceEditor> WeakEditor;
 
 	//void SetEditorTransformViewportToolbar(TWeakPtr<class SCustomizableObjectEditorTransformViewportToolbar> EditorTransformViewportToolbarParam);
 	void SetViewportToolbarTransformWidget(TWeakPtr<class SWidget> TransformWidget);
diff --git a/Engine/Plugins/Experimental/Mutable/Source/CustomizableObjectEditor/Private/MuCOE/SCustomizableObjectEditorViewportToolBar.cpp b/Engine/Plugins/Experimental/Mutable/Source/CustomizableObjectEditor/Private/MuCOE/SCustomizableObjectEditorViewportToolBar.cpp
--- a/Engine/Plugins/Experimental/Mutable/Source/CustomizableObjectEditor/Private/MuCOE/SCustomizableObjectEditorViewportToolBar.cpp
+++ b/Engine/Plugins/Experimental/Mutable/Source/CustomizableObjectEditor/Private/MuCOE/SCustomizableObjectEditorViewportToolBar.cpp
@@ -45,7 +45,7 @@ void SCustomizableObjectEditorViewportToolBar::Construct(const FArguments& InArg
 
 	TSharedRef<SCustomizableObjectEditorViewportTabBody> ViewportRef = Viewport.Pin().ToSharedRef();
 
-	WeakEditor = ViewportRef->CustomizableObjectEditorPtr;
+	WeakEditor = ViewportRef->WeakEditor;
 	
 	TSharedRef<SHorizontalBox> LeftToolbar = SNew(SHorizontalBox)
 
@@ -126,50 +126,19 @@ void SCustomizableObjectEditorViewportToolBar::Construct(const FArguments& InArg
 
 	ChildSlot
 		[
-			SNew(SVerticalBox)
-			+ SVerticalBox::Slot()
-			.VAlign(VAlign_Top)
+		SNew(SBorder)
+			.BorderImage(UE_MUTABLE_GET_BRUSH("NoBorder"))
+			.ForegroundColor(UE_MUTABLE_GET_SLATECOLOR(DefaultForegroundName))
 			[
-				SNew(SBorder)
-				.BorderImage(UE_MUTABLE_GET_BRUSH("NoBorder"))
-				.ForegroundColor(UE_MUTABLE_GET_SLATECOLOR(DefaultForegroundName))
+				SNew(SVerticalBox)
+				+ SVerticalBox::Slot()
+				.AutoHeight()
 				[
-					SNew(SVerticalBox)
-					+ SVerticalBox::Slot()
-					.AutoHeight()
-					[
-						SNew(SHorizontalBox)
-						+ SHorizontalBox::Slot()
-						.HAlign(HAlign_Left)
-						[
-							LeftToolbar
-						]
-					]
-				]
-			]
-	
-			+ SVerticalBox::Slot()
-			.HAlign(HAlign_Center)
-			.VAlign(VAlign_Center)
-			.AutoHeight()
-			[
-				SNew(SScaleBox)
-				.Stretch(EStretch::ScaleToFit)
-				[
-					SAssignNew(CompileErrorLayout, SButton)
-					.ButtonStyle(&FButtonStyle::GetDefault())
-					.ButtonColorAndOpacity(ButtonColor2)
-					.ForegroundColor(ButtonColor2)
-					.HAlign(HAlign_Center)
-					.VAlign(VAlign_Center)
-					.Visibility(this, &SCustomizableObjectEditorViewportToolBar::GetShowCompileErrorOverlay)
-					.Content()
+					SNew(SHorizontalBox)
+					+ SHorizontalBox::Slot()
+					.HAlign(HAlign_Left)
 					[
-						SNew(STextBlock)
-						.Text(this, &SCustomizableObjectEditorViewportToolBar::GetCompileErrorOverlayText)
-						.Justification(ETextJustify::Center)
-						.ColorAndOpacity(FSlateColor(TextColor2))
-						.Font(Info)
+						LeftToolbar
 					]
 				]
 			]
@@ -179,71 +148,6 @@ void SCustomizableObjectEditorViewportToolBar::Construct(const FArguments& InArg
 }
 
 
-EVisibility SCustomizableObjectEditorViewportToolBar::GetShowCompileErrorOverlay() const
-{
-	return GetCompileErrorOverlayText().IsEmpty() ? EVisibility::Hidden : EVisibility::Visible;
-}
-
-
-FText SCustomizableObjectEditorViewportToolBar::GetCompileErrorOverlayText() const
-{
-	const TSharedPtr<ICustomizableObjectInstanceEditor> Editor = WeakEditor.Pin();
-	if (!Editor)
-	{
-		return {};
-	}
-
-	bool bLoading = false;
-
-	const UObject* EditingObject = (*Editor->GetObjectsCurrentlyBeingEdited())[0];
-	if (const UCustomizableObject* CustomizableObject = Cast<UCustomizableObject>(EditingObject))
-	{
-		bLoading = CustomizableObject->GetPrivate()->Status.Get() == FCustomizableObjectStatusTypes::EState::Loading;
-	}
-	else if (const UCustomizableObjectInstance* Instance = Cast<UCustomizableObjectInstance>(EditingObject))
-	{
-		const UCustomizableObject* InstanceCustomizableObject = Instance->GetCustomizableObject();
-		bLoading = InstanceCustomizableObject && InstanceCustomizableObject->GetPrivate()->Status.Get() == FCustomizableObjectStatusTypes::EState::Loading;
-	}
-
-	if (bLoading)
-	{
-		return LOCTEXT("LoadingAssets", "Loading Customizable Object");
-	}
-
-	const UCustomizableObjectInstance* Instance = Editor->GetPreviewInstance();
-	if (!Instance)
-	{
-		return LOCTEXT("NoPreviewInstance", "No Preview Instance");
-	}
-
-	const UCustomizableObjectSystem* System = UCustomizableObjectSystem::GetInstanceChecked();
-	
-	if (System->IsUpdating(Instance))
-	{
-		return LOCTEXT("UpdatingSkeletalMesh", "Updating Skeletal Mesh");
-	}
-
-	const UCustomizableInstancePrivate* PrivateInstance = Instance->GetPrivate();
-	
-	switch (PrivateInstance->SkeletalMeshStatus)
-	{
-	case ESkeletalMeshStatus::NotGenerated:
-		return LOCTEXT("NoSkeletalMeshGenerated", "No Skeletal Mesh Generated");
-				
-	case ESkeletalMeshStatus::Error:
-		return LOCTEXT("ErrorUpdatingSkeletalMesh", "Error Updating Skeletal Mesh");
-
-	case ESkeletalMeshStatus::Success:
-		return {};
-
-	default:
-		unimplemented();
-		return {};
-	}
-}
-
-
 TSharedRef<SWidget> SCustomizableObjectEditorViewportToolBar::GenerateViewMenu() const
 {
 	const FCustomizableObjectEditorViewportMenuCommands& Actions = FCustomizableObjectEditorViewportMenuCommands::Get();
diff --git a/Engine/Plugins/Experimental/Mutable/Source/CustomizableObjectEditor/Public/MuCOE/CustomizableObjectEditorFunctionLibrary.h b/Engine/Plugins/Experimental/Mutable/Source/CustomizableObjectEditor/Public/MuCOE/CustomizableObjectEditorFunctionLibrary.h
--- a/Engine/Plugins/Experimental/Mutable/Source/CustomizableObjectEditor/Public/MuCOE/CustomizableObjectEditorFunctionLibrary.h
+++ b/Engine/Plugins/Experimental/Mutable/Source/CustomizableObjectEditor/Public/MuCOE/CustomizableObjectEditorFunctionLibrary.h
@@ -15,6 +15,17 @@ enum class ECustomizableObjectOptimizationLevel : uint8
 	Maximum
 };
 
+
+UENUM()
+enum class ECustomizableObjectCompilationState : uint8
+{
+	None,
+	InProgress,
+	Completed,
+	Failed
+};
+
+
 /**
  * Functions we want to be able to call on CustomizableObjects at edit time - could
  * be exposed to cook as well.
diff --git a/Engine/Plugins/Experimental/Mutable/Source/CustomizableObjectEditor/Public/MuCOE/SCustomizableObjectEditorViewportToolBar.h b/Engine/Plugins/Experimental/Mutable/Source/CustomizableObjectEditor/Public/MuCOE/SCustomizableObjectEditorViewportToolBar.h
--- a/Engine/Plugins/Experimental/Mutable/Source/CustomizableObjectEditor/Public/MuCOE/SCustomizableObjectEditorViewportToolBar.h
+++ b/Engine/Plugins/Experimental/Mutable/Source/CustomizableObjectEditor/Public/MuCOE/SCustomizableObjectEditorViewportToolBar.h
@@ -103,13 +103,7 @@ private:
 	TOptional<float> OnGetFloorOffset() const;
 	/** Called when the floor offset slider is adjusted in the perspective viewport */
 	void OnFloorOffsetChanged(float NewValue);
-
-	// Called to show / hide Customizable Object compile error
-	EVisibility GetShowCompileErrorOverlay() const;
-
-	// Called to modify the instance compile status overlay information text
-	FText GetCompileErrorOverlayText() const;
-
+	
 	ECheckBoxState IsRotationGridSnapChecked() const;
 	void HandleToggleRotationGridSnap(ECheckBoxState InState);
 	FText GetRotationGridLabel() const;
diff --git a/Engine/Plugins/Experimental/Mutable/Source/MutableValidation/Private/MuV/AssetValidator_CustomizableObjects.cpp b/Engine/Plugins/Experimental/Mutable/Source/MutableValidation/Private/MuV/AssetValidator_CustomizableObjects.cpp
--- a/Engine/Plugins/Experimental/Mutable/Source/MutableValidation/Private/MuV/AssetValidator_CustomizableObjects.cpp
+++ b/Engine/Plugins/Experimental/Mutable/Source/MutableValidation/Private/MuV/AssetValidator_CustomizableObjects.cpp
@@ -128,11 +128,11 @@ EDataValidationResult UAssetValidator_CustomizableObjects::IsCustomizableObjectV
 	TArray<FText> CachedValidationWarnings;
 	
 	// Map with all the possible compilation states. We use this so at the end we can know if any of those states was returned by any of the compilation runs
-	TMap<ECustomizableObjectCompilationState,bool>PossibleEndCompilationStates;
-	PossibleEndCompilationStates.Add(ECustomizableObjectCompilationState::Completed,false);
-	PossibleEndCompilationStates.Add(ECustomizableObjectCompilationState::None,false);
-	PossibleEndCompilationStates.Add(ECustomizableObjectCompilationState::Failed,false);
-	PossibleEndCompilationStates.Add(ECustomizableObjectCompilationState::InProgress,false);
+	TMap<ECompilationResultPrivate,bool>PossibleEndCompilationStates;
+	PossibleEndCompilationStates.Add(ECompilationResultPrivate::Unknown, false);
+	PossibleEndCompilationStates.Add(ECompilationResultPrivate::Success, false);
+	PossibleEndCompilationStates.Add(ECompilationResultPrivate::Errors, false);
+	PossibleEndCompilationStates.Add(ECompilationResultPrivate::Warnings, false);
 	PossibleEndCompilationStates.Shrink();
 	
 	// Iterate over the compilation options that we want to test and perform the compilation
@@ -167,7 +167,7 @@ EDataValidationResult UAssetValidator_CustomizableObjects::IsCustomizableObjectV
 		}
 
 		// Flag the array with end results to have the current output as true since it was produced by this execution
-		ECustomizableObjectCompilationState CompilationEndResult = Compiler->GetCompilationState();
+		ECompilationResultPrivate CompilationEndResult = RootObject->GetPrivate()->CompilationResult;
 		bool* Value = PossibleEndCompilationStates.Find(CompilationEndResult);
 		check (Value);		// If this fails it may mean we are getting a compilation state we are not considering. 
 		*Value = true;
@@ -194,13 +194,11 @@ EDataValidationResult UAssetValidator_CustomizableObjects::IsCustomizableObjectV
 	// Return informed guess about what the validation state of this object should be
 
 	// If it contains invalid states then notify about it too:
-	// ECustomizableObjectCompilationState::InProgress should not be possible since we are compiling synchronously.
-	check (*PossibleEndCompilationStates.Find(ECustomizableObjectCompilationState::InProgress) == false);
 	// ECustomizableObjectCompilationState::None would mean the resource is locked (and should not be)
-	check (*PossibleEndCompilationStates.Find(ECustomizableObjectCompilationState::None) == false);
+	check (*PossibleEndCompilationStates.Find(ECompilationResultPrivate::Unknown) == false);
 	
 	// If one or more tests failed to ran then the result must be invalid
-	if (*PossibleEndCompilationStates.Find(ECustomizableObjectCompilationState::Failed) == true)
+	if (*PossibleEndCompilationStates.Find(ECompilationResultPrivate::Errors) == true)
 	{
 		// Early CO compilation error (before starting mutable compilation) -> Output is invalid
 		Result = EDataValidationResult::Invalid;
diff --git a/Engine/Plugins/Experimental/Mutable/Source/MutableValidation/Private/MuV/ValidationUtils.cpp b/Engine/Plugins/Experimental/Mutable/Source/MutableValidation/Private/MuV/ValidationUtils.cpp
--- a/Engine/Plugins/Experimental/Mutable/Source/MutableValidation/Private/MuV/ValidationUtils.cpp
+++ b/Engine/Plugins/Experimental/Mutable/Source/MutableValidation/Private/MuV/ValidationUtils.cpp
@@ -10,6 +10,7 @@
 #include "MuCO/CustomizableObjectPrivate.h"
 #include "MuCO/CustomizableObjectSystem.h"
 #include "Interfaces/ITargetPlatform.h"
+#include "MuCOE/CustomizableObjectEditorFunctionLibrary.h"
 
 void PrepareAssetRegistry()
 {
@@ -81,9 +82,8 @@ bool CompileCustomizableObject(UCustomizableObject* InCustomizableObject, const
     const int64 CompilationEndRealPeakBytes = CompilationStartBytes + CompilationEndPeakBytes;
 	
     // Get the compilation result
-    const ECustomizableObjectCompilationState CompilationEndResult = Compiler->GetCompilationState();
-	
-    const bool bWasCoCompilationSuccessful = CompilationEndResult == ECustomizableObjectCompilationState::Completed;
+    const bool bWasCoCompilationSuccessful = InCustomizableObject->GetPrivate()->CompilationResult == ECompilationResultPrivate::Success ||
+    	InCustomizableObject->GetPrivate()->CompilationResult == ECompilationResultPrivate::Warnings;
     if (bWasCoCompilationSuccessful)
     {
     	UE_LOG(LogMutable, Display, TEXT("The compilation of the %s model was succesfull."), *InCustomizableObject->GetName());
@@ -93,6 +93,7 @@ bool CompileCustomizableObject(UCustomizableObject* InCustomizableObject, const
     	UE_LOG(LogMutable, Error, TEXT("The compilation of the %s model failed."), *InCustomizableObject->GetName());
     }
 
+	const ECustomizableObjectCompilationState CompilationEndResult = bWasCoCompilationSuccessful ? ECustomizableObjectCompilationState::Completed : ECustomizableObjectCompilationState::Failed;
 	
 	// Print MTU parseable logs only if asked. The addition of duplicated entries in MongoDB is not available so this way we avoid having to handle them
 	// when we are sure we do not require it.
