[Mutable] Added Undo/Redo of the instance parameters.

#rnx

diff --git a/Engine/Plugins/Experimental/Mutable/Source/CustomizableObject/Classes/MuCO/CustomizableObjectInstance.h b/Engine/Plugins/Experimental/Mutable/Source/CustomizableObject/Classes/MuCO/CustomizableObjectInstance.h
--- a/Engine/Plugins/Experimental/Mutable/Source/CustomizableObject/Classes/MuCO/CustomizableObjectInstance.h
+++ b/Engine/Plugins/Experimental/Mutable/Source/CustomizableObject/Classes/MuCO/CustomizableObjectInstance.h
@@ -243,7 +243,9 @@ public:
 #if WITH_EDITOR
 	virtual void PostEditChangeProperty(FPropertyChangedEvent& PropertyChangedEvent) override;
 	virtual bool CanEditChange( const FProperty* InProperty ) const override;
+	virtual void PostTransacted(const FTransactionObjectEvent& TransactionEvent) override;
 #endif //WITH_EDITOR
+
 	virtual void Serialize(FArchive& Ar) override;
 	virtual void PostLoad() override;
 	virtual void BeginDestroy() override;
diff --git a/Engine/Plugins/Experimental/Mutable/Source/CustomizableObject/Internal/MuCO/CustomizableObjectInstancePrivate.h b/Engine/Plugins/Experimental/Mutable/Source/CustomizableObject/Internal/MuCO/CustomizableObjectInstancePrivate.h
--- a/Engine/Plugins/Experimental/Mutable/Source/CustomizableObject/Internal/MuCO/CustomizableObjectInstancePrivate.h
+++ b/Engine/Plugins/Experimental/Mutable/Source/CustomizableObject/Internal/MuCO/CustomizableObjectInstancePrivate.h
@@ -162,6 +162,11 @@ struct FAnimBpGeneratedPhysicsAssets
 };
 
 
+#if WITH_EDITOR
+DECLARE_MULTICAST_DELEGATE_OneParam(FObjectInstanceTransactedDelegate, const FTransactionObjectEvent&);
+#endif //WIT_EDITOR
+
+
 UCLASS()
 class CUSTOMIZABLEOBJECT_API UCustomizableInstancePrivate : public UObject
 {
@@ -390,5 +395,10 @@ public:
 	/** Preview Instance Properties search box filter. Saved here to avoid losing the text during UI refreshes. */
 	FText ParametersSearchFilter;
 #endif
+
+#if WITH_EDITOR
+	/** Delegate called when the Instance has been transacted */
+	FObjectInstanceTransactedDelegate OnInstanceTransactedDelegate;
+#endif // WITH_EDITOR
 };
 
diff --git a/Engine/Plugins/Experimental/Mutable/Source/CustomizableObject/Private/MuCO/CustomizableObjectInstance.cpp b/Engine/Plugins/Experimental/Mutable/Source/CustomizableObject/Private/MuCO/CustomizableObjectInstance.cpp
--- a/Engine/Plugins/Experimental/Mutable/Source/CustomizableObject/Private/MuCO/CustomizableObjectInstance.cpp
+++ b/Engine/Plugins/Experimental/Mutable/Source/CustomizableObject/Private/MuCO/CustomizableObjectInstance.cpp
@@ -58,6 +58,7 @@
 #include "MessageLogModule.h"
 #include "UnrealEdMisc.h"
 #include "Subsystems/AssetEditorSubsystem.h"
+#include "Misc/TransactionObjectEvent.h"
 #endif
 
 #ifndef REQUIRES_SINGLEUSE_FLAG_FOR_RUNTIME_TEXTURES
@@ -392,6 +393,13 @@ bool UCustomizableObjectInstance::CanEditChange(const FProperty* InProperty) con
 	return bIsMutable;
 }
 
+void UCustomizableObjectInstance::PostTransacted(const FTransactionObjectEvent& TransactionEvent)
+{
+	Super::PostTransacted(TransactionEvent);
+
+	GetPrivate()->OnInstanceTransactedDelegate.Broadcast(TransactionEvent);
+}
+
 #endif // WITH_EDITOR
 
 
diff --git a/Engine/Plugins/Experimental/Mutable/Source/CustomizableObject/Public/MuCO/CustomizableObjectInstanceDescriptor.h b/Engine/Plugins/Experimental/Mutable/Source/CustomizableObject/Public/MuCO/CustomizableObjectInstanceDescriptor.h
--- a/Engine/Plugins/Experimental/Mutable/Source/CustomizableObject/Public/MuCO/CustomizableObjectInstanceDescriptor.h
+++ b/Engine/Plugins/Experimental/Mutable/Source/CustomizableObject/Public/MuCO/CustomizableObjectInstanceDescriptor.h
@@ -374,7 +374,8 @@ private:
 	UPROPERTY()
 	TArray<FCustomizableObjectProjectorParameterValue> ProjectorParameters;
 
-	/** Mutable parameters optimization state. */
+	/** Mutable parameters optimization state. Transient UProperty to make it transactable. */
+	UPROPERTY(Transient)
 	int32 State = 0;
 	
 	/** If this is set to true, when updating the instance an additional step will be performed to calculate the list of instance parameters that are relevant for the current parameter values. */
diff --git a/Engine/Plugins/Experimental/Mutable/Source/CustomizableObjectEditor/Private/MuCOE/CustomizableInstanceDetails.cpp b/Engine/Plugins/Experimental/Mutable/Source/CustomizableObjectEditor/Private/MuCOE/CustomizableInstanceDetails.cpp
--- a/Engine/Plugins/Experimental/Mutable/Source/CustomizableObjectEditor/Private/MuCOE/CustomizableInstanceDetails.cpp
+++ b/Engine/Plugins/Experimental/Mutable/Source/CustomizableObjectEditor/Private/MuCOE/CustomizableInstanceDetails.cpp
@@ -10,6 +10,7 @@
 #include "HAL/PlatformApplicationMisc.h"
 #include "IDetailsView.h"
 #include "IDetailGroup.h"
+#include "Misc/TransactionObjectEvent.h"
 #include "ScopedTransaction.h"
 #include "Serialization/BufferArchive.h"
 #include "Toolkits/ToolkitManager.h"
@@ -23,6 +24,7 @@
 
 #include "Slate/DeferredCleanupSlateBrush.h"
 #include "SSearchableComboBox.h"
+#include "Types/SlateEnums.h"
 #include "Widgets/Colors/SColorBlock.h"
 #include "Widgets/Colors/SColorPicker.h"
 #include "Widgets/Input/SButton.h"
@@ -83,6 +85,7 @@ void FCustomizableInstanceDetails::CustomizeDetails(const TSharedPtr<IDetailLayo
 
 	// Delegate to refresh the detils when the instance has finished the Update
 	CustomInstance->UpdatedNativeDelegate.AddSP(this, &FCustomizableInstanceDetails::InstanceUpdated);
+	CustomInstance->GetPrivate()->OnInstanceTransactedDelegate.AddSP(this, &FCustomizableInstanceDetails::OnInstanceTransacted);
 
 	// New Category that will store all properties widgets
 	IDetailCategoryBuilder& ResourcesCategory = DetailBuilder->EditCategory("Generated Resources");
@@ -273,8 +276,16 @@ void FCustomizableInstanceDetails::Refresh() const
 }
 
 
+void FCustomizableInstanceDetails::UpdateInstance()
+{
+	CustomInstance->UpdateSkeletalMeshAsync(true, true);
+	CustomInstance->PostEditChange();
+}
+
+
 void FCustomizableInstanceDetails::InstanceUpdated(UCustomizableObjectInstance* Instance) const
 {
+	// Check the instance update context to aboid unecessary UI updates.
 	if (!bUpdatingSlider)
 	{
 		Refresh();
@@ -327,13 +338,16 @@ TSharedRef<SWidget> FCustomizableInstanceDetails::GenerateStateSelector()
 
 void FCustomizableInstanceDetails::OnStateComboBoxSelectionChanged(TSharedPtr<FString> Selection, ESelectInfo::Type SelectInfo)
 {
-	CustomInstance->PreEditChange(nullptr);
-	CustomInstance->SetCurrentState(*Selection);
-	CustomInstance->UpdateSkeletalMeshAsync(true, true);
-	CustomInstance->PostEditChange();
+	if (Selection.IsValid())
+	{
+		FScopedTransaction LocalTransaction(LOCTEXT("OnStateSelectionChanged", "Change State"));
+		CustomInstance->Modify();
+		CustomInstance->SetCurrentState(*Selection);
+		UpdateInstance();
 
-	// Non-continuous change: collect garbage.
-	GEngine->ForceGarbageCollection();
+		// Non-continuous change: collect garbage.
+		GEngine->ForceGarbageCollection();
+	}
 }
 
 
@@ -465,15 +479,19 @@ FReply FCustomizableInstanceDetails::RemoveParameterProfile()
 		return FReply::Handled();
 	}
 
+	//TODO(Max):UE-212345
+	//BeginTransaction(LOCTEXT("OnStateSelectionChanged", "Remove Profile"), true);
+
 	TArray<FProfileParameterDat>& Profiles = CustomizableObject->GetPrivate()->GetInstancePropertiesProfiles();
 
 	Profiles.RemoveAt(ProfileIdx);
 	CustomInstance->GetPrivate()->SelectedProfileIndex = INDEX_NONE;
 	CustomizableObject->Modify();
 
-	CustomInstance->PreEditChange(nullptr);
-	CustomInstance->UpdateSkeletalMeshAsync(true, true);
-	CustomInstance->PostEditChange();
+	UpdateInstance();
+
+	//TODO(Max):UE-212345
+	//EndTransaction();
 
 	// Non-continuous change: collect garbage.
 	GEngine->ForceGarbageCollection();
@@ -484,37 +502,44 @@ FReply FCustomizableInstanceDetails::RemoveParameterProfile()
 
 void FCustomizableInstanceDetails::OnProfileSelectedChanged(TSharedPtr<FString> Selection, ESelectInfo::Type SelectInfo)
 {
-	const int32 ProfileIdx = CustomInstance->GetPrivate()->SelectedProfileIndex;
-	if (CustomInstance->GetPrivate()->IsSelectedParameterProfileDirty())
+	if (Selection.IsValid())
 	{
-		CustomInstance->GetPrivate()->SaveParametersToProfile(ProfileIdx);
-	}
+		const int32 ProfileIdx = CustomInstance->GetPrivate()->SelectedProfileIndex;
+		if (CustomInstance->GetPrivate()->IsSelectedParameterProfileDirty())
+		{
+			CustomInstance->GetPrivate()->SaveParametersToProfile(ProfileIdx);
+		}
 
-	if (*Selection == "None")
-	{
-		CustomInstance->GetPrivate()->SelectedProfileIndex = INDEX_NONE;
-	}
-	else
-	{
-		//Set selected profile
-		TArray<FProfileParameterDat>& Profiles = CustomInstance->GetCustomizableObject()->GetPrivate()->GetInstancePropertiesProfiles();
-		for (int32 Idx = 0; Idx < Profiles.Num(); ++Idx)
+		//TODO(Max):UE-212345
+		//BeginTransaction(LOCTEXT("OnStateSelectionChanged", "Select Profile"), true);
+
+		if (*Selection == "None")
+		{
+			CustomInstance->GetPrivate()->SelectedProfileIndex = INDEX_NONE;
+		}
+		else
 		{
-			if (Profiles[Idx].ProfileName == *Selection)
+			//Set selected profile
+			TArray<FProfileParameterDat>& Profiles = CustomInstance->GetCustomizableObject()->GetPrivate()->GetInstancePropertiesProfiles();
+			for (int32 Idx = 0; Idx < Profiles.Num(); ++Idx)
 			{
-				CustomInstance->GetPrivate()->SelectedProfileIndex = Idx;
-				break;
+				if (Profiles[Idx].ProfileName == *Selection)
+				{
+					CustomInstance->GetPrivate()->SelectedProfileIndex = Idx;
+					break;
+				}
 			}
 		}
-	}	
 
-	CustomInstance->GetPrivate()->LoadParametersFromProfile(CustomInstance->GetPrivate()->SelectedProfileIndex);
-	CustomInstance->PreEditChange(nullptr);
-	CustomInstance->UpdateSkeletalMeshAsync(true, true);
-	CustomInstance->PostEditChange();
+		CustomInstance->GetPrivate()->LoadParametersFromProfile(CustomInstance->GetPrivate()->SelectedProfileIndex);
+		UpdateInstance();
 
-	// Non-continuous change: collect garbage.
-	GEngine->ForceGarbageCollection();
+		//TODO(Max):UE-212345
+		//EndTransaction();
+
+		// Non-continuous change: collect garbage.
+		GEngine->ForceGarbageCollection();
+	}
 }
 
 
@@ -972,41 +997,16 @@ TSharedRef<SWidget> FCustomizableInstanceDetails::OnGenerateWidgetIntParameter(T
 
 void FCustomizableInstanceDetails::OnIntParameterComboBoxChanged(TSharedPtr<FString> Selection, ESelectInfo::Type SelectInfo, FString ParamName)
 {
-	TArray<FCustomizableObjectIntParameterValue>& IntParameters = CustomInstance->GetIntParameters();
-	const UCustomizableObject* CustomizableObject = CustomInstance->GetCustomizableObject();
-
-	for (int32 i = 0; i < IntParameters.Num(); ++i)
+	if (Selection.IsValid())
 	{
-		if (IntParameters[i].ParameterName == ParamName)
-		{
-			const int32 ParamIndexInObject = CustomizableObject->FindParameter(ParamName);
-			if (ParamIndexInObject != INDEX_NONE)
-			{
-				for (int32 v = 0; v < CustomizableObject->GetIntParameterNumOptions(ParamIndexInObject); ++v)
-				{
-					FString ValueName = CustomizableObject->GetIntParameterAvailableOption(ParamIndexInObject, v);
+		FScopedTransaction LocalTransaction(FText::Format(LOCTEXT("OnIntParameterSet", "Set Int Parameter: {0}"), FText::FromString(ParamName)));
+		CustomInstance->Modify();
+		CustomInstance->SetIntParameterSelectedOption(ParamName, *Selection);
+		UpdateInstance();
 
-					if (ValueName == *Selection)
-					{
-						IntParameters[i].ParameterValueName = ValueName;
-						break;
-					}
-				}
-			}
-			else
-			{
-				UE_LOG(LogMutable, Error, TEXT("Failed to find parameter."));
-			}
-		}
+		// Non-continuous change: collect garbage.
+		GEngine->ForceGarbageCollection();
 	}
-
-	CustomInstance->GetPrivate()->SetSelectedParameterProfileDirty();
-	CustomInstance->PreEditChange(nullptr);
-	CustomInstance->UpdateSkeletalMeshAsync(true, true);
-	CustomInstance->PostEditChange();
-
-	// Non-continuous change: collect garbage.
-	GEngine->ForceGarbageCollection();
 }
 
 
@@ -1059,6 +1059,7 @@ TSharedRef<SWidget> FCustomizableInstanceDetails::GenerateFloatWidget(const int3
 		.MinValue(UIMetadata.MinimumValue)
 		.MaxValue(UIMetadata.MaximumValue)
 		.OnValueChanged(this, &FCustomizableInstanceDetails::OnFloatParameterChanged, ParamName, -1)
+		.OnValueCommitted(this, &FCustomizableInstanceDetails::OnFloatParameterCommited, ParamName, -1)
 		.OnBeginSliderMovement(this, &FCustomizableInstanceDetails::OnFloatParameterSliderBegin)
 		.OnEndSliderMovement(this, &FCustomizableInstanceDetails::OnFloatParameterSpinBoxEnd, ParamName, -1);
 }
@@ -1085,24 +1086,27 @@ float FCustomizableInstanceDetails::GetFloatParameterValue(FString ParamName, in
 
 void FCustomizableInstanceDetails::OnFloatParameterChanged(float Value, FString ParamName, int32 RangeIndex)
 {
-	CustomInstance->PreEditChange(nullptr);
-	CustomInstance->GetPrivate()->SetSelectedParameterProfileDirty();
+	float OldValue = CustomInstance->GetFloatParameterSelectedOption(ParamName, RangeIndex);
 
-	CustomInstance->SetFloatParameterSelectedOption(ParamName, Value, RangeIndex);
-	CustomInstance->UpdateSkeletalMeshAsync(true, true);
-	
-	CustomInstance->PostEditChange();
+	if (OldValue != Value)
+	{
+		//No transaction is needed here as this is called when the transaction has already started
+		CustomInstance->SetFloatParameterSelectedOption(ParamName, Value, RangeIndex);
+		UpdateInstance();
+	}
 }
 
 
 void FCustomizableInstanceDetails::OnFloatParameterSliderBegin()
 {
+	BeginTransaction(LOCTEXT("OnFloatParameterCommited", "Set Float Slider"));
 	bUpdatingSlider = true;
 }
 
 
 void FCustomizableInstanceDetails::OnFloatParameterSliderEnd()
 {
+	EndTransaction();
 	bUpdatingSlider = false;
 }
 
@@ -1111,17 +1115,36 @@ void FCustomizableInstanceDetails::OnFloatParameterSpinBoxEnd(float Value, FStri
 {
 	bUpdatingSlider = false;
 
-	CustomInstance->PreEditChange(nullptr);
-	CustomInstance->GetPrivate()->SetSelectedParameterProfileDirty();
 	CustomInstance->SetFloatParameterSelectedOption(ParamName, Value, RangeIndex);
-	CustomInstance->UpdateSkeletalMeshAsync(true, true);
-	CustomInstance->PostEditChange();
-	
+	UpdateInstance();
+	EndTransaction();
+
 	// Non-continuous change: collect garbage.
 	GEngine->ForceGarbageCollection();
 }
 
 
+void FCustomizableInstanceDetails::OnFloatParameterCommited(float Value, ETextCommit::Type Type, const FString ParamName, int32 RangeIndex)
+{
+	if (Type == ETextCommit::OnEnter)
+	{
+		// Making sure that setting a float by text generates a transaction. OnSpinBoxSliderEnd is considered a floatcommit of type OnEnter (IDK why...)
+		// and already generates a transaction when the Slid begins.
+		if (!Transaction)
+		{
+			BeginTransaction(FText::Format(LOCTEXT("OnFloatParameterCommited", "Set Float Parameter: {0}"), FText::FromString(ParamName)));
+		}
+
+		CustomInstance->SetFloatParameterSelectedOption(ParamName, Value, RangeIndex);
+		UpdateInstance();
+		EndTransaction();
+
+		// Non-continuous change: collect garbage.
+		GEngine->ForceGarbageCollection();
+	}
+}
+
+
 // TEXTURE PARAMETERS -----------------------------------------------------------------------------------------------------------------
 
 TSharedRef<SWidget> FCustomizableInstanceDetails::GenerateTextureWidget(const int32 ParamIndexInObject)
@@ -1207,30 +1230,16 @@ void FCustomizableInstanceDetails::GenerateTextureParameterOptions()
 
 void FCustomizableInstanceDetails::OnTextureParameterComboBoxSelectionChanged(TSharedPtr<FString> Selection, ESelectInfo::Type SelectInfo, FString ParamName)
 {
-	CustomInstance->PreEditChange(nullptr);
-
-	auto& TextureParameters = CustomInstance->GetTextureParameters();
-
-	for (int i = 0; i < TextureParameters.Num(); ++i)
+	if (Selection.IsValid())
 	{
-		if (TextureParameters[i].ParameterName == ParamName)
-		{
-			for (int o = 0; o < TextureParameterValueNames.Num(); ++o)
-			{
-				if (Selection.IsValid() && *TextureParameterValueNames[o] == *Selection)
-				{
-					TextureParameters[i].ParameterValue = TextureParameterValues[o];
-				}
-			}
-		}
-	}
-
-	CustomInstance->GetPrivate()->SetSelectedParameterProfileDirty();
-	CustomInstance->UpdateSkeletalMeshAsync(true, true);
-	CustomInstance->PostEditChange();
+		FScopedTransaction LocalTransaction(FText::Format(LOCTEXT("OnTextureParameterSelected", "Set Texture Parameter: {0}"), FText::FromString(ParamName)));
+		CustomInstance->Modify();
+		CustomInstance->SetTextureParameterSelectedOption(ParamName, *Selection);
+		UpdateInstance();
 
-	// Non-continuous change: collect garbage.
-	GEngine->ForceGarbageCollection();
+		// Non-continuous change: collect garbage.
+		GEngine->ForceGarbageCollection();
+	}
 }
 
 
@@ -1278,16 +1287,12 @@ FReply FCustomizableInstanceDetails::OnColorBlockMouseButtonDown(const FGeometry
 }
 
 
-void FCustomizableInstanceDetails::OnSetColorFromColorPicker(FLinearColor NewColor, FString PickerParamName) const
+void FCustomizableInstanceDetails::OnSetColorFromColorPicker(FLinearColor NewColor, FString PickerParamName)
 {
-	CustomInstance->PreEditChange(nullptr);
+	FScopedTransaction LocalTransaction(FText::Format(LOCTEXT("SetColorParameter", "Set Color Parameter: {0}"), FText::FromString(PickerParamName)));
+	CustomInstance->Modify();
 	CustomInstance->SetColorParameterSelectedOption(PickerParamName, NewColor);
-	CustomInstance->GetPrivate()->SetSelectedParameterProfileDirty();
-
-	PickerParamName = FString();
-
-	CustomInstance->UpdateSkeletalMeshAsync(true, true);
-	CustomInstance->PostEditChange();
+	UpdateInstance();
 }
 
 
@@ -1426,15 +1431,14 @@ FReply FCustomizableInstanceDetails::OnProjectorCopyTransform(const FString Para
 
 FReply FCustomizableInstanceDetails::OnProjectorPasteTransform(const FString ParamName, const int32 RangeIndex)
 {
-	FScopedTransaction Transaction(LOCTEXT("PasteTransform", "Paste Transform"));
-	CustomInstance->Modify();
-
 	FString ClipboardText;
 	FPlatformApplicationMisc::ClipboardPaste(ClipboardText);
 	FCustomizableObjectProjector DefaultValue;
 	UScriptStruct* Struct = DefaultValue.StaticStruct();
 	Struct->ImportText(*ClipboardText, &DefaultValue, nullptr, 0, GLog, GetPathNameSafe(Struct));
 
+	FScopedTransaction LocalTransaction(LOCTEXT("PasteTransform", "Paste Projector Transform"));
+	CustomInstance->Modify();
 	CustomInstance->SetProjectorValue(ParamName,
 		static_cast<FVector>(DefaultValue.Position),
 		static_cast<FVector>(DefaultValue.Direction),
@@ -1446,7 +1450,7 @@ FReply FCustomizableInstanceDetails::OnProjectorPasteTransform(const FString Par
 	const TSharedPtr<ICustomizableObjectInstanceEditor> Editor = GetEditorChecked();
 
 	Editor->ShowGizmoProjectorParameter(ParamName, RangeIndex);
-	CustomInstance->UpdateSkeletalMeshAsync(true, true);
+	UpdateInstance();
 
 	return FReply::Handled();
 }
@@ -1454,11 +1458,10 @@ FReply FCustomizableInstanceDetails::OnProjectorPasteTransform(const FString Par
 
 FReply FCustomizableInstanceDetails::OnProjectorResetTransform(const FString ParamName, const int32 RangeIndex)
 {
-	FScopedTransaction Transaction(LOCTEXT("ResetTransform", "Reset Transform"));
-	CustomInstance->Modify();
-
 	const FCustomizableObjectProjector DefaultValue = CustomInstance->GetCustomizableObject()->GetProjectorParameterDefaultValue(ParamName);
 
+	FScopedTransaction LocalTransaction(LOCTEXT("ResetTransform", "Reset Projector Transform"));
+	CustomInstance->Modify();
 	CustomInstance->SetProjectorValue(ParamName,
 		static_cast<FVector>(DefaultValue.Position),
 		static_cast<FVector>(DefaultValue.Direction),
@@ -1470,7 +1473,7 @@ FReply FCustomizableInstanceDetails::OnProjectorResetTransform(const FString Par
 	const TSharedPtr<ICustomizableObjectInstanceEditor> Editor = GetEditorChecked();
 
 	Editor->ShowGizmoProjectorParameter(ParamName, RangeIndex);
-	CustomInstance->UpdateSkeletalMeshAsync(true, true);
+	UpdateInstance();
 
 	return FReply::Handled();
 }
@@ -1734,6 +1737,7 @@ TSharedRef<SWidget> FCustomizableInstanceDetails::GenerateMultidimensionalProjec
 				.MaxValue(1.0f)
 				.Value(this, &FCustomizableInstanceDetails::GetFloatParameterValue, OpacitySliderParamName, RangeIndex)
 				.OnValueChanged(this, &FCustomizableInstanceDetails::OnFloatParameterChanged, OpacitySliderParamName, RangeIndex)
+				.OnValueCommitted(this, &FCustomizableInstanceDetails::OnFloatParameterCommited, OpacitySliderParamName, RangeIndex)
 				.OnBeginSliderMovement(this, &FCustomizableInstanceDetails::OnFloatParameterSliderBegin)
 				.OnEndSliderMovement(this, &FCustomizableInstanceDetails::OnFloatParameterSpinBoxEnd, OpacitySliderParamName, RangeIndex)
 				.Font(LayoutBuilder.Pin()->GetDetailFont())
@@ -1804,8 +1808,11 @@ TSharedRef<SWidget> FCustomizableInstanceDetails::OnGenerateWidgetProjectorParam
 }
 
 
-FReply FCustomizableInstanceDetails::OnProjectorLayerAdded(FString ParamName) const
+FReply FCustomizableInstanceDetails::OnProjectorLayerAdded(FString ParamName)
 {
+	FScopedTransaction LocalTransaction(LOCTEXT("AddProjectorLayer", "Add Projector Layer"));
+	CustomInstance->Modify();
+
 	const int32 NumLayers = CustomInstance->AddValueToProjectorRange(ParamName) + 1;
 	if (NumLayers == 0)
 	{
@@ -1828,10 +1835,7 @@ FReply FCustomizableInstanceDetails::OnProjectorLayerAdded(FString ParamName) co
 	check(CustomInstance->FindFloatParameterNameIndex(OpacitySliderParamName) != INDEX_NONE);
 	check(NumLayers == FloatParameters[CustomInstance->FindFloatParameterNameIndex(OpacitySliderParamName)].ParameterRangeValues.Num());
 
-	CustomInstance->GetPrivate()->SetSelectedParameterProfileDirty();
-	CustomInstance->PreEditChange(nullptr);
-	CustomInstance->UpdateSkeletalMeshAsync(true, true);
-	CustomInstance->PostEditChange();
+	UpdateInstance();
 
 	FCoreUObjectDelegates::BroadcastOnObjectModified(CustomInstance.Get());
 
@@ -1839,8 +1843,17 @@ FReply FCustomizableInstanceDetails::OnProjectorLayerAdded(FString ParamName) co
 }
 
 
-FReply FCustomizableInstanceDetails::OnProjectorLayerRemoved(const FString ParamName, const int32 RangeIndex) const
+FReply FCustomizableInstanceDetails::OnProjectorLayerRemoved(const FString ParamName, const int32 RangeIndex)
 {
+	FScopedTransaction LocalTransaction(LOCTEXT("RemoveProjectorLayer", "Remove Projector Layer"));
+	CustomInstance->Modify();
+
+	// Unselect projector if it's the deleted one
+	if (GetEditorChecked()->GetProjectorParameter()->IsProjectorSelected(ParamName, RangeIndex))
+	{
+		GetEditorChecked()->HideGizmo();
+	}
+
 	const int32 projectorParameterIndex = CustomInstance->FindProjectorParameterNameIndex(ParamName);
 	if (projectorParameterIndex == INDEX_NONE
 		|| CustomInstance->GetProjectorParameters()[projectorParameterIndex].RangeValues.Num() <= 0)
@@ -1866,10 +1879,7 @@ FReply FCustomizableInstanceDetails::OnProjectorLayerRemoved(const FString Param
 	check(CustomInstance->FindFloatParameterNameIndex(OpacitySliderParamName) != INDEX_NONE);
 	check(NumLayers == FloatParameters[CustomInstance->FindFloatParameterNameIndex(OpacitySliderParamName)].ParameterRangeValues.Num());
 
-	CustomInstance->GetPrivate()->SetSelectedParameterProfileDirty();
-	CustomInstance->PreEditChange(nullptr);
-	CustomInstance->UpdateSkeletalMeshAsync(true, true);
-	CustomInstance->PostEditChange();
+	UpdateInstance();
 
 	FCoreUObjectDelegates::BroadcastOnObjectModified(CustomInstance.Get());
 
@@ -1877,49 +1887,18 @@ FReply FCustomizableInstanceDetails::OnProjectorLayerRemoved(const FString Param
 }
 
 
-void FCustomizableInstanceDetails::OnProjectorTextureParameterComboBoxChanged(TSharedPtr<FString> Selection, ESelectInfo::Type SelectInfo, FString ParamName, int32 RangeIndex) const
+void FCustomizableInstanceDetails::OnProjectorTextureParameterComboBoxChanged(TSharedPtr<FString> Selection, ESelectInfo::Type SelectInfo, FString ParamName, int32 RangeIndex)
 {
-	TArray<FCustomizableObjectIntParameterValue>& IntParameters = CustomInstance->GetIntParameters();
-
-	const UCustomizableObject* CustomObject = CustomInstance->GetCustomizableObject();
-	for (int32 i = 0; i < IntParameters.Num(); ++i)
+	if (Selection.IsValid())
 	{
-		if (IntParameters[i].ParameterName == ParamName)
-		{
-			const int32 ParamIndexInObject = CustomObject->FindParameter(ParamName);
-			if (ParamIndexInObject != INDEX_NONE)
-			{
-				for (int32 v = 0; v < CustomObject->GetIntParameterNumOptions(ParamIndexInObject); ++v)
-				{
-					FString ValueName = CustomObject->GetIntParameterAvailableOption(ParamIndexInObject, v);
-					if (ValueName == *Selection)
-					{
-						if (IntParameters[i].ParameterRangeValueNames.IsValidIndex(RangeIndex))
-						{
-							IntParameters[i].ParameterRangeValueNames[RangeIndex] = ValueName;
-						}
-						else
-						{
-							IntParameters[i].ParameterValueName = ValueName;
-						}
-
-						break;
-					}
-				}
-			}
-			else
-			{
-				UE_LOG(LogMutable, Error, TEXT("Failed to find parameter."));
-			}
-		}
-	}
-
-	CustomInstance->GetPrivate()->SetSelectedParameterProfileDirty();
-	CustomInstance->UpdateSkeletalMeshAsync(true, true);
-	CustomInstance->PostEditChange();
+		FScopedTransaction LocalTransaction(LOCTEXT("SelectProjectorImage", "Select Projecto Image"));
+		CustomInstance->Modify();
+		CustomInstance->SetIntParameterSelectedOption(ParamName, *Selection, RangeIndex);
+		UpdateInstance();
 
-	// Non-continuous change: collect garbage.
-	GEngine->ForceGarbageCollection();
+		// Non-continuous change: collect garbage.
+		GEngine->ForceGarbageCollection();
+	}
 }
 
 
@@ -1966,11 +1945,10 @@ ECheckBoxState FCustomizableInstanceDetails::GetBoolParameterValue(FString Param
 
 void FCustomizableInstanceDetails::OnBoolParameterChanged(ECheckBoxState InCheckboxState, FString ParamName)
 {
-	CustomInstance->PreEditChange(nullptr);
+	FScopedTransaction LocalTransaction(FText::Format(LOCTEXT("SetParameterBool", "Set Bool Parameter: {0}"), FText::FromString(ParamName)));
+	CustomInstance->Modify();
 	CustomInstance->SetBoolParameterSelectedOption(ParamName, InCheckboxState == ECheckBoxState::Checked);
-	CustomInstance->GetPrivate()->SetSelectedParameterProfileDirty();
-	CustomInstance->UpdateSkeletalMeshAsync(true, true);
-	CustomInstance->PostEditChange();
+	UpdateInstance();
 
 	// Non-continuous change: collect garbage.
 	GEngine->ForceGarbageCollection();
@@ -1990,8 +1968,6 @@ FReply FCustomizableInstanceDetails::OnCopyAllParameters()
 
 FReply FCustomizableInstanceDetails::OnPasteAllParameters()
 {
-	CustomInstance->Modify();
-
 	FString ClipText;
 	FPlatformApplicationMisc::ClipboardPaste(ClipText);
 	
@@ -2001,17 +1977,17 @@ FReply FCustomizableInstanceDetails::OnPasteAllParameters()
 	const int32 MinLOD = Descriptor.GetMinLod();
 	const TArray<uint16> RequestedLODLevels = Descriptor.GetRequestedLODLevels();
 	
+	FScopedTransaction LocalTransaction(LOCTEXT("OnPasteAllParameters", "Paste All Parameters"));
+	CustomInstance->Modify();
+
 	if (Struct->ImportText(*ClipText, &Descriptor, nullptr, 0, GLog, GetPathNameSafe(Struct)))
 	{
 		// Keep current LOD
 		Descriptor.SetMinLod(MinLOD);
 		Descriptor.SetRequestedLODLevels(RequestedLODLevels);
 		
-		CustomInstance->UpdateSkeletalMeshAsync(true, true);		
+		UpdateInstance();
 	}
-
-	CustomInstance->GetPrivate()->SetSelectedParameterProfileDirty();
-
 	
 	return FReply::Handled();
 }
@@ -2019,9 +1995,8 @@ FReply FCustomizableInstanceDetails::OnPasteAllParameters()
 
 FReply FCustomizableInstanceDetails::OnResetAllParameters()
 {
-	CustomInstance->PreEditChange(nullptr);
-
-	TArray<FCustomizableObjectIntParameterValue>& IntParameters = CustomInstance->GetIntParameters();
+	FScopedTransaction LocalTransaction(LOCTEXT("OnResetAllParameters", "Reset All Parameters"));
+	CustomInstance->Modify();
 
 	const UCustomizableObject* CustomObject = CustomInstance->GetCustomizableObject();
 	const int32 NumObjectParameter = CustomObject->GetParameterCount();
@@ -2031,9 +2006,8 @@ FReply FCustomizableInstanceDetails::OnResetAllParameters()
 		SetParameterValueToDefault(ParameterIndex);
 	}
 
-	CustomInstance->UpdateSkeletalMeshAsync(true, true);
-	CustomInstance->PostEditChange();
 	CustomInstance->GetPrivate()->SelectedProfileIndex = INDEX_NONE;
+	UpdateInstance();
 
 	// Non-continuous change: collect garbage.
 	GEngine->ForceGarbageCollection();
@@ -2044,11 +2018,12 @@ FReply FCustomizableInstanceDetails::OnResetAllParameters()
 
 void FCustomizableInstanceDetails::OnResetParameterButtonClicked(int32 ParameterIndex)
 {	
-	CustomInstance->PreEditChange(nullptr);
+	FString ParameterName = CustomInstance->GetCustomizableObject()->GetParameterName(ParameterIndex);
+
+	FScopedTransaction LocalTransaction(FText::Format(LOCTEXT("OnResetParameter", "Reset Parameter: {0}"), FText::FromString(ParameterName)));
+	CustomInstance->Modify();
 	SetParameterValueToDefault(ParameterIndex);
-	CustomInstance->GetPrivate()->SetSelectedParameterProfileDirty();
-	CustomInstance->UpdateSkeletalMeshAsync(true, true);
-	CustomInstance->PostEditChange();
+	UpdateInstance();
 }
 
 
@@ -2165,6 +2140,47 @@ void FCustomizableInstanceDetails::SetParameterValueToDefault(int32 ParameterInd
 }
 
 
+// TRANSACTION SYSTEM -----------------------------------------------------------------------------------------------------------------
+
+void FCustomizableInstanceDetails::BeginTransaction(const FText& TransactionDesc, bool bModifyCustomizableObject)
+{
+	// We only allow the BeginTransaction to be called with the EndTransaction pair. We should never call a second transaction before the first was ended.
+	check(!Transaction);
+
+	Transaction.Reset(new FScopedTransaction(TransactionDesc));
+	CustomInstance->Modify();
+
+	//TODO(Max):UE-212345
+	/*UCustomizableObject* CustomizableObject = CustomInstance->GetCustomizableObject();
+	if (CustomizableObject && bModifyCustomizableObject)
+	{
+		CustomizableObject->Modify();
+	}*/
+}
+
+
+void FCustomizableInstanceDetails::EndTransaction()
+{
+	Transaction.Reset();
+
+	//TODO(Max):UE-212345
+	CustomInstance->GetPrivate()->SetSelectedParameterProfileDirty();
+}
+
+
+void FCustomizableInstanceDetails::OnInstanceTransacted(const FTransactionObjectEvent& TransactionEvent)
+{
+	if (TransactionEvent.GetEventType() == ETransactionObjectEventType::UndoRedo)
+	{
+		// Update instance on undo/redo
+		UpdateInstance();
+
+		// We want to make sure that the gizom is hidden when we do an undo/redo transaction
+		GetEditorChecked()->HideGizmo();
+	}
+}
+
+
 // PROFILES WINDOW -----------------------------------------------------------------------------------------------------------------
 
 void SProfileParametersWindow::Construct(const FArguments& InArgs)
@@ -2256,7 +2272,12 @@ void SProfileParametersWindow::OnNameChange(const FText& NewName, ETextCommit::T
 		
 		RequestDestroyWindow();
 
+		//TODO(Max):UE-212345
+		//const FScopedTransaction Transaction(LOCTEXT("OnEnterAddProfile", "Add Profile"));
+		
 		UCustomizableObject* CustomizableObject = CustomInstance->GetCustomizableObject();
+		//CustomizableObject->Modify();
+
 		CustomizableObject->GetPrivate()->AddNewParameterProfile(GetFileName(), *CustomInstance.Get());
 
 		if (CustomInstance->GetPrivate()->bSelectedProfileDirty && CustomInstance->GetPrivate()->SelectedProfileIndex != INDEX_NONE)
@@ -2270,6 +2291,10 @@ void SProfileParametersWindow::OnNameChange(const FText& NewName, ETextCommit::T
 			InstanceDetails->Refresh();
 		}
 	}
+	else
+	{
+		FileName = NewName;
+	}
 }
 
 
@@ -2281,6 +2306,9 @@ FReply SProfileParametersWindow::OnButtonClick(EAppReturnType::Type ButtonID)
 
 		RequestDestroyWindow();
 
+		//TODO(Max):UE-212345
+		//const FScopedTransaction Transaction(LOCTEXT("OnOkAddProfile", "Add Profile"));
+		
 		UCustomizableObject* CustomizableObject = CustomInstance->GetCustomizableObject();
 		CustomizableObject->GetPrivate()->AddNewParameterProfile(GetFileName(), *CustomInstance.Get());
 
diff --git a/Engine/Plugins/Experimental/Mutable/Source/CustomizableObjectEditor/Private/MuCOE/CustomizableInstanceDetails.h b/Engine/Plugins/Experimental/Mutable/Source/CustomizableObjectEditor/Private/MuCOE/CustomizableInstanceDetails.h
--- a/Engine/Plugins/Experimental/Mutable/Source/CustomizableObjectEditor/Private/MuCOE/CustomizableInstanceDetails.h
+++ b/Engine/Plugins/Experimental/Mutable/Source/CustomizableObjectEditor/Private/MuCOE/CustomizableInstanceDetails.h
@@ -8,6 +8,8 @@
 
 class FDetailWidgetRow;
 class FReply;
+class FScopedTransaction;
+class FTransactionObjectEvent;
 class ICustomizableObjectInstanceEditor;
 class IDetailCategoryBuilder;
 class IDetailGroup;
@@ -17,9 +19,11 @@ class UCustomizableObjectInstance;
 
 enum class ECheckBoxState : uint8;
 namespace ESelectInfo { enum Type : int; }
+namespace ETextCommit { enum Type : int; }
 
 struct FGeometry;
 struct FPointerEvent;
+struct FUpdateContext;
 
 class FCustomizableInstanceDetails : public IDetailCustomization
 {
@@ -40,6 +44,10 @@ public:
 
 private:
 
+	// Updat the Instance if a parameter has been modified. Also creates a Delegate to refresh the UI if the 
+	// Instance has been updated successfully.
+	void UpdateInstance();
+
 	// Callback to regenerate the details when the instance has finished an update
 	void InstanceUpdated(UCustomizableObjectInstance* Instance) const;
 
@@ -86,6 +94,7 @@ private:
 	void OnFloatParameterSliderBegin();
 	void OnFloatParameterSpinBoxEnd(float Value, const FString ParamName, int32 RangeIndex);
 	void OnFloatParameterSliderEnd();
+	void OnFloatParameterCommited(float Value, ETextCommit::Type Type, const FString ParamName, int32 RangeIndex);	// Needed to have undo/redo
  
 	// Texture Parameters Functions
 	TSharedRef<SWidget> GenerateTextureWidget(const int32 ParamIndexInObject);
@@ -96,7 +105,7 @@ private:
 	TSharedRef<SWidget> GenerateColorWidget(const int32 ParamIndexInObject);
 	FLinearColor GetColorParameterValue(const FString ParamName) const;
 	FReply OnColorBlockMouseButtonDown(const FGeometry& MyGeometry, const FPointerEvent& MouseEvent, const FString ParamName);
-	void OnSetColorFromColorPicker(FLinearColor NewColor, const FString PickerParamName) const;
+	void OnSetColorFromColorPicker(FLinearColor NewColor, const FString PickerParamName);
 
 	// Bool Parameters Functions
 	TSharedRef<SWidget> GenerateBoolWidget(const int32 ParamIndexInObject);
@@ -111,9 +120,9 @@ private:
 	FReply OnProjectorCopyTransform(const FString ParamName, const int32 RangeIndex) const;
 	FReply OnProjectorPasteTransform(const FString ParamName, const int32 RangeIndex);
 	FReply OnProjectorResetTransform(const FString ParamName, const int32 RangeIndex);
-	FReply OnProjectorLayerAdded(const FString ParamName) const;
-	FReply OnProjectorLayerRemoved(const FString ParamName, const int32 RangeIndex) const;
-	void OnProjectorTextureParameterComboBoxChanged(TSharedPtr<FString> Selection, ESelectInfo::Type SelectInfo, const FString ParamName, int32 RangeIndex) const;
+	FReply OnProjectorLayerAdded(const FString ParamName);
+	FReply OnProjectorLayerRemoved(const FString ParamName, const int32 RangeIndex);
+	void OnProjectorTextureParameterComboBoxChanged(TSharedPtr<FString> Selection, ESelectInfo::Type SelectInfo, const FString ParamName, int32 RangeIndex);
 	TSharedRef<SWidget> MakeTextureComboEntryWidget(TSharedPtr<FString> InItem) const;
 	TSharedRef<SWidget> OnGenerateWidgetProjectorParameter(TSharedPtr<FString> InItem) const;
 
@@ -124,6 +133,11 @@ private:
 	void OnResetParameterButtonClicked(int32 ParameterIndex);
 	void SetParameterValueToDefault(int32 ParameterIndex);
 
+	// Transaction System
+	void BeginTransaction(const FText& TransactionDesc, bool bModifyCustomizableObject = false);
+	void EndTransaction();
+	void OnInstanceTransacted(const FTransactionObjectEvent& TransactionEvent);
+
 private:
 
 	/** Pointer to the Customizable Object Instance */
@@ -172,5 +186,8 @@ private:
 
 	/** Array to store dynamic brushes. Neede because an image widget only stores a pointer to a Brush. */
 	TArray< TSharedPtr<class FDeferredCleanupSlateBrush> > DynamicBrushes;
+
+	// Unique transaction pointer to allow transactions that start and finish in different funtion scopes.
+	TUniquePtr<FScopedTransaction> Transaction;
 };
 
