[mutable] Reporting peak memory usage during the mutable compilation process using LLM
- The value reported is the one with the "Total" LLM tag.
- Moved to a discrete file the CustomizableObjectUpdater Utility class
- Created in a discrete file the class CustomizableObjectCompiler Utility class.
- Removed innecesary includes and forward declarations.
- The memory usage to be reported will be the starting LLM Total memory usage, the real Peak memory usage DURING the compilation and the peak memory usage (pre-compilation memory minus peak compilation memory to represent the memory usage OF the compilation).
- Changes in the .xml and .cs files to make this work will be pushed separately

diff --git a/Engine/Plugins/Experimental/Mutable/Source/MutableValidation/Private/MuV/COIBakingTestCommandlet.cpp b/Engine/Plugins/Experimental/Mutable/Source/MutableValidation/Private/MuV/COIBakingTestCommandlet.cpp
--- a/Engine/Plugins/Experimental/Mutable/Source/MutableValidation/Private/MuV/COIBakingTestCommandlet.cpp
+++ b/Engine/Plugins/Experimental/Mutable/Source/MutableValidation/Private/MuV/COIBakingTestCommandlet.cpp
@@ -2,6 +2,7 @@
 
 #include "../../Public/MuV/COIBakingTestCommandlet.h"
 
+#include "CustomizableObjectCompilationUtility.h"
 #include "ValidationUtils.h"
 #include "AssetRegistry/AssetData.h"
 #include "HAL/FileManager.h"
@@ -10,8 +11,6 @@
 #include "MuCO/CustomizableObjectSystem.h"
 #include "MuCOE/CustomizableObjectInstanceBakingUtils.h"
 #include "HAL/FileManager.h"
-#include "MuCO/CustomizableObjectPrivate.h"
-#include "MuT/UnrealPixelFormatOverride.h"
 #include "Interfaces/ITargetPlatformManagerModule.h"
 
 
@@ -88,12 +87,15 @@ int32 UCOIBakingTestCommandlet::Main(const FString& Params)
 	ITargetPlatformManagerModule& TPM = GetTargetPlatformManagerRef();
 	CompilationOptions.TargetPlatform = TPM.GetRunningTargetPlatform();
 	CompilationOptions.bUseDiskCompilation = false;
-	if (!CompileCustomizableObject(InstanceCustomizableObject,true, &CompilationOptions))
+
+	TSharedRef<FCustomizableObjectCompilationUtility> CompilationUtility = MakeShared<FCustomizableObjectCompilationUtility>();
+	if( !CompilationUtility->CompileCustomizableObject(InstanceCustomizableObject, true, &CompilationOptions ))
 	{
 		UE_LOG(LogMutable,Error,TEXT("Failed to compile the target CO. Exitting commandlet."));
 		return 1;
 	}
 	
+	
 	// Update the instance
 	{
 		// If this fail something is very wrong
diff --git a/Engine/Plugins/Experimental/Mutable/Source/MutableValidation/Private/MuV/COIBulkUpdateTestCommandlet.cpp b/Engine/Plugins/Experimental/Mutable/Source/MutableValidation/Private/MuV/COIBulkUpdateTestCommandlet.cpp
--- a/Engine/Plugins/Experimental/Mutable/Source/MutableValidation/Private/MuV/COIBulkUpdateTestCommandlet.cpp
+++ b/Engine/Plugins/Experimental/Mutable/Source/MutableValidation/Private/MuV/COIBulkUpdateTestCommandlet.cpp
@@ -2,6 +2,8 @@
 
 #include "MuV/COIBulkUpdateTestCommandlet.h"
 
+#include "CustomizableObjectCompilationUtility.h"
+#include "CustomizableObjectInstanceUpdateUtility.h"
 #include "ValidationUtils.h"
 #include "AssetRegistry/AssetData.h"
 #include "AssetRegistry/AssetRegistryModule.h"
@@ -119,14 +121,14 @@ int32 UCOIBulkUpdateTestCommandlet::Main(const FString& Params)
 	
 	// Make sure there is nothing else that the engine needs to do before starting our test
 	Wait(60);
-	
-	// Create the updater object so we can later call for the update a target instance
-	InstanceUpdater = NewObject<UCOIUpdater>();
-	
+
 	// Cache the target compilation platform so we can override the compilation configs of the target COs
 	ITargetPlatformManagerModule& TPM = GetTargetPlatformManagerRef();
 	const ITargetPlatform* TargetCompilationPlatform = TPM.GetRunningTargetPlatform();
 	
+	TSharedRef<FCustomizableObjectCompilationUtility> CompilationUtility = MakeShared<FCustomizableObjectCompilationUtility>();
+	TSharedRef<FCustomizableObjectInstanceUpdateUtility> InstanceUpdatingUtility = MakeShared<FCustomizableObjectInstanceUpdateUtility>();
+	
 	// Compile all found COs one by one
 	uint32 CurrentInstanceIndex = 1;
 	for (TTuple<UCustomizableObject*, TSet<UCustomizableObjectInstance*>>& MutableResourceTuple : MutableResources)
@@ -134,15 +136,16 @@ int32 UCOIBulkUpdateTestCommandlet::Main(const FString& Params)
 		UCustomizableObject* CustomizableObjectToCompile = MutableResourceTuple.Key;
 		check (CustomizableObjectToCompile);
 
+		const FString CustomizableObjectName = CustomizableObjectToCompile->GetName();
+		
 		// Set the compilation platform based on what the system is currently running on
 		FCompilationOptions CompilationOptions = CustomizableObjectToCompile->CompileOptions;
 		CompilationOptions.TargetPlatform = TargetCompilationPlatform;
-
 		CompilationOptions.bUseDiskCompilation = false;
 		
 		// Compile the current CO object
-		const FString CustomizableObjectName = CustomizableObjectToCompile->GetName();
-		if (!CompileCustomizableObject(CustomizableObjectToCompile, false, &CompilationOptions))	// Do not log mutable data since mongoDB will not be able to handle it correctly 
+
+		if (!CompilationUtility->CompileCustomizableObject(CustomizableObjectToCompile, false, &CompilationOptions))	// Do not log mutable data since mongoDB will not be able to handle it correctly 
 		{
 			UE_LOG(LogMutable,Error,TEXT("The CO %s could not be compiled succesfully. Skipping the update of all COIs that use it."), *CustomizableObjectName )
 			continue;
@@ -158,10 +161,7 @@ int32 UCOIBulkUpdateTestCommandlet::Main(const FString& Params)
 			UE_LOG(LogMutable,Display,TEXT("\t( %u / %u ) Processing instance : \"%s\" ."),CurrentInstanceIndex++, TotalAmountOfInstances ,*Instance->GetName());
 			
 			// Update each one of the instances and notify if the update failed in any manner
-			if (InstanceUpdater != nullptr)
-			{
-				InstanceUpdater->UpdateInstance(Instance);
-			}
+			InstanceUpdatingUtility->UpdateInstance(Instance);
 		}
 	}
 
diff --git a/Engine/Plugins/Experimental/Mutable/Source/MutableValidation/Private/MuV/CustomizableObjectCompilationUtility.cpp b/Engine/Plugins/Experimental/Mutable/Source/MutableValidation/Private/MuV/CustomizableObjectCompilationUtility.cpp
new file mode 100644
--- /dev/null
+++ b/Engine/Plugins/Experimental/Mutable/Source/MutableValidation/Private/MuV/CustomizableObjectCompilationUtility.cpp
@@ -0,0 +1,193 @@
+﻿// Copyright Epic Games, Inc. All Rights Reserved.
+
+#include "CustomizableObjectCompilationUtility.h"
+
+#include "Commandlets/Commandlet.h"
+#include "MuCO/CustomizableObject.h"
+#include "MuCO/CustomizableObjectPrivate.h"
+#include "MuCO/CustomizableObjectSystem.h"
+#include "Interfaces/ITargetPlatform.h"
+#include "MuCOE/CustomizableObjectEditorFunctionLibrary.h"
+
+
+
+bool FCustomizableObjectCompilationUtility::CompileCustomizableObject(UCustomizableObject* InCustomizableObject, const bool bShouldLogMutableLogs /* = true */, const FCompilationOptions* InCompilationOptionsOverride  /* nullptr */)
+{
+	check(InCustomizableObject);
+	
+	CustomizableObject = TStrongObjectPtr(InCustomizableObject);
+	
+	// Handle the overriding of the compilation configuration of the CO ------------------------------------------------------------------------------
+	
+	// Override the compilation options if an override has been provided by the user
+	FCompilationOptions CompilationOptions;
+	const bool bOverrideCompilationOptions = InCompilationOptionsOverride != nullptr;
+	if (bOverrideCompilationOptions)
+	{
+		CompilationOptions = *InCompilationOptionsOverride;
+		UE_LOG(LogMutable,Display,TEXT("CO Compilation options overriden by the user defined ones."));
+	}
+	else
+	{
+		CompilationOptions = InCustomizableObject->CompileOptions;
+		UE_LOG(LogMutable,Display,TEXT("Compiling CO using it's own compilation options."));
+	}
+
+	// Ensure that the user has provided a target compilation platform.
+	// Mutable is able to run without one but we want to be explicit in the context of testing.
+	if (!CompilationOptions.TargetPlatform)
+	{
+		UE_LOG(LogMutable, Error, TEXT("The compilation of the %s model could not be started : No explicit platform was provided."), *InCustomizableObject->GetName());
+		return false;
+	}
+
+	// Report the compilation configuration : We use this flag just to be sure we only generate logs when required since MongoDB does not want to find
+	// duplicated keys. The addition of duplicated entries in MongoDB is not available so this way we avoid having to handle them
+	// when we are sure we do not require it.
+	if (bShouldLogMutableLogs)
+	{
+		UE_LOG(LogMutable, Log, TEXT("(string) model_compile_options_overriden : %s "), bOverrideCompilationOptions ? TEXT("true") : TEXT("false"));
+		UE_LOG(LogMutable, Log, TEXT("(int) model_optimization_level : %d "), CompilationOptions.OptimizationLevel);
+		UE_LOG(LogMutable, Log, TEXT("(string) model_texture_compression : %s "), *UEnum::GetValueAsString(CompilationOptions.TextureCompression));
+		UE_LOG(LogMutable, Log, TEXT("(string) model_disk_compilation : %s "), CompilationOptions.bUseDiskCompilation ? TEXT("true") : TEXT("false"));
+		UE_LOG(LogMutable, Log, TEXT("(string) model_compile_platform_name : %s "), *CompilationOptions.TargetPlatform->PlatformName());
+	}
+
+	
+	// Run and wait for the compilation to be completed ----------------------------------------------------------------------------------------------
+	{
+		// Bind our OnCompilationFinished() method to the compilation completion callback to get notified once it is done
+		FPostCompileDelegate& PostCompilationDelegate = CustomizableObject->GetPostCompileDelegate();
+		PostCompilationDelegate.AddSP(this, &FCustomizableObjectCompilationUtility::OnCompilationFinished);
+		
+		// Create the compiler that will take care of the ASYNC compilation
+		UCustomizableObjectSystem* CustomizableObjectSystem = UCustomizableObjectSystem::GetInstanceChecked();
+		TUniquePtr<FCustomizableObjectCompilerBase> Compiler = TUniquePtr<FCustomizableObjectCompilerBase>(CustomizableObjectSystem->GetNewCompiler());
+
+		// Get the memory usage before staring the compilation
+#if ENABLE_LOW_LEVEL_MEM_TRACKER
+		FLowLevelMemTracker& LowLevelMemoryTracker = FLowLevelMemTracker::Get();
+
+		// Memory in use before starting the compilation itself
+		int64 CompilationStartTotalBytes = 0;
+		// Peak memory usage as reported by LLM during the compilation operation.
+		// Note : This value will not discriminate and therefore will report memory usage by the compilation of the CO as well as other operations that
+		// may be running during it's operation
+		int64 CompilationEndTotalRealPeakBytes = 0;
+		
+		if (FLowLevelMemTracker::IsEnabled())
+		{
+			UE_LOG(LogMutable,Display,TEXT("LLM system enabled: Peak memory usage during mutable CO compilation will be logged after said compilation takes place."));
+			
+			LowLevelMemoryTracker.UpdateStatsPerFrame();
+			CompilationStartTotalBytes = LowLevelMemoryTracker.GetTagAmountForTracker(ELLMTracker::Default, ELLMTag::Total, UE::LLM::ESizeParams::ReportCurrent);
+		}
+#endif
+		
+		// Get the memory being used by mutable before the compilation
+		const int64 CompilationStartMutableBytes = mu::FGlobalMemoryCounter::GetCounter();
+		mu::FGlobalMemoryCounter::Zero();
+		
+		UE_LOG(LogMutable,Display,TEXT("Compiling Customizable Object..."));
+		check(Compiler);
+		Compiler->Compile(*CustomizableObject, CompilationOptions, true);
+		
+		// Wait while the compilation takes place
+		const double CompilationStartSeconds = FPlatformTime::Seconds();
+		bIsTheCompilationInProcess = true;
+		while (bIsTheCompilationInProcess)
+		{
+			// Tick the engine
+			CommandletHelpers::TickEngine();
+
+			// todo: Will this be required in the future? It feels odd having to tick the compiler manually when we are already ticking the engine itself
+			Compiler->Tick(false);
+			
+			// Cache the peak value found during the compilation of the CO
+#if ENABLE_LOW_LEVEL_MEM_TRACKER
+			if (FLowLevelMemTracker::IsEnabled())
+			{
+				LowLevelMemoryTracker.UpdateStatsPerFrame();
+				const int64 CurrentTotalMemoryUsage = LowLevelMemoryTracker.GetTagAmountForTracker(ELLMTracker::Default, ELLMTag::Total, UE::LLM::ESizeParams::ReportCurrent);
+				CompilationEndTotalRealPeakBytes = CompilationEndTotalRealPeakBytes > CurrentTotalMemoryUsage ? CompilationEndTotalRealPeakBytes : CurrentTotalMemoryUsage;
+			}
+#endif
+		
+			// Stop if exit was requested
+			if (IsEngineExitRequested())
+			{
+				break;
+			}
+		}
+		
+		// Report the time we took to run the compilation
+		const double CompilationEndSeconds = FPlatformTime::Seconds() - CompilationStartSeconds;
+		UE_LOG(LogMutable, Display, TEXT("The compilation of the %s CO model took %f seconds."), *InCustomizableObject->GetName(), CompilationEndSeconds);
+		if (bShouldLogMutableLogs)
+		{
+			UE_LOG(LogMutable, Log, TEXT("(double) model_compile_time_ms : %f "), CompilationEndSeconds * 1000);
+
+			// Also report compilation end status
+			const ECustomizableObjectCompilationState CompilationEndResult = bWasCoCompilationSuccessful ? ECustomizableObjectCompilationState::Completed : ECustomizableObjectCompilationState::Failed;
+			UE_LOG(LogMutable, Log, TEXT("(string) model_compile_end_state : %s "), *UEnum::GetValueAsString(CompilationEndResult));
+		}
+		
+		// Get the peak mutable memory used during the compilation operation
+		const int64 CompilationEndPeakMutableBytes = mu::FGlobalMemoryCounter::GetPeak();
+		const int64 CompilationEndRealMutablePeakBytes = CompilationStartMutableBytes + CompilationEndPeakMutableBytes;
+		UE_LOG(LogMutable, Display, TEXT("Pre-Compilation Mutable memory usage : %lld"), CompilationStartMutableBytes);
+		UE_LOG(LogMutable, Display, TEXT("Compilation Mutable peak memory usage : %lld"), CompilationEndPeakMutableBytes);
+		UE_LOG(LogMutable, Display, TEXT("Compilation Mutable real peak memory usage : %lld"), CompilationEndRealMutablePeakBytes);
+		if (bShouldLogMutableLogs)
+		{
+			// Mutable system reported memory usage
+			UE_LOG(LogMutable, Log, TEXT("(int) model_compilation_start_bytes : %lld "), CompilationStartMutableBytes);
+			UE_LOG(LogMutable, Log, TEXT("(int) model_compilation_end_peak_bytes : %lld "), CompilationEndPeakMutableBytes);
+			UE_LOG(LogMutable, Log, TEXT("(int) model_compilation_end_real_peak_bytes : %lld "), CompilationEndRealMutablePeakBytes);
+		}
+
+		
+		// Now report the peak memory usage reported by the "total" tag from the LLM system
+#if ENABLE_LOW_LEVEL_MEM_TRACKER
+		if (FLowLevelMemTracker::IsEnabled())
+		{
+			UE_LOG(LogMutable, Display, TEXT("LLM Pre-Compilation Total memory usage (bytes) : %lld"), CompilationStartTotalBytes);
+			// CompilationPeakBytes : Peak memory recorded during the compilation subtracting from it the memory in use before the compilation itself.
+			// It should reflect the peak memory usage OF the compilation (as direct result of the compilation itself) but it should be used as an approximation
+			// since we can never be too sure no other system is doing work in parallel that is not part of the mutable compilation operation.
+			const int64 CompilationEndTotalPeakBytes = CompilationEndTotalRealPeakBytes - CompilationStartTotalBytes;
+			UE_LOG(LogMutable, Display, TEXT("LLM Total Compilation Peak memory usage (bytes) : %lld"), CompilationEndTotalPeakBytes);
+			UE_LOG(LogMutable, Display, TEXT("LLM Total Compilation Real Peak memory usage (bytes) : %lld"), CompilationEndTotalRealPeakBytes);
+			if (bShouldLogMutableLogs)
+			{
+				UE_LOG(LogMutable, Log, TEXT("(int) model_compilation_llm_total_start_bytes : %lld "), CompilationStartTotalBytes);
+				UE_LOG(LogMutable, Log, TEXT("(int) model_compilation_llm_total_end_peak_bytes : %lld "), CompilationEndTotalPeakBytes);
+				UE_LOG(LogMutable, Log, TEXT("(int) model_compilation_llm_total_end_real_peak_bytes : %lld "), CompilationEndTotalRealPeakBytes);
+			}
+		}
+#endif
+	}
+	
+	// Return the success state of the operation
+	return bWasCoCompilationSuccessful;
+}
+
+
+void FCustomizableObjectCompilationUtility::OnCompilationFinished()
+{
+	UE_LOG(LogMutable,Display,TEXT("Compilation of CO completed!"));
+
+	// Store the success end state for the operation.
+	bWasCoCompilationSuccessful =	CustomizableObject->GetPrivate()->CompilationResult == ECompilationResultPrivate::Success ||
+									CustomizableObject->GetPrivate()->CompilationResult == ECompilationResultPrivate::Warnings;
+	if (bWasCoCompilationSuccessful)
+	{
+		UE_LOG(LogMutable, Display, TEXT("The compilation of the %s model was succesfull."), *CustomizableObject->GetName());
+	}
+	else
+	{
+		UE_LOG(LogMutable, Error, TEXT("The compilation of the %s model failed."), *CustomizableObject->GetName());
+	}
+	
+	bIsTheCompilationInProcess = false;
+}
diff --git a/Engine/Plugins/Experimental/Mutable/Source/MutableValidation/Private/MuV/CustomizableObjectCompilationUtility.h b/Engine/Plugins/Experimental/Mutable/Source/MutableValidation/Private/MuV/CustomizableObjectCompilationUtility.h
new file mode 100644
--- /dev/null
+++ b/Engine/Plugins/Experimental/Mutable/Source/MutableValidation/Private/MuV/CustomizableObjectCompilationUtility.h
@@ -0,0 +1,40 @@
+﻿// Copyright Epic Games, Inc. All Rights Reserved.
+#pragma once
+
+#include "UObject/StrongObjectPtr.h"
+#include "Templates/SharedPointer.h"
+
+// Forward declarations
+class UCustomizableObject;
+struct FCompilationOptions;
+
+/**
+ * Helper class that allows the calling commandlet to treat the asynchronous compilation of a CO as a synchronous operation.
+ */
+class FCustomizableObjectCompilationUtility :  public TSharedFromThis<FCustomizableObjectCompilationUtility>
+{
+public:
+	
+	/**
+	 * Run the asyncronous compilation of the provided CO but as part of this sync method. 
+	 * @param InCustomizableObject The Customizable Object we want to synchronously compile
+	 * @param bShouldLogMutableLogs Enables or disables the logging of Log category logs relevant to the CO compilation. Required to avoid a MongoDB limitation with the duplication of MongoDB document names.
+	 * @param InCompilationOptionsOverride The configuration for the compilation of the CO we want to use instead of the one part of the CO.
+	 * @return True if the compilation was successful and false if it failed.
+	 */
+	bool CompileCustomizableObject(UCustomizableObject* InCustomizableObject, const bool bShouldLogMutableLogs = true, const FCompilationOptions* InCompilationOptionsOverride = nullptr );
+
+private:
+
+	/** The CO that is currently being compiled.  */
+	TStrongObjectPtr<UCustomizableObject> CustomizableObject = nullptr;
+
+	/** Variable being used to know when the compilation is running and when it is not. */
+	bool bIsTheCompilationInProcess = false;
+
+	/** Variable holding the success state of the compilation operation. */
+	bool bWasCoCompilationSuccessful = false;
+	
+	/** Callback invoked once the compilation of the target CO has been completed. It will set a value for bWasCoCompilationSuccessful. */
+	void OnCompilationFinished();
+};
\ No newline at end of file
diff --git a/Engine/Plugins/Experimental/Mutable/Source/MutableValidation/Private/MuV/CustomizableObjectInstanceUpdateUtility.cpp b/Engine/Plugins/Experimental/Mutable/Source/MutableValidation/Private/MuV/CustomizableObjectInstanceUpdateUtility.cpp
new file mode 100644
--- /dev/null
+++ b/Engine/Plugins/Experimental/Mutable/Source/MutableValidation/Private/MuV/CustomizableObjectInstanceUpdateUtility.cpp
@@ -0,0 +1,125 @@
+﻿// Copyright Epic Games, Inc. All Rights Reserved.
+
+#include "CustomizableObjectInstanceUpdateUtility.h"
+
+#include "Commandlets/Commandlet.h"
+#include "MuCO/CustomizableObject.h"
+#include "MuCO/CustomizableObjectSystem.h"
+#include "Engine/TextureStreamingTypes.h"
+#include "SceneTypes.h"
+#include "Components/SkeletalMeshComponent.h"
+
+
+bool FCustomizableObjectInstanceUpdateUtility::UpdateInstance(UCustomizableObjectInstance* InInstance)
+{
+	check (InInstance);
+	check (ComponentsBeingUpdated.IsEmpty());
+	
+	// Cache the instance being updated for reference once in the update end callback
+	Instance = TStrongObjectPtr(InInstance);
+	
+	// Schedule the update of the COI
+	{
+		UE_LOG(LogMutable, Display, TEXT("Invoking update for %s instance."), *InInstance->GetName());
+
+		// Instance update delegate	
+		FInstanceUpdateNativeDelegate InstanceUpdateNativeDelegate;
+		InstanceUpdateNativeDelegate.AddSP(this, &FCustomizableObjectInstanceUpdateUtility::OnInstanceUpdateResult);
+
+		bIsInstanceBeingUpdated = true;
+		InInstance->UpdateSkeletalMeshAsyncResult(InstanceUpdateNativeDelegate, true, true);
+	}
+	
+	
+	// Wait until the update has been completed and the mips streamed
+	while (bIsInstanceBeingUpdated)
+	{
+		// Tick the engine
+		CommandletHelpers::TickEngine();
+
+		// Stop if exit was requested
+		if (IsEngineExitRequested())
+		{
+			break;
+		}
+
+		// Wait until all MIPs gets streamed 
+		if (!ComponentsBeingUpdated.IsEmpty())
+		{
+			bool bFullyStreamed = true;
+			for (TIndexedContainerIterator<TArray<TStrongObjectPtr<USkeletalMeshComponent>>, TArray<TStrongObjectPtr<USkeletalMeshComponent>>::ElementType, TArray<TStrongObjectPtr<USkeletalMeshComponent>>::SizeType> It = ComponentsBeingUpdated.CreateIterator(); It && bFullyStreamed; ++It)
+			{
+				TStrongObjectPtr<USkeletalMeshComponent>& ComponentBeingUpdated = *It;
+            		
+				FStreamingTextureLevelContext LevelContext(EMaterialQualityLevel::Num, ComponentBeingUpdated.Get());
+				TArray<FStreamingRenderAssetPrimitiveInfo> RenderAssetInfoArray;
+				ComponentBeingUpdated->GetStreamingRenderAssetInfo(LevelContext, RenderAssetInfoArray);
+
+				for (TIndexedContainerIterator<TArray<FStreamingRenderAssetPrimitiveInfo>, TArray<
+					                               FStreamingRenderAssetPrimitiveInfo>::ElementType, TArray<
+					                               FStreamingRenderAssetPrimitiveInfo>::SizeType> ItAsset =
+					     RenderAssetInfoArray.CreateIterator(); ItAsset && bFullyStreamed; ++ItAsset)
+				{
+					bFullyStreamed = ItAsset->RenderAsset->IsFullyStreamedIn();
+				}
+			}
+
+			if (bFullyStreamed)
+			{
+				UE_LOG(LogMutable, Display, TEXT("Instance %s finished streaming all MIPs."), *InInstance->GetName());
+				ComponentsBeingUpdated.Reset();
+				
+				bIsInstanceBeingUpdated = false;		// Exit the while loop
+			}
+		}
+	}
+	
+	// Return true if the update was successful and false if not
+	return !bInstanceFailedUpdate;
+}
+
+
+void FCustomizableObjectInstanceUpdateUtility::OnInstanceUpdateResult(const FUpdateContext& Result)
+{
+	const FString InstanceName = Instance->GetName();
+	
+	if (UCustomizableObjectSystem::IsUpdateResultValid(Result.UpdateResult))
+	{
+		UE_LOG(LogMutable, Display, TEXT("Instance %s finished update succesfully."), *InstanceName);
+		bInstanceFailedUpdate = false;
+		
+		// Request load all MIPs
+		UE_LOG(LogMutable, Display, TEXT("Instance %s rquesting streaming all MIPs."), *InstanceName);
+
+		check(ComponentsBeingUpdated.IsEmpty());
+		for (int32 Index = 0; Index < Instance->GetNumComponents(); ++Index)
+		{
+			TStrongObjectPtr<USkeletalMeshComponent> SkeletalComponent = TStrongObjectPtr(NewObject<USkeletalMeshComponent>());
+			SkeletalComponent->SetSkeletalMesh(Instance->GetSkeletalMesh(Index));
+			
+			ComponentsBeingUpdated.Add(SkeletalComponent);            			
+		}
+
+		// Request the streaming in of all the components affected by the update
+		for (TStrongObjectPtr<USkeletalMeshComponent>& ComponentBeingUpdated : ComponentsBeingUpdated)
+		{
+			FStreamingTextureLevelContext LevelContext(EMaterialQualityLevel::Num, ComponentBeingUpdated.Get());
+			TArray<FStreamingRenderAssetPrimitiveInfo> RenderAssetInfoArray;
+			ComponentBeingUpdated->GetStreamingRenderAssetInfo(LevelContext, RenderAssetInfoArray);
+
+			for (const FStreamingRenderAssetPrimitiveInfo& Info : RenderAssetInfoArray)
+			{
+				Info.RenderAsset->StreamIn(MAX_int32, true);
+			}
+		}
+	}
+	else
+	{
+		const FString OutputStatus = UEnum::GetValueAsString(Result.UpdateResult);
+		UE_LOG(LogMutable, Error, TEXT("Instance %s finished update with anomalous state : %s."), *InstanceName, *OutputStatus);
+		bInstanceFailedUpdate = true;
+		
+		// Tell the system the instance finished it's update so we can continue the execution without waitting for the mips to stream in
+		bIsInstanceBeingUpdated = false;
+	}	
+}
\ No newline at end of file
diff --git a/Engine/Plugins/Experimental/Mutable/Source/MutableValidation/Private/MuV/CustomizableObjectInstanceUpdateUtility.h b/Engine/Plugins/Experimental/Mutable/Source/MutableValidation/Private/MuV/CustomizableObjectInstanceUpdateUtility.h
new file mode 100644
--- /dev/null
+++ b/Engine/Plugins/Experimental/Mutable/Source/MutableValidation/Private/MuV/CustomizableObjectInstanceUpdateUtility.h
@@ -0,0 +1,46 @@
+﻿// Copyright Epic Games, Inc. All Rights Reserved.
+#pragma once
+
+#include "UObject/StrongObjectPtr.h"
+#include "Templates/SharedPointer.h"
+
+// Forward declarations
+class UCustomizableObjectInstance;
+class USkeletalMeshComponent;
+struct FUpdateContext;
+
+/**
+ * Helping class that handles the async update of the provided instance. It will also wait for the mips of it so they get streamed.
+ */
+class FCustomizableObjectInstanceUpdateUtility : public TSharedFromThis<FCustomizableObjectInstanceUpdateUtility>
+{
+public:
+	
+	/**
+	 * Updates the provided customizable object instance.
+	 * @param InInstance Instance to be compiled
+	 * @return True if the update was successful and false otherwise
+	 */
+	bool UpdateInstance(UCustomizableObjectInstance* InInstance);
+	
+
+private:
+	
+	/**
+	 * Callback executed when the instance being updated finishes it's mesh update.
+	 * @param Result The result of the updating operation
+	 */
+	void OnInstanceUpdateResult(const FUpdateContext& Result);
+
+	/** The instance that is currently being handled by this class. */
+	TStrongObjectPtr<UCustomizableObjectInstance> Instance = nullptr;
+
+	/** The components of the Instance that we are currently waiting for their mips to be streamed in */
+	TArray<TStrongObjectPtr<USkeletalMeshComponent>> ComponentsBeingUpdated;
+
+	/** Flag used to control if we are updating and instance or not. Once it gets set to false then the update gets halted and the program continues */
+	bool bIsInstanceBeingUpdated = false;
+
+	/** False if the instance did update successfully, true if it failed */
+	bool bInstanceFailedUpdate = false;
+};
diff --git a/Engine/Plugins/Experimental/Mutable/Source/MutableValidation/Private/MuV/CustomizableObjectValidationCommandlet.cpp b/Engine/Plugins/Experimental/Mutable/Source/MutableValidation/Private/MuV/CustomizableObjectValidationCommandlet.cpp
--- a/Engine/Plugins/Experimental/Mutable/Source/MutableValidation/Private/MuV/CustomizableObjectValidationCommandlet.cpp
+++ b/Engine/Plugins/Experimental/Mutable/Source/MutableValidation/Private/MuV/CustomizableObjectValidationCommandlet.cpp
@@ -2,6 +2,8 @@
 
 #include "MuV/CustomizableObjectValidationCommandlet.h"
 
+#include "CustomizableObjectCompilationUtility.h"
+#include "CustomizableObjectInstanceUpdateUtility.h"
 #include "ValidationUtils.h"
 #include "Components/SkeletalMeshComponent.h"
 #include "Containers/Array.h"
@@ -84,8 +86,9 @@ int32 UCustomizableObjectValidationCommandlet::Main(const FString& Params)
 		// Disk cache usage for compilation operation : Override if the user provided an argument with a different value than the default one of the CO
 		CompilationOptions.bUseDiskCompilation = false;
 		FParse::Bool(*Params,TEXT("UseDiskCompilation="),CompilationOptions.bUseDiskCompilation);
-
-		bWasCoCompilationSuccessful = CompileCustomizableObject(ToTestCustomizableObject, true, &CompilationOptions);
+		
+		TSharedRef<FCustomizableObjectCompilationUtility> CompilationUtility = MakeShared<FCustomizableObjectCompilationUtility>();
+		bWasCoCompilationSuccessful = CompilationUtility->CompileCustomizableObject(ToTestCustomizableObject, true, &CompilationOptions);
 	}
 	// -------------------------------------------------------------------------------------------------------------- //
 
@@ -197,10 +200,10 @@ int32 UCustomizableObjectValidationCommandlet::Main(const FString& Params)
 	bool bInstanceFailedUpdate = false;
 	const double InstancesUpdateStartSeconds = FPlatformTime::Seconds();
 	{
-		InstanceUpdater = NewObject<UCOIUpdater>();
+		TSharedRef<FCustomizableObjectInstanceUpdateUtility> InstanceUpdatingUtility = MakeShared<FCustomizableObjectInstanceUpdateUtility>();
 		for (UCustomizableObjectInstance* InstanceToUpdate : InstancesToProcess)
 		{
-			if (InstanceUpdater && !InstanceUpdater->UpdateInstance(InstanceToUpdate))
+			if (!InstanceUpdatingUtility->UpdateInstance(InstanceToUpdate))
 			{
 				bInstanceFailedUpdate = true;
 			}
diff --git a/Engine/Plugins/Experimental/Mutable/Source/MutableValidation/Private/MuV/ValidationUtils.cpp b/Engine/Plugins/Experimental/Mutable/Source/MutableValidation/Private/MuV/ValidationUtils.cpp
--- a/Engine/Plugins/Experimental/Mutable/Source/MutableValidation/Private/MuV/ValidationUtils.cpp
+++ b/Engine/Plugins/Experimental/Mutable/Source/MutableValidation/Private/MuV/ValidationUtils.cpp
@@ -4,13 +4,9 @@
 
 #include "AssetRegistry/AssetRegistryModule.h"
 #include "Commandlets/Commandlet.h"
-#include "Components/SkeletalMeshComponent.h"
 #include "MuCO/CustomizableObject.h"
-#include "MuCO/CustomizableObjectInstance.h"
-#include "MuCO/CustomizableObjectPrivate.h"
 #include "MuCO/CustomizableObjectSystem.h"
 #include "Interfaces/ITargetPlatform.h"
-#include "MuCOE/CustomizableObjectEditorFunctionLibrary.h"
 
 void PrepareAssetRegistry()
 {
@@ -26,98 +22,6 @@ void PrepareAssetRegistry()
 }
 
 
-bool CompileCustomizableObject(UCustomizableObject* InCustomizableObject, const bool bLogMutableLogs /* = true */, const FCompilationOptions* InCompilationOptionsOverride  /* nullptr */)
-{
-	if (!InCustomizableObject)
-	{
-		UE_LOG(LogMutable, Error, TEXT("Unable to compile CO : CO pointer is null!"));
-		return false;
-	}
-	
-	// Override the compilation options if an override has been provided by the user
-	FCompilationOptions CompilationOptions;
-	const bool bOverrideCompilationOptions = InCompilationOptionsOverride != nullptr;
-	if (bOverrideCompilationOptions)
-	{
-		CompilationOptions = *InCompilationOptionsOverride;
-		UE_LOG(LogMutable,Display,TEXT("CO Compilation options overriden by the user defined ones."));
-	}
-	else
-	{
-		CompilationOptions = InCustomizableObject->CompileOptions;
-		UE_LOG(LogMutable,Display,TEXT("Compiling CO using it's own compilation options."));
-	}
-
-	// Ensure that the user has provided a target compilation platform.
-	// Mutable is able to run without one but we want to be explicit in the context of testing.
-	if (!CompilationOptions.TargetPlatform)
-	{
-		UE_LOG(LogMutable, Error, TEXT("The compilation of the %s model could not be started : No explicit platform was provided."), *InCustomizableObject->GetName());
-		return false;
-	}
-	
-	// Request a compiler to be able to locate the root and to compile it
-	const TUniquePtr<FCustomizableObjectCompilerBase> Compiler =
-		TUniquePtr<FCustomizableObjectCompilerBase>(UCustomizableObjectSystem::GetInstanceChecked()->GetNewCompiler());
-	
-    // Compile the CO with the selected compilation options
-    // Run Sync compilation -> Warning : Potentially long operation -------------
-
-	// Get the memory being used by mutable before the compilation
-	const int64 CompilationStartBytes = mu::FGlobalMemoryCounter::GetCounter();
-	mu::FGlobalMemoryCounter::Zero();
-	
-	UE_LOG(LogMutable,Display,TEXT("Compiling Customizable Object..."));
-	const double CompilationStartSeconds = FPlatformTime::Seconds();
-	{
-		Compiler->Compile(*InCustomizableObject, CompilationOptions, false);
-	}
-	const double CompilationEndSeconds = FPlatformTime::Seconds() - CompilationStartSeconds;
-    // --------------------------------------------------------------------------
-	UE_LOG(LogMutable,Display,TEXT("Compilation of CO completed!"));
-	UE_LOG(LogMutable, Display, TEXT("The compilation of the %s model took %f seconds."), *InCustomizableObject->GetName(), CompilationEndSeconds);
-
-	// Get the peak mutable memory used during the compilation operation
-    const int64 CompilationEndPeakBytes = mu::FGlobalMemoryCounter::GetPeak();
-    const int64 CompilationEndRealPeakBytes = CompilationStartBytes + CompilationEndPeakBytes;
-	
-    // Get the compilation result
-    const bool bWasCoCompilationSuccessful = InCustomizableObject->GetPrivate()->CompilationResult == ECompilationResultPrivate::Success ||
-    	InCustomizableObject->GetPrivate()->CompilationResult == ECompilationResultPrivate::Warnings;
-    if (bWasCoCompilationSuccessful)
-    {
-    	UE_LOG(LogMutable, Display, TEXT("The compilation of the %s model was succesfull."), *InCustomizableObject->GetName());
-    }
-    else
-    {
-    	UE_LOG(LogMutable, Error, TEXT("The compilation of the %s model failed."), *InCustomizableObject->GetName());
-    }
-
-	const ECustomizableObjectCompilationState CompilationEndResult = bWasCoCompilationSuccessful ? ECustomizableObjectCompilationState::Completed : ECustomizableObjectCompilationState::Failed;
-	
-	// Print MTU parseable logs only if asked. The addition of duplicated entries in MongoDB is not available so this way we avoid having to handle them
-	// when we are sure we do not require it.
-	if (bLogMutableLogs)
-	{
-		UE_LOG(LogMutable, Log, TEXT("(string) model_compile_options_overriden : %s "), bOverrideCompilationOptions ? TEXT("true") : TEXT("false"));
-		UE_LOG(LogMutable, Log, TEXT("(int) model_optimization_level : %d "), CompilationOptions.OptimizationLevel);
-		UE_LOG(LogMutable, Log, TEXT("(string) model_texture_compression : %s "), *UEnum::GetValueAsString(CompilationOptions.TextureCompression));
-		UE_LOG(LogMutable, Log, TEXT("(string) model_disk_compilation : %s "), CompilationOptions.bUseDiskCompilation ? TEXT("true") : TEXT("false"));
-		UE_LOG(LogMutable, Log, TEXT("(string) model_compile_platform_name : %s "), *CompilationOptions.TargetPlatform->PlatformName());
-
-		UE_LOG(LogMutable, Log, TEXT("(double) model_compile_time_ms : %f "), CompilationEndSeconds * 1000);
-		UE_LOG(LogMutable, Log, TEXT("(string) model_compile_end_state : %s "), *UEnum::GetValueAsString(CompilationEndResult));
-
-		UE_LOG(LogMutable, Log, TEXT("(int) model_compilation_start_bytes : %lld "), CompilationStartBytes);
-		UE_LOG(LogMutable, Log, TEXT("(int) model_compilation_end_peak_bytes : %lld "), CompilationEndPeakBytes);
-		UE_LOG(LogMutable, Log, TEXT("(int) model_compilation_end_real_peak_bytes : %lld "), CompilationEndRealPeakBytes);
-
-		// TODO: Add logs for the other relevant configs of the model being compiled
-	}
-
-	return bWasCoCompilationSuccessful;
-}
-
 void LogMutableSettings()
 {
 	const int32 WorkingMemory = UCustomizableObjectSystem::GetInstanceChecked()->GetWorkingMemory() ;
@@ -127,6 +31,7 @@ void LogMutableSettings()
 	// Expand this when adding new controls from the .xml file
 }
 
+
 void Wait(const double ToWaitSeconds)
 {
 	check (ToWaitSeconds > 0);
@@ -149,113 +54,4 @@ void Wait(const double ToWaitSeconds)
 }
 
 
-bool UCOIUpdater::UpdateInstance(UCustomizableObjectInstance* InInstance)
-{
-	check (InInstance);
-	check (ComponentsBeingUpdated.IsEmpty());
-	
-	// Cache the instance being updated for reference once in the update end callback
-	Instance = InInstance;
-	
-	// Schedule the update of the COI
-	{
-		UE_LOG(LogMutable,Display,TEXT("Invoking update for %s instance."),*InInstance->GetName());
-
-		// Instance update delegate	
-		FInstanceUpdateDelegate InstanceUpdateDelegate;
-		InstanceUpdateDelegate.BindDynamic(this, &UCOIUpdater::OnInstanceUpdateResult);
-
-		bIsInstanceBeingUpdated = true;
-		InInstance->UpdateSkeletalMeshAsyncResult(InstanceUpdateDelegate,true, true);
-	}
-	
-	
-	// Wait until the update has been completed and the mips streamed
-	while (bIsInstanceBeingUpdated)
-	{
-		// Tick the engine
-		CommandletHelpers::TickEngine();
-
-		// Stop if exit was requested
-		if (IsEngineExitRequested())
-		{
-			break;
-		}
-
-		// Wait until all MIPs gets streamed 
-		if (!ComponentsBeingUpdated.IsEmpty())
-		{
-			bool bFullyStreamed = true;
-			for (auto It = ComponentsBeingUpdated.CreateIterator(); It && bFullyStreamed; ++It)
-			{
-				TObjectPtr<USkeletalMeshComponent>& ComponentBeingUpdated = *It;
-            		
-				FStreamingTextureLevelContext LevelContext(EMaterialQualityLevel::Num, ComponentBeingUpdated);
-				TArray<FStreamingRenderAssetPrimitiveInfo> RenderAssetInfoArray;
-				ComponentBeingUpdated->GetStreamingRenderAssetInfo(LevelContext, RenderAssetInfoArray);
-
-				for (auto ItAsset = RenderAssetInfoArray.CreateIterator(); ItAsset && bFullyStreamed; ++ItAsset)
-				{
-					bFullyStreamed = ItAsset->RenderAsset->IsFullyStreamedIn();
-				}
-			}
-
-			if (bFullyStreamed)
-			{
-				UE_LOG(LogMutable,Display,TEXT("Instance %s finished streaming all MIPs."), *InInstance->GetName());
-				ComponentsBeingUpdated.Reset();
-				
-				bIsInstanceBeingUpdated = false;		// Exit the while loop
-			}
-		}
-	}
-	
-	// Return true if the update was successful and false if not
-	return !bInstanceFailedUpdate;
-}
-
-
-void UCOIUpdater::OnInstanceUpdateResult(const FUpdateContext& Result)
-{
-	const FString InstanceName = Instance->GetName();
-	
-	if (UCustomizableObjectSystem::IsUpdateResultValid(Result.UpdateResult))
-	{
-		UE_LOG(LogMutable,Display,TEXT("Instance %s finished update succesfully."),*InstanceName);
-		bInstanceFailedUpdate = false;
-		
-		// Request load all MIPs
-		UE_LOG(LogMutable,Display,TEXT("Instance %s rquesting streaming all MIPs."), *InstanceName);
 
-		check(ComponentsBeingUpdated.IsEmpty());
-		for (int32 Index = 0; Index < Instance->GetNumComponents(); ++Index)
-		{
-			USkeletalMeshComponent* SkeletalComponent = NewObject<USkeletalMeshComponent>();
-			SkeletalComponent->SetSkeletalMesh(Instance->GetSkeletalMesh(Index));
-			
-			ComponentsBeingUpdated.Add(SkeletalComponent);            			
-		}
-
-		// Request the streaming in of all the components affected by the update
-		for (TObjectPtr<USkeletalMeshComponent>& ComponentBeingUpdated : ComponentsBeingUpdated)
-		{
-			FStreamingTextureLevelContext LevelContext(EMaterialQualityLevel::Num, ComponentBeingUpdated);
-			TArray<FStreamingRenderAssetPrimitiveInfo> RenderAssetInfoArray;
-			ComponentBeingUpdated->GetStreamingRenderAssetInfo(LevelContext, RenderAssetInfoArray);
-
-			for (const FStreamingRenderAssetPrimitiveInfo& Info : RenderAssetInfoArray)
-			{
-				Info.RenderAsset->StreamIn(MAX_int32, true);
-			}
-		}
-	}
-	else
-	{
-		const FString OutputStatus = UEnum::GetValueAsString(Result.UpdateResult);
-		UE_LOG(LogMutable,Error,TEXT("Instance %s finished update with anomalous state : %s."), *InstanceName, *OutputStatus);
-		bInstanceFailedUpdate = true;
-		
-		// Tell the system the instance finished it's update so we can continue the execution without waitting for the mips to stream in
-		bIsInstanceBeingUpdated = false;
-	}	
-}
diff --git a/Engine/Plugins/Experimental/Mutable/Source/MutableValidation/Private/MuV/ValidationUtils.h b/Engine/Plugins/Experimental/Mutable/Source/MutableValidation/Private/MuV/ValidationUtils.h
--- a/Engine/Plugins/Experimental/Mutable/Source/MutableValidation/Private/MuV/ValidationUtils.h
+++ b/Engine/Plugins/Experimental/Mutable/Source/MutableValidation/Private/MuV/ValidationUtils.h
@@ -1,18 +1,6 @@
 ﻿// Copyright Epic Games, Inc. All Rights Reserved.
 #pragma once
 
-#include "MuCO/CustomizableObjectInstance.h"
-#include "UObject/Object.h"
-
-#include "ValidationUtils.generated.h"
-
-// Forward declarations
-class UCustomizableObject;
-class UCustomizableObjectInstance;
-class USkeletalMeshComponent;
-struct FCompilationOptions;
-
-
 /**
  * Prepare the asset registry so we can later use it to search assets. It is required by Mutable to compile.
  */
@@ -26,62 +14,9 @@ void PrepareAssetRegistry();
 void Wait(const double ToWaitSeconds);
 
 
-/**
- * Compiles a CO synchronously
- * @param InCustomizableObject The Customizable Object we want to synchronously compile
- * @param bLogMutableLogs Enables or disables the logging of Log category logs relevant to the CO compilation. Required to avoid a MongoDB limitation with the duplication of MongoDB document names.
- * @param InCompilationOptionsOverride The configuration for the compilation of the CO we want to use instead of the one part of the CO.
- * @return True if the compilation was successful and false if it failed.
- */
-bool CompileCustomizableObject(UCustomizableObject* InCustomizableObject, const bool bLogMutableLogs = true, const FCompilationOptions* InCompilationOptionsOverride = nullptr);
-
-
 /**
  * Logs some configuration data related to how mutable will compile and then generate instances. We do this so we can later
  * Isolate tests using different configurations.
  * @note Add new logs each time you add a way to change the configuration of the test from the .xml testing file
  */
 void LogMutableSettings();
-
-
-/**
- * Helping class that handles the async update of the provided instance. It will also wait for the mips of it so they get streamed.
- */
-UCLASS()
-class UCOIUpdater : public UObject
-{
-	GENERATED_BODY()
-	
-public:
-	
-	/**
-	 * Updates the provided customizable object instance.
-	 * @param InInstance Instance to be compiled
-	 * @return True if the update was successful and false otherwise
-	 */
-	bool UpdateInstance(UCustomizableObjectInstance* InInstance);
-	
-
-private:
-	
-	/**
-	 * Callback executed when the instance being updated finishes it's mesh update.
-	 * @param Result The result of the updating operation
-	 */
-	UFUNCTION()
-	void OnInstanceUpdateResult(const FUpdateContext& Result);
-
-	/** The instance that is currently being handled by this class. */
-	UPROPERTY()
-	TObjectPtr<UCustomizableObjectInstance> Instance;
-
-	/** The components of the Instance that we are currently waiting for their mips to be streamed in */
-	UPROPERTY()
-	TArray<TObjectPtr<USkeletalMeshComponent>> ComponentsBeingUpdated;
-
-	/** Flag used to control if we are updating and instance or not. Once it gets set to false then the update gets halted and the program continues */
-	bool bIsInstanceBeingUpdated = false;
-
-	/** False if the instance did update successfully, true if it failed */
-	bool bInstanceFailedUpdate = false;
-};
diff --git a/Engine/Plugins/Experimental/Mutable/Source/MutableValidation/Public/MuV/COIBakingTestCommandlet.h b/Engine/Plugins/Experimental/Mutable/Source/MutableValidation/Public/MuV/COIBakingTestCommandlet.h
--- a/Engine/Plugins/Experimental/Mutable/Source/MutableValidation/Public/MuV/COIBakingTestCommandlet.h
+++ b/Engine/Plugins/Experimental/Mutable/Source/MutableValidation/Public/MuV/COIBakingTestCommandlet.h
@@ -3,6 +3,7 @@
 
 #include "CoreMinimal.h"
 #include "Commandlets/Commandlet.h"
+#include "MuV/CustomizableObjectCompilationUtility.h"
 
 #include "COIBakingTestCommandlet.generated.h"
 
diff --git a/Engine/Plugins/Experimental/Mutable/Source/MutableValidation/Public/MuV/COIBulkUpdateTestCommandlet.h b/Engine/Plugins/Experimental/Mutable/Source/MutableValidation/Public/MuV/COIBulkUpdateTestCommandlet.h
--- a/Engine/Plugins/Experimental/Mutable/Source/MutableValidation/Public/MuV/COIBulkUpdateTestCommandlet.h
+++ b/Engine/Plugins/Experimental/Mutable/Source/MutableValidation/Public/MuV/COIBulkUpdateTestCommandlet.h
@@ -2,7 +2,8 @@
 #pragma once
 
 #include "Commandlets/Commandlet.h"
-#include "MuV/ValidationUtils.h"
+#include "MuV/CustomizableObjectCompilationUtility.h"
+#include "MuV/CustomizableObjectInstanceUpdateUtility.h"
 
 #include "COIBulkUpdateTestCommandlet.generated.h"
 
@@ -20,11 +21,4 @@ class UCOIBulkUpdateTestCommandlet : public UCommandlet
 
 public:
 	virtual int32 Main(const FString& Params) override;
-
-private:
-
-	/** Helper object that handles the update of the instance for us.
-	 * It allows us to work with the instance update as if it was a sync operation */
-	UPROPERTY()
-	TObjectPtr<UCOIUpdater> InstanceUpdater = nullptr;
 };
diff --git a/Engine/Plugins/Experimental/Mutable/Source/MutableValidation/Public/MuV/CustomizableObjectValidationCommandlet.h b/Engine/Plugins/Experimental/Mutable/Source/MutableValidation/Public/MuV/CustomizableObjectValidationCommandlet.h
--- a/Engine/Plugins/Experimental/Mutable/Source/MutableValidation/Public/MuV/CustomizableObjectValidationCommandlet.h
+++ b/Engine/Plugins/Experimental/Mutable/Source/MutableValidation/Public/MuV/CustomizableObjectValidationCommandlet.h
@@ -4,13 +4,14 @@
 
 #include "Commandlets/Commandlet.h"
 #include "Containers/Array.h"
-#include "MuCO/CustomizableObjectInstance.h"
+#include "MuV/CustomizableObjectCompilationUtility.h"
+#include "MuV/CustomizableObjectInstanceUpdateUtility.h"
 
 #include "CustomizableObjectValidationCommandlet.generated.h"
 
 // Forward declarations
 class UCustomizableObject;
-class UCOIUpdater;
+class UCustomizableObjectInstance;
 class ITargetPlatform;
 
 UCLASS()
@@ -38,8 +39,4 @@ private:
 	/** Array of COI to be generated with randomized parameter values */
 	UPROPERTY()
 	TArray<TObjectPtr<UCustomizableObjectInstance>> InstancesToProcess;
-
-	/** Helper object designed to aid in the update of the CO Instances.*/
-	UPROPERTY()
-	TObjectPtr<UCOIUpdater> InstanceUpdater;
 };
\ No newline at end of file
