[Mutable] Fix crash when launching the editor with -game.

- Now all CO module checks if the COE module is valid before accessing it.
- Remove the compiler requirement to call non-related compiler functions.

#rnx

diff --git a/Engine/Plugins/Experimental/Mutable/Source/CustomizableObject/Classes/MuCO/CustomizableObjectSystem.h b/Engine/Plugins/Experimental/Mutable/Source/CustomizableObject/Classes/MuCO/CustomizableObjectSystem.h
--- a/Engine/Plugins/Experimental/Mutable/Source/CustomizableObject/Classes/MuCO/CustomizableObjectSystem.h
+++ b/Engine/Plugins/Experimental/Mutable/Source/CustomizableObject/Classes/MuCO/CustomizableObjectSystem.h
@@ -199,8 +199,7 @@ public:
 	// End UObject interface.
 
 	// Creates a new Customizable Object Compiler (Only does real work in editor builds). The caller is responsible for freeing the new compiler
-	FCustomizableObjectCompilerBase* GetNewCompiler();
-	void SetNewCompilerFunc(FCustomizableObjectCompilerBase* (*NewCompilerFunc)());
+	static TSharedPtr<FCustomizableObjectCompilerBase> GetNewCompiler();
 
 	bool IsReplaceDiscardedWithReferenceMeshEnabled() const;
 	void SetReplaceDiscardedWithReferenceMeshEnabled(bool bIsEnabled);
diff --git a/Engine/Plugins/Experimental/Mutable/Source/CustomizableObject/Internal/MuCO/CustomizableObjectPrivate.h b/Engine/Plugins/Experimental/Mutable/Source/CustomizableObject/Internal/MuCO/CustomizableObjectPrivate.h
--- a/Engine/Plugins/Experimental/Mutable/Source/CustomizableObject/Internal/MuCO/CustomizableObjectPrivate.h
+++ b/Engine/Plugins/Experimental/Mutable/Source/CustomizableObject/Internal/MuCO/CustomizableObjectPrivate.h
@@ -51,9 +51,6 @@ public:
 	virtual bool Tick(bool bBlocking) { return true; }
 	virtual void ForceFinishCompilation() {}
 
-	// Return true if this object doesn't reference a parent object.
-	virtual bool IsRootObject(const UCustomizableObject* Object) const { return true; }
-
 	/** Provides the caller with the warning and error messages produced during compilation */
 	virtual void GetCompilationMessages(TArray<FText>& OutWarningMessages, TArray<FText>& OutErrorMessages) const = 0;
 };
diff --git a/Engine/Plugins/Experimental/Mutable/Source/CustomizableObject/Internal/MuCO/CustomizableObjectSystemPrivate.h b/Engine/Plugins/Experimental/Mutable/Source/CustomizableObject/Internal/MuCO/CustomizableObjectSystemPrivate.h
--- a/Engine/Plugins/Experimental/Mutable/Source/CustomizableObject/Internal/MuCO/CustomizableObjectSystemPrivate.h
+++ b/Engine/Plugins/Experimental/Mutable/Source/CustomizableObject/Internal/MuCO/CustomizableObjectSystemPrivate.h
@@ -714,8 +714,6 @@ public:
 
 	bool bSupport16BitBoneIndex = false;
 
-	static FCustomizableObjectCompilerBase* (*NewCompilerFunc)();
-
 	TMap<FMutableImageCacheKey, uint32> TextureReferenceCount; // Keeps a count of texture usage to decide if they have to be blocked from GC during an update
 
 	UPROPERTY(Transient)
@@ -791,7 +789,7 @@ public:
 	FStreamableManager StreamableManager;
 	
 #if WITH_EDITOR
-	FCustomizableObjectCompilerBase* RecompileCustomizableObjectsCompiler = nullptr;
+	TSharedPtr<FCustomizableObjectCompilerBase> RecompileCustomizableObjectsCompiler;
 	
 	TArray<FAssetData> ObjectsToRecompile;
 	uint32 TotalNumObjectsToRecompile = 0;
diff --git a/Engine/Plugins/Experimental/Mutable/Source/CustomizableObject/Internal/MuCO/ICustomizableObjectEditorModule.h b/Engine/Plugins/Experimental/Mutable/Source/CustomizableObject/Internal/MuCO/ICustomizableObjectEditorModule.h
--- a/Engine/Plugins/Experimental/Mutable/Source/CustomizableObject/Internal/MuCO/ICustomizableObjectEditorModule.h
+++ b/Engine/Plugins/Experimental/Mutable/Source/CustomizableObject/Internal/MuCO/ICustomizableObjectEditorModule.h
@@ -15,6 +15,7 @@ class UCustomizableObject;
 class UCustomizableObjectPrivate;
 class UCustomizableObjectInstance;
 class FExtensibilityManager;
+class FCustomizableObjectCompilerBase;
 class FBakeOperationCompletedDelegate;
 
 extern const FName CustomizableObjectEditorAppIdentifier;
@@ -50,6 +51,12 @@ public:
 	  * @param OutOfDatePackages list of out of date packages.
    	  * @return true if the compilation is out of date. */
 	virtual bool IsCompilationOutOfDate(const UCustomizableObject& Object, TArray<FName>* OutOfDatePackages = nullptr) const = 0;
+
+	/** Create a new compiler. */
+	virtual TSharedRef<FCustomizableObjectCompilerBase> CreateCompiler() const = 0;
+
+	/** See GraphTraversal::IsRootObject(...) */
+	virtual bool IsRootObject(const UCustomizableObject& Object) const = 0;
 	
 	/**
 	 * Execute this method in order to bake the provided instance. It will schedule a special type of instance update before proceeding with the bake itself.
diff --git a/Engine/Plugins/Experimental/Mutable/Source/CustomizableObject/Private/MuCO/CustomizableObject.cpp b/Engine/Plugins/Experimental/Mutable/Source/CustomizableObject/Private/MuCO/CustomizableObject.cpp
--- a/Engine/Plugins/Experimental/Mutable/Source/CustomizableObject/Private/MuCO/CustomizableObject.cpp
+++ b/Engine/Plugins/Experimental/Mutable/Source/CustomizableObject/Private/MuCO/CustomizableObject.cpp
@@ -113,17 +113,16 @@ void UCustomizableObject::GetAssetRegistryTags(TArray<FAssetRegistryTag>& OutTag
 	PRAGMA_ENABLE_DEPRECATION_WARNINGS;
 }
 
+
 void UCustomizableObject::GetAssetRegistryTags(FAssetRegistryTagsContext Context) const
 {
 	int32 isRoot = 0;
 
-	FCustomizableObjectCompilerBase* Compiler = UCustomizableObjectSystem::GetInstance()->GetNewCompiler();
-	if (Compiler)
+	if (const ICustomizableObjectEditorModule* Module = ICustomizableObjectEditorModule::Get())
 	{
-		isRoot = Compiler->IsRootObject(this) ? 1 : 0;
-		delete Compiler;
+		isRoot = Module->IsRootObject(*this) ? 1 : 0;
 	}
-
+	
 	Context.AddTag(FAssetRegistryTag("IsRoot", FString::FromInt(isRoot), FAssetRegistryTag::TT_Numerical));
 	Super::GetAssetRegistryTags(Context);
 }
@@ -191,11 +190,9 @@ void UCustomizableObject::PostSaveRoot(FObjectPostSaveRootContext ObjectSaveCont
 
 bool UCustomizableObjectPrivate::TryUpdateIsChildObject()
 {
-	FCustomizableObjectCompilerBase* Compiler = UCustomizableObjectSystem::GetInstance()->GetNewCompiler();
-	if (Compiler)
+	if (const ICustomizableObjectEditorModule* Module = ICustomizableObjectEditorModule::Get())
 	{
-		GetPublic()->bIsChildObject = !Compiler->IsRootObject(GetPublic());
-		delete Compiler;
+		GetPublic()->bIsChildObject = !Module->IsRootObject(*GetPublic());
 		return true;
 	}
 	else
@@ -831,22 +828,32 @@ void UCustomizableObjectPrivate::CompileForTargetPlatform(const ITargetPlatform*
 		return;
 	}
 
-	FCustomizableObjectCompilerBase* Compiler = UCustomizableObjectSystem::GetInstance()->GetNewCompiler();
+	const ICustomizableObjectEditorModule* EditorModule = ICustomizableObjectEditorModule::Get();
+	if (!EditorModule)
+	{
+		return;
+	}
 
-	const bool bIsRootObject = Compiler->IsRootObject(GetPublic());
-	if (bIsRootObject)
+	const bool bIsRootObject = EditorModule->IsRootObject(*GetPublic());
+	if (!bIsRootObject)
 	{
-		FCompilationOptions Options;
-		Options.OptimizationLevel = UE_MUTABLE_MAX_OPTIMIZATION;	// max optimization when packaging.
-		Options.TextureCompression = ECustomizableObjectTextureCompression::HighQuality;
-		Options.bIsCooking = true;
-		Options.TargetPlatform = TargetPlatform;
-		Options.CustomizableObjectNumBoneInfluences = ICustomizableObjectModule::Get().GetNumBoneInfluences();
+		return;
+	}
 
-		Compiler->Compile(*GetPublic(), Options, false);
+	const TSharedPtr<FCustomizableObjectCompilerBase> Compiler = UCustomizableObjectSystem::GetInstance()->GetNewCompiler();
+	if (!Compiler)
+	{
+		return;
 	}
 
-	delete Compiler;
+	FCompilationOptions Options;
+	Options.OptimizationLevel = UE_MUTABLE_MAX_OPTIMIZATION;	// max optimization when packaging.
+	Options.TextureCompression = ECustomizableObjectTextureCompression::HighQuality;
+	Options.bIsCooking = true;
+	Options.TargetPlatform = TargetPlatform;
+	Options.CustomizableObjectNumBoneInfluences = ICustomizableObjectModule::Get().GetNumBoneInfluences();
+
+	Compiler->Compile(*GetPublic(), Options, false);
 }
 
 
@@ -901,9 +908,10 @@ bool UCustomizableObject::ConditionalAutoCompile()
 	// Sync/Async compilation
 	if (System->IsAutoCompilationSync())
 	{
-		FCustomizableObjectCompilerBase* Compiler = UCustomizableObjectSystem::GetInstance()->GetNewCompiler();
-		Compiler->Compile(*this, this->CompileOptions, false);
-		delete Compiler;
+		if (const TSharedPtr<FCustomizableObjectCompilerBase> Compiler = UCustomizableObjectSystem::GetNewCompiler())
+		{		
+			Compiler->Compile(*this, this->CompileOptions, false);
+		}
 	}
 	else
 	{
diff --git a/Engine/Plugins/Experimental/Mutable/Source/CustomizableObject/Private/MuCO/CustomizableObjectSystem.cpp b/Engine/Plugins/Experimental/Mutable/Source/CustomizableObject/Private/MuCO/CustomizableObjectSystem.cpp
--- a/Engine/Plugins/Experimental/Mutable/Source/CustomizableObject/Private/MuCO/CustomizableObjectSystem.cpp
+++ b/Engine/Plugins/Experimental/Mutable/Source/CustomizableObject/Private/MuCO/CustomizableObjectSystem.cpp
@@ -33,6 +33,7 @@
 #include "Components/SkeletalMeshComponent.h"
 #include "MuCO/CustomizableObjectSystemPrivate.h"
 #include "CustomizableObjectSettings.h"
+#include "MuCO/ICustomizableObjectEditorModule.h"
 
 #if WITH_EDITOR
 #include "Editor.h"
@@ -542,7 +543,6 @@ void UCustomizableObjectSystem::InitSystem()
 	// Everything initialized in Init() instead of constructor to prevent the default UCustomizableObjectSystem from registering a tick function
 	Private = NewObject<UCustomizableObjectSystemPrivate>(this, FName("Private"));
 	check(Private != nullptr);
-	Private->NewCompilerFunc = nullptr;
 
 	Private->bReplaceDiscardedWithReferenceMesh = false;
 
@@ -608,7 +608,6 @@ void UCustomizableObjectSystem::BeginDestroy()
 		if (Private->RecompileCustomizableObjectsCompiler)
 		{
 			Private->RecompileCustomizableObjectsCompiler->ForceFinishCompilation();
-			delete Private->RecompileCustomizableObjectsCompiler;
 		}
 #endif
 
@@ -654,15 +653,11 @@ FString UCustomizableObjectSystem::GetDesc()
 }
 
 
-FCustomizableObjectCompilerBase* (*UCustomizableObjectSystemPrivate::NewCompilerFunc)() = nullptr;
-
-
-FCustomizableObjectCompilerBase* UCustomizableObjectSystem::GetNewCompiler()
+TSharedPtr<FCustomizableObjectCompilerBase> UCustomizableObjectSystem::GetNewCompiler()
 {
-	check(Private != nullptr);
-	if (Private->NewCompilerFunc != nullptr)
+	if (const ICustomizableObjectEditorModule* Module = ICustomizableObjectEditorModule::Get())
 	{
-		return Private->NewCompilerFunc();
+		return Module->CreateCompiler();
 	}
 	else
 	{
@@ -671,12 +666,6 @@ FCustomizableObjectCompilerBase* UCustomizableObjectSystem::GetNewCompiler()
 }
 
 
-void UCustomizableObjectSystem::SetNewCompilerFunc(FCustomizableObjectCompilerBase* (*InNewCompilerFunc)())
-{
-	GetPrivate()->NewCompilerFunc = InNewCompilerFunc;
-}
-
-
 int32 UCustomizableObjectSystemPrivate::EnableMutableAnimInfoDebugging = 0;
 
 static FAutoConsoleVariableRef CVarEnableMutableAnimInfoDebugging(
@@ -4192,9 +4181,7 @@ void UCustomizableObjectSystem::TickRecompileCustomizableObjects()
 		}
 		else // All objects compiled, clean up
 		{
-			// Delete compiler
-			delete GetPrivate()->RecompileCustomizableObjectsCompiler;
-			GetPrivate()->RecompileCustomizableObjectsCompiler = nullptr;
+			GetPrivate()->RecompileCustomizableObjectsCompiler.Reset();
 
 			// Remove progress bar
 			FSlateNotificationManager::Get().UpdateProgressNotification(GetPrivate()->RecompileNotificationHandle, GetPrivate()->NumObjectsCompiled, GetPrivate()->TotalNumObjectsToRecompile);
diff --git a/Engine/Plugins/Experimental/Mutable/Source/CustomizableObjectEditor/Private/MuCOE/CustomizableObjectCompiler.cpp b/Engine/Plugins/Experimental/Mutable/Source/CustomizableObjectEditor/Private/MuCOE/CustomizableObjectCompiler.cpp
--- a/Engine/Plugins/Experimental/Mutable/Source/CustomizableObjectEditor/Private/MuCOE/CustomizableObjectCompiler.cpp
+++ b/Engine/Plugins/Experimental/Mutable/Source/CustomizableObjectEditor/Private/MuCOE/CustomizableObjectCompiler.cpp
@@ -138,31 +138,6 @@ bool FCustomizableObjectCompiler::Tick(bool bBlocking)
 }
 
 
-bool FCustomizableObjectCompiler::IsRootObject(const UCustomizableObject* Object) const
-{
-	// Look for the base object node
-	UCustomizableObjectNodeObject* Root = nullptr;
-	TArray<UCustomizableObjectNodeObject*> ObjectNodes;
-	if (!Object->GetPrivate()->GetSource() || !Object->GetPrivate()->GetSource()->Nodes.Num())
-	{
-		// Conservative approach.
-		return true;
-	}
-
-	Object->GetPrivate()->GetSource()->GetNodesOfClass<UCustomizableObjectNodeObject>(ObjectNodes);
-
-	for (TArray<UCustomizableObjectNodeObject*>::TIterator It(ObjectNodes); It; ++It)
-	{
-		if ((*It)->bIsBase)
-		{
-			Root = *It;
-		}
-	}
-
-	return Root && !Root->ParentObject;
-}
-
-
 void FCustomizableObjectCompiler::PreloadingReferencerAssetsCallback(bool bAsync)
 {
 	UE_LOG(LogMutable, Verbose, TEXT("PROFILE: [ %16.8f ] Preload asynchronously assets end."), FPlatformTime::Seconds());
diff --git a/Engine/Plugins/Experimental/Mutable/Source/CustomizableObjectEditor/Private/MuCOE/CustomizableObjectCompiler.h b/Engine/Plugins/Experimental/Mutable/Source/CustomizableObjectEditor/Private/MuCOE/CustomizableObjectCompiler.h
--- a/Engine/Plugins/Experimental/Mutable/Source/CustomizableObjectEditor/Private/MuCOE/CustomizableObjectCompiler.h
+++ b/Engine/Plugins/Experimental/Mutable/Source/CustomizableObjectEditor/Private/MuCOE/CustomizableObjectCompiler.h
@@ -25,9 +25,6 @@ public:
 	CUSTOMIZABLEOBJECTEDITOR_API FCustomizableObjectCompiler();
 	CUSTOMIZABLEOBJECTEDITOR_API virtual ~FCustomizableObjectCompiler() {}
 
-	// FCustomizableObjectCompilerBase interface
-	CUSTOMIZABLEOBJECTEDITOR_API bool IsRootObject(const class UCustomizableObject* Object) const override;
-	
 	/** Check for pending compilation process. Returns true if an object has been updated. */
 	CUSTOMIZABLEOBJECTEDITOR_API virtual bool Tick(bool bBlocking) override;
 
diff --git a/Engine/Plugins/Experimental/Mutable/Source/CustomizableObjectEditor/Private/MuCOE/CustomizableObjectEditorModule.cpp b/Engine/Plugins/Experimental/Mutable/Source/CustomizableObjectEditor/Private/MuCOE/CustomizableObjectEditorModule.cpp
--- a/Engine/Plugins/Experimental/Mutable/Source/CustomizableObjectEditor/Private/MuCOE/CustomizableObjectEditorModule.cpp
+++ b/Engine/Plugins/Experimental/Mutable/Source/CustomizableObjectEditor/Private/MuCOE/CustomizableObjectEditorModule.cpp
@@ -63,6 +63,7 @@
 #include "UObject/UObjectIterator.h"
 #include "Subsystems/PlacementSubsystem.h"
 #include "Components/SkeletalMeshComponent.h"
+#include "MuCOE/GraphTraversal.h"
 #include "MuCOE/CustomizableObjectInstanceBaker.h"
 
 class AActor;
@@ -185,16 +186,8 @@ static void CustomFree(void* mem)
 }
 
 
-FCustomizableObjectCompilerBase* NewCompiler()
-{
-	return new FCustomizableObjectCompiler();
-}
-
-
 void FCustomizableObjectEditorModule::StartupModule()
 {
-	UCustomizableObjectSystem::GetInstance()->SetNewCompilerFunc(NewCompiler);
-
 	// Register the thumbnail renderers
 	//UThumbnailManager::Get().RegisterCustomRenderer(UCustomizableObject::StaticClass(), UCustomizableObjectThumbnailRenderer::StaticClass());
 	//UThumbnailManager::Get().RegisterCustomRenderer(UCustomizableObjectInstance::StaticClass(), UCustomizableObjectInstanceThumbnailRenderer::StaticClass());
@@ -587,6 +580,18 @@ bool FCustomizableObjectEditorModule::IsCompilationOutOfDate(const UCustomizable
 }
 
 
+TSharedRef<FCustomizableObjectCompilerBase> FCustomizableObjectEditorModule::CreateCompiler() const
+{
+	return MakeShared<FCustomizableObjectCompiler>();
+}
+
+
+bool FCustomizableObjectEditorModule::IsRootObject(const UCustomizableObject& Object) const
+{
+	return GraphTraversal::IsRootObject(Object);
+}
+
+
 void FCustomizableObjectEditorModule::BakeCustomizableObjectInstance(UCustomizableObjectInstance* InTargetInstance, const FBakingConfiguration& InBakingConfig)
 {
 	UCustomizableObjectInstanceBaker* InstanceBaker = NewObject<UCustomizableObjectInstanceBaker>();
diff --git a/Engine/Plugins/Experimental/Mutable/Source/CustomizableObjectEditor/Private/MuCOE/CustomizableObjectEditorModule.h b/Engine/Plugins/Experimental/Mutable/Source/CustomizableObjectEditor/Private/MuCOE/CustomizableObjectEditorModule.h
--- a/Engine/Plugins/Experimental/Mutable/Source/CustomizableObjectEditor/Private/MuCOE/CustomizableObjectEditorModule.h
+++ b/Engine/Plugins/Experimental/Mutable/Source/CustomizableObjectEditor/Private/MuCOE/CustomizableObjectEditorModule.h
@@ -7,6 +7,7 @@
 #include "MuCOE/CustomizableObjectEditorLogger.h"
 #include "Toolkits/AssetEditorToolkit.h"
 
+class FCustomizableObjectCompilerBase;
 struct FBakingConfiguration;
 class USkeletalMeshComponent;
 class FPropertyEditorModule;
@@ -37,6 +38,8 @@ public:
 	// ICustomizableObjectEditorModule interface
 	virtual FCustomizableObjectEditorLogger& GetLogger() override;
 	virtual bool IsCompilationOutOfDate(const UCustomizableObject& Object, TArray<FName>* OutOfDatePackages) const override;
+	virtual TSharedRef<FCustomizableObjectCompilerBase> CreateCompiler() const override;
+	virtual bool IsRootObject(const UCustomizableObject& Object) const override;
 	virtual void BakeCustomizableObjectInstance(UCustomizableObjectInstance* InTargetInstance, const FBakingConfiguration& InBakingConfig) override;
 
 	virtual TSharedPtr<FExtensibilityManager> GetCustomizableObjectEditorToolBarExtensibilityManager() override { return CustomizableObjectEditor_ToolBarExtensibilityManager; }
diff --git a/Engine/Plugins/Experimental/Mutable/Source/CustomizableObjectEditor/Private/MuCOE/GraphTraversal.cpp b/Engine/Plugins/Experimental/Mutable/Source/CustomizableObjectEditor/Private/MuCOE/GraphTraversal.cpp
--- a/Engine/Plugins/Experimental/Mutable/Source/CustomizableObjectEditor/Private/MuCOE/GraphTraversal.cpp
+++ b/Engine/Plugins/Experimental/Mutable/Source/CustomizableObjectEditor/Private/MuCOE/GraphTraversal.cpp
@@ -573,4 +573,33 @@ TMultiMap<FGuid, UCustomizableObjectNodeObject*> GetNodeGroupObjectNodeMapping(U
 	GetNodeGroupObjectNodeMappingImmersive(Object, AssetRegistryModule, Visited, Mapping);
 	
 	return Mapping;
-}
\ No newline at end of file
+}
+
+
+namespace GraphTraversal
+{
+	bool IsRootObject(const UCustomizableObject& Object)
+	{
+		const TObjectPtr<UEdGraph> Source = Object.GetPrivate()->GetSource();
+		if (!Source || !Source->Nodes.Num())
+		{
+			// Conservative approach.
+			return true;
+		}
+
+		TArray<UCustomizableObjectNodeObject*> ObjectNodes;
+		Source->GetNodesOfClass<UCustomizableObjectNodeObject>(ObjectNodes);
+
+		// Look for the base object node
+		const UCustomizableObjectNodeObject* Root = nullptr;
+		for (TArray<UCustomizableObjectNodeObject*>::TIterator It(ObjectNodes); It; ++It)
+		{
+			if ((*It)->bIsBase)
+			{
+				Root = *It;
+			}
+		}
+
+		return Root && !Root->ParentObject;
+	}
+}
diff --git a/Engine/Plugins/Experimental/Mutable/Source/CustomizableObjectEditor/Public/MuCOE/GraphTraversal.h b/Engine/Plugins/Experimental/Mutable/Source/CustomizableObjectEditor/Public/MuCOE/GraphTraversal.h
--- a/Engine/Plugins/Experimental/Mutable/Source/CustomizableObjectEditor/Public/MuCOE/GraphTraversal.h
+++ b/Engine/Plugins/Experimental/Mutable/Source/CustomizableObjectEditor/Public/MuCOE/GraphTraversal.h
@@ -69,3 +69,9 @@ const UEdGraphPin* FindMeshBaseSource(const UEdGraphPin& Pin, const bool bOnlyLo
 /** Return the mapping of Group Object Nodes to Child Object Nodes of the given hierarchy.
  * @param Object Child or root Object to start the search from. */
 TMultiMap<FGuid, UCustomizableObjectNodeObject*> GetNodeGroupObjectNodeMapping(UCustomizableObject* Object);
+
+namespace GraphTraversal
+{
+    /** Return ture if the given Customizable Object is Root Object (not a Child Object). */
+    bool IsRootObject(const UCustomizableObject& Object); 
+}
diff --git a/Engine/Plugins/Experimental/Mutable/Source/MutableValidation/Private/MuV/AssetValidator_CustomizableObjects.cpp b/Engine/Plugins/Experimental/Mutable/Source/MutableValidation/Private/MuV/AssetValidator_CustomizableObjects.cpp
--- a/Engine/Plugins/Experimental/Mutable/Source/MutableValidation/Private/MuV/AssetValidator_CustomizableObjects.cpp
+++ b/Engine/Plugins/Experimental/Mutable/Source/MutableValidation/Private/MuV/AssetValidator_CustomizableObjects.cpp
@@ -89,8 +89,7 @@ EDataValidationResult UAssetValidator_CustomizableObjects::IsCustomizableObjectV
 	}
 	
 	// Request a compiler to be able to locate the root and to compile it
-	const TUniquePtr<FCustomizableObjectCompilerBase> Compiler =
-		TUniquePtr<FCustomizableObjectCompilerBase>(UCustomizableObjectSystem::GetInstance()->GetNewCompiler());
+	const TSharedPtr<FCustomizableObjectCompilerBase> Compiler = UCustomizableObjectSystem::GetNewCompiler();
 	
 	// Find out which is the root for this CO (it may be itself but that is OK)
 	UCustomizableObject* RootObject = GetRootObject(InCustomizableObject);
