[Mutable] Submit "Fix crash when launching the editor with -game." missing file

#rnx

diff --git a/Engine/Plugins/Experimental/Mutable/Source/MutableValidation/Private/MuV/CustomizableObjectCompilationUtility.cpp b/Engine/Plugins/Experimental/Mutable/Source/MutableValidation/Private/MuV/CustomizableObjectCompilationUtility.cpp
--- a/Engine/Plugins/Experimental/Mutable/Source/MutableValidation/Private/MuV/CustomizableObjectCompilationUtility.cpp
+++ b/Engine/Plugins/Experimental/Mutable/Source/MutableValidation/Private/MuV/CustomizableObjectCompilationUtility.cpp
@@ -62,7 +62,7 @@ bool FCustomizableObjectCompilationUtility::CompileCustomizableObject(UCustomiza
 		
 		// Create the compiler that will take care of the ASYNC compilation
 		UCustomizableObjectSystem* CustomizableObjectSystem = UCustomizableObjectSystem::GetInstanceChecked();
-		TUniquePtr<FCustomizableObjectCompilerBase> Compiler = TUniquePtr<FCustomizableObjectCompilerBase>(CustomizableObjectSystem->GetNewCompiler());
+		TSharedPtr<FCustomizableObjectCompilerBase> Compiler = CustomizableObjectSystem->GetNewCompiler();
 
 		// Get the memory usage before staring the compilation
 #if ENABLE_LOW_LEVEL_MEM_TRACKER
