[Mutable] Add compiled without optimization warning to Viewport.

#rnx

diff --git a/Engine/Plugins/Experimental/Mutable/Source/CustomizableObjectEditor/Private/MuCOE/SCustomizableObjectEditorViewport.cpp b/Engine/Plugins/Experimental/Mutable/Source/CustomizableObjectEditor/Private/MuCOE/SCustomizableObjectEditorViewport.cpp
--- a/Engine/Plugins/Experimental/Mutable/Source/CustomizableObjectEditor/Private/MuCOE/SCustomizableObjectEditorViewport.cpp
+++ b/Engine/Plugins/Experimental/Mutable/Source/CustomizableObjectEditor/Private/MuCOE/SCustomizableObjectEditorViewport.cpp
@@ -83,7 +83,18 @@ void SCustomizableObjectEditorViewport::PopulateViewportOverlays(TSharedRef<SOve
 	Overlay->AddSlot()
 	.VAlign(VAlign_Top)
 	[
-		SNew(SCustomizableObjectEditorViewportToolBar, TabBodyPtr.Pin(), SharedThis(this)).Cursor(EMouseCursor::Default)
+		SNew(SVerticalBox)
+		+SVerticalBox::Slot()
+		[
+			SNew(SCustomizableObjectEditorViewportToolBar, TabBodyPtr.Pin(), SharedThis(this)).Cursor(EMouseCursor::Default)
+		]
+		+SVerticalBox::Slot()
+		.Padding(4.0, 16.0, 0.0, 0.0)
+		[
+			SNew(STextBlock)
+			.Text(this, &SCustomizableObjectEditorViewport::GetWarningText)
+			.ColorAndOpacity(FLinearColor::Yellow)
+		]
 	];
 	Overlay->AddSlot()
 	.VAlign(VAlign_Center)
@@ -186,6 +197,23 @@ FText SCustomizableObjectEditorViewport::GetCompileErrorOverlayText() const
 }
 
 
+FText SCustomizableObjectEditorViewport::GetWarningText() const
+{
+	if (const UCustomizableObjectInstance* Instance = WeakEditor.Pin()->GetPreviewInstance())
+	{
+		if (const UCustomizableObject* Object = Instance->GetCustomizableObject())
+		{
+			if (Object->GetPrivate()->bIsCompiledWithoutOptimization)
+			{
+				return LOCTEXT("CompiledWithoutOptimization", "Compiled without maximum optimization. Updates will be slower!");
+			}
+		}
+	}
+
+	return {};
+}
+
+
 void SCustomizableObjectEditorViewport::OnUndoRedo()
 {
 	LevelViewportClient->Invalidate();
diff --git a/Engine/Plugins/Experimental/Mutable/Source/CustomizableObjectEditor/Private/MuCOE/SCustomizableObjectEditorViewport.h b/Engine/Plugins/Experimental/Mutable/Source/CustomizableObjectEditor/Private/MuCOE/SCustomizableObjectEditorViewport.h
--- a/Engine/Plugins/Experimental/Mutable/Source/CustomizableObjectEditor/Private/MuCOE/SCustomizableObjectEditorViewport.h
+++ b/Engine/Plugins/Experimental/Mutable/Source/CustomizableObjectEditor/Private/MuCOE/SCustomizableObjectEditorViewport.h
@@ -61,6 +61,8 @@ private:
 	EVisibility GetShowCompileErrorOverlay() const;
 
 	FText GetCompileErrorOverlayText() const;
+
+	FText GetWarningText() const;
 	
 	// Viewport client
 	TSharedPtr<FCustomizableObjectEditorViewportClient> LevelViewportClient;
