[mutable] Made Mutable textures respect the reference texture MipGenSettings when set to "NoMipmaps".
#rnx
#tests Editor

diff --git a/Engine/Plugins/Experimental/Mutable/Source/CustomizableObject/Internal/MuCO/CustomizableObjectPrivate.h b/Engine/Plugins/Experimental/Mutable/Source/CustomizableObject/Internal/MuCO/CustomizableObjectPrivate.h
--- a/Engine/Plugins/Experimental/Mutable/Source/CustomizableObject/Internal/MuCO/CustomizableObjectPrivate.h
+++ b/Engine/Plugins/Experimental/Mutable/Source/CustomizableObject/Internal/MuCO/CustomizableObjectPrivate.h
@@ -227,20 +227,22 @@ struct CUSTOMIZABLEOBJECT_API FMutableModelImageProperties
 		, FlipGreenChannel(0)
 		, IsPassThrough(0)
 		, LODBias(0)
+		, MipGenSettings(TextureMipGenSettings::TMGS_FromTextureGroup)
 		, LODGroup(TEXTUREGROUP_World)
 		, AddressX(TA_Clamp)
 		, AddressY(TA_Clamp)
 	{}
 
 	FMutableModelImageProperties(const FString& InTextureParameterName, TextureFilter InFilter, uint32 InSRGB,
-		uint32 InFlipGreenChannel, uint32 bInIsPassThrough, int32 InLODBias, TEnumAsByte<enum TextureGroup> InLODGroup,
-		TEnumAsByte<enum TextureAddress> InAddressX, TEnumAsByte<enum TextureAddress> InAddressY)
+		uint32 InFlipGreenChannel, uint32 bInIsPassThrough, int32 InLODBias, TEnumAsByte<TextureMipGenSettings> InMipGenSettings, 
+		TEnumAsByte<enum TextureGroup> InLODGroup, TEnumAsByte<enum TextureAddress> InAddressX, TEnumAsByte<enum TextureAddress> InAddressY)
 		: TextureParameterName(InTextureParameterName)
 		, Filter(InFilter)
 		, SRGB(InSRGB)
 		, FlipGreenChannel(InFlipGreenChannel)
 		, IsPassThrough(bInIsPassThrough)
 		, LODBias(InLODBias)
+		, MipGenSettings(InMipGenSettings)
 		, LODGroup(InLODGroup)
 		, AddressX(InAddressX)
 		, AddressY(InAddressY)
@@ -265,6 +267,9 @@ struct CUSTOMIZABLEOBJECT_API FMutableModelImageProperties
 	UPROPERTY()
 	int32 LODBias;
 
+	UPROPERTY()
+	TEnumAsByte<TextureMipGenSettings> MipGenSettings;
+
 	UPROPERTY()
 	TEnumAsByte<enum TextureGroup> LODGroup;
 
@@ -880,6 +885,6 @@ public:
 	// This is a manual version number for the binary blobs in this asset.
 	// Increasing it invalidates all the previously compiled models.
 	// Warning: If while merging code both versions have changed, take the highest+1.
-	static constexpr int32 CurrentSupportedVersion = 443;
+	static constexpr int32 CurrentSupportedVersion = 447;
 };
 
diff --git a/Engine/Plugins/Experimental/Mutable/Source/CustomizableObject/Private/MuCO/CustomizableObject.cpp b/Engine/Plugins/Experimental/Mutable/Source/CustomizableObject/Private/MuCO/CustomizableObject.cpp
--- a/Engine/Plugins/Experimental/Mutable/Source/CustomizableObject/Private/MuCO/CustomizableObject.cpp
+++ b/Engine/Plugins/Experimental/Mutable/Source/CustomizableObject/Private/MuCO/CustomizableObject.cpp
@@ -2333,6 +2333,7 @@ bool FMutableModelImageProperties::operator!=(const FMutableModelImageProperties
 		FlipGreenChannel != Other.FlipGreenChannel ||
 		IsPassThrough != Other.IsPassThrough ||
 		LODBias != Other.LODBias ||
+		MipGenSettings != Other.MipGenSettings ||
 		LODGroup != Other.LODGroup ||
 		AddressX != Other.AddressX ||
 		AddressY != Other.AddressY;
@@ -2404,6 +2405,7 @@ FArchive& operator<<(FArchive& Ar, FMutableModelImageProperties& ImageProps)
 	}
 
 	Ar << ImageProps.LODBias;
+	Ar << ImageProps.MipGenSettings;
 	Ar << ImageProps.LODGroup;
 
 	Ar << ImageProps.AddressX;
diff --git a/Engine/Plugins/Experimental/Mutable/Source/CustomizableObject/Private/MuCO/CustomizableObjectInstance.cpp b/Engine/Plugins/Experimental/Mutable/Source/CustomizableObject/Private/MuCO/CustomizableObjectInstance.cpp
--- a/Engine/Plugins/Experimental/Mutable/Source/CustomizableObject/Private/MuCO/CustomizableObjectInstance.cpp
+++ b/Engine/Plugins/Experimental/Mutable/Source/CustomizableObject/Private/MuCO/CustomizableObjectInstance.cpp
@@ -2786,7 +2786,16 @@ void SetTexturePropertiesFromMutableImageProps(UTexture2D* Texture, const FMutab
 	Texture->Filter = Props.Filter;
 	Texture->LODBias = Props.LODBias;
 
+	// TODO: MTBL-1957
+	//if (Props.MipGenSettings == TextureMipGenSettings::TMGS_NoMipmaps)
+	//{
+		//Texture->NeverStream = true;
+	//}
+
 #if WITH_EDITORONLY_DATA
+	// TODO: MTBL-1957
+	//Texture->MipGenSettings = Props.MipGenSettings;
+
 	Texture->bFlipGreenChannel = Props.FlipGreenChannel;
 #endif
 
diff --git a/Engine/Plugins/Experimental/Mutable/Source/CustomizableObjectEditor/Private/MuCOE/CustomizableObjectCompiler.cpp b/Engine/Plugins/Experimental/Mutable/Source/CustomizableObjectEditor/Private/MuCOE/CustomizableObjectCompiler.cpp
--- a/Engine/Plugins/Experimental/Mutable/Source/CustomizableObjectEditor/Private/MuCOE/CustomizableObjectCompiler.cpp
+++ b/Engine/Plugins/Experimental/Mutable/Source/CustomizableObjectEditor/Private/MuCOE/CustomizableObjectCompiler.cpp
@@ -1024,6 +1024,7 @@ void FCustomizableObjectCompiler::CompileInternal(bool bAsync)
 										ImageProp.bFlipGreenChannel,
 										ImageProp.bIsPassThrough,
 										ImageProp.LODBias,
+										ImageProp.MipGenSettings,
 										ImageProp.LODGroup,
 										ImageProp.AddressX, ImageProp.AddressY });
 		}
diff --git a/Engine/Plugins/Experimental/Mutable/Source/CustomizableObjectEditor/Private/MuCOE/GenerateMutableSource/GenerateMutableSourceSurface.cpp b/Engine/Plugins/Experimental/Mutable/Source/CustomizableObjectEditor/Private/MuCOE/GenerateMutableSource/GenerateMutableSourceSurface.cpp
--- a/Engine/Plugins/Experimental/Mutable/Source/CustomizableObjectEditor/Private/MuCOE/GenerateMutableSource/GenerateMutableSourceSurface.cpp
+++ b/Engine/Plugins/Experimental/Mutable/Source/CustomizableObjectEditor/Private/MuCOE/GenerateMutableSource/GenerateMutableSourceSurface.cpp
@@ -631,28 +631,31 @@ mu::NodeSurfacePtr GenerateMutableSourceSurface(const UEdGraphPin * Pin, FMutabl
 							const uint32 BaseLODBias = ComputeLODBiasForTexture(GenerationContext, *ReferenceTexture) + SurfaceLODBias;
 							mu::NodeImagePtr LastImage = ResizeTextureByNumMips(ImageNode, BaseLODBias);
 
-							mu::EMipmapFilterType MipGenerationFilterType = Invoke([&]()
+							if (ReferenceTexture->MipGenSettings != TextureMipGenSettings::TMGS_NoMipmaps)
 							{
-								if (ReferenceTexture)
+								mu::EMipmapFilterType MipGenerationFilterType = Invoke([&]()
 								{
-									switch(ReferenceTexture->MipGenSettings)
+									if (ReferenceTexture)
 									{
-									case TextureMipGenSettings::TMGS_SimpleAverage: return mu::EMipmapFilterType::SimpleAverage;
-									case TextureMipGenSettings::TMGS_Unfiltered:    return mu::EMipmapFilterType::Unfiltered;
-									default: return mu::EMipmapFilterType::SimpleAverage;
+										switch (ReferenceTexture->MipGenSettings)
+										{
+											case TextureMipGenSettings::TMGS_SimpleAverage: return mu::EMipmapFilterType::SimpleAverage;
+											case TextureMipGenSettings::TMGS_Unfiltered:    return mu::EMipmapFilterType::Unfiltered;
+											default: return mu::EMipmapFilterType::SimpleAverage;
+										}
 									}
-								}
 
-								return mu::EMipmapFilterType::SimpleAverage;
-							});
+									return mu::EMipmapFilterType::SimpleAverage;
+								});
 
 
-							mu::NodeImageMipmapPtr MipmapImage = new mu::NodeImageMipmap();
-							MipmapImage->SetSource(LastImage.get());
-							MipmapImage->SetMipmapGenerationSettings(MipGenerationFilterType, mu::EAddressMode::None);
+								mu::NodeImageMipmapPtr MipmapImage = new mu::NodeImageMipmap();
+								MipmapImage->SetSource(LastImage.get());
+								MipmapImage->SetMipmapGenerationSettings(MipGenerationFilterType, mu::EAddressMode::None);
 
-							MipmapImage->SetMessageContext(Node);
-							LastImage = MipmapImage;
+								MipmapImage->SetMessageContext(Node);
+								LastImage = MipmapImage;
+							}
 
 							// Apply composite image. This needs to be computed after mipmaps generation. 	
 							if (ReferenceTexture && ReferenceTexture->GetCompositeTexture() && ReferenceTexture->CompositeTextureMode != CTM_Disabled)
@@ -665,12 +668,12 @@ mu::NodeSurfacePtr GenerateMutableSourceSurface(const UEdGraphPin * Pin, FMutabl
 								{
 									switch (CompositeTextureMode)
 									{
-									case CTM_NormalRoughnessToRed: return mu::ECompositeImageMode::CIM_NormalRoughnessToRed;
-									case CTM_NormalRoughnessToGreen: return mu::ECompositeImageMode::CIM_NormalRoughnessToGreen;
-									case CTM_NormalRoughnessToBlue: return mu::ECompositeImageMode::CIM_NormalRoughnessToBlue;
-									case CTM_NormalRoughnessToAlpha: return mu::ECompositeImageMode::CIM_NormalRoughnessToAlpha;
+										case CTM_NormalRoughnessToRed: return mu::ECompositeImageMode::CIM_NormalRoughnessToRed;
+										case CTM_NormalRoughnessToGreen: return mu::ECompositeImageMode::CIM_NormalRoughnessToGreen;
+										case CTM_NormalRoughnessToBlue: return mu::ECompositeImageMode::CIM_NormalRoughnessToBlue;
+										case CTM_NormalRoughnessToAlpha: return mu::ECompositeImageMode::CIM_NormalRoughnessToAlpha;
 
-									default: return mu::ECompositeImageMode::CIM_Disabled;
+										default: return mu::ECompositeImageMode::CIM_Disabled;
 									}
 								}();
 
