[mutable] Fix for CIS raised error (Sizeof of incomplete type) when setting a TStrongObjectPtr value to nullptr

#rnx

diff --git a/Engine/Plugins/Experimental/Mutable/Source/MutableValidation/Private/MuV/CustomizableObjectCompilationUtility.h b/Engine/Plugins/Experimental/Mutable/Source/MutableValidation/Private/MuV/CustomizableObjectCompilationUtility.h
--- a/Engine/Plugins/Experimental/Mutable/Source/MutableValidation/Private/MuV/CustomizableObjectCompilationUtility.h
+++ b/Engine/Plugins/Experimental/Mutable/Source/MutableValidation/Private/MuV/CustomizableObjectCompilationUtility.h
@@ -1,11 +1,11 @@
 ﻿// Copyright Epic Games, Inc. All Rights Reserved.
 #pragma once
 
+#include "MuCO/CustomizableObject.h"
 #include "UObject/StrongObjectPtr.h"
 #include "Templates/SharedPointer.h"
 
 // Forward declarations
-class UCustomizableObject;
 struct FCompilationOptions;
 
 /**
diff --git a/Engine/Plugins/Experimental/Mutable/Source/MutableValidation/Private/MuV/CustomizableObjectInstanceUpdateUtility.h b/Engine/Plugins/Experimental/Mutable/Source/MutableValidation/Private/MuV/CustomizableObjectInstanceUpdateUtility.h
--- a/Engine/Plugins/Experimental/Mutable/Source/MutableValidation/Private/MuV/CustomizableObjectInstanceUpdateUtility.h
+++ b/Engine/Plugins/Experimental/Mutable/Source/MutableValidation/Private/MuV/CustomizableObjectInstanceUpdateUtility.h
@@ -1,13 +1,12 @@
 ﻿// Copyright Epic Games, Inc. All Rights Reserved.
 #pragma once
 
+#include "MuCO/CustomizableObjectInstance.h"
 #include "UObject/StrongObjectPtr.h"
 #include "Templates/SharedPointer.h"
 
 // Forward declarations
-class UCustomizableObjectInstance;
 class USkeletalMeshComponent;
-struct FUpdateContext;
 
 /**
  * Helping class that handles the async update of the provided instance. It will also wait for the mips of it so they get streamed.
