[mutable] Fixed Groups with the "None" option selected showing their child goup parameters in the CO Editor preview UI.
#rnx
#tests Editor

diff --git a/Engine/Plugins/Experimental/Mutable/Source/CustomizableObjectEditor/Private/MuCOE/GenerateMutableSource/GenerateMutableSource.cpp b/Engine/Plugins/Experimental/Mutable/Source/CustomizableObjectEditor/Private/MuCOE/GenerateMutableSource/GenerateMutableSource.cpp
--- a/Engine/Plugins/Experimental/Mutable/Source/CustomizableObjectEditor/Private/MuCOE/GenerateMutableSource/GenerateMutableSource.cpp
+++ b/Engine/Plugins/Experimental/Mutable/Source/CustomizableObjectEditor/Private/MuCOE/GenerateMutableSource/GenerateMutableSource.cpp
@@ -1311,10 +1311,22 @@ mu::NodeObjectPtr GenerateMutableSource(const UEdGraphPin * Pin, FMutableGraphGe
 			{
 				UCustomizableObjectNodeObject* NodeObject = Cast<UCustomizableObjectNodeObject>(ConnectedPin->GetOwningNode());
 
-				if (NodeObject && NodeObject->bIsBase)
+				if (NodeObject)
 				{
 					const FGuid* ParentId = GenerationContext.GroupIdToExternalNodeMap.FindKey(NodeObject);
 
+					// Group objects in the same graph aren't in the GroupIdToExternalNodeMap, so follow the pins instead
+					if (!ParentId && NodeObject->OutputPin())
+					{
+						if (const UEdGraphPin* ConnectedPinToObject = FollowOutputPin(*NodeObject->OutputPin()))
+						{
+							if (UCustomizableObjectNodeObjectGroup* ParentGroupNode = Cast<UCustomizableObjectNodeObjectGroup>(ConnectedPinToObject->GetOwningNode()))
+							{
+								ParentId = &ParentGroupNode->NodeGuid;
+							}
+						}
+					}
+
 					if (ParentId)
 					{
 						FString* ParentParamName = GenerationContext.GuidToParamNameMap.Find(*ParentId);
