[mutable] Fixed some table node bugs:
+ Now material nodes linked to a mesh of a table node will be able to get the slot name from the default skeletal mesh.
+ Now table nodes will be able to interpret UMaterialInterfaces. A warning will erase if this material interface is a UMaterial which are unsupported right now.

#rnx

diff --git a/Engine/Plugins/Experimental/Mutable/Source/CustomizableObjectEditor/Classes/MuCOE/Nodes/CustomizableObjectNodeTable.h b/Engine/Plugins/Experimental/Mutable/Source/CustomizableObjectEditor/Classes/MuCOE/Nodes/CustomizableObjectNodeTable.h
--- a/Engine/Plugins/Experimental/Mutable/Source/CustomizableObjectEditor/Classes/MuCOE/Nodes/CustomizableObjectNodeTable.h
+++ b/Engine/Plugins/Experimental/Mutable/Source/CustomizableObjectEditor/Classes/MuCOE/Nodes/CustomizableObjectNodeTable.h
@@ -21,6 +21,7 @@ class UTexture2D;
 class UTexture2DArray;
 class UAnimInstance;
 struct FGuid;
+struct FSkeletalMaterial;
 
 
 /** Enum class for the different types of image pins */
@@ -400,6 +401,9 @@ public:
 
 	/** Return the list of UDataTable that will be used to compose the final UDataTable. */
 	TArray<FAssetData> GetParentTables() const;
+
+	/** Returns the skeletal material associated to the given skeletal mesh output pin. */
+	FSkeletalMaterial* GetDefaultSkeletalMaterialFor(const UEdGraphPin& MeshPin) const;
 	
 private:
 
diff --git a/Engine/Plugins/Experimental/Mutable/Source/CustomizableObjectEditor/Private/MuCOE/GenerateMutableSource/GenerateMutableSourceSurface.cpp b/Engine/Plugins/Experimental/Mutable/Source/CustomizableObjectEditor/Private/MuCOE/GenerateMutableSource/GenerateMutableSourceSurface.cpp
--- a/Engine/Plugins/Experimental/Mutable/Source/CustomizableObjectEditor/Private/MuCOE/GenerateMutableSource/GenerateMutableSourceSurface.cpp
+++ b/Engine/Plugins/Experimental/Mutable/Source/CustomizableObjectEditor/Private/MuCOE/GenerateMutableSource/GenerateMutableSourceSurface.cpp
@@ -295,7 +295,7 @@ mu::NodeSurfacePtr GenerateMutableSourceSurface(const UEdGraphPin * Pin, FMutabl
 		if (TypedNodeMat->Material)
 		{
 			GenerationContext.AddParticipatingObject(*TypedNodeMat->Material);
-			
+
 			const int32 lastMaterialAmount = GenerationContext.ReferencedMaterials.Num();
 			ReferencedMaterialsIndex = GenerationContext.ReferencedMaterials.AddUnique(TypedNodeMat->Material);
 			// Used ReferencedMaterialsIndex instead of TypedNodeMat->Material->GetName() to prevent material name collisions
@@ -312,23 +312,30 @@ mu::NodeSurfacePtr GenerateMutableSourceSurface(const UEdGraphPin * Pin, FMutabl
 				{
 					if (const UEdGraphPin* SkeletalMeshPin = FindMeshBaseSource(*ConnectedMaterialPin, false))
 					{
+						FSkeletalMaterial* SkeletalMaterial = nullptr;
+
 						if (const UCustomizableObjectNodeSkeletalMesh* SkeletalMeshNode = Cast<UCustomizableObjectNodeSkeletalMesh>(SkeletalMeshPin->GetOwningNode()))
 						{
-							if (const FSkeletalMaterial* SkeletalMaterial = SkeletalMeshNode->GetSkeletalMaterialFor(*SkeletalMeshPin))
+							SkeletalMaterial = SkeletalMeshNode->GetSkeletalMaterialFor(*SkeletalMeshPin);
+						}
+
+						else if (const UCustomizableObjectNodeTable* TableNode = Cast<UCustomizableObjectNodeTable>(SkeletalMeshPin->GetOwningNode()))
+						{
+							SkeletalMaterial = TableNode->GetDefaultSkeletalMaterialFor(*SkeletalMeshPin);
+						}
+
+						if (SkeletalMaterial)
+						{
+							if (IsNewSlotName)
 							{
-								if (IsNewSlotName)
-								{
-									GenerationContext.ReferencedMaterialSlotNames.Add(SkeletalMaterial->MaterialSlotName);
-									SlotNameFound = true;
-								}
-								else if (GenerationContext.ReferencedMaterialSlotNames[ReferencedMaterialsIndex].IsNone())
-								{
-									GenerationContext.ReferencedMaterialSlotNames[ReferencedMaterialsIndex] = SkeletalMaterial->MaterialSlotName;
-									SlotNameFound = true;
-								}
+								GenerationContext.ReferencedMaterialSlotNames.Add(SkeletalMaterial->MaterialSlotName);
+								SlotNameFound = true;
+							}
+							else if (GenerationContext.ReferencedMaterialSlotNames[ReferencedMaterialsIndex].IsNone())
+							{
+								GenerationContext.ReferencedMaterialSlotNames[ReferencedMaterialsIndex] = SkeletalMaterial->MaterialSlotName;
+								SlotNameFound = true;
 							}
-
-							//TODO(max): Add support for table nodes. Convert this to a function as it will be also used in the GenerateMutableSourceTable file
 						}
 					}
 				}
diff --git a/Engine/Plugins/Experimental/Mutable/Source/CustomizableObjectEditor/Private/MuCOE/GenerateMutableSource/GenerateMutableSourceTable.cpp b/Engine/Plugins/Experimental/Mutable/Source/CustomizableObjectEditor/Private/MuCOE/GenerateMutableSource/GenerateMutableSourceTable.cpp
--- a/Engine/Plugins/Experimental/Mutable/Source/CustomizableObjectEditor/Private/MuCOE/GenerateMutableSource/GenerateMutableSourceTable.cpp
+++ b/Engine/Plugins/Experimental/Mutable/Source/CustomizableObjectEditor/Private/MuCOE/GenerateMutableSource/GenerateMutableSourceTable.cpp
@@ -28,7 +28,7 @@
 #define LOCTEXT_NAMESPACE "CustomizableObjectEditor"
 
 
-bool FillTableColumn(const UCustomizableObjectNodeTable* TableNode,	mu::TablePtr MutableTable,	const FString& ColumnName,	const FString& RowName,	const int32 RowIdx,	uint8* CellData, const FProperty* ColumnProperty,
+bool FillTableColumn(const UCustomizableObjectNodeTable* TableNode, mu::TablePtr MutableTable,	const FString& ColumnName,	const FString& RowName,	const int32 RowIdx,	uint8* CellData, const FProperty* ColumnProperty,
 	const int LODIndexConnected, const int32 SectionIndexConnected, int32 LODIndex, int32 SectionIndex, const bool bOnlyConnectedLOD, FMutableGraphGenerationContext& GenerationContext)
 {
 	int32 CurrentColumn;
@@ -464,17 +464,21 @@ bool FillTableColumn(const UCustomizableObjectNodeTable* TableNode,	mu::TablePtr
 			MutableTable->SetCell(CurrentColumn, RowIdx, Proxy.get());
 		}
 
-		else if (SoftObjectProperty->PropertyClass->IsChildOf(UMaterialInstance::StaticClass()))
+		else if (SoftObjectProperty->PropertyClass->IsChildOf(UMaterialInterface::StaticClass()))
 		{
-			// Getting the name of material column of the data table
-			FString MaterialColumnName = ColumnProperty->GetDisplayNameText().ToString();
+			// Get display name of the column of the data table (name showed in the table and struct editors)
+			// Will be used in the warnings to help to identify a column with errors.
+			FString MaterialColumnDisplayName = ColumnProperty->GetDisplayNameText().ToString();
+			
+			// Get the real name of the Property column
+			FString MaterialColumnName = ColumnProperty->GetName();
 
-			UMaterialInstance* Material = Cast<UMaterialInstance>(Object);
+			UMaterialInstance* MaterialInstance = Cast<UMaterialInstance>(Object);
 			UMaterialInstance* ReferenceMaterial = TableNode->GetColumnDefaultAssetByType<UMaterialInstance>(MaterialColumnName);
 
 			if (!ReferenceMaterial)
 			{
-				FString msg = FString::Printf(TEXT("Reference Material not found for column [%s]."), *MaterialColumnName);
+				FString msg = FString::Printf(TEXT("Default Material Instance not found for column [%s]."), *MaterialColumnDisplayName);
 				GenerationContext.Compiler->CompilerLog(FText::FromString(msg), TableNode);
 
 				return false;
@@ -483,30 +487,38 @@ bool FillTableColumn(const UCustomizableObjectNodeTable* TableNode,	mu::TablePtr
 			GenerationContext.AddParticipatingObject(*ReferenceMaterial);
 
 			const bool bTableMaterialCheckDisabled = GenerationContext.Object->bDisableTableMaterialsParentCheck;
-			const bool bMaterialParentMismatch = !bTableMaterialCheckDisabled && Material 
-												 && ReferenceMaterial->GetMaterial() != Material->GetMaterial();
+			const bool bMaterialParentMismatch = !bTableMaterialCheckDisabled && MaterialInstance
+												 && ReferenceMaterial->GetMaterial() != MaterialInstance->GetMaterial();
 
-			if (!Material || bMaterialParentMismatch)
+			if (!MaterialInstance || bMaterialParentMismatch)
 			{
 				FText Warning;
 
-				if (!Material)
+				if (!MaterialInstance)
 				{
-					Warning = FText::Format(LOCTEXT("NullMaterialInstance", "Material Instance from column [{0}] row [{1}] is null. The default Material Instance will be used instead."),
-						FText::FromString(MaterialColumnName), FText::FromString(RowName));
+					if (UMaterial* Material = Cast<UMaterial>(Object))
+					{
+						Warning = FText::Format(LOCTEXT("IsAMaterial", "Asset from column [{0}] row [{1}] is a Material and not a MaterialInstance. The default Material Instance will be used instead."),
+							FText::FromString(MaterialColumnDisplayName), FText::FromString(RowName));
+					}
+					else
+					{
+						Warning = FText::Format(LOCTEXT("NullMaterialInstance", "Material Instance from column [{0}] row [{1}] is null. The default Material Instance will be used instead."),
+							FText::FromString(MaterialColumnDisplayName), FText::FromString(RowName));
+					}
 				}
 				else
 				{
 					Warning = FText::Format(LOCTEXT("MatInstanceFromDifferentParent","Material Instance from column [{0}] row [{1}] has a different Material Parent than the Default Material Instance. The Default Material Instance will be used instead."),
-						FText::FromString(MaterialColumnName), FText::FromString(RowName));
+						FText::FromString(MaterialColumnDisplayName), FText::FromString(RowName));
 				}
 
-				Material = ReferenceMaterial;
+				MaterialInstance = ReferenceMaterial;
 
 				LogRowGenerationMessage(TableNode, DataTablePtr, GenerationContext, Warning.ToString(), RowName);
 			}
 
-			GenerationContext.AddParticipatingObject(*Material);
+			GenerationContext.AddParticipatingObject(*MaterialInstance);
 			
 			FString EncodedSwitchParameterName = "__MutableMaterialId";
 			if (ColumnName.Contains(EncodedSwitchParameterName))
@@ -519,7 +531,7 @@ bool FillTableColumn(const UCustomizableObjectNodeTable* TableNode,	mu::TablePtr
 				}
 
 				const int32 lastMaterialAmount = GenerationContext.ReferencedMaterials.Num();
-				int32 ReferenceMaterialId = GenerationContext.ReferencedMaterials.AddUnique(Material);
+				int32 ReferenceMaterialId = GenerationContext.ReferencedMaterials.AddUnique(MaterialInstance);
 
 				// Take slot name from skeletal mesh if one can be found, else leave empty.
 				// Keep Referenced Materials and Materail Slot Names synchronized even if no material name can be found.
@@ -541,7 +553,7 @@ bool FillTableColumn(const UCustomizableObjectNodeTable* TableNode,	mu::TablePtr
 			TArray<FMaterialParameterInfo> ParameterInfos;
 			TArray<FGuid> ParameterGuids;
 
-			Material->GetMaterial()->GetAllParameterInfoOfType(EMaterialParameterType::Texture, ParameterInfos, ParameterGuids);
+			MaterialInstance->GetMaterial()->GetAllParameterInfoOfType(EMaterialParameterType::Texture, ParameterInfos, ParameterGuids);
 			
 			FGuid ParameterId(GenerationContext.CurrentMaterialTableParameterId);
 			int32 ParameterIndex = ParameterGuids.Find(ParameterId);
@@ -558,7 +570,7 @@ bool FillTableColumn(const UCustomizableObjectNodeTable* TableNode,	mu::TablePtr
 				}
 
 				UTexture* ParentTextureValue = nullptr;
-				Material->GetMaterial()->GetTextureParameterValue(ParameterInfos[ParameterIndex], ParentTextureValue);
+				MaterialInstance->GetMaterial()->GetTextureParameterValue(ParameterInfos[ParameterIndex], ParentTextureValue);
 				
 				UTexture2D* ParentParameterTexture = Cast<UTexture2D>(ParentTextureValue);
 				if (!ParentParameterTexture)
@@ -566,14 +578,14 @@ bool FillTableColumn(const UCustomizableObjectNodeTable* TableNode,	mu::TablePtr
 					FString ParamName = ParameterInfos[ParameterIndex].Name.ToString();
 					FString Message = Cast<UObject>(ParentParameterTexture) ? "not a Texture2D" : "null";
 					
-					FString msg = FString::Printf(TEXT("Parameter [%s] from Default Material Instance of column [%s] is %s. This parameter will be ignored."), *ParamName, *MaterialColumnName, *Message);
+					FString msg = FString::Printf(TEXT("Parameter [%s] from Default Material Instance of column [%s] is %s. This parameter will be ignored."), *ParamName, *MaterialColumnDisplayName, *Message);
 					LogRowGenerationMessage(TableNode, DataTablePtr, GenerationContext, msg, RowName);
 					 
 					 return false;
 				}
 
 				UTexture* TextureValue = nullptr;
-				Material->GetTextureParameterValue(ParameterInfos[ParameterIndex], TextureValue);
+				MaterialInstance->GetTextureParameterValue(ParameterInfos[ParameterIndex], TextureValue);
 
 				UTexture2D* ParameterTexture = Cast<UTexture2D>(TextureValue);
 
@@ -584,7 +596,7 @@ bool FillTableColumn(const UCustomizableObjectNodeTable* TableNode,	mu::TablePtr
 					FString ParamName = GenerationContext.CurrentMaterialTableParameter;
 					FString Message = Cast<UObject>(TextureValue) ? "not a Texture2D" : "null";
 
-					FString msg = FString::Printf(TEXT("Parameter [%s] from material instance of column [%s] row [%s] is %s. The parameter texture of the default material will be used instead."), *ParamName, *MaterialColumnName, *RowName, *Message);
+					FString msg = FString::Printf(TEXT("Parameter [%s] from material instance of column [%s] row [%s] is %s. The parameter texture of the default material will be used instead."), *ParamName, *MaterialColumnDisplayName, *RowName, *Message);
 					LogRowGenerationMessage(TableNode, DataTablePtr, GenerationContext, msg, RowName);
 				}
 
@@ -599,6 +611,9 @@ bool FillTableColumn(const UCustomizableObjectNodeTable* TableNode,	mu::TablePtr
 		else
 		{
 			// Unsuported Variable Type
+			FString msg = FString::Printf(TEXT("[%s] is not a supported class for mutable nodes."), *SoftObjectProperty->PropertyClass.GetName());
+			GenerationContext.Compiler->CompilerLog(FText::FromString(msg), TableNode);
+
 			return false;
 		}
 	}
@@ -656,7 +671,7 @@ bool FillTableColumn(const UCustomizableObjectNodeTable* TableNode,	mu::TablePtr
 
 	else
 	{
-		// Unsuported Variable Type
+		// Unsuported Variable Type		
 		return false;
 	}
 
diff --git a/Engine/Plugins/Experimental/Mutable/Source/CustomizableObjectEditor/Private/MuCOE/Nodes/CustomizableObjectNodeTable.cpp b/Engine/Plugins/Experimental/Mutable/Source/CustomizableObjectEditor/Private/MuCOE/Nodes/CustomizableObjectNodeTable.cpp
--- a/Engine/Plugins/Experimental/Mutable/Source/CustomizableObjectEditor/Private/MuCOE/Nodes/CustomizableObjectNodeTable.cpp
+++ b/Engine/Plugins/Experimental/Mutable/Source/CustomizableObjectEditor/Private/MuCOE/Nodes/CustomizableObjectNodeTable.cpp
@@ -723,7 +723,6 @@ bool UCustomizableObjectNodeTable::IsNodeOutDatedAndNeedsRefresh()
 
 					NumPins++;
 				}
-				
 			}
 			else if (const FStructProperty* StructProperty = CastField<FStructProperty>(ColumnProperty))
 			{
@@ -1403,4 +1402,49 @@ TArray<FAssetData> UCustomizableObjectNodeTable::GetParentTables() const
 }
 
 
+FSkeletalMaterial* UCustomizableObjectNodeTable::GetDefaultSkeletalMaterialFor(const UEdGraphPin& MeshPin) const
+{
+	USkeletalMesh* SkeletalMesh = GetColumnDefaultAssetByType<USkeletalMesh>(&MeshPin);
+
+	if (!SkeletalMesh)
+	{
+		return nullptr;
+	}
+
+	int32 LODIndex;
+	int32 SectionIndex;
+	GetPinLODAndSection(&MeshPin, LODIndex, SectionIndex);
+
+	// We assume that LODIndex and MaterialIndex are valid for the imported model
+	int32 SkeletalMeshMaterialIndex = INDEX_NONE;
+
+	// Check if we have LOD info map to get the correct material index
+	if (const FSkeletalMeshLODInfo* LodInfo = SkeletalMesh->GetLODInfo(LODIndex))
+	{
+		if (LodInfo->LODMaterialMap.IsValidIndex(SectionIndex))
+		{
+			SkeletalMeshMaterialIndex = LodInfo->LODMaterialMap[SectionIndex];
+		}
+	}
+
+	// Only deduce index when the explicit mapping is not found or there is no remap
+	if (SkeletalMeshMaterialIndex == INDEX_NONE)
+	{
+		const FSkeletalMeshModel* ImportedModel = SkeletalMesh->GetImportedModel();
+
+		if (ImportedModel && ImportedModel->LODModels.IsValidIndex(LODIndex) && ImportedModel->LODModels[LODIndex].Sections.IsValidIndex(SectionIndex))
+		{
+			SkeletalMeshMaterialIndex = ImportedModel->LODModels[LODIndex].Sections[SectionIndex].MaterialIndex;
+		}
+	}
+
+	if (SkeletalMesh->GetMaterials().IsValidIndex(SkeletalMeshMaterialIndex))
+	{
+		return &SkeletalMesh->GetMaterials()[SkeletalMeshMaterialIndex];
+	}
+
+	return nullptr;
+}
+
+
 #undef LOCTEXT_NAMESPACE
