[mutable] Fixed LOCTEXT conflict.

#rnx

diff --git a/Engine/Plugins/Experimental/Mutable/Source/CustomizableObjectEditor/Private/MuCOE/CustomizableInstanceDetails.cpp b/Engine/Plugins/Experimental/Mutable/Source/CustomizableObjectEditor/Private/MuCOE/CustomizableInstanceDetails.cpp
--- a/Engine/Plugins/Experimental/Mutable/Source/CustomizableObjectEditor/Private/MuCOE/CustomizableInstanceDetails.cpp
+++ b/Engine/Plugins/Experimental/Mutable/Source/CustomizableObjectEditor/Private/MuCOE/CustomizableInstanceDetails.cpp
@@ -1099,7 +1099,7 @@ void FCustomizableInstanceDetails::OnFloatParameterChanged(float Value, FString
 
 void FCustomizableInstanceDetails::OnFloatParameterSliderBegin()
 {
-	BeginTransaction(LOCTEXT("OnFloatParameterCommited", "Set Float Slider"));
+	BeginTransaction(LOCTEXT("OnFloatParameterSliderBegin", "Set Float Slider"));
 	bUpdatingSlider = true;
 }
 
