Skeletal Mesh: Build fixes for deprecated GetLODInfoArray

#rnx

diff --git a/Engine/Plugins/Experimental/Mutable/Source/CustomizableObject/Private/MuCO/CustomizableObjectInstance.cpp b/Engine/Plugins/Experimental/Mutable/Source/CustomizableObject/Private/MuCO/CustomizableObjectInstance.cpp
--- a/Engine/Plugins/Experimental/Mutable/Source/CustomizableObject/Private/MuCO/CustomizableObjectInstance.cpp
+++ b/Engine/Plugins/Experimental/Mutable/Source/CustomizableObject/Private/MuCO/CustomizableObjectInstance.cpp
@@ -6005,7 +6005,7 @@ void UCustomizableInstancePrivate::BuildMaterials(const TSharedRef<FUpdateContex
 			// Copy data from the FirstGeneratedLOD into the LODs below.
 			for (int32 LODIndex = OperationData->FirstLODAvailable; LODIndex < FirstGeneratedLOD; ++LODIndex)
 			{
-				SkeletalMesh->GetLODInfo(LODIndex)->LODMaterialMap = SkeletalMesh->GetLODInfo(LODIndex)->LODMaterialMap;
+				SkeletalMesh->GetLODInfo(LODIndex)->LODMaterialMap = SkeletalMesh->GetLODInfo(FirstGeneratedLOD)->LODMaterialMap;
 
 				TIndirectArray<FSkeletalMeshLODRenderData>& LODRenderData = SkeletalMesh->GetResourceForRendering()->LODRenderData;
 
