[Mutable] Clothing data ideterminism fixes. Added cvar to disable clothing mesh edits propagation.

- Removed use of mesh tags for clothing physics asset loading.
- Now clothing vertex data is loaded from disk asyncronously after the mutable mesh generation.
- Mutable mesh clothing buffers are only present if the source mesh has clothing. This can reduce siginicantly the size of constant meshes in disk when compiling with clothing enabled.
- Added CO compilation options to debugger.
- Added Sim mesh disable flag to CO. Currently not exposed in editor.
- Fix bug where only the first LOD would be explored to determine which assets to load.
- Changed how clothing asset names are generated, now names will only change if clashes are found.
- Added cvar mutable.DisableClothingPhysicsEditsPropagation.

#rnx

diff --git a/Engine/Plugins/Experimental/Mutable/Source/CustomizableObject/Classes/MuCO/CustomizableObject.h b/Engine/Plugins/Experimental/Mutable/Source/CustomizableObject/Classes/MuCO/CustomizableObject.h
--- a/Engine/Plugins/Experimental/Mutable/Source/CustomizableObject/Classes/MuCO/CustomizableObject.h
+++ b/Engine/Plugins/Experimental/Mutable/Source/CustomizableObject/Classes/MuCO/CustomizableObject.h
@@ -213,7 +213,7 @@ struct FCustomizableObjectMeshToMeshVertData
 	explicit FCustomizableObjectMeshToMeshVertData(const FMeshToMeshVertData& Original)
 		: Weight(Original.Weight)
 	{
-		for (int i = 0; i < 4; ++i)
+		for (int32 i = 0; i < 4; ++i)
 		{
 			PositionBaryCoordsAndDist[i] = Original.PositionBaryCoordsAndDist[i];
 			NormalBaryCoordsAndDist[i] = Original.NormalBaryCoordsAndDist[i];
@@ -227,7 +227,7 @@ struct FCustomizableObjectMeshToMeshVertData
 	{
 		FMeshToMeshVertData ReturnValue;
 
-		for (int i = 0; i < 4; ++i)
+		for (int32 i = 0; i < 4; ++i)
 		{
 			ReturnValue.PositionBaryCoordsAndDist[i] = PositionBaryCoordsAndDist[i];
 			ReturnValue.NormalBaryCoordsAndDist[i] = NormalBaryCoordsAndDist[i];
@@ -235,6 +235,7 @@ struct FCustomizableObjectMeshToMeshVertData
 			ReturnValue.SourceMeshVertIndices[i] = SourceMeshVertIndices[i];
 		}
 		ReturnValue.Weight = Weight;
+		ReturnValue.Padding = 0;
 
 		return ReturnValue;
 	}
@@ -263,13 +264,8 @@ struct FCustomizableObjectMeshToMeshVertData
 	UPROPERTY()
 	float Weight = 0.0f;
 
-	// Dummy for alignment (8 bytes). Originally not used.
-	UPROPERTY()
-	int16 SourceAssetIndex = 0;
-	
-	//
-	UPROPERTY()
-	int16 SourceAssetLodIndex = 0;
+	// Non serialized, unused padding. This is present in the FMeshToMeshVertData struct as Padding for alignment.
+	uint32 UnusedPadding = 0;
 
 	/**
 	 * Serializer
@@ -300,11 +296,10 @@ struct FCustomizableObjectMeshToMeshVertData
 
 		Ar << V.Weight;
 
-		Ar << V.SourceAssetIndex;
-		Ar << V.SourceAssetLodIndex;
 		return Ar;
 	}
 };
+static_assert(sizeof(FCustomizableObjectMeshToMeshVertData) == sizeof(float)*4*3 + sizeof(uint16)*4 + sizeof(float) + sizeof(uint32));
 template<> struct TCanBulkSerialize<FCustomizableObjectMeshToMeshVertData> { enum { Value = true }; };
 
 USTRUCT()
@@ -403,7 +398,8 @@ private:
 	{
 		None = 0,
 		Model,
-		RealTimeMorph
+		RealTimeMorph,
+		Clothing,
 	};
 
 	struct FBlock
@@ -491,15 +487,6 @@ public:
 	UPROPERTY(EditAnywhere, Category = CustomizableObject, meta = (DisplayName = "LOD Settings"))
 	FMutableLODSettings LODSettings;
 	
-	UPROPERTY(Transient)
-	TArray<FCustomizableObjectClothConfigData> ClothSharedConfigsData;	
-
-	UPROPERTY(Transient)
-	TArray<FCustomizableObjectClothingAssetData> ContributingClothingAssetsData;
-	
-	UPROPERTY(Transient)
-	TArray<FCustomizableObjectMeshToMeshVertData> ClothMeshToMeshVertData;
-
 private:
 	// mu::ExtensionData::Index is an index into this array when mu::ExtensionData::Origin is ConstantAlwaysLoaded
 	UPROPERTY()
diff --git a/Engine/Plugins/Experimental/Mutable/Source/CustomizableObject/Internal/MuCO/CustomizableObjectClothingTypes.h b/Engine/Plugins/Experimental/Mutable/Source/CustomizableObject/Internal/MuCO/CustomizableObjectClothingTypes.h
--- a/Engine/Plugins/Experimental/Mutable/Source/CustomizableObject/Internal/MuCO/CustomizableObjectClothingTypes.h
+++ b/Engine/Plugins/Experimental/Mutable/Source/CustomizableObject/Internal/MuCO/CustomizableObjectClothingTypes.h
@@ -58,9 +58,6 @@ struct FCustomizableObjectClothingAssetData
 	UPROPERTY()
 	TArray<FCustomizableObjectClothConfigData> ConfigsData;
 	
-	UPROPERTY()
-	FString PhysicsAssetPath;
-
 	UPROPERTY()
 	FName Name;
 	
@@ -87,7 +84,6 @@ struct FCustomizableObjectClothingAssetData
 		Ar << ClothData.UsedBoneIndices;
 		Ar << ClothData.ReferenceBoneIndex;
 		Ar << ClothData.ConfigsData;
-		Ar << ClothData.PhysicsAssetPath;
 		Ar << ClothData.Name;
 		Ar << ClothData.OriginalAssetGuid;
 		
diff --git a/Engine/Plugins/Experimental/Mutable/Source/CustomizableObject/Internal/MuCO/CustomizableObjectPrivate.h b/Engine/Plugins/Experimental/Mutable/Source/CustomizableObject/Internal/MuCO/CustomizableObjectPrivate.h
--- a/Engine/Plugins/Experimental/Mutable/Source/CustomizableObject/Internal/MuCO/CustomizableObjectPrivate.h
+++ b/Engine/Plugins/Experimental/Mutable/Source/CustomizableObject/Internal/MuCO/CustomizableObjectPrivate.h
@@ -81,6 +81,7 @@ struct FCustomizableObjectStreameableResourceId
 		None                  = 0,
 		AssetUserData         = 1,
 		RealTimeMorphTarget   = 2,
+		Clothing              = 3,
 	};
 
 	uint64 Id   : 64 - 8;
@@ -538,6 +539,35 @@ struct FRealTimeMorphStreamable
 	}
 };
 
+
+USTRUCT()
+struct FClothingStreamable
+{
+	GENERATED_USTRUCT_BODY()
+	
+	UPROPERTY()
+	int32 ClothingAssetIndex = INDEX_NONE;
+	
+	UPROPERTY()
+	int32 ClothingAssetLOD = INDEX_NONE;
+	
+	UPROPERTY()
+	int32 PhysicsAssetIndex = INDEX_NONE;
+
+	UPROPERTY()
+	FMutableStreamableBlock Block;
+
+	friend FArchive& operator<<(FArchive& Ar, FClothingStreamable& Elem)
+	{
+		Ar << Elem.ClothingAssetIndex;
+		Ar << Elem.ClothingAssetLOD;
+		Ar << Elem.PhysicsAssetIndex;
+		Ar << Elem.Block;
+
+		return Ar;
+	}
+};
+
 USTRUCT()
 struct FMorphTargetVertexData
 {
@@ -631,9 +661,25 @@ struct FModelResources
 	UPROPERTY()
 	TMap<uint32, FRealTimeMorphStreamable> RealTimeMorphStreamables;
 
+	UPROPERTY()
+	TArray<FCustomizableObjectClothConfigData> ClothSharedConfigsData;	
+
+	UPROPERTY()
+	TArray<FCustomizableObjectClothingAssetData> ClothingAssetsData;
+
+	UPROPERTY()
+	TMap<uint32, FClothingStreamable> ClothingStreamables;
+
+	/** Currently not used, this option should be selectable from editor maybe as a compilation flag */
+	UPROPERTY()
+	bool bAllowClothingPhysicsEditsPropagation = true;
+
 #if WITH_EDITORONLY_DATA
 	UPROPERTY()
 	TArray<FMorphTargetVertexData> EditorOnlyMorphTargetReconstructionData;
+
+	UPROPERTY()
+	TArray<FCustomizableObjectMeshToMeshVertData> EditorOnlyClothingMeshToMeshVertData;
 #endif
 	
 	/** Map of Hash to Streaming blocks, used to stream a block of data representing a resource from the BulkData */
@@ -655,7 +701,6 @@ struct FModelResources
 	/** First LOD available, some platforms may remove lower LODs when cooking, this MinLOD represents the first LOD we can generate */
 	UPROPERTY()
 	uint8 FirstLODAvailable = 0;
-
 };
 
 struct CUSTOMIZABLEOBJECT_API FMutableCachedPlatformData
@@ -668,6 +713,9 @@ struct CUSTOMIZABLEOBJECT_API FMutableCachedPlatformData
 
 	/** */
 	TArray64<uint8> MorphData;
+	
+	/** */
+	TArray64<uint8> ClothingData;
 };
 
 
@@ -772,7 +820,7 @@ public:
 	void LoadCompiledDataFromDisk();
 
 	/** Cache platform data for cook */
-	void CachePlatformData(const ITargetPlatform* InTargetPlatform, TArray64<uint8>& InObjectBytes, TArray64<uint8>& InBulkBytes, TArray64<uint8>& InMorphBytes);
+	void CachePlatformData(const ITargetPlatform* InTargetPlatform, TArray64<uint8>& InObjectBytes, TArray64<uint8>& InBulkBytes, TArray64<uint8>& InMorphBytes, TArray64<uint8>& InClothingBytes);
 	
 	/** Loads data previously compiled in BeginCacheForCookedPlatformData onto the UProperties in *this,
 	  * in preparation for saving the cooked package for *this or for a CustomizableObjectInstance using *this.
@@ -885,6 +933,6 @@ public:
 	// This is a manual version number for the binary blobs in this asset.
 	// Increasing it invalidates all the previously compiled models.
 	// Warning: If while merging code both versions have changed, take the highest+1.
-	static constexpr int32 CurrentSupportedVersion = 447;
+	static constexpr int32 CurrentSupportedVersion = 448;
 };
 
diff --git a/Engine/Plugins/Experimental/Mutable/Source/CustomizableObject/Internal/MuCO/CustomizableObjectSystemPrivate.h b/Engine/Plugins/Experimental/Mutable/Source/CustomizableObject/Internal/MuCO/CustomizableObjectSystemPrivate.h
--- a/Engine/Plugins/Experimental/Mutable/Source/CustomizableObject/Internal/MuCO/CustomizableObjectSystemPrivate.h
+++ b/Engine/Plugins/Experimental/Mutable/Source/CustomizableObject/Internal/MuCO/CustomizableObjectSystemPrivate.h
@@ -430,6 +430,16 @@ struct FInstanceUpdateData
 
 	TArray<FRealTimeMorphsComponentData> RealTimeMorphTargets;
 
+	struct FClothingMeshData
+	{
+		int32 ClothingAssetIndex = INDEX_NONE;
+		int32 ClothingAssetLOD = INDEX_NONE;
+		int32 PhysicsAssetIndex = INDEX_NONE;
+		TArray<FCustomizableObjectMeshToMeshVertData> Data;
+	};
+
+	TMap<uint32, FClothingMeshData> ClothingMeshData;
+
 	struct FSkeletonData
 	{
 		TArray<uint16> SkeletonIds;
diff --git a/Engine/Plugins/Experimental/Mutable/Source/CustomizableObject/Private/MuCO/CustomizableObject.cpp b/Engine/Plugins/Experimental/Mutable/Source/CustomizableObject/Private/MuCO/CustomizableObject.cpp
--- a/Engine/Plugins/Experimental/Mutable/Source/CustomizableObject/Private/MuCO/CustomizableObject.cpp
+++ b/Engine/Plugins/Experimental/Mutable/Source/CustomizableObject/Private/MuCO/CustomizableObject.cpp
@@ -362,10 +362,6 @@ void UCustomizableObjectPrivate::ClearCompiledData(bool bIsCooking)
 {
 	GetModelResources(bIsCooking) = FModelResources();
 
-	GetPublic()->ClothMeshToMeshVertData.Empty();
-	GetPublic()->ContributingClothingAssetsData.Empty();
-	GetPublic()->ClothSharedConfigsData.Empty();
-
 #if WITH_EDITORONLY_DATA
 	CustomizableObjectPathMap.Empty();
 	GroupNodeMap.Empty();
@@ -557,13 +553,13 @@ void UCustomizableObjectPrivate::SaveCompiledData(FArchive& MemoryWriter, bool b
 	MemoryWriter << LocalModelResources.StateUIDataMap;
 
 	MemoryWriter << LocalModelResources.RealTimeMorphStreamables;
+	
+	MemoryWriter << LocalModelResources.ClothingStreamables;
+	MemoryWriter << LocalModelResources.ClothingAssetsData; 	
+	MemoryWriter << LocalModelResources.ClothSharedConfigsData; 
 
 	MemoryWriter << LocalModelResources.HashToStreamableBlock;
 
-	MemoryWriter << GetPublic()->ClothMeshToMeshVertData;
-	MemoryWriter << GetPublic()->ContributingClothingAssetsData;
-	MemoryWriter << GetPublic()->ClothSharedConfigsData; 
-
 	MemoryWriter << LocalModelResources.NumComponents;
 	MemoryWriter << LocalModelResources.NumLODs;
 	MemoryWriter << LocalModelResources.NumLODsToStream;
@@ -579,6 +575,7 @@ void UCustomizableObjectPrivate::SaveCompiledData(FArchive& MemoryWriter, bool b
 	if (!bIsCooking)
 	{
 		MemoryWriter << LocalModelResources.EditorOnlyMorphTargetReconstructionData;
+		MemoryWriter << LocalModelResources.EditorOnlyClothingMeshToMeshVertData;
 	}
 }
 
@@ -677,13 +674,13 @@ void UCustomizableObjectPrivate::LoadCompiledData(FArchive& MemoryReader, const
 		MemoryReader << LocalModelResource.StateUIDataMap;
 
 		MemoryReader << LocalModelResource.RealTimeMorphStreamables;
-
+		
+		MemoryReader << LocalModelResource.ClothingStreamables;
+		MemoryReader << LocalModelResource.ClothingAssetsData; 
+		MemoryReader << LocalModelResource.ClothSharedConfigsData; 
+		
 		MemoryReader << LocalModelResource.HashToStreamableBlock;
 
-		MemoryReader << GetPublic()->ClothMeshToMeshVertData;
-		MemoryReader << GetPublic()->ContributingClothingAssetsData;
-		MemoryReader << GetPublic()->ClothSharedConfigsData;
-
 		MemoryReader << LocalModelResource.NumComponents;
 		MemoryReader << LocalModelResource.NumLODs;
 		MemoryReader << LocalModelResource.NumLODsToStream;
@@ -702,6 +699,7 @@ void UCustomizableObjectPrivate::LoadCompiledData(FArchive& MemoryReader, const
 			if (!bIsCooking)
 			{
 				MemoryReader << LocalModelResource.EditorOnlyMorphTargetReconstructionData;
+				MemoryReader << LocalModelResource.EditorOnlyClothingMeshToMeshVertData;
 				
 				DirtyParticipatingObjects.Empty();
 
@@ -800,7 +798,12 @@ void UCustomizableObjectPrivate::LoadCompiledDataFromDisk()
 }
 
 
-void UCustomizableObjectPrivate::CachePlatformData(const ITargetPlatform* InTargetPlatform, TArray64<uint8>& InModelBytes, TArray64<uint8>& InBulkBytes, TArray64<uint8>& InMorphBytes)
+void UCustomizableObjectPrivate::CachePlatformData(
+		const ITargetPlatform* InTargetPlatform, 
+		TArray64<uint8>& InModelBytes, 
+		TArray64<uint8>& InBulkBytes, 
+		TArray64<uint8>& InMorphBytes, 
+		TArray64<uint8>& InClothingBytes)
 {
 	MUTABLE_CPUPROFILER_SCOPE(CachePlatformData)
 
@@ -818,6 +821,7 @@ void UCustomizableObjectPrivate::CachePlatformData(const ITargetPlatform* InTarg
 	// Cache streamable bulk data
 	Data.StreamableData = MoveTemp(InBulkBytes);
 	Data.MorphData = MoveTemp(InMorphBytes);
+	Data.ClothingData = MoveTemp(InClothingBytes);
 }
 
 
@@ -991,12 +995,6 @@ void UCustomizableObjectPrivate::SaveEmbeddedData(FArchive& Ar)
 
 	if (GetModel())
 	{	
-		{
-			Ar << GetPublic()->ClothMeshToMeshVertData;
-			Ar << GetPublic()->ContributingClothingAssetsData;
-			Ar << GetPublic()->ClothSharedConfigsData;
-		}
-
 		// Serialise the entire model, but unload the streamable data first.
 		{
 			GetModel()->UnloadExternalData();
@@ -1024,13 +1022,7 @@ void UCustomizableObjectPrivate::LoadEmbeddedData(FArchive& Ar)
 	check(CurrentSupportedVersion == InternalVersion);
 
 	if(CurrentSupportedVersion == InternalVersion)
-	{	
-		{
-			Ar << GetPublic()->ClothMeshToMeshVertData;
-			Ar << GetPublic()->ContributingClothingAssetsData;
-			Ar << GetPublic()->ClothSharedConfigsData;
-		}
-		
+	{		
 		// Load model
 		UnrealMutableInputStream Stream(Ar);
 		mu::InputArchive Arch(&Stream);
@@ -2076,6 +2068,10 @@ void UCustomizableObjectBulk::CookAdditionalFilesOverride(const TCHAR* PackageFi
 		{
 			SourceData = PlatformData->MorphData.GetData();
 		}
+		else if (CurrentFile.DataType == EDataType::Clothing)
+		{
+			SourceData = PlatformData->ClothingData.GetData();
+		}
 		else
 		{
 			checkf(false, TEXT("Unknown file DataType found."));
@@ -2203,6 +2199,42 @@ void UCustomizableObjectBulk::PrepareBulkData(UCustomizableObject* InOuter, cons
 		}
 	}
 
+	// TODO: This should create a new classification branch when the tree is implemented.
+	// For now append after Model roms. 
+	{
+		uint64 SourceOffset = 0;
+		
+		constexpr bool bGetCooked = false;
+
+		TArray<FMutableStreamableBlock> ClothingBlocks;
+		
+		const TMap<uint32, FClothingStreamable>& ClothingStreamables = 
+				CustomizableObject->GetPrivate()->GetModelResources(bGetCooked).ClothingStreamables;
+
+		const int32 NumBlocks = ClothingBlocks.Num();
+		for (const TPair<uint32, FClothingStreamable>& ClothStreamable : ClothingStreamables)
+		{
+			const uint32 BlockSize = ClothStreamable.Value.Block.Size;
+			
+			check(SourceOffset == ClothStreamable.Value.Block.Offset);
+			FBlock CurrentBlock = { EDataType::Clothing, ClothStreamable.Key, BlockSize, SourceOffset };
+			
+			if (BlockSize > TargetBulkDataFileBytes)
+			{
+				// It will go to its own file
+				BulkDataFiles.Add(FFile{ EDataType::Clothing, 0, {CurrentBlock} });
+			}
+			else
+			{
+				// It may merge with other small blocks
+				RootNode.Blocks.Add(CurrentBlock);
+			}
+
+			SourceOffset += BlockSize;
+		}
+	}
+
+
 	// Temp: Group by order in the array
 	for (int32 BlockIndex = 0; BlockIndex < RootNode.Blocks.Num(); )
 	{
@@ -2303,6 +2335,24 @@ void UCustomizableObjectBulk::PrepareBulkData(UCustomizableObject* InOuter, cons
 				OffsetInFile += ThisBlock.Size;
 			}
 		}
+		else if (CurrentFile.DataType == EDataType::Clothing)
+		{
+			constexpr bool bGetCooked = true;
+			TMap<uint32, FClothingStreamable>& ClothBlocks = 
+					CustomizableObject->GetPrivate()->GetModelResources(bGetCooked).ClothingStreamables;
+			// Set it to all streamable blocks
+			uint32 OffsetInFile = 0;
+			for (int32 FileBlockIndex = 0; FileBlockIndex < CurrentFile.Blocks.Num(); ++FileBlockIndex)
+			{
+				FBlock ThisBlock = CurrentFile.Blocks[FileBlockIndex];
+
+				FMutableStreamableBlock& StreamableBlock = ClothBlocks[ThisBlock.Id].Block;
+				check(StreamableBlock.Size == ThisBlock.Size);
+				StreamableBlock.FileId = FileId;
+				StreamableBlock.Offset = OffsetInFile;
+				OffsetInFile += ThisBlock.Size;
+			}
+		}
 		else
 		{
 			UE_LOG(LogMutable, Error, TEXT("Unknown DataType found while fixing streaming block files ids."));
diff --git a/Engine/Plugins/Experimental/Mutable/Source/CustomizableObject/Private/MuCO/CustomizableObjectInstance.cpp b/Engine/Plugins/Experimental/Mutable/Source/CustomizableObject/Private/MuCO/CustomizableObjectInstance.cpp
--- a/Engine/Plugins/Experimental/Mutable/Source/CustomizableObject/Private/MuCO/CustomizableObjectInstance.cpp
+++ b/Engine/Plugins/Experimental/Mutable/Source/CustomizableObject/Private/MuCO/CustomizableObjectInstance.cpp
@@ -61,10 +61,21 @@
 #include "Misc/TransactionObjectEvent.h"
 #endif
 
+namespace
+{
 #ifndef REQUIRES_SINGLEUSE_FLAG_FOR_RUNTIME_TEXTURES
 	#define REQUIRES_SINGLEUSE_FLAG_FOR_RUNTIME_TEXTURES !PLATFORM_DESKTOP
 #endif
 
+bool bDisableClothingPhysicsEditsPropagation = false;
+static FAutoConsoleVariableRef CVarDisableClothingPhysicsEditsPropagation(
+	TEXT("mutable.DisableClothingPhysicsEditsPropagation"),
+	bDisableClothingPhysicsEditsPropagation,
+	TEXT("If set to true, disables clothing physics edits propagation from the render mesh."),
+	ECVF_Default);
+}
+
+
 
 // Struct used by BuildMaterials() to identify common materials between LODs
 struct FMutableMaterialPlaceholder
@@ -3539,7 +3550,6 @@ void UCustomizableInstancePrivate::BuildOrCopyElementData(const TSharedRef<FUpda
 	}
 }
 
-
 void UCustomizableInstancePrivate::BuildOrCopyMorphTargetsData(const TSharedRef<FUpdateContextPrivate>& OperationData, USkeletalMesh* SkeletalMesh, const USkeletalMesh* LastUpdateSkeletalMesh, UCustomizableObjectInstance* CustomizableObjectInstance, int32 ComponentIndex)
 {
 	MUTABLE_CPUPROFILER_SCOPE(UCustomizableInstancePrivate::BuildOrCopyMorphTargetsData);
@@ -3612,7 +3622,7 @@ namespace
 
 }
 
-void UCustomizableInstancePrivate::BuildOrCopyClothingData(const TSharedRef<FUpdateContextPrivate>&OperationData, USkeletalMesh * SkeletalMesh, const USkeletalMesh* LastUpdateSkeletalMesh, UCustomizableObjectInstance * CustomizableObjectInstance, int32 ComponentIndex)
+void UCustomizableInstancePrivate::BuildOrCopyClothingData(const TSharedRef<FUpdateContextPrivate>& OperationData, USkeletalMesh * SkeletalMesh, const USkeletalMesh* LastUpdateSkeletalMesh, UCustomizableObjectInstance* CustomizableObjectInstance, int32 ComponentIndex)
 {
 	MUTABLE_CPUPROFILER_SCOPE(UCustomizableInstancePrivate::BuildOrCopyClothingData);
 
@@ -3621,16 +3631,17 @@ void UCustomizableInstancePrivate::BuildOrCopyClothingData(const TSharedRef<FUpd
 		return;
 	}
 
-	const UCustomizableObject* CustomizableObject = CustomizableObjectInstance->GetCustomizableObject();
-	const TArray<FCustomizableObjectClothingAssetData>& ContributingClothingAssetsData = CustomizableObject->ContributingClothingAssetsData;
-	const TArray<FCustomizableObjectClothConfigData>& ClothSharedConfigsData = CustomizableObject->ClothSharedConfigsData;
-	const TArray<FCustomizableObjectMeshToMeshVertData>& ClothMeshToMeshVertData = CustomizableObject->ClothMeshToMeshVertData;
 
-	if (!(ContributingClothingAssetsData.Num() && ClothMeshToMeshVertData.Num()))
+	const FModelResources& ModelResources = CustomizableObjectInstance->GetCustomizableObject()->GetPrivate()->GetModelResources();
+	const TArray<FCustomizableObjectClothingAssetData>& ClothingAssetsData = ModelResources.ClothingAssetsData;
+	const TArray<FCustomizableObjectClothConfigData>& ClothSharedConfigsData = ModelResources.ClothSharedConfigsData;
+
+	if (!(ClothingAssetsData.Num() && OperationData->InstanceUpdateData.ClothingMeshData.Num()))
 	{
 		return;
 	}
 
+	const bool bAllowClothingPhysicsEdits = !bDisableClothingPhysicsEditsPropagation && ModelResources.bAllowClothingPhysicsEditsPropagation;
 	// First we need to discover if any clothing asset is used for the instance. 
 
 	struct FSectionWithClothData
@@ -3643,18 +3654,21 @@ void UCustomizableInstancePrivate::BuildOrCopyClothingData(const TSharedRef<FUpd
 		TArrayView<const uint16> SectionIndex16View;
 		TArrayView<const uint32> SectionIndex32View;
 		TArrayView<const int32> ClothingDataIndicesView;
+		TArrayView<FCustomizableObjectMeshToMeshVertData> ClothingDataView;
 		TArray<FMeshToMeshVertData> MappingData;
 	};
 
 	TArray<FSectionWithClothData> SectionsWithCloth;
 	SectionsWithCloth.Reserve(32);
 
+	int32 NumClothingDataNotFound = 0;
+
 	const int32 LODCount = OperationData->InstanceUpdateData.LODs.Num();
 
 	{
 		MUTABLE_CPUPROFILER_SCOPE(DiscoverSectionsWithCloth);
 
-		for (int32 LODIndex = OperationData->GetMinLOD(); LODIndex < OperationData->NumLODsAvailable; ++LODIndex)
+		for (int32 LODIndex = OperationData->FirstLODAvailable; LODIndex < OperationData->NumLODsAvailable; ++LODIndex)
 		{
 			const FInstanceUpdateData::FLOD& LOD = OperationData->InstanceUpdateData.LODs[LODIndex];
 			if (ComponentIndex >= LOD.ComponentCount)
@@ -3672,28 +3686,26 @@ void UCustomizableInstancePrivate::BuildOrCopyClothingData(const TSharedRef<FUpd
 			if (mu::MeshPtrConst MutableMesh = Component.Mesh)
 			{
 				const mu::FMeshBufferSet& MeshSet = MutableMesh->GetVertexBuffers();
-
 				const mu::FMeshBufferSet& IndicesSet = MutableMesh->GetIndexBuffers();
 
-				// Semantics index may vary depending on whether realtime morph targets are enabled.
-				const int32 ClothingDataBufferIndex = [&MeshSet]()
-				{
-					int32 BufferIndex, Channel;
-					MeshSet.FindChannel(mu::MBS_OTHER, 2, &BufferIndex, &Channel);
-					return BufferIndex;
-				}(); // lambda is invoked
+				int32 ClothingIndexBuffer, ClothingIndexChannel;
+				MeshSet.FindChannel(mu::MBS_OTHER, 2, &ClothingIndexBuffer, &ClothingIndexChannel);
 
-				if (ClothingDataBufferIndex < 0)
+				int32 ClothingResourceBuffer, ClothingResourceChannel;
+				MeshSet.FindChannel(mu::MBS_OTHER, 3, &ClothingResourceBuffer, &ClothingResourceChannel);
+
+				if (ClothingIndexBuffer < 0 || ClothingResourceBuffer < 0)
 				{
 					continue;
 				}
 
-				const int32* const ClothingDataBuffer = reinterpret_cast<const int32*>(MeshSet.GetBufferData(ClothingDataBufferIndex));
+				const int32* const ClothingDataBuffer = reinterpret_cast<const int32*>(MeshSet.GetBufferData(ClothingIndexBuffer));
+				const uint32* const ClothingDataResource = reinterpret_cast<const uint32*>(MeshSet.GetBufferData(ClothingResourceBuffer));
 
 				const int32 SurfaceCount = MutableMesh->GetSurfaceCount();
 				for (int32 Section = 0; Section < SurfaceCount; ++Section)
 				{
-					int FirstVertex, VerticesCount, FirstIndex, IndicesCount;
+					int32 FirstVertex, VerticesCount, FirstIndex, IndicesCount;
 					MutableMesh->GetSurface(Section, &FirstVertex, &VerticesCount, &FirstIndex, &IndicesCount, nullptr, nullptr, nullptr);
 
 					if (VerticesCount == 0 || IndicesCount == 0)
@@ -3705,6 +3717,7 @@ void UCustomizableInstancePrivate::BuildOrCopyClothingData(const TSharedRef<FUpd
 					// It can be determined if this section has clothing data just looking at the 
 					// first vertex of the section.
 					TArrayView<const int32> ClothingDataView(ClothingDataBuffer + FirstVertex, VerticesCount);
+					TArrayView<const uint32> ClothingResourceView(ClothingDataResource + FirstVertex, VerticesCount);
 
 					const int32 IndexCount = MutableMesh->GetIndexBuffers().GetElementCount();
 
@@ -3724,61 +3737,95 @@ void UCustomizableInstancePrivate::BuildOrCopyClothingData(const TSharedRef<FUpd
 							IndicesView32Bits = TArrayView<const uint32>(IndexPtr + FirstIndex, IndicesCount);
 						 }
 					}
-					
+				
 					if (!ClothingDataView.Num())
 					{
 						continue;
 					}
 
-					const int32 ClothDataIndex = ClothingDataView[0];
-					if (ClothDataIndex < 0)
+					const uint32 ClothResourceId = ClothingResourceView[0];
+					if (ClothResourceId == 0)
 					{
 						continue;
 					}
 
-					const int32 ClothAssetIndex = ClothMeshToMeshVertData[ClothDataIndex].SourceAssetIndex;
-					const int32 ClothAssetLodIndex = ClothMeshToMeshVertData[ClothDataIndex].SourceAssetLodIndex;
+					FInstanceUpdateData::FClothingMeshData* SectionClothingData = 
+							OperationData->InstanceUpdateData.ClothingMeshData.Find(ClothResourceId);
+	
+					if (!SectionClothingData)
+					{
+						++NumClothingDataNotFound;
+						continue;
+					}
+
+					check(SectionClothingData->Data.Num());
+					check(SectionClothingData->ClothingAssetIndex != INDEX_NONE);
+					check(SectionClothingData->ClothingAssetLOD != INDEX_NONE);
+
+					const int32 ClothAssetIndex = SectionClothingData->ClothingAssetIndex;
+					const int32 ClothAssetLodIndex = SectionClothingData->ClothingAssetLOD;
+
+					check(SectionClothingData->ClothingAssetIndex == ClothAssetIndex);
 
 					// Defensive check, this indicates the clothing data might be stale and needs to be recompiled.
 					// Should never happen.
-					if (!ensure(ClothAssetIndex >= 0 && ClothAssetIndex < ContributingClothingAssetsData.Num()
-						&& ContributingClothingAssetsData[ClothAssetIndex].LodData.Num()))
+					if (!ensure(ClothAssetIndex >= 0 && ClothAssetIndex < ClothingAssetsData.Num() && 
+								ClothingAssetsData[ClothAssetIndex].LodData.Num()))
 					{
 						continue;
 					}
 
-					SectionsWithCloth.Add
-							(FSectionWithClothData{ ClothAssetIndex, ClothAssetLodIndex, Section, LODIndex, FirstVertex, IndicesView16Bits, IndicesView32Bits, ClothingDataView, TArray<FMeshToMeshVertData>() });
+					SectionsWithCloth.Add(FSectionWithClothData
+							{
+							 	ClothAssetIndex, ClothAssetLodIndex, 
+								Section, LODIndex, FirstVertex, 
+								IndicesView16Bits, IndicesView32Bits, 
+								ClothingDataView, 
+								MakeArrayView(SectionClothingData->Data.GetData(), SectionClothingData->Data.Num()),
+								TArray<FMeshToMeshVertData>() 
+							});
 				}
 			}
 		}
 	}
 
+	if (NumClothingDataNotFound > 0)
+	{
+		UE_LOG(LogMutable, Error, TEXT("Some clothing data could not be loaded properly, clothing assets may not behave as expected."));
+	}
+
 	if (!SectionsWithCloth.Num())
 	{
 		return; // Nothing to do.
 	}
-	
+		
 	TArray<FCustomizableObjectClothingAssetData> NewClothingAssetsData;
-	NewClothingAssetsData.SetNum( ContributingClothingAssetsData.Num() );
+	NewClothingAssetsData.SetNum(ClothingAssetsData.Num());
 
 	{
-
 		for (FSectionWithClothData& SectionWithCloth : SectionsWithCloth)
 		{
-			const FCustomizableObjectClothingAssetData& SrcAssetData = ContributingClothingAssetsData[SectionWithCloth.ClothAssetIndex];
+			const FCustomizableObjectClothingAssetData& SrcAssetData = ClothingAssetsData[SectionWithCloth.ClothAssetIndex];
 			FCustomizableObjectClothingAssetData& DstAssetData = NewClothingAssetsData[SectionWithCloth.ClothAssetIndex];
 		
 			// Only initilialize once, multiple sections with cloth could point to the same cloth asset.
 			if (!DstAssetData.LodMap.Num())
 			{
-				DstAssetData.LodMap.Init( INDEX_NONE, OperationData->NumLODsAvailable);
+				DstAssetData.LodMap.Init(INDEX_NONE, OperationData->NumLODsAvailable);
 
-				DstAssetData.LodData.SetNum( SrcAssetData.LodData.Num() );
 				DstAssetData.UsedBoneNames = SrcAssetData.UsedBoneNames;
 				DstAssetData.UsedBoneIndices = SrcAssetData.UsedBoneIndices;
 				DstAssetData.ReferenceBoneIndex = SrcAssetData.ReferenceBoneIndex;
 				DstAssetData.Name = SrcAssetData.Name;
+
+				if (!bAllowClothingPhysicsEdits)
+				{
+					DstAssetData.LodData = SrcAssetData.LodData;
+				}
+				else
+				{
+					DstAssetData.LodData.SetNum(SrcAssetData.LodData.Num());
+				}
 			}
 
 			DstAssetData.LodMap[SectionWithCloth.Lod] = SectionWithCloth.ClothAssetLodIndex;
@@ -3808,9 +3855,9 @@ void UCustomizableInstancePrivate::BuildOrCopyClothingData(const TSharedRef<FUpd
 				static_assert(TIsTrivial<FMeshToMeshVertData>::Value, "");
 
 				const int32 VertexDataIndex = SectionWithCloth.ClothingDataIndicesView[VertexIdx];
-				check(VertexDataIndex > 0);
+				check(VertexDataIndex >= 0);
 
-				const FCustomizableObjectMeshToMeshVertData& SrcData = ClothMeshToMeshVertData[VertexDataIndex];
+				const FCustomizableObjectMeshToMeshVertData& SrcData = SectionWithCloth.ClothingDataView[VertexDataIndex];
 
 				FMeshToMeshVertData& DstData = ClothMappingData[VertexIdx];
 				FMemory::Memcpy(&DstData, &SrcData, sizeof(FMeshToMeshVertData));
@@ -3818,467 +3865,470 @@ void UCustomizableInstancePrivate::BuildOrCopyClothingData(const TSharedRef<FUpd
 		}
 	}
 
-	// Indices remaps for {Section, AssetLod}, needed to recreate the lod transition data.
-	TMap<int32, TArray<TArray<int32>>> PhysicsSectionLodsIndicesRemaps;
-
-	check( SectionsWithCloth.Num() > 0 );
-	FSectionWithClothData* MaxSection = Algo::MaxElement(SectionsWithCloth, 
-			[](const FSectionWithClothData& A, const FSectionWithClothData& B) { return  A.ClothAssetIndex < B.ClothAssetIndex; } );
-	
-	PhysicsSectionLodsIndicesRemaps.Reserve(MaxSection->ClothAssetIndex + 1);
-		
-	for (FSectionWithClothData& SectionLods : SectionsWithCloth)
+	if (bAllowClothingPhysicsEdits)
 	{
-		TArray<TArray<int32>>& Value = PhysicsSectionLodsIndicesRemaps.FindOrAdd(SectionLods.ClothAssetIndex, {});
-		Value.SetNum(FMath::Max(Value.Num(), SectionLods.ClothAssetLodIndex + 1));
-	}
+		// Indices remaps for {Section, AssetLod}, needed to recreate the lod transition data.
+		TMap<int32, TArray<TArray<int32>>> PhysicsSectionLodsIndicesRemaps;
 
-	{
-		MUTABLE_CPUPROFILER_SCOPE(RemapPhysicsMesh)
+		check( SectionsWithCloth.Num() > 0 );
+		FSectionWithClothData* MaxSection = Algo::MaxElement(SectionsWithCloth, 
+				[](const FSectionWithClothData& A, const FSectionWithClothData& B) { return  A.ClothAssetIndex < B.ClothAssetIndex; } );
+		
+		PhysicsSectionLodsIndicesRemaps.Reserve(MaxSection->ClothAssetIndex + 1);
+			
+		for (FSectionWithClothData& SectionLods : SectionsWithCloth)
+		{
+			TArray<TArray<int32>>& Value = PhysicsSectionLodsIndicesRemaps.FindOrAdd(SectionLods.ClothAssetIndex, {});
+			Value.SetNum(FMath::Max(Value.Num(), SectionLods.ClothAssetLodIndex + 1));
+		}
 
-		for (FSectionWithClothData& SectionWithCloth : SectionsWithCloth)
 		{
-			const FCustomizableObjectClothingAssetData& SrcClothingAssetData = ContributingClothingAssetsData[SectionWithCloth.ClothAssetIndex];
-			FCustomizableObjectClothingAssetData& NewClothingAssetData = NewClothingAssetsData[SectionWithCloth.ClothAssetIndex];
-					
-			const FClothLODDataCommon& SrcLodData = SrcClothingAssetData.LodData[SectionWithCloth.ClothAssetLodIndex];
-			FClothLODDataCommon& NewLodData = NewClothingAssetData.LodData[SectionWithCloth.ClothAssetLodIndex];
-			
-			const int32 PhysicalMeshVerticesNum = SrcLodData.PhysicalMeshData.Vertices.Num();
+			MUTABLE_CPUPROFILER_SCOPE(RemapPhysicsMesh)
 
-			if (!PhysicalMeshVerticesNum)
+			for (FSectionWithClothData& SectionWithCloth : SectionsWithCloth)
 			{
-				// Nothing to do.
-				continue;
-			}
+				const FCustomizableObjectClothingAssetData& SrcClothingAssetData = ClothingAssetsData[SectionWithCloth.ClothAssetIndex];
+				FCustomizableObjectClothingAssetData& NewClothingAssetData = NewClothingAssetsData[SectionWithCloth.ClothAssetIndex];
+						
+				const FClothLODDataCommon& SrcLodData = SrcClothingAssetData.LodData[SectionWithCloth.ClothAssetLodIndex];
+				FClothLODDataCommon& NewLodData = NewClothingAssetData.LodData[SectionWithCloth.ClothAssetLodIndex];
+				
+				const int32 PhysicalMeshVerticesNum = SrcLodData.PhysicalMeshData.Vertices.Num();
+
+				if (!PhysicalMeshVerticesNum)
+				{
+					// Nothing to do.
+					continue;
+				}
+				
+				// Vertices not indexed in the mesh to mesh data generated for this section need to be removed.
+				TArray<uint8> VertexUtilizationBuffer;
+				VertexUtilizationBuffer.Init(0, PhysicalMeshVerticesNum);
 			
-			// Vertices not indexed in the mesh to mesh data generated for this section need to be removed.
-			TArray<uint8> VertexUtilizationBuffer;
-			VertexUtilizationBuffer.Init(0, PhysicalMeshVerticesNum);
-		
-			// Discover used vertices.
-			const int32 SectionVerticesNum = SectionWithCloth.ClothingDataIndicesView.Num();
+				// Discover used vertices.
+				const int32 SectionVerticesNum = SectionWithCloth.ClothingDataIndicesView.Num();
 
-			TArray<uint8> RenderVertexUtilizationBuffer;
-			RenderVertexUtilizationBuffer.Init(0, SectionVerticesNum);
+				TArray<uint8> RenderVertexUtilizationBuffer;
+				RenderVertexUtilizationBuffer.Init(0, SectionVerticesNum);
 
-			// Sometimes, at least when a clip morph is applied, vertices are not removed from the section
-			// and only the triangles (indices) that form the mesh are modified.
+				// Sometimes, at least when a clip morph is applied, vertices are not removed from the section
+				// and only the triangles (indices) that form the mesh are modified.
 
-			auto GenerateRenderUtilizationBuffer = [&RenderVertexUtilizationBuffer](const auto& IndicesView, int32 SectionBaseVertex )
-			{
-				const int32 IndicesCount = IndicesView.Num();
-				check(IndicesCount % 3 == 0);
-				for ( int32 I = 0; I < IndicesCount; I += 3 )
+				auto GenerateRenderUtilizationBuffer = [&RenderVertexUtilizationBuffer](const auto& IndicesView, int32 SectionBaseVertex)
 				{
-					RenderVertexUtilizationBuffer[int32(IndicesView[I + 0]) - SectionBaseVertex] = 1;
-					RenderVertexUtilizationBuffer[int32(IndicesView[I + 1]) - SectionBaseVertex] = 1;
-					RenderVertexUtilizationBuffer[int32(IndicesView[I + 2]) - SectionBaseVertex] = 1;
+					const int32 IndicesCount = IndicesView.Num();
+					check(IndicesCount % 3 == 0);
+					for ( int32 I = 0; I < IndicesCount; I += 3 )
+					{
+						RenderVertexUtilizationBuffer[int32(IndicesView[I + 0]) - SectionBaseVertex] = 1;
+						RenderVertexUtilizationBuffer[int32(IndicesView[I + 1]) - SectionBaseVertex] = 1;
+						RenderVertexUtilizationBuffer[int32(IndicesView[I + 2]) - SectionBaseVertex] = 1;
+					}
+				};
+				
+				if (SectionWithCloth.SectionIndex16View.Num())
+				{
+					GenerateRenderUtilizationBuffer(SectionWithCloth.SectionIndex16View, SectionWithCloth.BaseVertex);
+				}
+				else
+				{
+					check(SectionWithCloth.SectionIndex32View.Num());
+					GenerateRenderUtilizationBuffer(SectionWithCloth.SectionIndex32View, SectionWithCloth.BaseVertex);
 				}
-			};
-			
-			if (SectionWithCloth.SectionIndex16View.Num())
-			{
-				GenerateRenderUtilizationBuffer(SectionWithCloth.SectionIndex16View, SectionWithCloth.BaseVertex);
-			}
-			else
-			{
-				check(SectionWithCloth.SectionIndex32View.Num());
-				GenerateRenderUtilizationBuffer(SectionWithCloth.SectionIndex32View, SectionWithCloth.BaseVertex);
-			}
 
-			const TArray<FMeshToMeshVertData>& SectionClothMappingData = SectionWithCloth.MappingData;
-			for (int32 Idx = 0; Idx < SectionVerticesNum; ++Idx)
-			{
-				if (RenderVertexUtilizationBuffer[Idx])
+				const TArray<FMeshToMeshVertData>& SectionClothMappingData = SectionWithCloth.MappingData;
+				for (int32 Idx = 0; Idx < SectionVerticesNum; ++Idx)
 				{
-					const uint16* Indices = SectionClothMappingData[Idx].SourceMeshVertIndices;
+					if (RenderVertexUtilizationBuffer[Idx])
+					{
+						const uint16* Indices = SectionClothMappingData[Idx].SourceMeshVertIndices;
 
-					VertexUtilizationBuffer[Indices[0]] = 1;
-					VertexUtilizationBuffer[Indices[1]] = 1;
-					VertexUtilizationBuffer[Indices[2]] = 1;
+						VertexUtilizationBuffer[Indices[0]] = 1;
+						VertexUtilizationBuffer[Indices[1]] = 1;
+						VertexUtilizationBuffer[Indices[2]] = 1;
+					}
 				}
-			}
 
-			TArray<int32>& IndexMap = PhysicsSectionLodsIndicesRemaps[SectionWithCloth.ClothAssetIndex][SectionWithCloth.ClothAssetLodIndex];
-			IndexMap.SetNumUninitialized(PhysicalMeshVerticesNum);
-
-			// Compute index remap and number of remaining physics vertices.
-			// -1 indicates the vertex has been removed.
-			int32 NewPhysicalMeshVerticesNum = 0;
-			for (int32 Idx = 0; Idx < PhysicalMeshVerticesNum; ++Idx)
-			{
-				IndexMap[Idx] = VertexUtilizationBuffer[Idx] ? NewPhysicalMeshVerticesNum++ : -1;
-			}
+				TArray<int32>& IndexMap = PhysicsSectionLodsIndicesRemaps[SectionWithCloth.ClothAssetIndex][SectionWithCloth.ClothAssetLodIndex];
+				IndexMap.SetNumUninitialized(PhysicalMeshVerticesNum);
 
-			const bool bHasVerticesRemoved = NewPhysicalMeshVerticesNum < PhysicalMeshVerticesNum;
-			if (!bHasVerticesRemoved)
-			{
-				// If no vertices are removed the IndexMap is no longer needed. The lack of data in the map 
-				// can indicates that no vertex has been removed to subsequent operations.
-				IndexMap.Reset();
-			}
-			
-			const auto CopyIfUsed = [&VertexUtilizationBuffer, bHasVerticesRemoved](auto& Dst, const auto& Src)
-			{	
-				const int32 SrcNumElems = Src.Num();
+				// Compute index remap and number of remaining physics vertices.
+				// -1 indicates the vertex has been removed.
+				int32 NewPhysicalMeshVerticesNum = 0;
+				for (int32 Idx = 0; Idx < PhysicalMeshVerticesNum; ++Idx)
+				{
+					IndexMap[Idx] = VertexUtilizationBuffer[Idx] ? NewPhysicalMeshVerticesNum++ : -1;
+				}
 
+				const bool bHasVerticesRemoved = NewPhysicalMeshVerticesNum < PhysicalMeshVerticesNum;
 				if (!bHasVerticesRemoved)
 				{
-					for (int32 Idx = 0; Idx < SrcNumElems; ++Idx)
-					{
-						Dst[Idx] = Src[Idx];
-					}
-
-					return;
+					// If no vertices are removed the IndexMap is no longer needed. The lack of data in the map 
+					// can indicates that no vertex has been removed to subsequent operations.
+					IndexMap.Reset();
 				}
+				
+				const auto CopyIfUsed = [&VertexUtilizationBuffer, bHasVerticesRemoved](auto& Dst, const auto& Src)
+				{	
+					const int32 SrcNumElems = Src.Num();
 
-				for (int32 Idx = 0, DstNumElems = 0; Idx < SrcNumElems; ++Idx)
-				{
-					if (VertexUtilizationBuffer[Idx])
+					if (!bHasVerticesRemoved)
 					{
-						Dst[DstNumElems++] = Src[Idx];
+						for (int32 Idx = 0; Idx < SrcNumElems; ++Idx)
+						{
+							Dst[Idx] = Src[Idx];
+						}
+
+						return;
 					}
-				}
-			};
 
-			NewLodData.PhysicalMeshData.MaxBoneWeights = SrcLodData.PhysicalMeshData.MaxBoneWeights;
+					for (int32 Idx = 0, DstNumElems = 0; Idx < SrcNumElems; ++Idx)
+					{
+						if (VertexUtilizationBuffer[Idx])
+						{
+							Dst[DstNumElems++] = Src[Idx];
+						}
+					}
+				};
 
-			NewLodData.PhysicalMeshData.Vertices.SetNum(NewPhysicalMeshVerticesNum);
-			NewLodData.PhysicalMeshData.Normals.SetNum(NewPhysicalMeshVerticesNum);
-			NewLodData.PhysicalMeshData.BoneData.SetNum(NewPhysicalMeshVerticesNum);
-			NewLodData.PhysicalMeshData.InverseMasses.SetNum(NewPhysicalMeshVerticesNum);
+				NewLodData.PhysicalMeshData.MaxBoneWeights = SrcLodData.PhysicalMeshData.MaxBoneWeights;
 
-			CopyIfUsed(NewLodData.PhysicalMeshData.Vertices, SrcLodData.PhysicalMeshData.Vertices);
-			CopyIfUsed(NewLodData.PhysicalMeshData.Normals, SrcLodData.PhysicalMeshData.Normals);
-			CopyIfUsed(NewLodData.PhysicalMeshData.BoneData, SrcLodData.PhysicalMeshData.BoneData);
-			CopyIfUsed(NewLodData.PhysicalMeshData.InverseMasses, SrcLodData.PhysicalMeshData.InverseMasses);
+				NewLodData.PhysicalMeshData.Vertices.SetNum(NewPhysicalMeshVerticesNum);
+				NewLodData.PhysicalMeshData.Normals.SetNum(NewPhysicalMeshVerticesNum);
+				NewLodData.PhysicalMeshData.BoneData.SetNum(NewPhysicalMeshVerticesNum);
+				NewLodData.PhysicalMeshData.InverseMasses.SetNum(NewPhysicalMeshVerticesNum);
 
-			const int32 PrevIndex = SectionWithCloth.Lod - 1;
-			const bool bNeedsTransitionUpData = NewClothingAssetData.LodMap.IsValidIndex(PrevIndex) && NewClothingAssetData.LodMap[PrevIndex] != INDEX_NONE;
-			if (bNeedsTransitionUpData)
-			{
-				NewLodData.TransitionUpSkinData.SetNum(SrcLodData.TransitionUpSkinData.Num() ? NewPhysicalMeshVerticesNum : 0);	
-				CopyIfUsed(NewLodData.TransitionUpSkinData, SrcLodData.TransitionUpSkinData);
-			}
+				CopyIfUsed(NewLodData.PhysicalMeshData.Vertices, SrcLodData.PhysicalMeshData.Vertices);
+				CopyIfUsed(NewLodData.PhysicalMeshData.Normals, SrcLodData.PhysicalMeshData.Normals);
+				CopyIfUsed(NewLodData.PhysicalMeshData.BoneData, SrcLodData.PhysicalMeshData.BoneData);
+				CopyIfUsed(NewLodData.PhysicalMeshData.InverseMasses, SrcLodData.PhysicalMeshData.InverseMasses);
 
-			const int32 NextIndex = SectionWithCloth.Lod + 1;
-			const bool bNeedsTransitionDownData = NewClothingAssetData.LodMap.IsValidIndex(NextIndex) && NewClothingAssetData.LodMap[NextIndex] != INDEX_NONE;
-			if (bNeedsTransitionDownData)
-			{
-				NewLodData.TransitionDownSkinData.SetNum(SrcLodData.TransitionDownSkinData.Num() ? NewPhysicalMeshVerticesNum : 0);
-				CopyIfUsed(NewLodData.TransitionDownSkinData, SrcLodData.TransitionDownSkinData);
-			}
-			
-			const TMap<uint32, FPointWeightMap>& SrcPhysWeightMaps = SrcLodData.PhysicalMeshData.WeightMaps;
-			TMap<uint32, FPointWeightMap>& NewPhysWeightMaps = NewLodData.PhysicalMeshData.WeightMaps;
-
-			for (const TPair<uint32, FPointWeightMap>& WeightMap : SrcPhysWeightMaps)
-			{
-				if (WeightMap.Value.Values.Num() > 0)
+				const int32 PrevIndex = SectionWithCloth.Lod - 1;
+				const bool bNeedsTransitionUpData = NewClothingAssetData.LodMap.IsValidIndex(PrevIndex) && NewClothingAssetData.LodMap[PrevIndex] != INDEX_NONE;
+				if (bNeedsTransitionUpData)
 				{
-					FPointWeightMap& NewWeightMap = NewLodData.PhysicalMeshData.AddWeightMap(WeightMap.Key);
-					NewWeightMap.Values.SetNum(NewPhysicalMeshVerticesNum);
-
-					CopyIfUsed(NewWeightMap.Values, WeightMap.Value.Values);
+					NewLodData.TransitionUpSkinData.SetNum(SrcLodData.TransitionUpSkinData.Num() ? NewPhysicalMeshVerticesNum : 0);	
+					CopyIfUsed(NewLodData.TransitionUpSkinData, SrcLodData.TransitionUpSkinData);
 				}
-			}
-			// Remap render mesh to mesh indices.
-			if (bHasVerticesRemoved)
-			{
-				for (FMeshToMeshVertData& VertClothData : SectionWithCloth.MappingData)
-				{
-					uint16* Indices = VertClothData.SourceMeshVertIndices;
-					Indices[0] = (uint16)IndexMap[Indices[0]];
-					Indices[1] = (uint16)IndexMap[Indices[1]];
-					Indices[2] = (uint16)IndexMap[Indices[2]];
-				}
-			}
-
-			// Remap and trim physics mesh vertices and self collision indices. 
-			
-			// Returns the final size of Dst.
-			const auto TrimAndRemapTriangles = [&IndexMap](TArray<uint32>& Dst, const TArray<uint32>& Src) -> int32
-			{
-				check(Src.Num() % 3 == 0);
 
-				const int32 SrcNumElems = Src.Num();
-				if (!IndexMap.Num())
+				const int32 NextIndex = SectionWithCloth.Lod + 1;
+				const bool bNeedsTransitionDownData = NewClothingAssetData.LodMap.IsValidIndex(NextIndex) && NewClothingAssetData.LodMap[NextIndex] != INDEX_NONE;
+				if (bNeedsTransitionDownData)
 				{
-					//for (int32 Idx = 0; Idx < SrcNumElems; ++Idx)
-					//{
-					//	Dst[Idx] = Src[Idx];
-					//}
-
-					FMemory::Memcpy( Dst.GetData(), Src.GetData(), SrcNumElems*sizeof(uint32) );
-					return SrcNumElems;
+					NewLodData.TransitionDownSkinData.SetNum(SrcLodData.TransitionDownSkinData.Num() ? NewPhysicalMeshVerticesNum : 0);
+					CopyIfUsed(NewLodData.TransitionDownSkinData, SrcLodData.TransitionDownSkinData);
 				}
+				
+				const TMap<uint32, FPointWeightMap>& SrcPhysWeightMaps = SrcLodData.PhysicalMeshData.WeightMaps;
+				TMap<uint32, FPointWeightMap>& NewPhysWeightMaps = NewLodData.PhysicalMeshData.WeightMaps;
 
-				int32 DstNumElems = 0;
-				for (int32 Idx = 0; Idx < SrcNumElems; Idx += 3)
+				for (const TPair<uint32, FPointWeightMap>& WeightMap : SrcPhysWeightMaps)
 				{
-					const int32 Idx0 = IndexMap[Src[Idx + 0]];
-					const int32 Idx1 = IndexMap[Src[Idx + 1]];
-					const int32 Idx2 = IndexMap[Src[Idx + 2]];
-
-					// triangles are only copied if all vertices are used.
-					if (!((Idx0 < 0) | (Idx1 < 0) | (Idx2 < 0)))
+					if (WeightMap.Value.Values.Num() > 0)
 					{
-						Dst[DstNumElems + 0] = Idx0;
-						Dst[DstNumElems + 1] = Idx1;
-						Dst[DstNumElems + 2] = Idx2;
+						FPointWeightMap& NewWeightMap = NewLodData.PhysicalMeshData.AddWeightMap(WeightMap.Key);
+						NewWeightMap.Values.SetNum(NewPhysicalMeshVerticesNum);
 
-						DstNumElems += 3;
+						CopyIfUsed(NewWeightMap.Values, WeightMap.Value.Values);
 					}
 				}
-
-				return DstNumElems;
-			};
-
-			const TArray<uint32>& SrcPhysicalMeshIndices = SrcLodData.PhysicalMeshData.Indices;
-			TArray<uint32>& NewPhysicalMeshIndices = NewLodData.PhysicalMeshData.Indices;
-			NewPhysicalMeshIndices.SetNum(SrcPhysicalMeshIndices.Num());
-			NewPhysicalMeshIndices.SetNum(TrimAndRemapTriangles(NewPhysicalMeshIndices, SrcPhysicalMeshIndices), EAllowShrinking::No);
-		
-			const auto TrimAndRemapVertexSet = [&IndexMap](TSet<int32>& Dst, const TSet<int32>& Src)
-			{	
-				if (!IndexMap.Num())
+				// Remap render mesh to mesh indices.
+				if (bHasVerticesRemoved)
 				{
-					Dst = Src;
-					return;
+					for (FMeshToMeshVertData& VertClothData : SectionWithCloth.MappingData)
+					{
+						uint16* Indices = VertClothData.SourceMeshVertIndices;
+						Indices[0] = (uint16)IndexMap[Indices[0]];
+						Indices[1] = (uint16)IndexMap[Indices[1]];
+						Indices[2] = (uint16)IndexMap[Indices[2]];
+					}
 				}
 
-				Dst.Reserve(Src.Num());
-				for(const int32 SrcIdx : Src)
+				// Remap and trim physics mesh vertices and self collision indices. 
+				
+				// Returns the final size of Dst.
+				const auto TrimAndRemapTriangles = [&IndexMap](TArray<uint32>& Dst, const TArray<uint32>& Src) -> int32
 				{
-					const int32 MappedIdx = IndexMap[SrcIdx];
+					check(Src.Num() % 3 == 0);
 
-					if (MappedIdx >= 0)
+					const int32 SrcNumElems = Src.Num();
+					if (!IndexMap.Num())
 					{
-						Dst.Add(MappedIdx);
+						//for (int32 Idx = 0; Idx < SrcNumElems; ++Idx)
+						//{
+						//	Dst[Idx] = Src[Idx];
+						//}
+
+						FMemory::Memcpy( Dst.GetData(), Src.GetData(), SrcNumElems*sizeof(uint32) );
+						return SrcNumElems;
 					}
-				}
-			};
 
-			const TSet<int32>& SrcSelfCollisionVertexSet = SrcLodData.PhysicalMeshData.SelfCollisionVertexSet;
-			TSet<int32>& NewSelfCollisionVertexSet = NewLodData.PhysicalMeshData.SelfCollisionVertexSet;
-			TrimAndRemapVertexSet(NewSelfCollisionVertexSet, SrcSelfCollisionVertexSet);
-						
-			{
-				MUTABLE_CPUPROFILER_SCOPE(BuildClothTetherData)
+					int32 DstNumElems = 0;
+					for (int32 Idx = 0; Idx < SrcNumElems; Idx += 3)
+					{
+						const int32 Idx0 = IndexMap[Src[Idx + 0]];
+						const int32 Idx1 = IndexMap[Src[Idx + 1]];
+						const int32 Idx2 = IndexMap[Src[Idx + 2]];
 
-				auto TrimAndRemapTethers = [&IndexMap](FClothTetherData& Dst, const FClothTetherData& Src)
-				{
+						// triangles are only copied if all vertices are used.
+						if (!((Idx0 < 0) | (Idx1 < 0) | (Idx2 < 0)))
+						{
+							Dst[DstNumElems + 0] = Idx0;
+							Dst[DstNumElems + 1] = Idx1;
+							Dst[DstNumElems + 2] = Idx2;
+
+							DstNumElems += 3;
+						}
+					}
+
+					return DstNumElems;
+				};
+
+				const TArray<uint32>& SrcPhysicalMeshIndices = SrcLodData.PhysicalMeshData.Indices;
+				TArray<uint32>& NewPhysicalMeshIndices = NewLodData.PhysicalMeshData.Indices;
+				NewPhysicalMeshIndices.SetNum(SrcPhysicalMeshIndices.Num());
+				NewPhysicalMeshIndices.SetNum(TrimAndRemapTriangles(NewPhysicalMeshIndices, SrcPhysicalMeshIndices), EAllowShrinking::No);
+			
+				const auto TrimAndRemapVertexSet = [&IndexMap](TSet<int32>& Dst, const TSet<int32>& Src)
+				{	
 					if (!IndexMap.Num())
 					{
-						Dst.Tethers = Src.Tethers;
+						Dst = Src;
 						return;
 					}
 
-					Dst.Tethers.Reserve(Src.Tethers.Num());
-					for ( const TArray<TTuple<int32, int32, float>>& SrcTetherCluster : Src.Tethers )
+					Dst.Reserve(Src.Num());
+					for(const int32 SrcIdx : Src)
 					{
-						TArray<TTuple<int32, int32, float>>& DstTetherCluster = Dst.Tethers.Emplace_GetRef();
-						DstTetherCluster.Reserve(SrcTetherCluster.Num());
-						for ( const TTuple<int32, int32, float>& Tether : SrcTetherCluster )
-						{
-							const int32 Index0 = IndexMap[Tether.Get<0>()];
-							const int32 Index1 = IndexMap[Tether.Get<1>()];
-							if ((Index0 >= 0) & (Index1 >= 0))
-							{
-								DstTetherCluster.Emplace(Index0, Index1, Tether.Get<2>());
-							}
-						}
+						const int32 MappedIdx = IndexMap[SrcIdx];
 
-						if (!DstTetherCluster.Num())
+						if (MappedIdx >= 0)
 						{
-							Dst.Tethers.RemoveAt( Dst.Tethers.Num() - 1, EAllowShrinking::No );
+							Dst.Add(MappedIdx);
 						}
 					}
 				};
 
-				TrimAndRemapTethers(NewLodData.PhysicalMeshData.GeodesicTethers, SrcLodData.PhysicalMeshData.GeodesicTethers);
-				TrimAndRemapTethers(NewLodData.PhysicalMeshData.EuclideanTethers, SrcLodData.PhysicalMeshData.EuclideanTethers);
-			}
-		}
-	}
+				const TSet<int32>& SrcSelfCollisionVertexSet = SrcLodData.PhysicalMeshData.SelfCollisionVertexSet;
+				TSet<int32>& NewSelfCollisionVertexSet = NewLodData.PhysicalMeshData.SelfCollisionVertexSet;
+				TrimAndRemapVertexSet(NewSelfCollisionVertexSet, SrcSelfCollisionVertexSet);
+							
+				{
+					MUTABLE_CPUPROFILER_SCOPE(BuildClothTetherData)
 
-	// Try to find plausible values for LodTransitionData vertices that have lost the triangle to which are attached.
-	{
-		MUTABLE_CPUPROFILER_SCOPE(BuildLodTransitionData)
-		
-		for (const FSectionWithClothData& SectionWithCloth : SectionsWithCloth)
-		{
+					auto TrimAndRemapTethers = [&IndexMap](FClothTetherData& Dst, const FClothTetherData& Src)
+					{
+						if (!IndexMap.Num())
+						{
+							Dst.Tethers = Src.Tethers;
+							return;
+						}
 
-			FCustomizableObjectClothingAssetData& NewClothingAssetData = NewClothingAssetsData[SectionWithCloth.ClothAssetIndex];
-			FClothLODDataCommon& NewLodData = NewClothingAssetData.LodData[SectionWithCloth.ClothAssetLodIndex];
+						Dst.Tethers.Reserve(Src.Tethers.Num());
+						for ( const TArray<TTuple<int32, int32, float>>& SrcTetherCluster : Src.Tethers )
+						{
+							TArray<TTuple<int32, int32, float>>& DstTetherCluster = Dst.Tethers.Emplace_GetRef();
+							DstTetherCluster.Reserve(SrcTetherCluster.Num());
+							for ( const TTuple<int32, int32, float>& Tether : SrcTetherCluster )
+							{
+								const int32 Index0 = IndexMap[Tether.Get<0>()];
+								const int32 Index1 = IndexMap[Tether.Get<1>()];
+								if ((Index0 >= 0) & (Index1 >= 0))
+								{
+									DstTetherCluster.Emplace(Index0, Index1, Tether.Get<2>());
+								}
+							}
 
-			const int32 PhysicalMeshVerticesNum = NewLodData.PhysicalMeshData.Vertices.Num();
+							if (!DstTetherCluster.Num())
+							{
+								Dst.Tethers.RemoveAt( Dst.Tethers.Num() - 1, 1, EAllowShrinking::No );
+							}
+						}
+					};
 
-			if (!PhysicalMeshVerticesNum)
-			{
-				// Nothing to do.
-				continue;
+					TrimAndRemapTethers(NewLodData.PhysicalMeshData.GeodesicTethers, SrcLodData.PhysicalMeshData.GeodesicTethers);
+					TrimAndRemapTethers(NewLodData.PhysicalMeshData.EuclideanTethers, SrcLodData.PhysicalMeshData.EuclideanTethers);
+				}
 			}
-		
-			auto RemapTransitionMeshToMeshVertData = []( TArray<FMeshToMeshVertData>& InOutVertData, TArray<int32>& IndexMap )
+		}
+
+		// Try to find plausible values for LodTransitionData vertices that have lost the triangle to which are attached.
+		{
+			MUTABLE_CPUPROFILER_SCOPE(BuildLodTransitionData)
+			
+			for (const FSectionWithClothData& SectionWithCloth : SectionsWithCloth)
 			{
-				for (FMeshToMeshVertData& VertData : InOutVertData)
-				{
-					uint16* Indices = VertData.SourceMeshVertIndices;
-					Indices[0] = (uint16)IndexMap[Indices[0]];
-					Indices[1] = (uint16)IndexMap[Indices[1]];
-					Indices[2] = (uint16)IndexMap[Indices[2]];
-				}
-			};
 
-			if (NewLodData.TransitionDownSkinData.Num() > 0)
-			{	
-				TArray<int32>& IndexMap = PhysicsSectionLodsIndicesRemaps[SectionWithCloth.ClothAssetIndex][SectionWithCloth.ClothAssetLodIndex + 1];
+				FCustomizableObjectClothingAssetData& NewClothingAssetData = NewClothingAssetsData[SectionWithCloth.ClothAssetIndex];
+				FClothLODDataCommon& NewLodData = NewClothingAssetData.LodData[SectionWithCloth.ClothAssetLodIndex];
 
-				if (IndexMap.Num())
+				const int32 PhysicalMeshVerticesNum = NewLodData.PhysicalMeshData.Vertices.Num();
+
+				if (!PhysicalMeshVerticesNum)
 				{
-					RemapTransitionMeshToMeshVertData( NewLodData.TransitionDownSkinData, IndexMap );
+					// Nothing to do.
+					continue;
 				}
-			}
 			
-			if (NewLodData.TransitionUpSkinData.Num() > 0)
-			{	
-				TArray<int32>& IndexMap = PhysicsSectionLodsIndicesRemaps[SectionWithCloth.ClothAssetIndex][SectionWithCloth.ClothAssetLodIndex - 1];
-				if (IndexMap.Num())
+				auto RemapTransitionMeshToMeshVertData = []( TArray<FMeshToMeshVertData>& InOutVertData, TArray<int32>& IndexMap )
 				{
-					RemapTransitionMeshToMeshVertData( NewLodData.TransitionUpSkinData, IndexMap );
-				}
-			}
+					for (FMeshToMeshVertData& VertData : InOutVertData)
+					{
+						uint16* Indices = VertData.SourceMeshVertIndices;
+						Indices[0] = (uint16)IndexMap[Indices[0]];
+						Indices[1] = (uint16)IndexMap[Indices[1]];
+						Indices[2] = (uint16)IndexMap[Indices[2]];
+					}
+				};
 
-			struct FMeshPhysicsDesc
-			{
-				const TArray<FVector3f>& Vertices;
-				const TArray<FVector3f>& Normals;
-				const TArray<uint32>& Indices;
-			};	
-		
-			auto RebindVertex = [](const FMeshPhysicsDesc& Mesh, const FVector3f& InPosition, const FVector3f& InNormal, FMeshToMeshVertData& Out)
-			{
-				const FVector3f Normal = InNormal;
+				if (NewLodData.TransitionDownSkinData.Num() > 0)
+				{	
+					TArray<int32>& IndexMap = PhysicsSectionLodsIndicesRemaps[SectionWithCloth.ClothAssetIndex][SectionWithCloth.ClothAssetLodIndex + 1];
 
-				// We don't have the mesh tangent, find something plausible.
-				FVector3f Tan0, Tan1;
-				Normal.FindBestAxisVectors(Tan0, Tan1);
-				const FVector3f Tangent = Tan0;
-				
-				// Some of the math functions take as argument FVector, we'd want to be FVector3f. 
-				// This should be changed once support for the single type in the FMath functions is added. 
-				const FVector Position = (FVector)InPosition;
-				int32 BestBaseTriangleIdx = INDEX_NONE;
-				FVector::FReal BestDistanceSq = TNumericLimits<FVector::FReal>::Max();
+					if (IndexMap.Num())
+					{
+						RemapTransitionMeshToMeshVertData( NewLodData.TransitionDownSkinData, IndexMap );
+					}
+				}
 				
-				const int32 NumIndices = Mesh.Indices.Num();
-				check(NumIndices % 3 == 0);
+				if (NewLodData.TransitionUpSkinData.Num() > 0)
+				{	
+					TArray<int32>& IndexMap = PhysicsSectionLodsIndicesRemaps[SectionWithCloth.ClothAssetIndex][SectionWithCloth.ClothAssetLodIndex - 1];
+					if (IndexMap.Num())
+					{
+						RemapTransitionMeshToMeshVertData( NewLodData.TransitionUpSkinData, IndexMap );
+					}
+				}
 
-				for (int32 I = 0; I < NumIndices; I += 3)
+				struct FMeshPhysicsDesc
+				{
+					const TArray<FVector3f>& Vertices;
+					const TArray<FVector3f>& Normals;
+					const TArray<uint32>& Indices;
+				};	
+			
+				auto RebindVertex = [](const FMeshPhysicsDesc& Mesh, const FVector3f& InPosition, const FVector3f& InNormal, FMeshToMeshVertData& Out)
 				{
-					const FVector& A = (FVector)Mesh.Vertices[Mesh.Indices[I + 0]];
-					const FVector& B = (FVector)Mesh.Vertices[Mesh.Indices[I + 1]];
-					const FVector& C = (FVector)Mesh.Vertices[Mesh.Indices[I + 2]];
+					const FVector3f Normal = InNormal;
 
-					FVector ClosestTrianglePoint = FMath::ClosestPointOnTriangleToPoint(Position, (FVector)A, (FVector)B, (FVector)C);
+					// We don't have the mesh tangent, find something plausible.
+					FVector3f Tan0, Tan1;
+					Normal.FindBestAxisVectors(Tan0, Tan1);
+					const FVector3f Tangent = Tan0;
+					
+					// Some of the math functions take as argument FVector, we'd want to be FVector3f. 
+					// This should be changed once support for the single type in the FMath functions is added. 
+					const FVector Position = (FVector)InPosition;
+					int32 BestBaseTriangleIdx = INDEX_NONE;
+					FVector::FReal BestDistanceSq = TNumericLimits<FVector::FReal>::Max();
+					
+					const int32 NumIndices = Mesh.Indices.Num();
+					check(NumIndices % 3 == 0);
 
-					const FVector::FReal CurrentDistSq = (ClosestTrianglePoint - Position).SizeSquared();
-					if (CurrentDistSq < BestDistanceSq)
+					for (int32 I = 0; I < NumIndices; I += 3)
 					{
-						BestDistanceSq = CurrentDistSq;
-						BestBaseTriangleIdx = I;
+						const FVector& A = (FVector)Mesh.Vertices[Mesh.Indices[I + 0]];
+						const FVector& B = (FVector)Mesh.Vertices[Mesh.Indices[I + 1]];
+						const FVector& C = (FVector)Mesh.Vertices[Mesh.Indices[I + 2]];
+
+						FVector ClosestTrianglePoint = FMath::ClosestPointOnTriangleToPoint(Position, (FVector)A, (FVector)B, (FVector)C);
+
+						const FVector::FReal CurrentDistSq = (ClosestTrianglePoint - Position).SizeSquared();
+						if (CurrentDistSq < BestDistanceSq)
+						{
+							BestDistanceSq = CurrentDistSq;
+							BestBaseTriangleIdx = I;
+						}
 					}
-				}
 
-				check(BestBaseTriangleIdx >= 0);
+					check(BestBaseTriangleIdx >= 0);
 
-				auto ComputeBaryCoordsAndDist = [](const FVector3f& A, const FVector3f& B, const FVector3f& C, const FVector3f& P) -> FVector4f
-				{
-					FPlane4f TrianglePlane(A, B, C);
+					auto ComputeBaryCoordsAndDist = [](const FVector3f& A, const FVector3f& B, const FVector3f& C, const FVector3f& P) -> FVector4f
+					{
+						FPlane4f TrianglePlane(A, B, C);
 
-					const FVector3f PointOnTriPlane = FVector3f::PointPlaneProject(P, TrianglePlane);
-					const FVector3f BaryCoords = (FVector3f)FMath::ComputeBaryCentric2D((FVector)PointOnTriPlane, (FVector)A, (FVector)B, (FVector)C);
+						const FVector3f PointOnTriPlane = FVector3f::PointPlaneProject(P, TrianglePlane);
+						const FVector3f BaryCoords = (FVector3f)FMath::ComputeBaryCentric2D((FVector)PointOnTriPlane, (FVector)A, (FVector)B, (FVector)C);
 
-					return FVector4f(BaryCoords, TrianglePlane.PlaneDot((FVector3f)P));
-				};
+						return FVector4f(BaryCoords, TrianglePlane.PlaneDot((FVector3f)P));
+					};
 
-				const FVector3f& A = Mesh.Vertices[Mesh.Indices[BestBaseTriangleIdx + 0]];
-				const FVector3f& B = Mesh.Vertices[Mesh.Indices[BestBaseTriangleIdx + 1]];
-				const FVector3f& C = Mesh.Vertices[Mesh.Indices[BestBaseTriangleIdx + 2]];
-
-				Out.PositionBaryCoordsAndDist = ComputeBaryCoordsAndDist(A, B, C, (FVector3f)Position );
-				Out.NormalBaryCoordsAndDist = ComputeBaryCoordsAndDist(A, B, C, (FVector3f)Position + Normal );
-				Out.TangentBaryCoordsAndDist = ComputeBaryCoordsAndDist(A, B, C, (FVector3f)Position + Tangent );
-				Out.SourceMeshVertIndices[0] = (uint16)Mesh.Indices[BestBaseTriangleIdx + 0];
-				Out.SourceMeshVertIndices[1] = (uint16)Mesh.Indices[BestBaseTriangleIdx + 1]; 
-				Out.SourceMeshVertIndices[2] = (uint16)Mesh.Indices[BestBaseTriangleIdx + 2];
-			};
-		
-			auto RecreateTransitionData = [&PhysicsSectionLodsIndicesRemaps, &RebindVertex]( 
-				const FMeshPhysicsDesc& ToMesh, const FMeshPhysicsDesc& FromMesh, const TArray<int32>& IndexMap, TArray<FMeshToMeshVertData>& InOutTransitionData )
-			{
-				if (!IndexMap.Num())
-				{
-					return;
-				}
+					const FVector3f& A = Mesh.Vertices[Mesh.Indices[BestBaseTriangleIdx + 0]];
+					const FVector3f& B = Mesh.Vertices[Mesh.Indices[BestBaseTriangleIdx + 1]];
+					const FVector3f& C = Mesh.Vertices[Mesh.Indices[BestBaseTriangleIdx + 2]];
 
-				if (!InOutTransitionData.Num())
+					Out.PositionBaryCoordsAndDist = ComputeBaryCoordsAndDist(A, B, C, (FVector3f)Position );
+					Out.NormalBaryCoordsAndDist = ComputeBaryCoordsAndDist(A, B, C, (FVector3f)Position + Normal );
+					Out.TangentBaryCoordsAndDist = ComputeBaryCoordsAndDist(A, B, C, (FVector3f)Position + Tangent );
+					Out.SourceMeshVertIndices[0] = (uint16)Mesh.Indices[BestBaseTriangleIdx + 0];
+					Out.SourceMeshVertIndices[1] = (uint16)Mesh.Indices[BestBaseTriangleIdx + 1]; 
+					Out.SourceMeshVertIndices[2] = (uint16)Mesh.Indices[BestBaseTriangleIdx + 2];
+				};
+			
+				auto RecreateTransitionData = [&PhysicsSectionLodsIndicesRemaps, &RebindVertex]( 
+					const FMeshPhysicsDesc& ToMesh, const FMeshPhysicsDesc& FromMesh, const TArray<int32>& IndexMap, TArray<FMeshToMeshVertData>& InOutTransitionData )
 				{
-					return;
-				}
+					if (!IndexMap.Num())
+					{
+						return;
+					}
 
-				const int32 TransitionDataNum = InOutTransitionData.Num();
-				
-				for (int32 I = 0; I < TransitionDataNum; ++I)
-				{
-					FMeshToMeshVertData& VertData = InOutTransitionData[I];
-					uint16* Indices = VertData.SourceMeshVertIndices;
+					if (!InOutTransitionData.Num())
+					{
+						return;
+					}
 
-					// If any original indices are missing but the vertex is still alive rebind the vertex.
-					// In general, the number of rebinds should be small.
+					const int32 TransitionDataNum = InOutTransitionData.Num();
+					
+					for (int32 I = 0; I < TransitionDataNum; ++I)
+					{
+						FMeshToMeshVertData& VertData = InOutTransitionData[I];
+						uint16* Indices = VertData.SourceMeshVertIndices;
 
-					// Currently, if any index is missing we rebind to the closest triangle but it could be nice to use the remaining indices, 
-					// if any, to find the most appropriate triangle to bind to. 
-					const bool bNeedsRebind = (Indices[0] == 0xFFFF) | (Indices[1] == 0xFFFF) | (Indices[2] == 0xFFFF);
+						// If any original indices are missing but the vertex is still alive rebind the vertex.
+						// In general, the number of rebinds should be small.
 
-					if (bNeedsRebind)
-					{
-						RebindVertex( ToMesh, FromMesh.Vertices[I], FromMesh.Normals[I], VertData );
+						// Currently, if any index is missing we rebind to the closest triangle but it could be nice to use the remaining indices, 
+						// if any, to find the most appropriate triangle to bind to. 
+						const bool bNeedsRebind = (Indices[0] == 0xFFFF) | (Indices[1] == 0xFFFF) | (Indices[2] == 0xFFFF);
+
+						if (bNeedsRebind)
+						{
+							RebindVertex( ToMesh, FromMesh.Vertices[I], FromMesh.Normals[I], VertData );
+						}
 					}
-				}
-			};
+				};
 
-			const FMeshPhysicsDesc CurrentPhysicsMesh {
-				NewClothingAssetData.LodData[SectionWithCloth.ClothAssetLodIndex].PhysicalMeshData.Vertices,
-				NewClothingAssetData.LodData[SectionWithCloth.ClothAssetLodIndex].PhysicalMeshData.Normals,
-				NewClothingAssetData.LodData[SectionWithCloth.ClothAssetLodIndex].PhysicalMeshData.Indices };
+				const FMeshPhysicsDesc CurrentPhysicsMesh {
+					NewClothingAssetData.LodData[SectionWithCloth.ClothAssetLodIndex].PhysicalMeshData.Vertices,
+					NewClothingAssetData.LodData[SectionWithCloth.ClothAssetLodIndex].PhysicalMeshData.Normals,
+					NewClothingAssetData.LodData[SectionWithCloth.ClothAssetLodIndex].PhysicalMeshData.Indices };
 
-			const TArray<TArray<int32>>& SectionIndexRemaps = PhysicsSectionLodsIndicesRemaps[SectionWithCloth.ClothAssetIndex];
-			
-			if (SectionWithCloth.ClothAssetLodIndex < SectionIndexRemaps.Num() - 1)
-			{
-				const TArray<int32>& IndexMap = SectionIndexRemaps[SectionWithCloth.ClothAssetLodIndex + 1];
+				const TArray<TArray<int32>>& SectionIndexRemaps = PhysicsSectionLodsIndicesRemaps[SectionWithCloth.ClothAssetIndex];
 				
-				const FMeshPhysicsDesc TransitionDownTarget {  
-					NewClothingAssetData.LodData[SectionWithCloth.ClothAssetLodIndex + 1].PhysicalMeshData.Vertices,
-					NewClothingAssetData.LodData[SectionWithCloth.ClothAssetLodIndex + 1].PhysicalMeshData.Normals,
-					NewClothingAssetData.LodData[SectionWithCloth.ClothAssetLodIndex + 1].PhysicalMeshData.Indices };
+				if (SectionWithCloth.ClothAssetLodIndex < SectionIndexRemaps.Num() - 1)
+				{
+					const TArray<int32>& IndexMap = SectionIndexRemaps[SectionWithCloth.ClothAssetLodIndex + 1];
 					
-				RecreateTransitionData( TransitionDownTarget, CurrentPhysicsMesh, IndexMap, NewLodData.TransitionDownSkinData );
-			}
-			
-			if (SectionWithCloth.ClothAssetLodIndex > 0)
-			{
-				const TArray<int32>& IndexMap = SectionIndexRemaps[SectionWithCloth.ClothAssetLodIndex - 1];
+					const FMeshPhysicsDesc TransitionDownTarget {  
+						NewClothingAssetData.LodData[SectionWithCloth.ClothAssetLodIndex + 1].PhysicalMeshData.Vertices,
+						NewClothingAssetData.LodData[SectionWithCloth.ClothAssetLodIndex + 1].PhysicalMeshData.Normals,
+						NewClothingAssetData.LodData[SectionWithCloth.ClothAssetLodIndex + 1].PhysicalMeshData.Indices };
+						
+					RecreateTransitionData( TransitionDownTarget, CurrentPhysicsMesh, IndexMap, NewLodData.TransitionDownSkinData );
+				}
 				
-				FMeshPhysicsDesc TransitionUpTarget{  
-					NewClothingAssetData.LodData[SectionWithCloth.ClothAssetLodIndex - 1].PhysicalMeshData.Vertices,
-					NewClothingAssetData.LodData[SectionWithCloth.ClothAssetLodIndex - 1].PhysicalMeshData.Normals,
-					NewClothingAssetData.LodData[SectionWithCloth.ClothAssetLodIndex - 1].PhysicalMeshData.Indices };
-		
-				RecreateTransitionData( TransitionUpTarget, CurrentPhysicsMesh, IndexMap, NewLodData.TransitionUpSkinData );
+				if (SectionWithCloth.ClothAssetLodIndex > 0)
+				{
+					const TArray<int32>& IndexMap = SectionIndexRemaps[SectionWithCloth.ClothAssetLodIndex - 1];
+					
+					FMeshPhysicsDesc TransitionUpTarget{  
+						NewClothingAssetData.LodData[SectionWithCloth.ClothAssetLodIndex - 1].PhysicalMeshData.Vertices,
+						NewClothingAssetData.LodData[SectionWithCloth.ClothAssetLodIndex - 1].PhysicalMeshData.Normals,
+						NewClothingAssetData.LodData[SectionWithCloth.ClothAssetLodIndex - 1].PhysicalMeshData.Indices };
+			
+					RecreateTransitionData( TransitionUpTarget, CurrentPhysicsMesh, IndexMap, NewLodData.TransitionUpSkinData );
+				}
 			}
 		}
 	}
@@ -4361,7 +4411,7 @@ void UCustomizableInstancePrivate::BuildOrCopyClothingData(const TSharedRef<FUpd
 	}
 
 	TArray<UCustomizableObjectClothingAsset*> NewClothingAssets;
-	NewClothingAssets.Init(nullptr, ContributingClothingAssetsData.Num());
+	NewClothingAssets.Init(nullptr, ClothingAssetsData.Num());
 
 	{ 
 		MUTABLE_CPUPROFILER_SCOPE(CreateClothingAssets)
@@ -4396,17 +4446,43 @@ void UCustomizableInstancePrivate::BuildOrCopyClothingData(const TSharedRef<FUpd
 			}
 		}
 
-		check(NewClothingAssets.Num() == ClothingPhysicsAssets.Num());
-		for (int32 I = 0; I < NewClothingAssetsData.Num(); ++I)
+		const int32 NumClothingAssets = NewClothingAssetsData.Num(); 
+		check(NumClothingAssets == ClothingPhysicsAssets.Num());
+	
+		bool bAllNamesUnique = true;
+		TArray<FName, TInlineAllocator<8>> UniqueAssetNames;
+
+		for (int32 I = 0; I < NumClothingAssets; ++I)
 		{
-			// skip assets not set.
+			// Skip assets not set.
 			if (!NewClothingAssetsData[I].LodData.Num())
 			{
 				continue;
 			}
-	
-			FName UniqueClothingAssetName = FName(FString::Printf(TEXT("%s_%d"), *NewClothingAssetsData[I].Name.ToString(), I));
-			NewClothingAssets[I] = NewObject<UCustomizableObjectClothingAsset>(SkeletalMesh, UniqueClothingAssetName);
+
+			const int32 PrevNumUniqueElems = UniqueAssetNames.Num();
+			const int32 ElemIndex = UniqueAssetNames.AddUnique(NewClothingAssetsData[I].Name);
+
+			if (ElemIndex < PrevNumUniqueElems)
+			{
+				bAllNamesUnique = false;
+				break;
+			}
+		}
+
+		for (int32 I = 0; I < NumClothingAssets; ++I)
+		{
+			// Skip assets not set.
+			if (!NewClothingAssetsData[I].LodData.Num())
+			{
+				continue;
+			}
+
+			FName ClothingAssetObjectName = bAllNamesUnique 
+					? NewClothingAssetsData[I].Name
+					: FName(FString::Printf(TEXT("%s_%d"), *NewClothingAssetsData[I].Name.ToString(), I));;
+
+			NewClothingAssets[I] = NewObject<UCustomizableObjectClothingAsset>(SkeletalMesh, ClothingAssetObjectName);
 			
 			// The data can be moved to the actual asset since it will not be used anymore.
 			NewClothingAssets[I]->LodMap = MoveTemp(NewClothingAssetsData[I].LodMap);
@@ -4419,7 +4495,7 @@ void UCustomizableInstancePrivate::BuildOrCopyClothingData(const TSharedRef<FUpd
 			NewClothingAssets[I]->CalculateReferenceBoneIndex();	
 			NewClothingAssets[I]->PhysicsAsset = ClothingPhysicsAssets[I];
 
-			for (const FCustomizableObjectClothConfigData& ConfigData : ContributingClothingAssetsData[I].ConfigsData)
+			for (const FCustomizableObjectClothConfigData& ConfigData : ClothingAssetsData[I].ConfigsData)
 			{
 				UClothConfigCommon* ClothConfig = CreateNewClothConfigFromData(NewClothingAssets[I], ConfigData);
 				if (ClothConfig)
@@ -4664,6 +4740,7 @@ FAutoConsoleVariableRef CVarMutableHighPriorityLoading(
 	bEnableHighPriorityLoading,
 	TEXT("If enabled, the request to load additional assets will have high priority."));
 
+
 UE::Tasks::FTask UCustomizableInstancePrivate::LoadAdditionalAssetsAndData(
 		const TSharedRef<FUpdateContextPrivate>& OperationData, FStreamableManager& StreamableManager, bool bAsync)
 {
@@ -4675,6 +4752,7 @@ UE::Tasks::FTask UCustomizableInstancePrivate::LoadAdditionalAssetsAndData(
 
 	TArray<FSoftObjectPath> AssetsToStream;
 	TArray<uint32> RealTimeMorphStreamableBlocksToStream;
+	TArray<uint32> ClothingStreamableBlocksToStream;
 
 	TArray<FInstanceUpdateData::FLOD>& LODs = OperationData->InstanceUpdateData.LODs;
 	TArray<FInstanceUpdateData::FComponent>& Components = OperationData->InstanceUpdateData.Components;
@@ -4682,7 +4760,7 @@ UE::Tasks::FTask UCustomizableInstancePrivate::LoadAdditionalAssetsAndData(
 	ObjectToInstanceIndexMap.Empty();
 	ReferencedMaterials.Empty();
 
-	const int32 NumClothingAssets = CustomizableObject->ContributingClothingAssetsData.Num();
+	const int32 NumClothingAssets = ModelResources.ClothingAssetsData.Num();
 	ClothingPhysicsAssets.Reset(NumClothingAssets);
 	ClothingPhysicsAssets.SetNum(NumClothingAssets);
 
@@ -4755,195 +4833,204 @@ UE::Tasks::FTask UCustomizableInstancePrivate::LoadAdditionalAssetsAndData(
 		}
 	}
 
-
 	// Load assets coming from SubMeshes of the newly generated Mesh
 	if (OperationData->InstanceUpdateData.LODs.Num())
 	{
-		const FInstanceUpdateData::FLOD& LOD = OperationData->InstanceUpdateData.LODs[OperationData->GetMinLOD()];
-		for (int32 ComponentIndex = 0; ComponentIndex < LOD.ComponentCount; ++ComponentIndex)
+		for (int32 LODIndex = OperationData->FirstLODAvailable; LODIndex < OperationData->NumLODsAvailable; ++LODIndex)
 		{
-			const FInstanceUpdateData::FComponent& Component = Components[LOD.FirstComponent + ComponentIndex];
-			mu::MeshPtrConst MutableMesh = Component.Mesh;
-
-			if (!MutableMesh)
+			const FInstanceUpdateData::FLOD& LOD = OperationData->InstanceUpdateData.LODs[LODIndex];
+			for (int32 ComponentIndex = 0; ComponentIndex < LOD.ComponentCount; ++ComponentIndex)
 			{
-				continue;
-			}
+				const FInstanceUpdateData::FComponent& Component = Components[LOD.FirstComponent + ComponentIndex];
+				mu::MeshPtrConst MutableMesh = Component.Mesh;
+
+				if (!MutableMesh)
+				{
+					continue;
+				}
 
-			FCustomizableInstanceComponentData* ComponentData = GetComponentData(Component.Id);
+				FCustomizableInstanceComponentData* ComponentData = GetComponentData(Component.Id);
 
-			const TArray<uint64>& StreamedResources = MutableMesh->GetStreamedResources();
+				const TArray<uint64>& StreamedResources = MutableMesh->GetStreamedResources();
 
-			for (uint64 ResourceId : StreamedResources)
-			{
-				const TArray<FCustomizableObjectStreamedResourceData>& StreamedResourcesData = CustomizableObject->GetPrivate()->GetStreamedResourceData();
-				FCustomizableObjectStreameableResourceId TypedResourceId = BitCast<FCustomizableObjectStreameableResourceId>(ResourceId);	
-	
-				if (TypedResourceId.Type == (uint8)FCustomizableObjectStreameableResourceId::EType::AssetUserData)
+				for (uint64 ResourceId : StreamedResources)
 				{
-					const uint32 ResourceIndex = TypedResourceId.Id;
-					if (!StreamedResourcesData.IsValidIndex(ResourceIndex))
-					{
-						UE_LOG(LogMutable, Error, TEXT("Invalid streamed resource index. Max Index [%d]. Resource Index [%d]."), StreamedResourcesData.Num(), ResourceIndex);
-						continue; 
-					}
-
-					const FCustomizableObjectStreamedResourceData& StreamedResource = StreamedResourcesData[ResourceIndex];
-					if (!StreamedResource.IsLoaded())
+					const TArray<FCustomizableObjectStreamedResourceData>& StreamedResourcesData = CustomizableObject->GetPrivate()->GetStreamedResourceData();
+					FCustomizableObjectStreameableResourceId TypedResourceId = BitCast<FCustomizableObjectStreameableResourceId>(ResourceId);	
+		
+					if (TypedResourceId.Type == (uint8)FCustomizableObjectStreameableResourceId::EType::AssetUserData)
 					{
-						AssetsToStream.AddUnique(StreamedResource.GetPath().ToSoftObjectPath());
-					}
+						const uint32 ResourceIndex = TypedResourceId.Id;
+						if (!StreamedResourcesData.IsValidIndex(ResourceIndex))
+						{
+							UE_LOG(LogMutable, Error, TEXT("Invalid streamed resource index. Max Index [%d]. Resource Index [%d]."), StreamedResourcesData.Num(), ResourceIndex);
+							continue; 
+						}
 
-					ComponentData->StreamedResourceIndex.Add(ResourceIndex);
-				}
-				else if (TypedResourceId.Type == (uint8)FCustomizableObjectStreameableResourceId::EType::RealTimeMorphTarget)
-				{
-					check(TypedResourceId.Id != 0 && TypedResourceId.Id <= TNumericLimits<uint32>::Max());
+						const FCustomizableObjectStreamedResourceData& StreamedResource = StreamedResourcesData[ResourceIndex];
+						if (!StreamedResource.IsLoaded())
+						{
+							AssetsToStream.AddUnique(StreamedResource.GetPath().ToSoftObjectPath());
+						}
 
-					const TMap<uint32, FRealTimeMorphStreamable>& MorphsStreamables = ModelResources.RealTimeMorphStreamables;
-					if (MorphsStreamables.Contains((uint32)TypedResourceId.Id))
-					{
-						RealTimeMorphStreamableBlocksToStream.AddUnique(TypedResourceId.Id);
+						ComponentData->StreamedResourceIndex.Add(ResourceIndex);
 					}
-					else
+					else if (TypedResourceId.Type == (uint8)FCustomizableObjectStreameableResourceId::EType::RealTimeMorphTarget)
 					{
-						UE_LOG(LogMutable, Error, TEXT("Invalid streamed real time morph target data block [%d] found."), TypedResourceId.Id);
-					}
-				}
-				else
-				{
-					UE_LOG(LogMutable, Error, TEXT("Unknown streamed resource type found."));
-					check(false);
-				}
-			}
-
-			const bool bReplacePhysicsAssets = HasCOInstanceFlags(ReplacePhysicsAssets);
-
-			for (int32 TagIndex = 0; TagIndex < MutableMesh->GetTagCount(); ++TagIndex)
-			{
-				FString Tag = MutableMesh->GetTag(TagIndex);
-				if (Tag.RemoveFromStart("__PA:"))
-				{
-					const int32 AssetIndex = FCString::Atoi(*Tag);
-					const TSoftObjectPtr<UPhysicsAsset>& PhysicsAsset = ModelResources.PhysicsAssets.IsValidIndex(AssetIndex) ? ModelResources.PhysicsAssets[AssetIndex] : nullptr;
+						check(TypedResourceId.Id != 0 && TypedResourceId.Id <= TNumericLimits<uint32>::Max());
 
-					if (!PhysicsAsset.IsNull())
-					{
-						if (PhysicsAsset.Get())
+						const TMap<uint32, FRealTimeMorphStreamable>& MorphsStreamables = ModelResources.RealTimeMorphStreamables;
+						if (MorphsStreamables.Contains((uint32)TypedResourceId.Id))
 						{
-							ComponentData->PhysicsAssets.PhysicsAssetsToMerge.Add(PhysicsAsset.Get());
+							RealTimeMorphStreamableBlocksToStream.AddUnique(TypedResourceId.Id);
 						}
 						else
 						{
-							ComponentData->PhysicsAssets.PhysicsAssetToLoad.Add(AssetIndex);
-							AssetsToStream.Add(PhysicsAsset.ToSoftObjectPath());
+							UE_LOG(LogMutable, Error, TEXT("Invalid streamed real time morph target data block [%d] found."), TypedResourceId.Id);
 						}
 					}
-				}
-				else if (Tag.RemoveFromStart("__ClothPA:"))
-				{
-					FString AssetIndexString, PhysicsAssetIndexString;
-
-					if (Tag.Split(TEXT("_"), &AssetIndexString, &PhysicsAssetIndexString))
+					else if (TypedResourceId.Type == (uint8)FCustomizableObjectStreameableResourceId::EType::Clothing)
 					{
-						if (AssetIndexString.IsNumeric() && PhysicsAssetIndexString.IsNumeric())
+						check(TypedResourceId.Id != 0 && TypedResourceId.Id <= TNumericLimits<uint32>::Max());
+						
+						const TMap<uint32, FClothingStreamable>& ClothingStreamables = ModelResources.ClothingStreamables;
+						if (const FClothingStreamable* ClothingStreamable = ClothingStreamables.Find(TypedResourceId.Id))
 						{
-							const int32 AssetIndex = FCString::Atoi(*AssetIndexString);
-							const int32 PhysicsAssetIndex = FCString::Atoi(*PhysicsAssetIndexString);
-
+							// TODO: Add async loading of ClothingAsset Data. This could be loaded as an streamead resource similar to and the asset user data.
+							int32 ClothingAssetIndex = ClothingStreamable->ClothingAssetIndex; 
+							int32 PhysicsAssetIndex = ClothingStreamable->PhysicsAssetIndex;					
 							const TSoftObjectPtr<UPhysicsAsset>& PhysicsAsset = ModelResources.PhysicsAssets.IsValidIndex(PhysicsAssetIndex) 
 									? ModelResources.PhysicsAssets[PhysicsAssetIndex] 
-									: nullptr;
+								    : nullptr;
 
 							// The entry should always be in the map
 							if (!PhysicsAsset.IsNull())
 							{
 								if (PhysicsAsset.Get())
 								{
-									if (ClothingPhysicsAssets.IsValidIndex(AssetIndex))
+									if (ClothingPhysicsAssets.IsValidIndex(ClothingAssetIndex))
 									{
-										ClothingPhysicsAssets[AssetIndex] = PhysicsAsset.Get();
+										ClothingPhysicsAssets[ClothingAssetIndex] = PhysicsAsset.Get();
 									}
 								}
 								else
 								{
-									ComponentData->ClothingPhysicsAssetsToStream.Emplace(AssetIndex, PhysicsAssetIndex);
-									AssetsToStream.Add(PhysicsAsset.ToSoftObjectPath());
+									ComponentData->ClothingPhysicsAssetsToStream.Emplace(ClothingAssetIndex, PhysicsAssetIndex);
+									AssetsToStream.AddUnique(PhysicsAsset.ToSoftObjectPath());
 								}
 							}
+
+							ClothingStreamableBlocksToStream.AddUnique(TypedResourceId.Id);
 						}
+						else
+						{
+							UE_LOG(LogMutable, Error, TEXT("Invalid streamed clothing data block [%d] found."), TypedResourceId.Id);
+						}
+					}
+					else
+					{
+						UE_LOG(LogMutable, Error, TEXT("Unknown streamed resource type found."));
+						check(false);
 					}
 				}
-				if (Tag.RemoveFromStart("__AnimBP:"))
-				{
-					FString SlotIndexString, AnimBpIndexString;
 
-					if (Tag.Split(TEXT("_Slot_"), &SlotIndexString, &AnimBpIndexString))
+				const bool bReplacePhysicsAssets = HasCOInstanceFlags(ReplacePhysicsAssets);
+
+				for (int32 TagIndex = 0; TagIndex < MutableMesh->GetTagCount(); ++TagIndex)
+				{
+					FString Tag = MutableMesh->GetTag(TagIndex);
+					if (Tag.RemoveFromStart("__PA:"))
 					{
-						if (SlotIndexString.IsEmpty() || AnimBpIndexString.IsEmpty())
+						const int32 AssetIndex = FCString::Atoi(*Tag);
+						const TSoftObjectPtr<UPhysicsAsset>& PhysicsAsset = ModelResources.PhysicsAssets.IsValidIndex(AssetIndex) ? ModelResources.PhysicsAssets[AssetIndex] : nullptr;
+
+						if (!PhysicsAsset.IsNull())
 						{
-							continue;
+							if (PhysicsAsset.Get())
+							{
+								ComponentData->PhysicsAssets.PhysicsAssetsToMerge.Add(PhysicsAsset.Get());
+							}
+							else
+							{
+								ComponentData->PhysicsAssets.PhysicsAssetToLoad.Add(AssetIndex);
+								AssetsToStream.AddUnique(PhysicsAsset.ToSoftObjectPath());
+							}
 						}
+					}
+					
+					if (Tag.RemoveFromStart("__AnimBP:"))
+					{
+						FString SlotIndexString, AnimBpIndexString;
 
-						const int32 AnimBpIndex = FCString::Atoi(*AnimBpIndexString);
-						if (!ModelResources.AnimBPs.IsValidIndex(AnimBpIndex))
+						if (Tag.Split(TEXT("_Slot_"), &SlotIndexString, &AnimBpIndexString))
 						{
-							continue;
-						}
+							if (SlotIndexString.IsEmpty() || AnimBpIndexString.IsEmpty())
+							{
+								continue;
+							}
 
-						FName SlotIndex = *SlotIndexString;
+							const int32 AnimBpIndex = FCString::Atoi(*AnimBpIndexString);
+							if (!ModelResources.AnimBPs.IsValidIndex(AnimBpIndex))
+							{
+								continue;
+							}
 
-						const TSoftClassPtr<UAnimInstance>& AnimBPAsset = ModelResources.AnimBPs[AnimBpIndex];
+							FName SlotIndex = *SlotIndexString;
 
-						if (!AnimBPAsset.IsNull())
-						{
-							if (!ComponentData->AnimSlotToBP.Contains(SlotIndex))
-							{
-								ComponentData->AnimSlotToBP.Add(SlotIndex, AnimBPAsset);
+							const TSoftClassPtr<UAnimInstance>& AnimBPAsset = ModelResources.AnimBPs[AnimBpIndex];
 
-								if (AnimBPAsset.Get())
+							if (!AnimBPAsset.IsNull())
+							{
+								const TSoftClassPtr<UAnimInstance>* FoundAnimBpSlot = ComponentData->AnimSlotToBP.Find(SlotIndex);
+								bool bIsSameAnimBp = FoundAnimBpSlot && AnimBPAsset == *FoundAnimBpSlot;
+								if (!FoundAnimBpSlot)
 								{
-									GatheredAnimBPs.Add(AnimBPAsset.Get());
+									ComponentData->AnimSlotToBP.Add(SlotIndex, AnimBPAsset);
+
+									if (AnimBPAsset.Get())
+									{
+										GatheredAnimBPs.Add(AnimBPAsset.Get());
+									}
+									else
+									{
+										AssetsToStream.AddUnique(AnimBPAsset.ToSoftObjectPath());
+									}
 								}
-								else
+								else if (!bIsSameAnimBp)
 								{
-									AssetsToStream.Add(AnimBPAsset.ToSoftObjectPath());
+									// Two submeshes should not have the same animation slot index
+									OperationData->UpdateResult = EUpdateResult::Warning;
+
+									FString WarningMessage = FString::Printf(TEXT("Two submeshes have the same anim slot index [%s] in a Mutable Instance."), *SlotIndex.ToString());
+									UE_LOG(LogMutable, Warning, TEXT("%s"), *WarningMessage);
+	#if WITH_EDITOR
+									FMessageLog MessageLog("Mutable");
+									MessageLog.Notify(FText::FromString(WarningMessage), EMessageSeverity::Warning, true);
+	#endif
 								}
 							}
-							else
-							{
-								// Two submeshes should not have the same animation slot index
-								OperationData->UpdateResult = EUpdateResult::Warning;
-
-								FString WarningMessage = FString::Printf(TEXT("Two submeshes have the same anim slot index [%s] in a Mutable Instance."), *SlotIndex.ToString());
-								UE_LOG(LogMutable, Warning, TEXT("%s"), *WarningMessage);
-#if WITH_EDITOR
-								FMessageLog MessageLog("Mutable");
-								MessageLog.Notify(FText::FromString(WarningMessage), EMessageSeverity::Warning, true);
-#endif
-							}
 						}
 					}
+					else if (Tag.RemoveFromStart("__AnimBPTag:"))
+					{
+						AnimBPGameplayTags.AddTag(FGameplayTag::RequestGameplayTag(*Tag));
+					}
+	#if WITH_EDITORONLY_DATA
+					else if (Tag.RemoveFromStart("__MeshPath:"))
+					{
+						ComponentData->MeshPartPaths.Add(Tag);
+					}
+	#endif
 				}
-				else if (Tag.RemoveFromStart("__AnimBPTag:"))
-				{
-					AnimBPGameplayTags.AddTag(FGameplayTag::RequestGameplayTag(*Tag));
-				}
-#if WITH_EDITORONLY_DATA
-				else if (Tag.RemoveFromStart("__MeshPath:"))
+
+				const int32 AdditionalPhysicsNum = MutableMesh->AdditionalPhysicsBodies.Num();
+				for (int32 I = 0; I < AdditionalPhysicsNum; ++I)
 				{
-					ComponentData->MeshPartPaths.Add(Tag);
+					const int32 ExternalId = MutableMesh->AdditionalPhysicsBodies[I]->CustomId;
+					
+					ComponentData->PhysicsAssets.AdditionalPhysicsAssetsToLoad.Add(ExternalId);
+					AssetsToStream.Add(ModelResources.AnimBpOverridePhysiscAssetsInfo[ExternalId].SourceAsset.ToSoftObjectPath());
 				}
-#endif
-			}
-
-			const int32 AdditionalPhysicsNum = MutableMesh->AdditionalPhysicsBodies.Num();
-			for (int32 I = 0; I < AdditionalPhysicsNum; ++I)
-			{
-				const int32 ExternalId = MutableMesh->AdditionalPhysicsBodies[I]->CustomId;
-				
-				ComponentData->PhysicsAssets.AdditionalPhysicsAssetsToLoad.Add(ExternalId);
-				AssetsToStream.Add(ModelResources.AnimBpOverridePhysiscAssetsInfo[ExternalId].SourceAsset.ToSoftObjectPath());
 			}
 		}
 	}
@@ -4981,30 +5068,57 @@ UE::Tasks::FTask UCustomizableInstancePrivate::LoadAdditionalAssetsAndData(
 	bool bHasInvalidMesh = false;
 	bool bUpdateMeshes = DoComponentsNeedUpdate(GetPublic(), OperationData, bHasInvalidMesh);
 
-	if (RealTimeMorphStreamableBlocksToStream.Num() && bUpdateMeshes)
+	bool bIsDataBlocksStreamNeeded = RealTimeMorphStreamableBlocksToStream.Num() || ClothingStreamableBlocksToStream.Num();
+
+	if (bIsDataBlocksStreamNeeded && bUpdateMeshes)
 	{
-		
 #if WITH_EDITOR
 		// On editor the data is always loaded, load directly form the ModelResources.
-		MUTABLE_CPUPROFILER_SCOPE(RealTimeMorphStreamingEditor);
-		for (uint32 BlockId : RealTimeMorphStreamableBlocksToStream)
-		{	
-			const FRealTimeMorphStreamable& RealTimeMorphStreamable = ModelResources.RealTimeMorphStreamables[BlockId]; 
-		
-			const FMutableStreamableBlock& Block = RealTimeMorphStreamable.Block;
+		{
+			MUTABLE_CPUPROFILER_SCOPE(RealTimeMorphStreamingEditor);
+			for (uint32 BlockId : RealTimeMorphStreamableBlocksToStream)
+			{	
+				const FRealTimeMorphStreamable& RealTimeMorphStreamable = ModelResources.RealTimeMorphStreamables[BlockId]; 
+			
+				const FMutableStreamableBlock& Block = RealTimeMorphStreamable.Block;
 
-			FInstanceUpdateData::FMorphTargetMeshData& MeshData = 
-					OperationData->InstanceUpdateData.RealTimeMorphTargetMeshData.FindOrAdd(BlockId);
+				FInstanceUpdateData::FMorphTargetMeshData& MeshData = 
+						OperationData->InstanceUpdateData.RealTimeMorphTargetMeshData.FindOrAdd(BlockId);
+				
+				MeshData.NameResolutionMap = RealTimeMorphStreamable.NameResolutionMap;
+
+				const TArray<FMorphTargetVertexData>& SourceData = ModelResources.EditorOnlyMorphTargetReconstructionData;	
+				const uint32 NumElems = Block.Size / sizeof(FMorphTargetVertexData);
+				const uint32 OffsetInElems = Block.Offset / sizeof(FMorphTargetVertexData);
+				MeshData.Data.SetNumUninitialized(NumElems);
+
+				check(SourceData.Num()*sizeof(FMorphTargetVertexData) >= Block.Offset + Block.Size);
+				FMemory::Memcpy(MeshData.Data.GetData(), SourceData.GetData() + OffsetInElems, Block.Size);
+			}
+		}
+		{
+			MUTABLE_CPUPROFILER_SCOPE(ClothingStreamingEditor);
+			for (uint32 BlockId : ClothingStreamableBlocksToStream)
+			{	
+				const FClothingStreamable& ClothingStreamable = ModelResources.ClothingStreamables[BlockId]; 
 			
-			MeshData.NameResolutionMap = RealTimeMorphStreamable.NameResolutionMap;
+				const FMutableStreamableBlock& Block = ClothingStreamable.Block;
 
-			const TArray<FMorphTargetVertexData>& SourceData = ModelResources.EditorOnlyMorphTargetReconstructionData;	
-			const uint32 NumElems = Block.Size / sizeof(FMorphTargetVertexData);
-			const uint32 OffsetInElems = Block.Offset / sizeof(FMorphTargetVertexData);
-			MeshData.Data.SetNumUninitialized(NumElems);
+				FInstanceUpdateData::FClothingMeshData& MeshData = 
+						OperationData->InstanceUpdateData.ClothingMeshData.FindOrAdd(BlockId);
+				
+				MeshData.ClothingAssetIndex = ClothingStreamable.ClothingAssetIndex;
+				MeshData.ClothingAssetLOD = ClothingStreamable.ClothingAssetLOD;
+				MeshData.PhysicsAssetIndex = ClothingStreamable.PhysicsAssetIndex;
 
-			check(SourceData.Num()*sizeof(FMorphTargetVertexData) >= Block.Offset + Block.Size);
-			FMemory::Memcpy(MeshData.Data.GetData(), SourceData.GetData() + OffsetInElems, Block.Size);
+				const TArray<FCustomizableObjectMeshToMeshVertData>& SourceData = ModelResources.EditorOnlyClothingMeshToMeshVertData;
+				const uint32 NumElems = Block.Size / sizeof(FCustomizableObjectMeshToMeshVertData);
+				const uint32 OffsetInElems = Block.Offset / sizeof(FCustomizableObjectMeshToMeshVertData);
+				MeshData.Data.SetNumUninitialized(NumElems);
+
+				check(SourceData.Num()*sizeof(FCustomizableObjectMeshToMeshVertData) >= Block.Offset + Block.Size);
+				FMemory::Memcpy(MeshData.Data.GetData(), SourceData.GetData() + OffsetInElems, Block.Size);
+			}
 		}
 #else	
 		MUTABLE_CPUPROFILER_SCOPE(RealTimeMorphStreaming);
@@ -5067,6 +5181,53 @@ UE::Tasks::FTask UCustomizableInstancePrivate::LoadAdditionalAssetsAndData(
 			});
 		}
 
+		const int32 NumClothBlocks = ClothingStreamableBlocksToStream.Num();
+		for (int32 I = 0; I < NumClothBlocks; ++I)
+		{
+			MUTABLE_CPUPROFILER_SCOPE(ClothingStreamingRequest_Alloc);
+
+			const int32 BlockId = ClothingStreamableBlocksToStream[I];
+			const FMutableStreamableBlock& Block = ModelResources.ClothingStreamables[BlockId].Block; 
+		
+			FInstanceUpdateData::FClothingMeshData& ReadDestData = 
+					OperationData->InstanceUpdateData.ClothingMeshData.FindOrAdd(BlockId);
+
+			// Only request blocks once.
+			if (ReadDestData.Data.Num())
+			{
+				continue;
+			}
+
+			const FClothingStreamable& ClothingStreamable = ModelResources.ClothingStreamables[BlockId]; 
+
+			ReadDestData.ClothingAssetIndex = ClothingStreamable.ClothingAssetIndex;
+			ReadDestData.ClothingAssetLOD = ClothingStreamable.ClothingAssetLOD;
+			ReadDestData.PhysicsAssetIndex = ClothingStreamable.PhysicsAssetIndex;
+
+			check(Block.Size % sizeof(FCustomizableObjectMeshToMeshVertData) == 0);
+			const uint32 NumElems = Block.Size / sizeof(FCustomizableObjectMeshToMeshVertData);
+
+			ReadDestData.Data.SetNumUninitialized(NumElems);
+
+			int32 FileHandleIndex = OpenFilesIds.Find(Block.FileId);
+			if (FileHandleIndex == INDEX_NONE && BulkData)
+			{
+				TUniquePtr<IAsyncReadFileHandle> ReadFileHandle = BulkData->OpenFileAsyncRead(Block.FileId);
+
+				OpenFileHandles.Emplace(MoveTemp(ReadFileHandle));
+				FileHandleIndex = OpenFilesIds.Add(Block.FileId);
+			}	
+
+			BlockReadInfos.Emplace(FBlockReadInfo
+			{ 
+				Block.Offset,
+				OpenFileHandles[FileHandleIndex].Get(), 
+				MakeArrayView(reinterpret_cast<uint8*>(ReadDestData.Data.GetData()), ReadDestData.Data.Num()*sizeof(FCustomizableObjectMeshToMeshVertData)),
+				Block.FileId
+			});
+		}
+
+
 		for (const FBlockReadInfo& BlockReadInfo : BlockReadInfos)
 		{
 			if (!BlockReadInfo.FileHandle)
diff --git a/Engine/Plugins/Experimental/Mutable/Source/CustomizableObjectEditor/Private/MuCOE/CustomizableObjectCompileRunnable.cpp b/Engine/Plugins/Experimental/Mutable/Source/CustomizableObjectEditor/Private/MuCOE/CustomizableObjectCompileRunnable.cpp
--- a/Engine/Plugins/Experimental/Mutable/Source/CustomizableObjectEditor/Private/MuCOE/CustomizableObjectCompileRunnable.cpp
+++ b/Engine/Plugins/Experimental/Mutable/Source/CustomizableObjectEditor/Private/MuCOE/CustomizableObjectCompileRunnable.cpp
@@ -278,15 +278,23 @@ FCustomizableObjectSaveDDRunnable::FCustomizableObjectSaveDDRunnable(UCustomizab
 #if WITH_EDITORONLY_DATA
 	else
 	{
-		// Do a copy of the MorphData generated at compile time. Only needed when cooking.
+		// Do a copy of the Morph and Clothing Data generated at compile time. Only needed when cooking.
 		
-		static_assert(TCanBulkSerialize<FMorphTargetVertexData>::Value);
 		constexpr bool bGetCookedFalse = false;
+		
+		static_assert(TCanBulkSerialize<FMorphTargetVertexData>::Value);
 		const TArray<FMorphTargetVertexData>& MorphVertexData = 
 				CustomizableObject->GetPrivate()->GetModelResources(bGetCookedFalse).EditorOnlyMorphTargetReconstructionData;
 
 		MorphDataBytes.SetNum(MorphVertexData.Num() * sizeof(FMorphTargetVertexData));
 		FMemory::Memcpy(MorphDataBytes.GetData(), MorphVertexData.GetData(), MorphDataBytes.Num());
+		
+		static_assert(TCanBulkSerialize<FCustomizableObjectMeshToMeshVertData>::Value);
+		const TArray<FCustomizableObjectMeshToMeshVertData>& ClothingVertexData = 
+				CustomizableObject->GetPrivate()->GetModelResources(bGetCookedFalse).EditorOnlyClothingMeshToMeshVertData;
+
+		ClothingDataBytes.SetNum(ClothingVertexData.Num() * sizeof(FCustomizableObjectMeshToMeshVertData));
+		FMemory::Memcpy(ClothingDataBytes.GetData(), ClothingVertexData.GetData(), ClothingDataBytes.Num());
 	}
 #endif // WITH_EDITORONLY_DATA
 }
@@ -296,8 +304,9 @@ uint32 FCustomizableObjectSaveDDRunnable::Run()
 {
 	MUTABLE_CPUPROFILER_SCOPE(FCustomizableObjectSaveDDRunnable::Run)
 
-	// MorphDataBytes has data only if cooking. 
+	// MorphDataBytes and ClothingDataBytes has data only if cooking. 
 	check(!!Options.bIsCooking || MorphDataBytes.IsEmpty());
+	check(!!Options.bIsCooking || ClothingDataBytes.IsEmpty());
 
 	bool bModelSerialized = Model.Get() != nullptr;
 
@@ -313,7 +322,7 @@ uint32 FCustomizableObjectSaveDDRunnable::Run()
 			FUnrealMutableModelBulkWriter Streamer(&ModelMemoryWriter, &StreamableMemoryWriter);
 			mu::Model::Serialise(Model.Get(), Streamer);
 
-			//MorphData is already in the corresponding buffer copied from the compilation thread.
+			//Morph and Clothing are already in the corresponding buffer copied from the compilation thread.
 		}
 	}
 	else if (bModelSerialized) // Save CO data + mu::Model and streamable resources to disk
diff --git a/Engine/Plugins/Experimental/Mutable/Source/CustomizableObjectEditor/Private/MuCOE/CustomizableObjectCompileRunnable.h b/Engine/Plugins/Experimental/Mutable/Source/CustomizableObjectEditor/Private/MuCOE/CustomizableObjectCompileRunnable.h
--- a/Engine/Plugins/Experimental/Mutable/Source/CustomizableObjectEditor/Private/MuCOE/CustomizableObjectCompileRunnable.h
+++ b/Engine/Plugins/Experimental/Mutable/Source/CustomizableObjectEditor/Private/MuCOE/CustomizableObjectCompileRunnable.h
@@ -127,4 +127,6 @@ public:
 	// Bytes store streameable files coming form the CO itself.
 	TArray64<uint8> MorphDataBytes;
 
+	// Bytes store streameable files coming form the CO itself.
+	TArray64<uint8> ClothingDataBytes;
 };
diff --git a/Engine/Plugins/Experimental/Mutable/Source/CustomizableObjectEditor/Private/MuCOE/CustomizableObjectCompiler.cpp b/Engine/Plugins/Experimental/Mutable/Source/CustomizableObjectEditor/Private/MuCOE/CustomizableObjectCompiler.cpp
--- a/Engine/Plugins/Experimental/Mutable/Source/CustomizableObjectEditor/Private/MuCOE/CustomizableObjectCompiler.cpp
+++ b/Engine/Plugins/Experimental/Mutable/Source/CustomizableObjectEditor/Private/MuCOE/CustomizableObjectCompiler.cpp
@@ -1057,11 +1057,37 @@ void FCustomizableObjectCompiler::CompileInternal(bool bAsync)
 			RealTimeMorphDataOffsetInBytes += DataSizeInBytes;
 			ModelResources.EditorOnlyMorphTargetReconstructionData.Append(MeshData.Value.Data);
 		}
+	
+		// Create the Clothing Blocks from the per mesh Morph data.
+		uint64 ClothingDataNum = 0;
+		for (const TPair<uint32, FClothingMeshData>& MeshData : GenerationContext.ClothingPerMeshData)
+		{
+			ClothingDataNum += MeshData.Value.Data.Num();
+		}
 		
-		// Clothing	
-		CurrentObject->ClothMeshToMeshVertData = MoveTemp(GenerationContext.ClothMeshToMeshVertData);
-		CurrentObject->ContributingClothingAssetsData = MoveTemp(GenerationContext.ContributingClothingAssetsData);
-		CurrentObject->ClothSharedConfigsData.Empty();
+		ModelResources.ClothingStreamables.Empty(32);
+		ModelResources.EditorOnlyClothingMeshToMeshVertData.Empty(ClothingDataNum);
+
+		uint64 ClothingDataOffsetInBytes = 0;
+		for (const TPair<uint32, FClothingMeshData>& MeshData : GenerationContext.ClothingPerMeshData)
+		{
+			const uint32 DataSizeInBytes = (uint32)MeshData.Value.Data.Num()*sizeof(FCustomizableObjectMeshToMeshVertData); 
+			FClothingStreamable& ResourceMeshData = ModelResources.ClothingStreamables.FindOrAdd(MeshData.Key);
+			
+			check(ResourceMeshData.ClothingAssetIndex == INDEX_NONE);
+			check(ResourceMeshData.ClothingAssetLOD == INDEX_NONE);
+			check(ResourceMeshData.Block.Size == 0);
+
+			ResourceMeshData.ClothingAssetIndex = MeshData.Value.ClothingAssetIndex;
+			ResourceMeshData.ClothingAssetLOD = MeshData.Value.ClothingAssetLOD;
+			ResourceMeshData.PhysicsAssetIndex = MeshData.Value.PhysicsAssetIndex;
+			ResourceMeshData.Block = FMutableStreamableBlock { uint32(0), DataSizeInBytes, ClothingDataOffsetInBytes };
+
+			ClothingDataOffsetInBytes += DataSizeInBytes;
+			ModelResources.EditorOnlyClothingMeshToMeshVertData.Append(MeshData.Value.Data);
+		}
+
+		ModelResources.ClothingAssetsData = MoveTemp(GenerationContext.ClothingAssetsData);
 
 		// A clothing backend, e.g. Chaos cloth, can use 2 config files, one owned by the asset, and another that is shared 
 		// among all assets in a SkeletalMesh. When merging different assets in a skeletalmesh we need to make sure only one of 
@@ -1075,13 +1101,13 @@ void FCustomizableObjectCompiler::CompileInternal(bool bAsync)
 		};
 		
 		// Find shared configs to be used (One of each type) 
-		for (FCustomizableObjectClothingAssetData& ClothingAssetData : CurrentObject->ContributingClothingAssetsData)
+		for (FCustomizableObjectClothingAssetData& ClothingAssetData : ModelResources.ClothingAssetsData)
 		{
 			 for (FCustomizableObjectClothConfigData& ClothConfigData : ClothingAssetData.ConfigsData)
 			 {
 				  if (IsSharedConfigData(ClothConfigData))
 				  {
-					  FCustomizableObjectClothConfigData* FoundConfig = CurrentObject->ClothSharedConfigsData.FindByPredicate(
+					  FCustomizableObjectClothConfigData* FoundConfig = ModelResources.ClothSharedConfigsData.FindByPredicate(
 						   [Name = ClothConfigData.ConfigName](const FCustomizableObjectClothConfigData& Other)
 						   {
 							   return Name == Other.ConfigName;
@@ -1089,14 +1115,14 @@ void FCustomizableObjectCompiler::CompileInternal(bool bAsync)
 
 					  if (!FoundConfig)
 					  {
-						   CurrentObject->ClothSharedConfigsData.AddDefaulted_GetRef() = ClothConfigData;
+						   ModelResources.ClothSharedConfigsData.AddDefaulted_GetRef() = ClothConfigData;
 					  }
 				  }
 			 }
 		}
 		
 		// Remove shared configs
-		for (FCustomizableObjectClothingAssetData& ClothingAssetData : CurrentObject->ContributingClothingAssetsData)
+		for (FCustomizableObjectClothingAssetData& ClothingAssetData : ModelResources.ClothingAssetsData)
 		{
 			 ClothingAssetData.ConfigsData.RemoveAllSwap(IsSharedConfigData);
 		}
@@ -1269,7 +1295,14 @@ mu::NodePtr FCustomizableObjectCompiler::Export(UCustomizableObject* Object, con
 
 	FCompilationOptions CompilerOptions = InCompilerOptions;
 	CompilerOptions.CustomizableObjectNumBoneInfluences = CustomizableObjectNumBoneInfluences;
-		
+
+	CompilerOptions.bRealTimeMorphTargetsEnabled = Object->bEnableRealTimeMorphTargets;
+	CompilerOptions.bClothingEnabled = Object->bEnableClothing;
+	CompilerOptions.b16BitBoneWeightsEnabled = Object->bEnable16BitBoneWeights;
+	CompilerOptions.bSkinWeightProfilesEnabled = Object->bEnableAltSkinWeightProfiles;
+	CompilerOptions.bPhysicsAssetMergeEnabled = Object->bEnablePhysicsAssetMerge;
+	CompilerOptions.bAnimBpPhysicsManipulationEnabled = Object->bEnableAnimBpPhysicsAssetsManipualtion;
+
 	FMutableGraphGenerationContext GenerationContext(Object, this, CompilerOptions);
 	GenerationContext.ParamNamesToSelectedOptions = ParamNamesToSelectedOptions;
 
@@ -1402,7 +1435,8 @@ void FCustomizableObjectCompiler::FinishSavingDerivedDataTask()
 				SaveDDTask->GetTargetPlatform(), 
 				SaveDDTask->Bytes, 
 				SaveDDTask->BulkDataBytes,
-				SaveDDTask->MorphDataBytes);
+				SaveDDTask->MorphDataBytes,
+				SaveDDTask->ClothingDataBytes);
 	}
 	
 
diff --git a/Engine/Plugins/Experimental/Mutable/Source/CustomizableObjectEditor/Private/MuCOE/GenerateMutableSource/GenerateMutableSource.h b/Engine/Plugins/Experimental/Mutable/Source/CustomizableObjectEditor/Private/MuCOE/GenerateMutableSource/GenerateMutableSource.h
--- a/Engine/Plugins/Experimental/Mutable/Source/CustomizableObjectEditor/Private/MuCOE/GenerateMutableSource/GenerateMutableSource.h
+++ b/Engine/Plugins/Experimental/Mutable/Source/CustomizableObjectEditor/Private/MuCOE/GenerateMutableSource/GenerateMutableSource.h
@@ -195,6 +195,8 @@ struct FMutableGraphMeshGenerationData
 {
 	// Did we find any mesh with vertex colours in the expression?
 	bool bHasVertexColors = false;
+	bool bHasRealTimeMorphs = false;
+	bool bHasClothing = false;
 
 	// Maximum number of texture channels found in the expression.
 	int NumTexCoordChannels = 0;
@@ -211,16 +213,18 @@ struct FMutableGraphMeshGenerationData
 	TArray<int32> SkinWeightProfilesSemanticIndices;
 
 	// Combine another generated data looking for the most general case.
-	void Combine(const FMutableGraphMeshGenerationData& other)
+	void Combine(const FMutableGraphMeshGenerationData& Other)
 	{
-		bHasVertexColors = bHasVertexColors || other.bHasVertexColors;
-		NumTexCoordChannels = FMath::Max(other.NumTexCoordChannels, NumTexCoordChannels);
-		MaxNumBonesPerVertex = FMath::Max(other.MaxNumBonesPerVertex, MaxNumBonesPerVertex);
-		MaxBoneIndexTypeSizeBytes = FMath::Max(other.MaxBoneIndexTypeSizeBytes, MaxBoneIndexTypeSizeBytes);
-		MaxNumTriangles = FMath::Max(other.MaxNumTriangles, MaxNumTriangles);
-		MinNumTriangles = FMath::Min(other.MinNumTriangles, MinNumTriangles);
-
-		for (int32 SemanticIndex : other.SkinWeightProfilesSemanticIndices)
+		bHasVertexColors = bHasVertexColors || Other.bHasVertexColors;
+		bHasRealTimeMorphs = bHasRealTimeMorphs || Other.bHasRealTimeMorphs;
+		bHasClothing = bHasClothing || Other.bHasClothing;
+		NumTexCoordChannels = FMath::Max(Other.NumTexCoordChannels, NumTexCoordChannels);
+		MaxNumBonesPerVertex = FMath::Max(Other.MaxNumBonesPerVertex, MaxNumBonesPerVertex);
+		MaxBoneIndexTypeSizeBytes = FMath::Max(Other.MaxBoneIndexTypeSizeBytes, MaxBoneIndexTypeSizeBytes);
+		MaxNumTriangles = FMath::Max(Other.MaxNumTriangles, MaxNumTriangles);
+		MinNumTriangles = FMath::Min(Other.MinNumTriangles, MinNumTriangles);
+
+		for (int32 SemanticIndex : Other.SkinWeightProfilesSemanticIndices)
 		{
 			SkinWeightProfilesSemanticIndices.AddUnique(SemanticIndex);
 		}
@@ -277,6 +281,14 @@ struct FRealTimeMorphMeshData
 	TArray<FMorphTargetVertexData> Data;
 };
 
+struct FClothingMeshData
+{
+	int32 ClothingAssetIndex = INDEX_NONE;
+	int32 ClothingAssetLOD = INDEX_NONE;
+	int32 PhysicsAssetIndex = INDEX_NONE;
+	TArray<FCustomizableObjectMeshToMeshVertData> Data;
+};
+
 struct FGroupProjectorTempData
 {
 	class UCustomizableObjectNodeGroupProjectorParameter* CustomizableObjectNodeGroupProjectorParameter;
@@ -718,8 +730,8 @@ struct FMutableGraphGenerationContext
 	TMap<uint32, FRealTimeMorphMeshData> RealTimeMorphTargetPerMeshData;
 
 	// Data used for Clothing reconstruction.
-	TArray<FCustomizableObjectMeshToMeshVertData> ClothMeshToMeshVertData;
-	TArray<FCustomizableObjectClothingAssetData> ContributingClothingAssetsData;
+	TArray<FCustomizableObjectClothingAssetData> ClothingAssetsData;
+	TMap<uint32, FClothingMeshData> ClothingPerMeshData;
 
 	// Data used for SkinWeightProfiles reconstruction
 	TArray<FMutableSkinWeightProfileInfo> SkinWeightProfilesInfo;
diff --git a/Engine/Plugins/Experimental/Mutable/Source/CustomizableObjectEditor/Private/MuCOE/GenerateMutableSource/GenerateMutableSourceMesh.cpp b/Engine/Plugins/Experimental/Mutable/Source/CustomizableObjectEditor/Private/MuCOE/GenerateMutableSource/GenerateMutableSourceMesh.cpp
--- a/Engine/Plugins/Experimental/Mutable/Source/CustomizableObjectEditor/Private/MuCOE/GenerateMutableSource/GenerateMutableSourceMesh.cpp
+++ b/Engine/Plugins/Experimental/Mutable/Source/CustomizableObjectEditor/Private/MuCOE/GenerateMutableSource/GenerateMutableSourceMesh.cpp
@@ -1170,6 +1170,8 @@ mu::MeshPtr ConvertSkeletalMeshToMutable(const USkeletalMesh* InSkeletalMesh, co
 
 		if (UsedMorphTargets.Num())
 		{
+			MUTABLE_CPUPROFILER_SCOPE(RealTimeMorphTargetProcessing);
+			
 			MutableMesh->GetVertexBuffers().SetBufferCount(NextBufferIndex + 2);
 			// MorphTarget vertex block offset.
 			{
@@ -1217,8 +1219,7 @@ mu::MeshPtr ConvertSkeletalMeshToMutable(const USkeletalMesh* InSkeletalMesh, co
 			{
 				Elem = InvalidResourceId;
 			}
-
-			MUTABLE_CPUPROFILER_SCOPE(RealTimeMorphTargetProcessing);
+			
 
 			TArray<FMorphTargetVertexData> MorphsMeshData;
 			MorphsMeshData.Reserve(32);
@@ -1383,53 +1384,43 @@ mu::MeshPtr ConvertSkeletalMeshToMutable(const USkeletalMesh* InSkeletalMesh, co
 	// Clothing vertex info.
 	if (GenerationContext.Options.bClothingEnabled)
 	{
-		{
-			using namespace mu;
-			const int32 ElementSize = sizeof(int32);
-			const int32 ChannelCount = 1;
-			const EMeshBufferSemantic Semantics[ChannelCount] = { MBS_OTHER };
-			const int32 SemanticIndices[ChannelCount] = { 2 };
-			const EMeshBufferFormat Formats[ChannelCount] = { MBF_INT32 };
-			int32 Components[ChannelCount] = { 1 };
-			const int32 Offsets[ChannelCount] = { 0 };
-
-			MutableMesh->GetVertexBuffers().SetBufferCount(NextBufferIndex + 1);
-			MutableMesh->GetVertexBuffers().SetBuffer(NextBufferIndex, ElementSize, ChannelCount, Semantics, SemanticIndices, Formats, Components, Offsets);
-		}
-
-		TArrayView<int32> ClothSectionBufferView(reinterpret_cast<int32*>(MutableMesh->GetVertexBuffers().GetBufferData(NextBufferIndex)), VertexCount);
-		for (int32& Elem : ClothSectionBufferView)
-		{
-			Elem = -1;
-		}
-
 		// Create new asset or find an already created one if the section has clothing assets.
 		// clothing assets are shared among all LODs in a section
-		const int32 ClothingAssetIndex = Invoke([&]() -> int32
+		int32 ClothingAssetIndex = INDEX_NONE;
+		int32 PhysicsAssetIndex = INDEX_NONE;
+
+		Invoke([&](int32& OutClothingAssetIndex, int32& OutPhysicsAssetIndex) -> void
 		{
 			const UClothingAssetBase* ClothingAssetBase = InSkeletalMesh->GetSectionClothingAsset(LODIndex, SectionIndex);
 
 			if (!ClothingAssetBase)
 			{
-				return INDEX_NONE;
+				return;
+			}
+
+			const UClothingAssetCommon* Asset = Cast<UClothingAssetCommon>(ClothingAssetBase);
+			if (!Asset)
+			{
+				return;
 			}
 
-			int32 FoundIndex = GenerationContext.ContributingClothingAssetsData.IndexOfByPredicate(
+			int32 FoundIndex = GenerationContext.ClothingAssetsData.IndexOfByPredicate(
 				[AssetGuid = ClothingAssetBase->GetAssetGuid()](const FCustomizableObjectClothingAssetData& Asset){ return Asset.OriginalAssetGuid == AssetGuid; });
 			
 			if (FoundIndex != INDEX_NONE)
 			{
-				return FoundIndex;
-			}
+				OutClothingAssetIndex = FoundIndex;
+				OutPhysicsAssetIndex = GenerationContext.PhysicsAssets.IndexOfByPredicate(
+					[PhysicsAsset = Asset->PhysicsAsset](const TSoftObjectPtr<UPhysicsAsset>& OtherPhysicsAsset)
+					{
+						return TSoftObjectPtr<UPhysicsAsset>(PhysicsAsset) == OtherPhysicsAsset;
+					});
 
-			const UClothingAssetCommon* Asset = Cast<UClothingAssetCommon>(ClothingAssetBase);
-			if (!Asset)
-			{
-				return INDEX_NONE;
+				return;
 			}
 
-			const int32 NewAssetIndex = GenerationContext.ContributingClothingAssetsData.AddDefaulted();
-			FCustomizableObjectClothingAssetData& AssetData = GenerationContext.ContributingClothingAssetsData[NewAssetIndex];
+			OutClothingAssetIndex = GenerationContext.ClothingAssetsData.AddDefaulted();
+			FCustomizableObjectClothingAssetData& AssetData = GenerationContext.ClothingAssetsData[OutClothingAssetIndex];
 			
 			AssetData.LodData = Asset->LodData;
 			AssetData.LodMap = Asset->LodMap;
@@ -1439,8 +1430,11 @@ mu::MeshPtr ConvertSkeletalMeshToMutable(const USkeletalMesh* InSkeletalMesh, co
 			AssetData.OriginalAssetGuid = Asset->GetAssetGuid();
 			AssetData.Name = Asset->GetFName();
 
+			OutPhysicsAssetIndex = GenerationContext.PhysicsAssets.AddUnique(Asset->PhysicsAsset);
+			GenerationContext.AddParticipatingObject(*Asset->PhysicsAsset);
+			
 			// Store raw clothing config serialized raw data, and info to recreate it afterwards.
-			for ( const TPair<FName, TObjectPtr<UClothConfigBase>>& ClothConfig : Asset->ClothConfigs )
+			for (const TPair<FName, TObjectPtr<UClothConfigBase>>& ClothConfig : Asset->ClothConfigs)
 			{
 				FCustomizableObjectClothConfigData& ConfigData = AssetData.ConfigsData.AddDefaulted_GetRef();
 				ConfigData.ClassPath = ClothConfig.Value->GetClass()->GetPathName();
@@ -1449,21 +1443,53 @@ mu::MeshPtr ConvertSkeletalMeshToMutable(const USkeletalMesh* InSkeletalMesh, co
 				FMemoryWriter MemoryWriter(ConfigData.ConfigBytes);
                 ClothConfig.Value->Serialize(MemoryWriter);
 			}
-
-			return NewAssetIndex;
-		});
+		}, ClothingAssetIndex, PhysicsAssetIndex);
 
 		if (ClothingAssetIndex != INDEX_NONE)
 		{
-			// Reserve first element as a way to indicate invalid data. Currently not used.
-			if (GenerationContext.ClothMeshToMeshVertData.Num() == 0)
+			MutableMesh->GetVertexBuffers().SetBufferCount(NextBufferIndex + 2);
 			{
-				FCustomizableObjectMeshToMeshVertData& FirstElem = GenerationContext.ClothMeshToMeshVertData.AddZeroed_GetRef();
-				FirstElem.SourceAssetIndex = INDEX_NONE;
+				using namespace mu;
+				const int32 ElementSize = sizeof(int32);
+				const int32 ChannelCount = 1;
+				const EMeshBufferSemantic Semantics[ChannelCount] = { MBS_OTHER };
+				const int32 SemanticIndices[ChannelCount] = { 2 };
+				const EMeshBufferFormat Formats[ChannelCount] = { MBF_INT32 };
+				int32 Components[ChannelCount] = { 1 };
+				const int32 Offsets[ChannelCount] = { 0 };
+
+				MutableMesh->GetVertexBuffers().SetBuffer(NextBufferIndex, ElementSize, ChannelCount, Semantics, SemanticIndices, Formats, Components, Offsets);
 			}
 
-			const TArray<FMeshToMeshVertData>& ClothMappingData = MeshSection.ClothMappingDataLODs[0];
+			{
+				using namespace mu;
+				const int32 ElementSize = sizeof(uint32);
+				const int32 ChannelCount = 1;
+				const EMeshBufferSemantic Semantics[ChannelCount] = { MBS_OTHER };
+				const int32 SemanticIndices[ChannelCount] = { 3 };
+				const EMeshBufferFormat Formats[ChannelCount] = { MBF_UINT32 };
+				int32 Components[ChannelCount] = { 1 };
+				const int32 Offsets[ChannelCount] = { 0 };
 
+				MutableMesh->GetVertexBuffers().SetBuffer(NextBufferIndex + 1, ElementSize, ChannelCount, Semantics, SemanticIndices, Formats, Components, Offsets);
+			}
+
+			TArrayView<int32> ClothSectionIndexView(reinterpret_cast<int32*>(
+						MutableMesh->GetVertexBuffers().GetBufferData(NextBufferIndex)), VertexCount);
+			for (int32& Elem : ClothSectionIndexView)
+			{
+				Elem = -1;
+			}
+
+			constexpr uint32 InvalidResourceId = 0;
+			TArrayView<uint32> ClothSectionResourceIdView(reinterpret_cast<uint32*>(
+						MutableMesh->GetVertexBuffers().GetBufferData(NextBufferIndex + 1)), VertexCount);
+			for (uint32& Elem : ClothSectionResourceIdView)
+			{
+				Elem = InvalidResourceId;
+			}
+
+			const TArray<FMeshToMeshVertData>& ClothMappingData = MeshSection.ClothMappingDataLODs[0];
 
 			// Similar test as the one used on FSkeletalMeshObjectGPUSkin::FVertexFactoryData::InitAPEXClothVertexFactories
 			// Here should work as expexted, but in the reference code I'm not sure it always works. It is worth investigate
@@ -1476,21 +1502,29 @@ mu::MeshPtr ConvertSkeletalMeshToMutable(const USkeletalMesh* InSkeletalMesh, co
 			// TODO: find a better place to keep this constant.
 			constexpr int32 NumInfluencesPerVertex = 5;
 
-			int32 MeshToMeshDataIndex = GenerationContext.ClothMeshToMeshVertData.Num();
+			int32 MeshToMeshDataIndex = 0; 
 
 			constexpr int32 MaxSupportedInfluences = 1;
-			for (int32& Elem : ClothSectionBufferView)
+			for (int32& Elem : ClothSectionIndexView)
 			{
 				Elem = MeshToMeshDataIndex;
 				MeshToMeshDataIndex += MaxSupportedInfluences;
 			}
 
-			const int32 ClothDataIndexBase = GenerationContext.ClothMeshToMeshVertData.Num();
+			const int32 ClothDataIndexBase = 0;
 
 			const int32 ClothDataStride = bUseMutlipleInfluences ? NumInfluencesPerVertex : 1;
 			const int32 NumClothMappingDataVerts = ClothMappingData.Num() / ClothDataStride;
 
-			GenerationContext.ClothMeshToMeshVertData.Reserve(NumClothMappingDataVerts);
+			FClothingMeshData ClothingMeshData;
+			
+			const FCustomizableObjectClothingAssetData& ClothingAssetData = GenerationContext.ClothingAssetsData[ClothingAssetIndex];
+			
+			ClothingMeshData.ClothingAssetIndex = ClothingAssetIndex;
+			ClothingMeshData.ClothingAssetLOD = ClothingAssetData.LodMap[LODIndex];
+			ClothingMeshData.PhysicsAssetIndex = PhysicsAssetIndex;
+			ClothingMeshData.Data.Reserve(NumClothMappingDataVerts * ClothDataStride);
+
 			for (int32 Idx = 0; Idx < NumClothMappingDataVerts * ClothDataStride; Idx += ClothDataStride)
 			{
 				// If bUseMutlipleInfluences we will only take the element with higher weight ignoring the other ones.
@@ -1498,28 +1532,102 @@ mu::MeshPtr ConvertSkeletalMeshToMutable(const USkeletalMesh* InSkeletalMesh, co
 				const FMeshToMeshVertData* MaxInfluence = MaxElement(Influences.begin(), Influences.end(),
 					[](const FMeshToMeshVertData& A, const FMeshToMeshVertData& B) { return A.Weight < B.Weight; });
 
-				GenerationContext.ClothMeshToMeshVertData.Emplace(*MaxInfluence);
+				ClothingMeshData.Data.Emplace(*MaxInfluence);
 			}
 
-			TArrayView<FCustomizableObjectMeshToMeshVertData> AppendedClothingDataView
-			(GenerationContext.ClothMeshToMeshVertData.GetData() + ClothDataIndexBase, NumClothMappingDataVerts);
-
-			const FCustomizableObjectClothingAssetData& ClothingAssetData = GenerationContext.ContributingClothingAssetsData[ClothingAssetIndex];
-			const int16 ClothingAssetLODIndex = static_cast<int16>(ClothingAssetData.LodMap[LODIndex]);
-
-			for (FCustomizableObjectMeshToMeshVertData& ClothingDataElem : AppendedClothingDataView)
+			for (FCustomizableObjectMeshToMeshVertData& ClothingDataElem : ClothingMeshData.Data)
 			{
-				ClothingDataElem.SourceAssetIndex = static_cast<int16>(ClothingAssetIndex);
-				ClothingDataElem.SourceAssetLodIndex = ClothingAssetLODIndex;
-
 				// Currently if the cloth mapping uses multiple influences, these are ignored and only 
 				// the one with the highest weight is used. We set the weight to 1.0, but
 				// this value will be ignored anyway.
 				ClothingDataElem.Weight = 1.0f;
 			}
+
+			const uint32 DataHash = CityHash32(
+					reinterpret_cast<const char*>(ClothingMeshData.Data.GetData()), 
+					ClothingMeshData.Data.Num()*sizeof(FCustomizableObjectMeshToMeshVertData));
+
+			uint32 UniqueHash = DataHash == InvalidResourceId ? DataHash + 1 : DataHash;
+
+			const FClothingMeshData* FoundHash = GenerationContext.ClothingPerMeshData.Find(UniqueHash);
+
+			bool bIsDataAlreadyCollected = false;
+			
+			auto CompareClothingData = [](const FClothingMeshData& A, const FClothingMeshData& B) -> bool
+			{
+				return A.ClothingAssetIndex == B.ClothingAssetIndex &&
+					   A.ClothingAssetLOD == B.ClothingAssetLOD &&
+					   A.PhysicsAssetIndex == B.PhysicsAssetIndex &&
+					   A.Data.Num() == B.Data.Num() &&
+					   FMemory::Memcmp(A.Data.GetData(), B.Data.GetData(), A.Data.Num()*sizeof(FCustomizableObjectMeshToMeshVertData)) == 0;
+			};
+
+			if (FoundHash)
+			{
+				bIsDataAlreadyCollected = CompareClothingData(*FoundHash, ClothingMeshData); 
+			}
+
+			// NOTE: This way of unique hash generation guarantees all valid values can be used but given its 
+			// sequential nature a cascade of changes can occur if new meshes are added. Not many hash collisions 
+			// are expected so it should not be problematic.
+			if (FoundHash && !bIsDataAlreadyCollected)
+			{
+				uint32 NumTries = 0;
+				for (; NumTries < TNumericLimits<uint32>::Max(); ++NumTries)
+				{
+					FoundHash = GenerationContext.ClothingPerMeshData.Find(UniqueHash);
+					
+					if (!FoundHash)
+					{
+						break;
+					}
+
+					bIsDataAlreadyCollected = CompareClothingData(*FoundHash, ClothingMeshData);
+
+					if (bIsDataAlreadyCollected)
+					{
+						break;
+					}
+
+					UniqueHash = UniqueHash + 1 == InvalidResourceId ? InvalidResourceId + 1 : UniqueHash + 1;
+				}
+
+				if (NumTries == TNumericLimits<uint32>::Max())
+				{
+					UE_LOG(LogMutable, Warning, TEXT("Maximum number of meshes with clothing reached, some cloth meshes may not work as expected."));
+					UniqueHash = InvalidResourceId;
+				}	
+			}
+			
+			// fill the resource buffer with the generated unique id.
+			for (uint32& Elem : ClothSectionResourceIdView)
+			{
+				Elem = UniqueHash;
+			}
+
+			FCustomizableObjectStreameableResourceId StreamedClothResource;
+			StreamedClothResource.Id = UniqueHash;
+			StreamedClothResource.Type = static_cast<uint8>(FCustomizableObjectStreameableResourceId::EType::Clothing);
+
+			MutableMesh->AddStreamedResource(BitCast<uint64>(StreamedClothResource));
+
+			if (!bIsDataAlreadyCollected)
+			{
+				FClothingMeshData& NewClothingMeshData = GenerationContext.ClothingPerMeshData.FindOrAdd(UniqueHash);
+
+				check(NewClothingMeshData.ClothingAssetIndex == INDEX_NONE);
+				check(NewClothingMeshData.Data.IsEmpty());
+
+				NewClothingMeshData.ClothingAssetIndex = ClothingMeshData.ClothingAssetIndex;
+				NewClothingMeshData.ClothingAssetLOD = ClothingMeshData.ClothingAssetLOD;
+				NewClothingMeshData.PhysicsAssetIndex = ClothingMeshData.PhysicsAssetIndex;
+
+				NewClothingMeshData.Data = MoveTemp(ClothingMeshData.Data);
+			}
+			
+			NextBufferIndex += 2;
 		}
 
-		NextBufferIndex += 1;
 	}
 
 
@@ -3092,29 +3200,6 @@ mu::NodeMeshPtr GenerateMutableSourceMesh(const UEdGraphPin* Pin,
 					AddTagToMutableMeshUnique(*MutableMesh, PhysicsAssetTag);
 				}
 
-				if (GenerationContext.Options.bClothingEnabled)
-				{
-					UClothingAssetBase* ClothingAssetBase = TypedNodeSkel->SkeletalMesh->GetSectionClothingAsset(LODIndex, SectionIndex);	
-					UClothingAssetCommon* ClothingAssetCommon = Cast<UClothingAssetCommon>(ClothingAssetBase);
-
-					if (ClothingAssetCommon && ClothingAssetCommon->PhysicsAsset)
-					{	
-						int32 AssetIndex = GenerationContext.ContributingClothingAssetsData.IndexOfByPredicate( 
-						[Guid = ClothingAssetBase->GetAssetGuid()](const FCustomizableObjectClothingAssetData& A)
-						{
-							return A.OriginalAssetGuid == Guid;
-						});
-
-						check(AssetIndex != INDEX_NONE);
-
-						GenerationContext.AddParticipatingObject(*ClothingAssetCommon->PhysicsAsset);
-						
-						const int32 PhysicsAssetIndex = GenerationContext.PhysicsAssets.AddUnique(ClothingAssetCommon->PhysicsAsset);
-						FString ClothPhysicsAssetTag = FString::Printf(TEXT("__ClothPA:%d_%d"), AssetIndex, PhysicsAssetIndex);
-						AddTagToMutableMeshUnique(*MutableMesh, ClothPhysicsAssetTag);
-					}
-				}
-
 				if (GenerationContext.Options.bSkinWeightProfilesEnabled && LODIndex>=0)
 				{
 					if (ImportedModel && ImportedModel->LODModels.IsValidIndex(LODIndex))
@@ -3155,15 +3240,37 @@ mu::NodeMeshPtr GenerateMutableSourceMesh(const UEdGraphPin* Pin,
 				if (ImportedModel->LODModels.IsValidIndex(LODIndex) &&
 					ImportedModel->LODModels[LODIndex].Sections.IsValidIndex(SectionIndex))
 				{
-					MeshData.bHasVertexColors = TypedNodeSkel->SkeletalMesh->GetHasVertexColors();
+					MeshData.bHasVertexColors = TypedNodeSkel->SkeletalMesh->GetHasVertexColors();	
 					MeshData.NumTexCoordChannels = ImportedModel->LODModels[LODIndex].NumTexCoords;
 					MeshData.MaxBoneIndexTypeSizeBytes = MutableMesh->GetBoneMap().Num() > 256 ? 2 : 1;
 					MeshData.MaxNumBonesPerVertex = ImportedModel->LODModels[LODIndex].GetMaxBoneInfluences();
 					
+
 					// When mesh data is combined we will get an upper and lower bound of the number of triangles.
 					MeshData.MaxNumTriangles = ImportedModel->LODModels[LODIndex].Sections[SectionIndex].NumTriangles;
 					MeshData.MinNumTriangles = ImportedModel->LODModels[LODIndex].Sections[SectionIndex].NumTriangles;
 				}
+
+				// Find if the mesh has realtime morphs and clothing looking at the generated mesh buffers.
+				{
+					const mu::FMeshBufferSet& MeshSet = MutableMesh->GetVertexBuffers();
+
+					int32 MorphIndexBuffer, MorphIndexChannel;
+					MeshSet.FindChannel(mu::MBS_OTHER, 0, &MorphIndexBuffer, &MorphIndexChannel);
+
+					int32 MorphResourceBuffer, MorphResourceChannel;
+					MeshSet.FindChannel(mu::MBS_OTHER, 1, &MorphResourceBuffer, &MorphResourceChannel);
+
+					MeshData.bHasRealTimeMorphs = MorphIndexBuffer >= 0 && MorphResourceBuffer >= 0; 
+
+					int32 ClothIndexBuffer, ClothIndexChannel;
+					MeshSet.FindChannel(mu::MBS_OTHER, 2, &ClothIndexBuffer, &ClothIndexChannel);
+
+					int32 ClothResourceBuffer, ClothResourceChannel;
+					MeshSet.FindChannel(mu::MBS_OTHER, 3, &ClothResourceBuffer, &ClothResourceChannel);
+
+					MeshData.bHasClothing = ClothIndexBuffer >= 0 && ClothResourceBuffer >= 0; 
+				}
 			}
 
 			// Layouts
@@ -3869,6 +3976,11 @@ mu::NodeMeshPtr GenerateMutableSourceMesh(const UEdGraphPin* Pin,
 								// When mesh data is combined we will get an upper and lower bound of the number of triangles.
 								MeshData.MaxNumTriangles = ImportedModel->LODModels[LODIndex].Sections[SectionIndex].NumTriangles;
 								MeshData.MinNumTriangles = ImportedModel->LODModels[LODIndex].Sections[SectionIndex].NumTriangles;
+
+								// With tables we have to use the default skeletal mesh to tell if the surface mesh needs morphs or clothing.
+								// For now tables will not have morphs as there isn't a way to select them.
+								MeshData.bHasRealTimeMorphs = false;
+								MeshData.bHasClothing = GenerationContext.Options.bClothingEnabled && ImportedModel->LODModels[LODIndex].HasClothData();
 							}
 						}
 
diff --git a/Engine/Plugins/Experimental/Mutable/Source/CustomizableObjectEditor/Private/MuCOE/GenerateMutableSource/GenerateMutableSourceSurface.cpp b/Engine/Plugins/Experimental/Mutable/Source/CustomizableObjectEditor/Private/MuCOE/GenerateMutableSource/GenerateMutableSourceSurface.cpp
--- a/Engine/Plugins/Experimental/Mutable/Source/CustomizableObjectEditor/Private/MuCOE/GenerateMutableSource/GenerateMutableSourceSurface.cpp
+++ b/Engine/Plugins/Experimental/Mutable/Source/CustomizableObjectEditor/Private/MuCOE/GenerateMutableSource/GenerateMutableSourceSurface.cpp
@@ -60,8 +60,7 @@
 
 void SetSurfaceFormat( FMutableGraphGenerationContext& GenerationContext,
 					   mu::FMeshBufferSet& OutVertexBufferFormat, mu::FMeshBufferSet& OutIndexBufferFormat, const FMutableGraphMeshGenerationData& MeshData, 
-					   ECustomizableObjectNumBoneInfluences ECustomizableObjectNumBoneInfluences, bool bWithRealTimeMorphs, 
-					   bool bWithClothing, bool bWith16BitWeights )
+					   ECustomizableObjectNumBoneInfluences ECustomizableObjectNumBoneInfluences, bool bWith16BitWeights)
 {
 	// Limit skinning weights if necessary
 	// \todo: make it more flexible to support 3 or 5 or 1 weight, since there is support for this in 4.25
@@ -88,14 +87,14 @@ void SetSurfaceFormat( FMutableGraphGenerationContext& GenerationContext,
 		++MutableBufferCount;
 	}
 
-	if (bWithRealTimeMorphs)
+	if (MeshData.bHasRealTimeMorphs)
 	{
-		MutableBufferCount += 3;
+		MutableBufferCount += 2;
 	}
 
-	if (bWithClothing)
+	if (MeshData.bHasClothing)
 	{
-		++MutableBufferCount;
+		MutableBufferCount += 2;
 	}
 
 	MutableBufferCount += MeshData.SkinWeightProfilesSemanticIndices.Num();
@@ -132,7 +131,7 @@ void SetSurfaceFormat( FMutableGraphGenerationContext& GenerationContext,
 	}
 
 	// MorphTarget vertex tracking info buffers
-	if (bWithRealTimeMorphs)
+	if (MeshData.bHasRealTimeMorphs)
 	{
 		using namespace mu;
 		{
@@ -163,19 +162,35 @@ void SetSurfaceFormat( FMutableGraphGenerationContext& GenerationContext,
 	}
 
 	//Clothing Data Buffer.
-	if (bWithClothing)
+	if (MeshData.bHasClothing)
 	{
-		using namespace mu;
-		const int32 ElementSize = sizeof(int32);
-		constexpr int32 ChannelCount = 1;
-		const EMeshBufferSemantic Semantics[ChannelCount] = { MBS_OTHER };
-		const int32 SemanticIndices[ChannelCount] = { 2 };
-		const EMeshBufferFormat Formats[ChannelCount] = { MBF_INT32 };
-		const int32 Components[ChannelCount] = { 1 };
-		const int32 Offsets[ChannelCount] = { 0 };
-
-		OutVertexBufferFormat.SetBuffer(CurrentVertexBuffer, ElementSize, ChannelCount, Semantics, SemanticIndices, Formats, Components, Offsets);
-		++CurrentVertexBuffer;
+		{
+			using namespace mu;
+			const int32 ElementSize = sizeof(int32);
+			constexpr int32 ChannelCount = 1;
+			const EMeshBufferSemantic Semantics[ChannelCount] = { MBS_OTHER };
+			const int32 SemanticIndices[ChannelCount] = { 2 };
+			const EMeshBufferFormat Formats[ChannelCount] = { MBF_INT32 };
+			const int32 Components[ChannelCount] = { 1 };
+			const int32 Offsets[ChannelCount] = { 0 };
+
+			OutVertexBufferFormat.SetBuffer(CurrentVertexBuffer, ElementSize, ChannelCount, Semantics, SemanticIndices, Formats, Components, Offsets);
+			++CurrentVertexBuffer;
+		}
+
+		{
+			using namespace mu;
+			const int32 ElementSize = sizeof(uint32);
+			constexpr int32 ChannelCount = 1;
+			const EMeshBufferSemantic Semantics[ChannelCount] = { MBS_OTHER };
+			const int32 SemanticIndices[ChannelCount] = { 3 };
+			const EMeshBufferFormat Formats[ChannelCount] = { MBF_UINT32 };
+			const int32 Components[ChannelCount] = { 1 };
+			const int32 Offsets[ChannelCount] = { 0 };
+
+			OutVertexBufferFormat.SetBuffer(CurrentVertexBuffer, ElementSize, ChannelCount, Semantics, SemanticIndices, Formats, Components, Offsets);
+			++CurrentVertexBuffer;
+		}
 	}
 
 	for (int32 ProfileSemanticIndex : MeshData.SkinWeightProfilesSemanticIndices)
@@ -363,11 +378,9 @@ mu::NodeSurfacePtr GenerateMutableSourceSurface(const UEdGraphPin * Pin, FMutabl
 
 				mu::Ptr<mu::NodeMeshFormat> MeshFormatNode = new mu::NodeMeshFormat();
 				MeshFormatNode->SetSource(MeshNode.get());
-				SetSurfaceFormat( GenerationContext,
+				SetSurfaceFormat(GenerationContext,
 						MeshFormatNode->GetVertexBuffers(), MeshFormatNode->GetIndexBuffers(), MeshData,
 						GenerationContext.Options.CustomizableObjectNumBoneInfluences,
-						GenerationContext.Options.bRealTimeMorphTargetsEnabled, 
-						GenerationContext.Options.bClothingEnabled,
 						GenerationContext.Options.b16BitBoneWeightsEnabled);
 
 				// \TODO: Make it an option?
@@ -1124,8 +1137,6 @@ mu::NodeSurfacePtr GenerateMutableSourceSurface(const UEdGraphPin * Pin, FMutabl
 				SetSurfaceFormat( GenerationContext,
 						MeshFormat->GetVertexBuffers(), MeshFormat->GetIndexBuffers(), MeshData,
 						GenerationContext.Options.CustomizableObjectNumBoneInfluences,
-						GenerationContext.Options.bRealTimeMorphTargetsEnabled, 
-						GenerationContext.Options.bClothingEnabled,
 						GenerationContext.Options.b16BitBoneWeightsEnabled);
 
 				MeshFormat->SetSource(MeshPtr.get());
diff --git a/Engine/Plugins/Experimental/Mutable/Source/CustomizableObjectEditor/Private/MuCOE/GenerateMutableSource/GenerateMutableSourceTable.cpp b/Engine/Plugins/Experimental/Mutable/Source/CustomizableObjectEditor/Private/MuCOE/GenerateMutableSource/GenerateMutableSourceTable.cpp
--- a/Engine/Plugins/Experimental/Mutable/Source/CustomizableObjectEditor/Private/MuCOE/GenerateMutableSource/GenerateMutableSourceTable.cpp
+++ b/Engine/Plugins/Experimental/Mutable/Source/CustomizableObjectEditor/Private/MuCOE/GenerateMutableSource/GenerateMutableSourceTable.cpp
@@ -297,29 +297,6 @@ bool FillTableColumn(const UCustomizableObjectNodeTable* TableNode, mu::TablePtr
 					AddTagToMutableMeshUnique(*MutableMesh, PhysicsAssetTag);
 				}
 
-				if (GenerationContext.Options.bClothingEnabled)
-				{
-					UClothingAssetBase* ClothingAssetBase = SkeletalMesh->GetSectionClothingAsset(LODIndex, SectionIndex);	
-					UClothingAssetCommon* ClothingAssetCommon = Cast<UClothingAssetCommon>(ClothingAssetBase);
-
-					if (ClothingAssetCommon && ClothingAssetCommon->PhysicsAsset)
-					{	
-						int32 AssetIndex = GenerationContext.ContributingClothingAssetsData.IndexOfByPredicate( 
-						[Guid = ClothingAssetBase->GetAssetGuid()](const FCustomizableObjectClothingAssetData& A)
-						{
-							return A.OriginalAssetGuid == Guid;
-						});
-
-						check(AssetIndex != INDEX_NONE);
-
-						GenerationContext.AddParticipatingObject(*ClothingAssetCommon->PhysicsAsset);
-						
-						const int32 PhysicsAssetIndex = GenerationContext.PhysicsAssets.AddUnique(ClothingAssetCommon->PhysicsAsset);
-						FString ClothPhysicsAssetTag = FString::Printf(TEXT("__ClothPA:%d_%d"), AssetIndex, PhysicsAssetIndex);
-						AddTagToMutableMeshUnique(*MutableMesh, ClothPhysicsAssetTag);
-					}
-				}
-
 				if (!AnimBPAssetTag.IsEmpty())
 				{
 					AddTagToMutableMeshUnique(*MutableMesh, AnimBPAssetTag);
diff --git a/Engine/Plugins/Experimental/Mutable/Source/MutableTools/Private/MuT/TypeInfo.cpp b/Engine/Plugins/Experimental/Mutable/Source/MutableTools/Private/MuT/TypeInfo.cpp
--- a/Engine/Plugins/Experimental/Mutable/Source/MutableTools/Private/MuT/TypeInfo.cpp
+++ b/Engine/Plugins/Experimental/Mutable/Source/MutableTools/Private/MuT/TypeInfo.cpp
@@ -58,7 +58,7 @@ namespace mu
 	static_assert(sizeof(TypeInfo::s_imageFormatName) / sizeof(void*) == int32(EImageFormat::IF_COUNT));
 
 
-	const char* TypeInfo::s_meshBufferSemanticName[MBS_COUNT] =
+	const char* TypeInfo::s_meshBufferSemanticName[] =
 	{
 		"None",
 
@@ -83,8 +83,9 @@ namespace mu
 
 		"TriangleIndex",
 		"BarycentricCoords",
-		"Distance"
+		"Distance",
 
+		"AltSkinWeight"
 	};
 
 	static_assert(sizeof(TypeInfo::s_meshBufferSemanticName) / sizeof(void*) == int32(MBS_COUNT));
@@ -167,4 +168,4 @@ namespace mu
 	};
 	static_assert(sizeof(TypeInfo::s_blendModeName) / sizeof(void*) == int32(EBlendType::_BT_COUNT));
 
-}
\ No newline at end of file
+}
