[mutable] Added new icon to indicate from which SKM are the UVs displayed in the layout editor of a Table Node.
+ Changed some widget to match the UE style.

#rnx

diff --git a/Engine/Plugins/Experimental/Mutable/Source/CustomizableObjectEditor/Private/MuCOE/Nodes/CustomizableObjectNodeTableDetails.cpp b/Engine/Plugins/Experimental/Mutable/Source/CustomizableObjectEditor/Private/MuCOE/Nodes/CustomizableObjectNodeTableDetails.cpp
--- a/Engine/Plugins/Experimental/Mutable/Source/CustomizableObjectEditor/Private/MuCOE/Nodes/CustomizableObjectNodeTableDetails.cpp
+++ b/Engine/Plugins/Experimental/Mutable/Source/CustomizableObjectEditor/Private/MuCOE/Nodes/CustomizableObjectNodeTableDetails.cpp
@@ -68,7 +68,7 @@ void FCustomizableObjectNodeTableDetails::CustomizeDetails(const TSharedPtr<IDet
 			SNew(STextBlock)
 			.Text(LOCTEXT("VersionColumn_SelectorText","Version Column"))
 			.ToolTipText(LOCTEXT("VersionColumn_SelectorTooltip","Select the column that contains the version of each row."))
-			.Font(DetailBuilder->GetDetailFont())
+			.Font(IDetailLayoutBuilder::GetDetailFont())
 		]
 		.ValueContent()
 		[
@@ -77,7 +77,7 @@ void FCustomizableObjectNodeTableDetails::CustomizeDetails(const TSharedPtr<IDet
 			.OptionsSource(&VersionColumnsOptionNames)
 			.OnComboBoxOpening(this, &FCustomizableObjectNodeTableDetails::OnOpenVersionColumnComboBox)
 			.OnSelectionChanged(this, &FCustomizableObjectNodeTableDetails::OnVersionColumnComboBoxSelectionChanged)
-			.Font(DetailBuilder->GetDetailFont())
+			.Font(IDetailLayoutBuilder::GetDetailFont())
 			.ColorAndOpacity(this, &FCustomizableObjectNodeTableDetails::GetVersionColumnComboBoxTextColor, &VersionColumnsOptionNames)
 		]
 		.OverrideResetToDefault(FResetToDefaultOverride::Create(FSimpleDelegate::CreateSP(this, &FCustomizableObjectNodeTableDetails::OnVersionColumnComboBoxSelectionReset)));
@@ -88,7 +88,7 @@ void FCustomizableObjectNodeTableDetails::CustomizeDetails(const TSharedPtr<IDet
 			SNew(STextBlock)
 			.Text(LOCTEXT("MutableUIMetadataColumn_SelectorText","Options UI Metadata Column"))
 			.ToolTipText(LOCTEXT("MutableUIMetadataColumn_SelectorTooltip","Select a column that contains a Parameter UI Metadata for each Parameter Option (table row)."))
-			.Font(DetailBuilder->GetDetailFont())
+			.Font(IDetailLayoutBuilder::GetDetailFont())
 		]
 		.ValueContent()
 		[
@@ -97,233 +97,147 @@ void FCustomizableObjectNodeTableDetails::CustomizeDetails(const TSharedPtr<IDet
 			.OptionsSource(&MutableMetaDataColumnsOptionNames)
 			.OnComboBoxOpening(this, &FCustomizableObjectNodeTableDetails::OnOpenMutableMetadataComboBox)
 			.OnSelectionChanged(this, &FCustomizableObjectNodeTableDetails::OnMutableMetaDataColumnComboBoxSelectionChanged)
-			.Font(DetailBuilder->GetDetailFont())
+			.Font(IDetailLayoutBuilder::GetDetailFont())
 			.ColorAndOpacity(this, &FCustomizableObjectNodeTableDetails::GetMetadataUIComboBoxTextColor, &MutableMetaDataColumnsOptionNames)
 		]
 		.OverrideResetToDefault(FResetToDefaultOverride::Create(FSimpleDelegate::CreateSP(this, &FCustomizableObjectNodeTableDetails::OnMutableMetaDataColumnComboBoxSelectionReset)));
 
-		AnimationCategory.AddCustomRow(LOCTEXT("AnimationProperties", "Animation Properties"))
-		[
-			SNew(SVerticalBox)
-			
-			// Mesh Column selection widget
-			+SVerticalBox::Slot()
-			.Padding(0.0f, 5.0f, 6.0f, 0.0f)
-			.AutoHeight()
-			[
-				SNew(SHorizontalBox)
-				+ SHorizontalBox::Slot()
-				.AutoWidth()
-				.Padding(0.0f, 5.0f, 6.0f, 0.0f)
-				[
-					SNew(STextBlock)
-					.Text(LOCTEXT("AnimMeshColumnText", "Mesh Column: "))
-					.ToolTipText(LOCTEXT("AnimMeshColumnTooltip","Select a mesh column from the Data Table to edit its animation options (Applied to all LODs)."))
-				]
 
-				+ SHorizontalBox::Slot()
-				[
-					SNew(SBorder)
-					.BorderImage(UE_MUTABLE_GET_BRUSH("NoBorder"))
-					.Padding(FMargin(0.0f, 0.0f, 10.0f, 0.0f))
-					[
-						SAssignNew(AnimMeshColumnComboBox, STextComboBox)
-						.OptionsSource(&AnimMeshColumnOptionNames)
-						.InitiallySelectedItem(nullptr)
-						.OnSelectionChanged(this, &FCustomizableObjectNodeTableDetails::OnAnimMeshColumnComboBoxSelectionChanged)
-					]
-				]
-			]
 
-			// Animation Blueprint selection widget
-			+SVerticalBox::Slot()
-			.Padding(0.0f, 5.0f, 6.0f, 0.0f)
-			.AutoHeight()
-			[
-				SNew(SHorizontalBox)
-				+ SHorizontalBox::Slot()
-				.AutoWidth()
-				.Padding(0.0f, 5.0f, 6.0f, 0.0f)
-				[
-					SNew(STextBlock)
-					.Text(LOCTEXT("AnimBPText", "Animation Blueprint Column: "))
-					.ToolTipText(LOCTEXT("AnimBlueprintColumnTooltip", "Select an animation blueprint column from the Data Table that will be applied to the mesh selected"))
-					.Visibility_Lambda([WeakDetails = SharedThis(this).ToWeakPtr()]() -> EVisibility
-					{
-						const TSharedPtr<FCustomizableObjectNodeTableDetails> Details = WeakDetails.Pin();
-						if (!Details)
-						{
-							return EVisibility::Collapsed;
-						}
-						
-						if (!Details->AnimComboBox.IsValid())
-						{
-							return EVisibility::Collapsed;
-						}
-
-						return Details->AnimComboBox->GetVisibility();
-					})
-				]
+		// Anim Category -----------------------------------
 
-				+ SHorizontalBox::Slot()
-				[
-					SNew(SBorder)
-					.BorderImage(UE_MUTABLE_GET_BRUSH("NoBorder"))
-					.Padding(FMargin(0.0f, 0.0f, 10.0f, 0.0f))
-					[
-						SAssignNew(AnimComboBox, STextComboBox)
-						.Visibility(EVisibility::Collapsed)
-						.OptionsSource(&AnimOptionNames)
-						.OnSelectionChanged(this, &FCustomizableObjectNodeTableDetails::OnAnimInstanceComboBoxSelectionChanged)
-					]
-				]
-			]
-
-			// Animation Slot selection widget
-			+SVerticalBox::Slot()
-			.Padding(0.0f, 5.0f, 6.0f, 0.0f)
-			.AutoHeight()
-			[
-				SNew(SHorizontalBox)
-				+ SHorizontalBox::Slot()
-				.AutoWidth()
-				.Padding(0.0f, 5.0f, 6.0f, 0.0f)
-				[
-					SNew(STextBlock)
-					.Text(LOCTEXT("AnimSlotText", "Animation Slot Column: "))
-					.ToolTipText(LOCTEXT("AnimSlotColumnTooltip", "Select an animation slot column from the Data Table that will set to the slot value of the animation blueprint"))
-					.Visibility_Lambda([WeakDetails = SharedThis(this).ToWeakPtr()]() -> EVisibility 
-						{
-							const TSharedPtr<FCustomizableObjectNodeTableDetails> Details = WeakDetails.Pin();
-							if (!Details)
-							{
-								return EVisibility::Collapsed;
-							}
+		// Mesh Column Selector
+		AnimationCategory.AddCustomRow(LOCTEXT("AnimationProperties", "Animation Properties"))
+		.NameContent()
+		[
+			SNew(STextBlock)
+			.Text(LOCTEXT("AnimMeshColumnText", "Mesh Column: "))
+			.ToolTipText(LOCTEXT("AnimMeshColumnTooltip", "Select a mesh column from the Data Table to edit its animation options (Applied to all LODs)."))
+			.Font(IDetailLayoutBuilder::GetDetailFont())
+		]
+		.ValueContent()
+		.HAlign(EHorizontalAlignment::HAlign_Fill)
+		[
+			SAssignNew(AnimMeshColumnComboBox, STextComboBox)
+			.OptionsSource(&AnimMeshColumnOptionNames)
+			.InitiallySelectedItem(AnimMeshColumnOptionNames[0])
+			.OnSelectionChanged(this, &FCustomizableObjectNodeTableDetails::OnAnimMeshColumnComboBoxSelectionChanged)
+			.Font(IDetailLayoutBuilder::GetDetailFont())
+		]
+		.OverrideResetToDefault(FResetToDefaultOverride::Create(FSimpleDelegate::CreateSP(this, &FCustomizableObjectNodeTableDetails::OnAnimMeshCustomRowResetButtonClicked)));
 
-							if (!Details->AnimSlotComboBox.IsValid())
-							{
-								return EVisibility::Collapsed;
-							}
 
-							return Details->AnimSlotComboBox->GetVisibility();
-						})
-				]
+		// AnimBP Column Selector
+		AnimationCategory.AddCustomRow(LOCTEXT("AnimationProperties", "Animation Properties"))
+		.NameContent()
+		[
+			SNew(STextBlock)
+			.Text(LOCTEXT("AnimBPText", "Animation Blueprint Column: "))
+			.ToolTipText(LOCTEXT("AnimBlueprintColumnTooltip", "Select an animation blueprint column from the Data Table that will be applied to the mesh selected"))
+			.Font(IDetailLayoutBuilder::GetDetailFont())
+		]
+		.ValueContent()
+		.HAlign(EHorizontalAlignment::HAlign_Fill)
+		[
+			SAssignNew(AnimComboBox, STextComboBox)
+			.OptionsSource(&AnimOptionNames)
+			.OnSelectionChanged(this, &FCustomizableObjectNodeTableDetails::OnAnimInstanceComboBoxSelectionChanged)
+			.Font(IDetailLayoutBuilder::GetDetailFont())
+		]
+		.OverrideResetToDefault(FResetToDefaultOverride::Create(FSimpleDelegate::CreateSP(this, &FCustomizableObjectNodeTableDetails::OnAnimCustomRowResetButtonClicked, EAnimColumnType::EACT_BluePrintColumn)))
+		.Visibility(TAttribute<EVisibility>(this, &FCustomizableObjectNodeTableDetails::AnimWidgetsVisibility));
 
-				+ SHorizontalBox::Slot()
-				[
-					SNew(SBorder)
-					.BorderImage(UE_MUTABLE_GET_BRUSH("NoBorder"))
-					.Padding(FMargin(0.0f, 0.0f, 10.0f, 0.0f))
-					[
-						SAssignNew(AnimSlotComboBox, STextComboBox)
-						.Visibility(EVisibility::Collapsed)
-						.OptionsSource(&AnimSlotOptionNames)
-						.OnSelectionChanged(this, &FCustomizableObjectNodeTableDetails::OnAnimSlotComboBoxSelectionChanged)
-					]
-				]
-			]
 
-			// Animation Tags selection widget
-			+SVerticalBox::Slot()
-			.Padding(0.0f, 5.0f, 6.0f, 0.0f)
-			.AutoHeight()
-			[
-				SNew(SHorizontalBox)
-				+ SHorizontalBox::Slot()
-				.AutoWidth()
-				.Padding(0.0f, 5.0f, 6.0f, 0.0f)
-				[
-					SNew(STextBlock)
-					.Text(LOCTEXT("AnimTagsText", "Animation Tags Column: "))
-					.ToolTipText(LOCTEXT("AnimTagColumnTooltip", "Select an animation tag column from the Data Table that will set to the animation tags of the animation blueprint"))
-					.Visibility_Lambda([WeakDetails = SharedThis(this).ToWeakPtr()]() -> EVisibility 
-						{
-							const TSharedPtr<FCustomizableObjectNodeTableDetails> Details = WeakDetails.Pin();
-							if (!Details)
-							{
-								return EVisibility::Collapsed;
-							}
+		// AnimSlot Column Selector
+		AnimationCategory.AddCustomRow(LOCTEXT("AnimationProperties", "Animation Properties"))
+		.NameContent()
+		[
+			SNew(STextBlock)
+			.Text(LOCTEXT("AnimSlotText", "Animation Slot Column: "))
+			.ToolTipText(LOCTEXT("AnimSlotColumnTooltip", "Select an animation slot column from the Data Table that will set to the slot value of the animation blueprint"))
+			.Font(IDetailLayoutBuilder::GetDetailFont())
+		]
+		.ValueContent()
+		.HAlign(EHorizontalAlignment::HAlign_Fill)
+		[
+			SAssignNew(AnimSlotComboBox, STextComboBox)
+			.OptionsSource(&AnimSlotOptionNames)
+			.OnSelectionChanged(this, &FCustomizableObjectNodeTableDetails::OnAnimSlotComboBoxSelectionChanged)
+			.Font(IDetailLayoutBuilder::GetDetailFont())
+		]
+		.OverrideResetToDefault(FResetToDefaultOverride::Create(FSimpleDelegate::CreateSP(this, &FCustomizableObjectNodeTableDetails::OnAnimCustomRowResetButtonClicked, EAnimColumnType::EACT_SlotColumn)))
+		.Visibility(TAttribute<EVisibility>(this, &FCustomizableObjectNodeTableDetails::AnimWidgetsVisibility));
+		
 
-							if (!Details->AnimTagsComboBox.IsValid())
-							{
-								return EVisibility::Collapsed;
-							}
+		// AnimTags Column Selector
+		AnimationCategory.AddCustomRow(LOCTEXT("AnimationProperties", "Animation Properties"))
+		.NameContent()
+		[
+			SNew(STextBlock)
+			.Text(LOCTEXT("AnimTagsText", "Animation Tags Column: "))
+			.ToolTipText(LOCTEXT("AnimTagColumnTooltip", "Select an animation tag column from the Data Table that will set to the animation tags of the animation blueprint"))
+			.Font(IDetailLayoutBuilder::GetDetailFont())
+		]
+		.ValueContent()
+		.HAlign(EHorizontalAlignment::HAlign_Fill)
+		[
+			SAssignNew(AnimTagsComboBox, STextComboBox)
+			.OptionsSource(&AnimTagsOptionNames)
+			.OnSelectionChanged(this, &FCustomizableObjectNodeTableDetails::OnAnimTagsComboBoxSelectionChanged)
+			.Font(IDetailLayoutBuilder::GetDetailFont())
+		]
+		.OverrideResetToDefault(FResetToDefaultOverride::Create(FSimpleDelegate::CreateSP(this, &FCustomizableObjectNodeTableDetails::OnAnimCustomRowResetButtonClicked, EAnimColumnType::EACT_TagsColumn)))
+		.Visibility(TAttribute<EVisibility>(this, &FCustomizableObjectNodeTableDetails::AnimWidgetsVisibility));
 
-							return Details->AnimTagsComboBox->GetVisibility();								
-						})
-				]
 
-				+ SHorizontalBox::Slot()
-				[
-					SNew(SBorder)
-					.BorderImage(UE_MUTABLE_GET_BRUSH("NoBorder"))
-					.Padding(FMargin(0.0f, 0.0f, 10.0f, 0.0f))
-					[
-						SAssignNew(AnimTagsComboBox, STextComboBox)
-						.Visibility(EVisibility::Collapsed)
-						.OptionsSource(&AnimTagsOptionNames)
-						.OnSelectionChanged(this, &FCustomizableObjectNodeTableDetails::OnAnimTagsComboBoxSelectionChanged)
-					]
-				]
-			]
 
-			+SVerticalBox::Slot()
-			.Padding(0.0f, 5.0f, 16.0f,5.0f)
-			.AutoHeight()
-			.HAlign(EHorizontalAlignment::HAlign_Right)
-			[
-				SNew(SHorizontalBox)
-				+ SHorizontalBox::Slot()
-				.AutoWidth()
-				[
-					SAssignNew(ClearButton,SButton)
-					.Text(FText::FromString("Clear"))
-					.ToolTipText(FText::FromString("Clear options of animation parameters"))
-					.HAlign(EHorizontalAlignment::HAlign_Center)
-					.Visibility(EVisibility::Collapsed)
-					.OnClicked(this, &FCustomizableObjectNodeTableDetails::OnClearButtonPressed)
-				]
-			]
-		];
+		// Layout Category -----------------------------------
 
+		//	Layout Selector
 		SelectedLayout = nullptr;
 		LayoutBlocksEditor = SNew(SCustomizableObjectNodeLayoutBlocksEditor);
 
+		// Mesh selector of the layout editor
 		LayoutCategory.AddCustomRow(LOCTEXT("TableLayoutEditor_MeshSelector", "Mesh Selector"))
+		.NameContent()
+		[
+			SNew(STextBlock)
+			.Text(LOCTEXT("LayoutMeshColumnText", "Mesh Column: "))
+			.ToolTipText(LOCTEXT("LayoutMeshColumnTooltip", "Select a mesh from the Data Table to edit its layout blocks."))
+			.Font(IDetailLayoutBuilder::GetDetailFont())
+		]
+		.ValueContent()
+		.HAlign(EHorizontalAlignment::HAlign_Fill)
 		[
-			SNew(SVerticalBox)
+			SNew(SHorizontalBox)
+			+ SHorizontalBox::Slot()
+			.FillWidth(0.9f)
+			[
+				SAssignNew(LayoutMeshColumnComboBox, STextComboBox)
+				.OptionsSource(&LayoutMeshColumnOptionNames)
+				.InitiallySelectedItem(LayoutMeshColumnOptionNames[0])
+				.OnSelectionChanged(this, &FCustomizableObjectNodeTableDetails::OnLayoutMeshColumnComboBoxSelectionChanged)
+				.Font(IDetailLayoutBuilder::GetDetailFont())
+			]
 			
-			// Mesh Column selection widget
-			+SVerticalBox::Slot()
-			.Padding(0.0f, 5.0f, 6.0f, 0.0f)
-			.AutoHeight()
+			+ SHorizontalBox::Slot()
+			.FillWidth(0.1f)
 			[
-				SNew(SHorizontalBox)
-				+ SHorizontalBox::Slot()
-				.AutoWidth()
-				.Padding(0.0f, 5.0f, 6.0f, 0.0f)
+				SNew(SBox)
+				.HAlign(EHorizontalAlignment::HAlign_Center)
+				.VAlign(EVerticalAlignment::VAlign_Center)
+				.Content()
 				[
-					SNew(STextBlock)
-					.Text(LOCTEXT("LayoutMeshColumnText", "Mesh Column: "))
-					.ToolTipText(LOCTEXT("LayoutMeshColumnTooltip", "Select a mesh from the Data Table to edit its layout blocks."))
-					.Font(IDetailLayoutBuilder::GetDetailFont())
-				]
-		
-				+ SHorizontalBox::Slot()
-				[
-					SNew(SBorder)
-					.BorderImage(UE_MUTABLE_GET_BRUSH("NoBorder"))
-					.Padding(FMargin(0.0f, 0.0f, 10.0f, 0.0f))
-					[
-						SAssignNew(LayoutMeshColumnComboBox, STextComboBox)
-						.OptionsSource(&LayoutMeshColumnOptionNames)
-						.InitiallySelectedItem(nullptr)
-						.OnSelectionChanged(this, &FCustomizableObjectNodeTableDetails::OnLayoutMeshColumnComboBoxSelectionChanged)
-						.Font(IDetailLayoutBuilder::GetDetailFont())
-					]
+					SNew(SImage)
+					.Image(UE_MUTABLE_GET_BRUSH(TEXT("Icons.Info")))
+					.ToolTipText(FText(LOCTEXT("LaytoutMeshNoteTooltipText","Note:"
+						"\nAs all meshes of a Data Table column share the same layout, the UVs shown"
+						"\nin the editor are from the Default Skeletal Mesh of the Structure.")))
 				]
 			]
-		];
+			
+		]
+		.OverrideResetToDefault(FResetToDefaultOverride::Create(FSimpleDelegate::CreateSP(this, &FCustomizableObjectNodeTableDetails::OnLayoutMeshCustomRowResetButtonClicked)));
 
 		// Layout size selector widget
 		LayoutCategory.AddCustomRow(LOCTEXT("TableBlocksDetails_SizeSelector", "SizeSelector"))
@@ -342,6 +256,7 @@ void FCustomizableObjectNodeTableDetails::CustomizeDetails(const TSharedPtr<IDet
 			.Font(IDetailLayoutBuilder::GetDetailFont())
 		];
 
+
 		// Layout strategy selector group widget
 		IDetailGroup* LayoutStrategyOptionsGroup = &LayoutCategory.AddGroup(TEXT("TableLayoutStrategyOptionsGroup"), LOCTEXT("TableLayoutStrategyGroup", "Table Layout Strategy Group"), false, true);
 		LayoutStrategyOptionsGroup->HeaderRow()
@@ -361,6 +276,7 @@ void FCustomizableObjectNodeTableDetails::CustomizeDetails(const TSharedPtr<IDet
 			.Font(IDetailLayoutBuilder::GetDetailFont())
 		];
 
+
 		// Max layout size selector widget
 		LayoutStrategyOptionsGroup->AddWidgetRow()
 		.Visibility(TAttribute<EVisibility>(this, &FCustomizableObjectNodeTableDetails::FixedStrategyOptionsVisibility))
@@ -378,6 +294,7 @@ void FCustomizableObjectNodeTableDetails::CustomizeDetails(const TSharedPtr<IDet
 			.Font(IDetailLayoutBuilder::GetDetailFont())
 		];
 
+
 		// Reduction method selector widget
 		LayoutStrategyOptionsGroup->AddWidgetRow()
 		.Visibility(TAttribute<EVisibility>(this, &FCustomizableObjectNodeTableDetails::FixedStrategyOptionsVisibility))
@@ -398,6 +315,7 @@ void FCustomizableObjectNodeTableDetails::CustomizeDetails(const TSharedPtr<IDet
 			.Font(IDetailLayoutBuilder::GetDetailFont())
 		];
 
+
 		// Block editor Widget
 		LayoutCategory.AddCustomRow(LOCTEXT("TableLayoutEditor", "Layout Editor"))
 		[
@@ -426,7 +344,11 @@ void FCustomizableObjectNodeTableDetails::GenerateMeshColumnComboBoxOptions()
 		return;
 	}
 
-	// we just need the mesh columns
+	// Add first element to clear selection
+	AnimMeshColumnOptionNames.Add(MakeShareable(new FString("- Nothing Selected -")));
+	LayoutMeshColumnOptionNames.Add(MakeShareable(new FString("- Nothing Selected -")));
+
+	// Get mesh columns only
 	for (TFieldIterator<FProperty> It(TableStruct); It; ++It)
 	{
 		FProperty* ColumnProperty = *It;
@@ -475,6 +397,17 @@ void FCustomizableObjectNodeTableDetails::GenerateMeshColumnComboBoxOptions()
 }
 
 
+void FCustomizableObjectNodeTableDetails::OnNodePinValueChanged()
+{
+	if (IDetailLayoutBuilder* DetailBuilder = DetailBuilderPtr.Pin().Get()) // Raw because we don't want to keep alive the details builder when calling the force refresh details
+	{
+		DetailBuilder->ForceRefreshDetails();
+	}
+}
+
+
+// Anim Category --------------------------------------------------------------------------------
+
 void FCustomizableObjectNodeTableDetails::GenerateAnimInstanceComboBoxOptions()
 {
 	// Options Reset
@@ -552,55 +485,31 @@ void FCustomizableObjectNodeTableDetails::GenerateAnimInstanceComboBoxOptions()
 }
 
 
-void FCustomizableObjectNodeTableDetails::OnAnimMeshColumnComboBoxSelectionChanged(TSharedPtr<FString> Selection, ESelectInfo::Type SelectInfo)
+EVisibility FCustomizableObjectNodeTableDetails::AnimWidgetsVisibility() const
 {
-	if (Selection.IsValid())
+	if (AnimMeshColumnComboBox.IsValid() && AnimMeshColumnComboBox->GetSelectedItem() != AnimMeshColumnOptionNames[0])
 	{
-		if (AnimComboBox.IsValid() && AnimSlotComboBox.IsValid() && AnimTagsComboBox.IsValid())
-		{
-			AnimComboBox->SetVisibility(EVisibility::Visible);
-			AnimSlotComboBox->SetVisibility(EVisibility::Visible);
-			AnimTagsComboBox->SetVisibility(EVisibility::Visible);
-			ClearButton->SetVisibility(EVisibility::Visible);
-
-			GenerateAnimInstanceComboBoxOptions();
-		}
+		return EVisibility::Visible;
 	}
+
+	return EVisibility::Collapsed;
 }
 
 
-void FCustomizableObjectNodeTableDetails::OnLayoutMeshColumnComboBoxSelectionChanged(TSharedPtr<FString> Selection, ESelectInfo::Type SelectInfo)
+void FCustomizableObjectNodeTableDetails::OnAnimMeshColumnComboBoxSelectionChanged(TSharedPtr<FString> Selection, ESelectInfo::Type SelectInfo)
 {
 	if (Selection.IsValid())
 	{
-		FString ColumnName;
-		(*Selection).Split(" LOD_", &ColumnName, NULL);
-
-		for (const UEdGraphPin* Pin : Node->GetAllNonOrphanPins())
-		{
-			const UCustomizableObjectNodeTableMeshPinData* PinData = Cast<UCustomizableObjectNodeTableMeshPinData >(Node->GetPinData(*Pin));
-
-			if (PinData && PinData->ColumnName == ColumnName)
-			{
-				for (int32 LayoutIndex = 0; LayoutIndex < PinData->Layouts.Num(); ++LayoutIndex)
-				{
-					if (PinData->Layouts[LayoutIndex]->GetLayoutName() == *Selection)
-					{
-						LayoutBlocksEditor->SetCurrentLayout(PinData->Layouts[LayoutIndex]);
-						SelectedLayout = PinData->Layouts[LayoutIndex];
-
-						FillLayoutComboBoxOptions();
-					}
-				}
-			}
-		}
+		GenerateAnimInstanceComboBoxOptions();
 	}
 }
 
 
 void FCustomizableObjectNodeTableDetails::OnAnimInstanceComboBoxSelectionChanged(TSharedPtr<FString> Selection, ESelectInfo::Type SelectInfo)
 {
-	if (Selection.IsValid() && AnimMeshColumnComboBox->GetSelectedItem().IsValid() && SelectInfo != ESelectInfo::Direct)
+	bool bIsMeshSelectionValid =  AnimMeshColumnComboBox->GetSelectedItem() != AnimMeshColumnOptionNames[0] && AnimMeshColumnComboBox->GetSelectedItem().IsValid();
+
+	if (bIsMeshSelectionValid && Selection.IsValid() && SelectInfo != ESelectInfo::Direct)
 	{
 		FString ColumnName = *AnimMeshColumnComboBox->GetSelectedItem();
 		FGuid ColumnId = Node->GetColumnIdByName(FName(*ColumnName));
@@ -625,7 +534,9 @@ void FCustomizableObjectNodeTableDetails::OnAnimInstanceComboBoxSelectionChanged
 
 void FCustomizableObjectNodeTableDetails::OnAnimSlotComboBoxSelectionChanged(TSharedPtr<FString> Selection, ESelectInfo::Type SelectInfo)
 {
-	if (Selection.IsValid() && AnimMeshColumnComboBox->GetSelectedItem().IsValid() && SelectInfo != ESelectInfo::Direct)
+	bool bIsMeshSelectionValid = AnimMeshColumnComboBox->GetSelectedItem() != AnimMeshColumnOptionNames[0] && AnimMeshColumnComboBox->GetSelectedItem().IsValid();
+
+	if (bIsMeshSelectionValid && Selection.IsValid() && SelectInfo != ESelectInfo::Direct)
 	{
 		FString ColumnName = *AnimMeshColumnComboBox->GetSelectedItem();
 		FGuid ColumnId = Node->GetColumnIdByName(FName(*ColumnName));
@@ -650,7 +561,9 @@ void FCustomizableObjectNodeTableDetails::OnAnimSlotComboBoxSelectionChanged(TSh
 
 void FCustomizableObjectNodeTableDetails::OnAnimTagsComboBoxSelectionChanged(TSharedPtr<FString> Selection, ESelectInfo::Type SelectInfo)
 {
-	if (Selection.IsValid() && AnimMeshColumnComboBox->GetSelectedItem().IsValid() && SelectInfo != ESelectInfo::Direct)
+	bool bIsMeshSelectionValid = AnimMeshColumnComboBox->GetSelectedItem() != AnimMeshColumnOptionNames[0] && AnimMeshColumnComboBox->GetSelectedItem().IsValid();
+
+	if (bIsMeshSelectionValid && Selection.IsValid() && SelectInfo != ESelectInfo::Direct)
 	{
 		FString ColumnName = *AnimMeshColumnComboBox->GetSelectedItem();
 		FGuid ColumnId = Node->GetColumnIdByName(FName(*ColumnName));
@@ -675,55 +588,231 @@ void FCustomizableObjectNodeTableDetails::OnAnimTagsComboBoxSelectionChanged(TSh
 }
 
 
-void FCustomizableObjectNodeTableDetails::OnNodePinValueChanged()
+void FCustomizableObjectNodeTableDetails::OnAnimMeshCustomRowResetButtonClicked()
 {
-	if (IDetailLayoutBuilder* DetailBuilder = DetailBuilderPtr.Pin().Get()) // Raw because we don't want to keep alive the details builder when calling the force refresh details
+	if (AnimMeshColumnOptionNames.Num())
 	{
-		DetailBuilder->ForceRefreshDetails();
+		AnimMeshColumnComboBox->SetSelectedItem(AnimMeshColumnOptionNames[0]);
 	}
 }
 
 
-FReply FCustomizableObjectNodeTableDetails::OnClearButtonPressed()
+void FCustomizableObjectNodeTableDetails::OnAnimCustomRowResetButtonClicked(EAnimColumnType ColumnType)
 {
 	if (!AnimMeshColumnComboBox->GetSelectedItem().IsValid())
 	{
-		return FReply::Unhandled();
+		return;
 	}
-		
+
 	FString ColumnName = *AnimMeshColumnComboBox->GetSelectedItem();
 	FGuid ColumnId = Node->GetColumnIdByName(FName(*ColumnName));
 	FTableNodeColumnData* MeshColumnData = Node->ColumnDataMap.Find(ColumnId);
 
-	if (MeshColumnData)
+	if (!MeshColumnData)
+	{
+		return;
+	}
+
+	switch (ColumnType)
+	{
+	case EAnimColumnType::EACT_BluePrintColumn:
 	{
 		MeshColumnData->AnimInstanceColumnName.Reset();
+		AnimComboBox->ClearSelection();
+
+		break;
+	}
+	case EAnimColumnType::EACT_SlotColumn:
+	{
 		MeshColumnData->AnimSlotColumnName.Reset();
+		AnimSlotComboBox->ClearSelection();
+
+		break;
+	}
+	case EAnimColumnType::EACT_TagsColumn:
+	{
 		MeshColumnData->AnimTagColumnName.Reset();
+		AnimTagsComboBox->ClearSelection();
+
+		break;
+	}
+	default:
+		break;
+	}
+
+	Node->MarkPackageDirty();
+}
+
+
+// Layout Category --------------------------------------------------------------------------------
+
+void FCustomizableObjectNodeTableDetails::OnLayoutMeshCustomRowResetButtonClicked()
+{
+	if (LayoutMeshColumnComboBox.IsValid() && LayoutMeshColumnOptionNames.Num())
+	{
+		LayoutMeshColumnComboBox->SetSelectedItem(LayoutMeshColumnOptionNames[0]);
+		SelectedLayout = nullptr;
+		LayoutBlocksEditor->SetCurrentLayout(nullptr);
+	}
+}
+
+
+void FCustomizableObjectNodeTableDetails::OnLayoutMeshColumnComboBoxSelectionChanged(TSharedPtr<FString> Selection, ESelectInfo::Type SelectInfo)
+{
+	if (Selection.IsValid() && Selection != LayoutMeshColumnOptionNames[0])
+	{
+		FString ColumnName;
+		(*Selection).Split(" LOD_", &ColumnName, NULL);
 
-		if (AnimComboBox.IsValid())
+		for (const UEdGraphPin* Pin : Node->GetAllNonOrphanPins())
 		{
-			AnimComboBox->ClearSelection();
+			const UCustomizableObjectNodeTableMeshPinData* PinData = Cast<UCustomizableObjectNodeTableMeshPinData >(Node->GetPinData(*Pin));
+
+			if (PinData && PinData->ColumnName == ColumnName)
+			{
+				for (int32 LayoutIndex = 0; LayoutIndex < PinData->Layouts.Num(); ++LayoutIndex)
+				{
+					if (PinData->Layouts[LayoutIndex]->GetLayoutName() == *Selection)
+					{
+						LayoutBlocksEditor->SetCurrentLayout(PinData->Layouts[LayoutIndex]);
+						SelectedLayout = PinData->Layouts[LayoutIndex];
+
+						FillLayoutComboBoxOptions();
+					}
+				}
+			}
 		}
+	}
+	else
+	{
+		SelectedLayout = nullptr;
+		LayoutBlocksEditor->SetCurrentLayout(nullptr);
+	}
+}
+
+
+EVisibility FCustomizableObjectNodeTableDetails::LayoutOptionsVisibility() const
+{
+	return SelectedLayout.IsValid() ? EVisibility::Visible : EVisibility::Collapsed;
+}
 
-		if (AnimSlotComboBox.IsValid())
+
+EVisibility FCustomizableObjectNodeTableDetails::FixedStrategyOptionsVisibility() const
+{
+	return (SelectedLayout.IsValid() && SelectedLayout->GetPackingStrategy() == ECustomizableObjectTextureLayoutPackingStrategy::Fixed) ? EVisibility::Visible : EVisibility::Collapsed;
+}
+
+
+void FCustomizableObjectNodeTableDetails::FillLayoutComboBoxOptions()
+{
+	if (SelectedLayout.IsValid() && GridSizeComboBox.IsValid() && StrategyComboBox.IsValid()
+		&& MaxGridSizeComboBox.IsValid() && ReductionMethodComboBox.IsValid())
+	{
+		// Static const variable?
+		int32 MaxGridSize = 32;
+		LayoutGridSizes.Empty();
+
+		for (int32 Size = 1; Size <= MaxGridSize; Size *= 2)
+		{
+			LayoutGridSizes.Add(MakeShareable(new FString(FString::Printf(TEXT("%d x %d"), Size, Size))));
+
+			if (SelectedLayout->GetGridSize() == FIntPoint(Size))
+			{
+				GridSizeComboBox->SetSelectedItem(LayoutGridSizes.Last());
+			}
+
+			if (SelectedLayout->GetMaxGridSize() == FIntPoint(Size))
+			{
+				MaxGridSizeComboBox->SetSelectedItem(LayoutGridSizes.Last());
+			}
+		}
+
+		LayoutPackingStrategies.Empty();
+		LayoutPackingStrategies.Add(MakeShareable(new FString("Resizable")));
+		LayoutPackingStrategies.Add(MakeShareable(new FString("Fixed")));
+		LayoutPackingStrategies.Add(MakeShareable(new FString("Overlay")));
+		StrategyComboBox->SetSelectedItem(LayoutPackingStrategies[(uint32)SelectedLayout->GetPackingStrategy()]);
+
+		BlockReductionMethods.Empty();
+		BlockReductionMethods.Add(MakeShareable(new FString("Halve")));
+		BlockReductionMethods.Add(MakeShareable(new FString("Unitary")));
+		ReductionMethodComboBox->SetSelectedItem(BlockReductionMethods[(uint32)SelectedLayout->GetBlockReductionMethod()]);
+	}
+}
+
+
+void FCustomizableObjectNodeTableDetails::OnGridSizeChanged(TSharedPtr<FString> NewSelection, ESelectInfo::Type SelectInfo)
+{
+	if (SelectedLayout.IsValid())
+	{
+		int Size = 1 << LayoutGridSizes.Find(NewSelection);
+
+		if (SelectedLayout->GetGridSize().X != Size || SelectedLayout->GetGridSize().Y != Size)
 		{
-			AnimSlotComboBox->ClearSelection();
+			SelectedLayout->SetGridSize(FIntPoint(Size));
+
+			// Adjust all the blocks sizes
+			for (int b = 0; b < SelectedLayout->Blocks.Num(); ++b)
+			{
+				SelectedLayout->Blocks[b].Min.X = FMath::Min(SelectedLayout->Blocks[b].Min.X, Size - 1);
+				SelectedLayout->Blocks[b].Min.Y = FMath::Min(SelectedLayout->Blocks[b].Min.Y, Size - 1);
+				SelectedLayout->Blocks[b].Max.X = FMath::Min(SelectedLayout->Blocks[b].Max.X, Size);
+				SelectedLayout->Blocks[b].Max.Y = FMath::Min(SelectedLayout->Blocks[b].Max.Y, Size);
+			}
+
+			Node->MarkPackageDirty();
 		}
+	}
+}
+
+
+void FCustomizableObjectNodeTableDetails::OnLayoutPackingStrategyChanged(TSharedPtr<FString> NewSelection, ESelectInfo::Type SelectInfo)
+{
+	if (SelectedLayout.IsValid())
+	{
+		uint32 selection = LayoutPackingStrategies.IndexOfByKey(NewSelection);
 
-		if (AnimTagsComboBox.IsValid())
+		if (SelectedLayout->GetPackingStrategy() != (ECustomizableObjectTextureLayoutPackingStrategy)selection)
 		{
-			AnimTagsComboBox->ClearSelection();
+			SelectedLayout->SetPackingStrategy((ECustomizableObjectTextureLayoutPackingStrategy)selection);
+			Node->MarkPackageDirty();
 		}
+	}
+}
 
-		Node->MarkPackageDirty();
-		return FReply::Handled();
+
+void FCustomizableObjectNodeTableDetails::OnMaxGridSizeChanged(TSharedPtr<FString> NewSelection, ESelectInfo::Type SelectInfo)
+{
+	if (SelectedLayout.IsValid())
+	{
+		int Size = 1 << LayoutGridSizes.Find(NewSelection);
+
+		if (SelectedLayout->GetMaxGridSize().X != Size || SelectedLayout->GetMaxGridSize().Y != Size)
+		{
+			SelectedLayout->SetMaxGridSize(FIntPoint(Size));
+			SelectedLayout->MarkPackageDirty();
+		}
 	}
+}
+
+
+void FCustomizableObjectNodeTableDetails::OnReductionMethodChanged(TSharedPtr<FString> NewSelection, ESelectInfo::Type SelectInfo)
+{
+	if (SelectedLayout.IsValid())
+	{
+		uint32 selection = BlockReductionMethods.IndexOfByKey(NewSelection);
 
-	return FReply::Unhandled();
+		if (SelectedLayout->GetBlockReductionMethod() != (ECustomizableObjectLayoutBlockReductionMethod)selection)
+		{
+			SelectedLayout->SetBlockReductionMethod((ECustomizableObjectLayoutBlockReductionMethod)selection);
+			Node->MarkPackageDirty();
+		}
+	}
 }
 
 
+// Metadata Category --------------------------------------------------------------------------------
+
 TSharedPtr<FString> FCustomizableObjectNodeTableDetails::GenerateMutableMetaDataColumnComboBoxOptions()
 {
 	const UScriptStruct* TableStruct = Node->GetTableNodeStruct();
@@ -805,6 +894,7 @@ FSlateColor FCustomizableObjectNodeTableDetails::GetMetadataUIComboBoxTextColor(
 	return FSlateColor(FLinearColor(0.9f, 0.05f, 0.05f, 1.0f));
 }
 
+
 void FCustomizableObjectNodeTableDetails::OnMutableMetaDataColumnComboBoxSelectionReset()
 {
 	Node->ParamUIMetadataColumn = NAME_None;
@@ -818,6 +908,8 @@ void FCustomizableObjectNodeTableDetails::OnMutableMetaDataColumnComboBoxSelecti
 }
 
 
+// Version Bridge Category --------------------------------------------------------------------------------
+
 TSharedPtr<FString> FCustomizableObjectNodeTableDetails::GenerateVersionColumnComboBoxOptions()
 {
 	const UScriptStruct* TableStruct = Node->GetTableNodeStruct();
@@ -907,123 +999,4 @@ void FCustomizableObjectNodeTableDetails::OnVersionColumnComboBoxSelectionReset(
 }
 
 
-EVisibility FCustomizableObjectNodeTableDetails::LayoutOptionsVisibility() const
-{
-	return SelectedLayout.IsValid() ? EVisibility::Visible : EVisibility::Collapsed;
-}
-
-
-EVisibility FCustomizableObjectNodeTableDetails::FixedStrategyOptionsVisibility() const
-{
-	return (SelectedLayout.IsValid() && SelectedLayout->GetPackingStrategy() == ECustomizableObjectTextureLayoutPackingStrategy::Fixed) ? EVisibility::Visible : EVisibility::Collapsed;
-}
-
-
-void FCustomizableObjectNodeTableDetails::FillLayoutComboBoxOptions()
-{
-	if (SelectedLayout.IsValid() && GridSizeComboBox.IsValid() && StrategyComboBox.IsValid()
-		&& MaxGridSizeComboBox.IsValid() && ReductionMethodComboBox.IsValid())
-	{
-		// Static const variable?
-		int32 MaxGridSize = 32;
-		LayoutGridSizes.Empty();
-
-		for (int32 Size = 1; Size <= MaxGridSize; Size *= 2)
-		{
-			LayoutGridSizes.Add(MakeShareable(new FString(FString::Printf(TEXT("%d x %d"), Size, Size))));
-
-			if (SelectedLayout->GetGridSize() == FIntPoint(Size))
-			{
-				GridSizeComboBox->SetSelectedItem(LayoutGridSizes.Last());
-			}
-
-			if (SelectedLayout->GetMaxGridSize() == FIntPoint(Size))
-			{
-				MaxGridSizeComboBox->SetSelectedItem(LayoutGridSizes.Last());
-			}
-		}
-
-		LayoutPackingStrategies.Empty();
-		LayoutPackingStrategies.Add(MakeShareable(new FString("Resizable")));
-		LayoutPackingStrategies.Add(MakeShareable(new FString("Fixed")));
-		LayoutPackingStrategies.Add(MakeShareable(new FString("Overlay")));
-		StrategyComboBox->SetSelectedItem(LayoutPackingStrategies[(uint32)SelectedLayout->GetPackingStrategy()]);
-
-		BlockReductionMethods.Empty();
-		BlockReductionMethods.Add(MakeShareable(new FString("Halve")));
-		BlockReductionMethods.Add(MakeShareable(new FString("Unitary")));
-		ReductionMethodComboBox->SetSelectedItem(BlockReductionMethods[(uint32)SelectedLayout->GetBlockReductionMethod()]);
-	}
-}
-
-
-void FCustomizableObjectNodeTableDetails::OnGridSizeChanged(TSharedPtr<FString> NewSelection, ESelectInfo::Type SelectInfo)
-{
-	if (SelectedLayout.IsValid())
-	{
-		int Size = 1 << LayoutGridSizes.Find(NewSelection);
-
-		if (SelectedLayout->GetGridSize().X != Size || SelectedLayout->GetGridSize().Y != Size)
-		{
-			SelectedLayout->SetGridSize(FIntPoint(Size));
-
-			// Adjust all the blocks sizes
-			for (int b = 0; b < SelectedLayout->Blocks.Num(); ++b)
-			{
-				SelectedLayout->Blocks[b].Min.X = FMath::Min(SelectedLayout->Blocks[b].Min.X, Size - 1);
-				SelectedLayout->Blocks[b].Min.Y = FMath::Min(SelectedLayout->Blocks[b].Min.Y, Size - 1);
-				SelectedLayout->Blocks[b].Max.X = FMath::Min(SelectedLayout->Blocks[b].Max.X, Size);
-				SelectedLayout->Blocks[b].Max.Y = FMath::Min(SelectedLayout->Blocks[b].Max.Y, Size);
-			}
-
-			Node->MarkPackageDirty();
-		}
-	}
-}
-
-
-void FCustomizableObjectNodeTableDetails::OnLayoutPackingStrategyChanged(TSharedPtr<FString> NewSelection, ESelectInfo::Type SelectInfo)
-{
-	if (SelectedLayout.IsValid())
-	{
-		uint32 selection = LayoutPackingStrategies.IndexOfByKey(NewSelection);
-
-		if (SelectedLayout->GetPackingStrategy() != (ECustomizableObjectTextureLayoutPackingStrategy)selection)
-		{
-			SelectedLayout->SetPackingStrategy((ECustomizableObjectTextureLayoutPackingStrategy)selection);
-			Node->MarkPackageDirty();
-		}
-	}
-}
-
-
-void FCustomizableObjectNodeTableDetails::OnMaxGridSizeChanged(TSharedPtr<FString> NewSelection, ESelectInfo::Type SelectInfo)
-{
-	if (SelectedLayout.IsValid())
-	{
-		int Size = 1 << LayoutGridSizes.Find(NewSelection);
-
-		if (SelectedLayout->GetMaxGridSize().X != Size || SelectedLayout->GetMaxGridSize().Y != Size)
-		{
-			SelectedLayout->SetMaxGridSize(FIntPoint(Size));
-			SelectedLayout->MarkPackageDirty();
-		}
-	}
-}
-
-
-void FCustomizableObjectNodeTableDetails::OnReductionMethodChanged(TSharedPtr<FString> NewSelection, ESelectInfo::Type SelectInfo)
-{
-	if (SelectedLayout.IsValid())
-	{
-		uint32 selection = BlockReductionMethods.IndexOfByKey(NewSelection);
-
-		if (SelectedLayout->GetBlockReductionMethod() != (ECustomizableObjectLayoutBlockReductionMethod)selection)
-		{
-			SelectedLayout->SetBlockReductionMethod((ECustomizableObjectLayoutBlockReductionMethod)selection);
-			Node->MarkPackageDirty();
-		}
-	}
-}
-
 #undef LOCTEXT_NAMESPACE
diff --git a/Engine/Plugins/Experimental/Mutable/Source/CustomizableObjectEditor/Private/MuCOE/Nodes/CustomizableObjectNodeTableDetails.h b/Engine/Plugins/Experimental/Mutable/Source/CustomizableObjectEditor/Private/MuCOE/Nodes/CustomizableObjectNodeTableDetails.h
--- a/Engine/Plugins/Experimental/Mutable/Source/CustomizableObjectEditor/Private/MuCOE/Nodes/CustomizableObjectNodeTableDetails.h
+++ b/Engine/Plugins/Experimental/Mutable/Source/CustomizableObjectEditor/Private/MuCOE/Nodes/CustomizableObjectNodeTableDetails.h
@@ -20,6 +20,14 @@
 struct EVisibility;
 struct FSlateColor;
 
+enum class EAnimColumnType 
+{
+	EACT_BluePrintColumn,
+	EACT_SlotColumn,
+	EACT_TagsColumn
+};
+
+
 /** Copy Material node details panel. Hides all properties from the inheret Material node. */
 class FCustomizableObjectNodeTableDetails : public IDetailCustomization
 {
@@ -36,27 +44,21 @@
 
 private:
 
+	// Details -------------
 	// Generates Mesh columns combobox options
 	void GenerateMeshColumnComboBoxOptions();
+	
+	// Function called when the table node has been refreshed
+	void OnNodePinValueChanged();
 
+
+	// Anim Category -------------
 	// Generates Animation Instance combobox options
 	void GenerateAnimInstanceComboBoxOptions();
 
-	// Generates Animation Slot combobox options
-	void GenerateAnimSlotComboBoxOptions();
-
-	// Generates Animation Tags combobox options
-	void GenerateAnimTagsComboBoxOptions();
-
-	// OnComboBoxSelectionChanged Callback for Column ComboBox
-	void OnColumnComboBoxSelectionChanged(TSharedPtr<FString> Selection, ESelectInfo::Type SelectInfo);
-
 	// OnComboBoxSelectionChanged Callback for Layout ComboBox
 	void OnAnimMeshColumnComboBoxSelectionChanged(TSharedPtr<FString> Selection, ESelectInfo::Type SelectInfo);
 
-	// OnComboBoxSelectionChanged Callback for Layout ComboBox
-	void OnLayoutMeshColumnComboBoxSelectionChanged(TSharedPtr<FString> Selection, ESelectInfo::Type SelectInfo);
-
 	// OnComboBoxSelectionChanged Callback for AnimInstance ComboBox
 	void OnAnimInstanceComboBoxSelectionChanged(TSharedPtr<FString> Selection, ESelectInfo::Type SelectInfo);
 
@@ -66,19 +68,44 @@
 	// OnComboBoxSelectionChanged Callback for Anim Tags ComboBox
 	void OnAnimTagsComboBoxSelectionChanged(TSharedPtr<FString> Selection, ESelectInfo::Type SelectInfo);
 
-	// Function called when the table node has been refreshed
-	void OnNodePinValueChanged();
+	// Set the visibility of the animation column selector widgets
+	EVisibility AnimWidgetsVisibility() const;
 
+	// Callback to clear the mesh column to edit its animation properties
+	void OnAnimMeshCustomRowResetButtonClicked();
+
 	// Callback to clear the animation combobox selections
-	FReply OnClearButtonPressed();
+	void OnAnimCustomRowResetButtonClicked(EAnimColumnType ColumnType);
+
+
+	// Layout Category -------------
+	// OnComboBoxSelectionChanged Callback for Layout ComboBox
+	void OnLayoutMeshColumnComboBoxSelectionChanged(TSharedPtr<FString> Selection, ESelectInfo::Type SelectInfo);
+
+	// Callback to clear the mesh column to edit its layout properties
+	void OnLayoutMeshCustomRowResetButtonClicked();
+
+	// Layout options visibility
+	EVisibility LayoutOptionsVisibility() const;
+
+	/** Returns the visibility of the Fixed layout widgets */
+	EVisibility FixedStrategyOptionsVisibility() const;
+
+	/** Fills the combo box arrays sources */
+	void FillLayoutComboBoxOptions();
+
+	/** Layout Options Callbacks */
+	void OnGridSizeChanged(TSharedPtr<FString> NewSelection, ESelectInfo::Type SelectInfo);
+	void OnMaxGridSizeChanged(TSharedPtr<FString> NewSelection, ESelectInfo::Type SelectInfo);
+	void OnReductionMethodChanged(TSharedPtr<FString> NewSelection, ESelectInfo::Type SelectInfo);
+	void OnLayoutPackingStrategyChanged(TSharedPtr<FString> NewSelection, ESelectInfo::Type SelectInfo);
+
+
+	// Mutable UI Metadata Category -------------
 
 	// Generates MutableMetadata columns combobox options
 	// Returns the current selected option or a null pointer
 	TSharedPtr<FString> GenerateMutableMetaDataColumnComboBoxOptions();
 
-	// Generates MutableMetadata columns combobox options
-	// Returns the current selected option or a null pointer
-	TSharedPtr<FString> GenerateVersionColumnComboBoxOptions();
-
 	// Callback to regenerate the combobox options
 	void OnOpenMutableMetadataComboBox();
@@ -92,6 +119,13 @@
 	// OnComboBoxSelectionChanged Callback for Layout ComboBox
 	void OnMutableMetaDataColumnComboBoxSelectionReset();
 
+
+	// Version Control Category -------------
+
+	// Generates MutableMetadata columns combobox options
+	// Returns the current selected option or a null pointer
+	TSharedPtr<FString> GenerateVersionColumnComboBoxOptions();
+
 	// Callback to regenerate the combobox options
 	void OnOpenVersionColumnComboBox();
 	
@@ -104,38 +138,23 @@
 	// OnComboBoxSelectionChanged Callback for Layout ComboBox
 	void OnVersionColumnComboBoxSelectionReset();
 
-	// Layout options visibility
-	EVisibility LayoutOptionsVisibility() const;
-
-	/** Returns the visibility of the Fixed layout widgets */
-	EVisibility FixedStrategyOptionsVisibility() const;
-
-	/** Fills the combo box arrays sources */
-	void FillLayoutComboBoxOptions();
-
-	/** Layout Options Callbacks */
-	void OnGridSizeChanged(TSharedPtr<FString> NewSelection, ESelectInfo::Type SelectInfo);
-	void OnMaxGridSizeChanged(TSharedPtr<FString> NewSelection, ESelectInfo::Type SelectInfo);
-	void OnReductionMethodChanged(TSharedPtr<FString> NewSelection, ESelectInfo::Type SelectInfo);
-	void OnLayoutPackingStrategyChanged(TSharedPtr<FString> NewSelection, ESelectInfo::Type SelectInfo);
-
 private:
 
+	// Details -------------
 	// Pointer to the node represented in this details
 	TWeakObjectPtr<UCustomizableObjectNodeTable> Node;
 
+	// Pointer to the Detail Builder to force the refresh on recontruct the node
+	TWeakPtr<IDetailLayoutBuilder> DetailBuilderPtr = nullptr;
+
+
+	// Anim -------------
 	// ComboBox widget to select a column from the NodeTable
-	TSharedPtr<STextComboBox> LayoutMeshColumnComboBox;
-
-	// ComboBox widget to select a column from the NodeTable
 	TSharedPtr<STextComboBox> AnimMeshColumnComboBox;
 
 	// Array with the name of the table columns as combobox options
 	TArray<TSharedPtr<FString>> AnimMeshColumnOptionNames;
 
-	// Array with the name of the table columns as combobox options
-	TArray<TSharedPtr<FString>> LayoutMeshColumnOptionNames;
-
 	// ComboBox widget to select an Animation Instance column from the NodeTable
 	TSharedPtr<STextComboBox> AnimComboBox;
 
@@ -153,30 +172,20 @@
 
 	// Array with the name of the Animation Tags columns as combobox options
 	TArray<TSharedPtr<FString>> AnimTagsOptionNames;
+
+	// Layout -------------
+	// ComboBox widget to select a column from the NodeTable
+	TSharedPtr<STextComboBox> LayoutMeshColumnComboBox;
 	
-	// Array with the name of the MutableMetaData columns
-	TArray<TSharedPtr<FString>> MutableMetaDataColumnsOptionNames;
-
-	// ComboBox widget to select a MutableMetaDatacolumn from the NodeTable
-	TSharedPtr<STextComboBox> MutableMetaDataComboBox;
-
-	// Array with the name of the Version columns
-	TArray<TSharedPtr<FString>> VersionColumnsOptionNames;
+	// Array with the name of the table columns as combobox options
+	TArray<TSharedPtr<FString>> LayoutMeshColumnOptionNames;
 	
-	// ComboBox widget to select a VersionColumn from the NodeTable
-	TSharedPtr<STextComboBox> VersionColumnsComboBox;
-	
-	// Button to clear selections of the animation comboboxes
-	TSharedPtr<SButton> ClearButton;
-
 	// Layout block editor widget
 	TSharedPtr<SCustomizableObjectNodeLayoutBlocksEditor> LayoutBlocksEditor;
 
+	// Weakpointer to the selected layout objet
 	TWeakObjectPtr<UCustomizableObjectLayout> SelectedLayout;
 
-	// Pointer to the Detail Builder to force the refresh on recontruct the node
-	TWeakPtr<IDetailLayoutBuilder> DetailBuilderPtr = nullptr;
-
 	/** List of available layout grid sizes. */
 	TArray< TSharedPtr< FString > > LayoutGridSizes;
 
@@ -194,4 +203,21 @@
 	TSharedPtr<STextComboBox> MaxGridSizeComboBox;
 	// ComboBox widget to select a Reduction Method from the Selected Layout
 	TSharedPtr<STextComboBox> ReductionMethodComboBox;
+
+
+	// Mutable UI Metadata -------------
+	// Array with the name of the MutableMetaData columns
+	TArray<TSharedPtr<FString>> MutableMetaDataColumnsOptionNames;
+
+	// ComboBox widget to select a MutableMetaDatacolumn from the NodeTable
+	TSharedPtr<STextComboBox> MutableMetaDataComboBox;
+
+
+	// Version Bridge -------------
+	// Array with the name of the Version columns
+	TArray<TSharedPtr<FString>> VersionColumnsOptionNames;
+	
+	// ComboBox widget to select a VersionColumn from the NodeTable
+	TSharedPtr<STextComboBox> VersionColumnsComboBox;
+
 };
