[mutable] Added a new combobox option to default values of group nodes. This option is automatically selected when the user has not specified a default value.

#rnx

diff --git a/Engine/Plugins/Experimental/Mutable/Source/CustomizableObjectEditor/Private/MuCOE/CustomizableObjectNodeObjectGroupDetails.cpp b/Engine/Plugins/Experimental/Mutable/Source/CustomizableObjectEditor/Private/MuCOE/CustomizableObjectNodeObjectGroupDetails.cpp
--- a/Engine/Plugins/Experimental/Mutable/Source/CustomizableObjectEditor/Private/MuCOE/CustomizableObjectNodeObjectGroupDetails.cpp
+++ b/Engine/Plugins/Experimental/Mutable/Source/CustomizableObjectEditor/Private/MuCOE/CustomizableObjectNodeObjectGroupDetails.cpp
@@ -96,6 +96,7 @@ void FCustomizableObjectNodeObjectGroupDetails::CustomizeDetails(IDetailLayoutBu
 			.OnComboBoxOpening(this, &FCustomizableObjectNodeObjectGroupDetails::GenerateChildrenObjectNames)
 			.OnSelectionChanged(this, &FCustomizableObjectNodeObjectGroupDetails::OnSetDefaultValue)
 			.Font(DetailBuilder.GetDetailFont())
+			.ToolTipText(this, &FCustomizableObjectNodeObjectGroupDetails::DefaultValueComboBoxTooltip)
 		];
 	}
 }
@@ -106,11 +107,17 @@ void FCustomizableObjectNodeObjectGroupDetails::GenerateChildrenObjectNames()
 	ChildrenNameOptions.Reset();
 	InitialNameOption = nullptr;
 
+	// if needed add a None option and set it as the default value
 	if (NodeGroup->GroupType == ECustomizableObjectGroupType::COGT_ONE_OR_NONE)
 	{
 		ChildrenNameOptions.Add(MakeShareable(new FString("None")));
 		InitialNameOption = ChildrenNameOptions.Last();
 	}
+	else
+	{
+		ChildrenNameOptions.Add(MakeShareable(new FString("- Not Selected -")));
+		InitialNameOption = ChildrenNameOptions.Last();
+	}
 
 	// Adding linked children names
 	const TArray<UEdGraphPin*> ConnectedChildrenPins = FollowInputPinArray(*NodeGroup->ObjectsPin());
@@ -164,11 +171,6 @@ void FCustomizableObjectNodeObjectGroupDetails::GenerateChildrenObjectNames()
 		}
 	}
 
-	if (ChildrenNameOptions.Num() && InitialNameOption == nullptr)
-	{
-		InitialNameOption = ChildrenNameOptions[0];
-	}
-
 	if (DefaultValueSelector.IsValid())
 	{
 		DefaultValueSelector->RefreshOptions();
@@ -196,4 +198,22 @@ EVisibility FCustomizableObjectNodeObjectGroupDetails::DefaultValueSelectorVisib
 	return EVisibility::Hidden;
 }
 
+
+FText FCustomizableObjectNodeObjectGroupDetails::DefaultValueComboBoxTooltip() const
+{
+	if (DefaultValueSelector.IsValid())
+	{
+		if (ChildrenNameOptions.Num() && DefaultValueSelector->GetSelectedItem() == ChildrenNameOptions[0] && NodeGroup->GroupType == ECustomizableObjectGroupType::COGT_ONE)
+		{
+			return FText::FromString("When nothing selected, the first compiled option will be used as the default value.");
+		}
+		else if (DefaultValueSelector->GetSelectedItem().IsValid())
+		{
+			return FText::FromString(*DefaultValueSelector->GetSelectedItem());
+		}
+	}
+
+	return FText();
+}
+
 #undef LOCTEXT_NAMESPACE
diff --git a/Engine/Plugins/Experimental/Mutable/Source/CustomizableObjectEditor/Private/MuCOE/CustomizableObjectNodeObjectGroupDetails.h b/Engine/Plugins/Experimental/Mutable/Source/CustomizableObjectEditor/Private/MuCOE/CustomizableObjectNodeObjectGroupDetails.h
--- a/Engine/Plugins/Experimental/Mutable/Source/CustomizableObjectEditor/Private/MuCOE/CustomizableObjectNodeObjectGroupDetails.h
+++ b/Engine/Plugins/Experimental/Mutable/Source/CustomizableObjectEditor/Private/MuCOE/CustomizableObjectNodeObjectGroupDetails.h
@@ -28,6 +28,8 @@ public:
 	/** Determines the visibility of the default value selector */
 	EVisibility DefaultValueSelectorVisibility() const;
 
+	FText DefaultValueComboBoxTooltip() const;
+
 private:
 
 	/** Pointer to the original group node */
