[Mutable] Fix CO stuck in Compiling state.

#rnx

diff --git a/Engine/Plugins/Experimental/Mutable/Source/CustomizableObjectEditor/Private/MuCOE/CustomizableObjectCompiler.cpp b/Engine/Plugins/Experimental/Mutable/Source/CustomizableObjectEditor/Private/MuCOE/CustomizableObjectCompiler.cpp
--- a/Engine/Plugins/Experimental/Mutable/Source/CustomizableObjectEditor/Private/MuCOE/CustomizableObjectCompiler.cpp
+++ b/Engine/Plugins/Experimental/Mutable/Source/CustomizableObjectEditor/Private/MuCOE/CustomizableObjectCompiler.cpp
@@ -951,6 +951,8 @@ void FCustomizableObjectCompiler::CompileInternal(bool bAsync)
 			CompilerLog(FText::FromString(TEXT("Failed to generate the mutable node graph. Object not built.")), nullptr);
 		}
 
+		SetCompilationState(ECompilationStatePrivate::Completed, GetCompilationResult());
+		
 		if (CurrentObject->GetPrivate()->IsLocked())
 		{
 			UCustomizableObjectSystem::GetInstance()->UnlockObject(CurrentObject);
