[Mutable] Move FCustomizableObjectInstanceDescriptor to private.

Descriptor being public is usafe since it can be serialized without us being able to execute backwards compatible code. Since there are no real benefits of external users using the descriptor, it should be kept private.

- Added IsParameterDirty.

#rnx

diff --git a/Engine/Plugins/Experimental/Mutable/Source/CustomizableObject/Classes/MuCO/CustomizableObjectInstance.h b/Engine/Plugins/Experimental/Mutable/Source/CustomizableObject/Classes/MuCO/CustomizableObjectInstance.h
--- a/Engine/Plugins/Experimental/Mutable/Source/CustomizableObject/Classes/MuCO/CustomizableObjectInstance.h
+++ b/Engine/Plugins/Experimental/Mutable/Source/CustomizableObject/Classes/MuCO/CustomizableObjectInstance.h
@@ -335,6 +335,10 @@ public:
 	UFUNCTION(BlueprintCallable, Category = CustomizableObjectInstance)
 	void SetRandomValuesFromStream(const FRandomStream& InStream);
 
+	/** Set all parameters to their default value. */
+	UFUNCTION(BlueprintCallable, Category = CustomizableObjectInstance)
+	void SetDefaultValues();
+	
 	/** Returns the AssetUserData that was gathered from all the constituent mesh parts during the last update. 
 	  * It requires that the CustomizableObject had the bEnableAssetUserDataMerge set to true during compilation. */
 	UFUNCTION(BlueprintCallable, Category = CustomizableObjectInstance)
@@ -375,6 +379,10 @@ public:
 	UFUNCTION(BlueprintCallable, Category = CustomizableObjectInstance)
 	bool IsParameterRelevant(const FString& ParamName) const;
 
+	/** Return true if the parameter has changed but the Instance has not yet been updated. */
+	UFUNCTION(BlueprintCallable, Category = CustomizableObjectInstance)
+	bool IsParameterDirty(const FString& ParamName, int32 RangeIndex = -1) const;
+	
 	/** For multidimensional parameters, return the number of dimensions that the given projector parameter supports. */
 	UFUNCTION(BlueprintCallable, Category = CustomizableObjectInstance)
 	int32 GetProjectorValueRange(const FString& ParamName) const;
diff --git a/Engine/Plugins/Experimental/Mutable/Source/CustomizableObject/Private/MuCO/CustomizableObjectInstance.cpp b/Engine/Plugins/Experimental/Mutable/Source/CustomizableObject/Private/MuCO/CustomizableObjectInstance.cpp
--- a/Engine/Plugins/Experimental/Mutable/Source/CustomizableObject/Private/MuCO/CustomizableObjectInstance.cpp
+++ b/Engine/Plugins/Experimental/Mutable/Source/CustomizableObject/Private/MuCO/CustomizableObjectInstance.cpp
@@ -725,6 +725,220 @@ bool UCustomizableObjectInstance::IsParameterRelevant(const FString& ParamName)
 }
 
 
+bool UCustomizableObjectInstance::IsParameterDirty(const FString& ParamName, const int32 RangeIndex) const
+{
+	switch (Descriptor.CustomizableObject->GetParameterTypeByName(ParamName))
+	{
+	case EMutableParameterType::None:
+		return false;
+
+	case EMutableParameterType::Projector:
+		{
+			const FCustomizableObjectProjectorParameterValue* Result = Descriptor.GetProjectorParameters().FindByPredicate([&](const FCustomizableObjectProjectorParameterValue& Value)
+				{
+					return Value.ParameterName == ParamName;
+				});
+			
+			const FCustomizableObjectProjectorParameterValue* ResultCommited = GetPrivate()->CommittedDescriptor.GetProjectorParameters().FindByPredicate([&](const FCustomizableObjectProjectorParameterValue& Value)
+				{
+					return Value.ParameterName == ParamName;
+				});
+			
+			if (Result && ResultCommited)
+			{
+				if (RangeIndex != INDEX_NONE)
+				{
+					return Result->Value == ResultCommited->Value;					
+				}
+				else
+				{
+					if (Result->RangeValues.IsValidIndex(RangeIndex) && ResultCommited->RangeValues.IsValidIndex(RangeIndex))
+					{
+						return Result->RangeValues[RangeIndex] == ResultCommited->RangeValues[RangeIndex];
+					}
+					else
+					{
+						return Result->RangeValues.Num() != ResultCommited->RangeValues.Num();
+					}
+				}
+			}
+			else
+			{
+				return Result != ResultCommited;
+			}
+		}		
+	case EMutableParameterType::Texture:
+		{
+			const FCustomizableObjectTextureParameterValue* Result = Descriptor.GetTextureParameters().FindByPredicate([&](const FCustomizableObjectTextureParameterValue& Value)
+				{
+					return Value.ParameterName == ParamName;
+				});
+			
+			const FCustomizableObjectTextureParameterValue* ResultCommited = GetPrivate()->CommittedDescriptor.GetTextureParameters().FindByPredicate([&](const FCustomizableObjectTextureParameterValue& Value)
+				{
+					return Value.ParameterName == ParamName;
+				});
+			
+			if (Result && ResultCommited)
+			{
+				if (RangeIndex != INDEX_NONE)
+				{
+					return Result->ParameterValue == ResultCommited->ParameterValue;					
+				}
+				else
+				{
+					if (Result->ParameterRangeValues.IsValidIndex(RangeIndex) && ResultCommited->ParameterRangeValues.IsValidIndex(RangeIndex))
+					{
+						return Result->ParameterRangeValues[RangeIndex] == ResultCommited->ParameterRangeValues[RangeIndex];
+					}
+					else
+					{
+						return Result->ParameterRangeValues.Num() != ResultCommited->ParameterRangeValues.Num();
+					}
+				}
+			}
+			else
+			{
+				return Result != ResultCommited;
+			}
+		}
+
+	case EMutableParameterType::Bool:
+		{
+			const FCustomizableObjectBoolParameterValue* Result = Descriptor.GetBoolParameters().FindByPredicate([&](const FCustomizableObjectBoolParameterValue& Value)
+				{
+					return Value.ParameterName == ParamName;
+				});
+			
+			const FCustomizableObjectBoolParameterValue* ResultCommited = GetPrivate()->CommittedDescriptor.GetBoolParameters().FindByPredicate([&](const FCustomizableObjectBoolParameterValue& Value)
+				{
+					return Value.ParameterName == ParamName;
+				});
+			
+			if (Result && ResultCommited)
+			{
+				if (RangeIndex != INDEX_NONE)
+				{
+					return Result->ParameterValue == ResultCommited->ParameterValue;					
+				}
+				else
+				{
+					return false;
+				}
+			}
+			else
+			{
+				return Result != ResultCommited;
+			}
+		}
+	case EMutableParameterType::Int:
+		{
+			const FCustomizableObjectIntParameterValue* Result = Descriptor.GetIntParameters().FindByPredicate([&](const FCustomizableObjectIntParameterValue& Value)
+				{
+					return Value.ParameterName == ParamName;
+				});
+			
+			const FCustomizableObjectIntParameterValue* ResultCommited = GetPrivate()->CommittedDescriptor.GetIntParameters().FindByPredicate([&](const FCustomizableObjectIntParameterValue& Value)
+				{
+					return Value.ParameterName == ParamName;
+				});
+			
+			if (Result && ResultCommited)
+			{
+				if (RangeIndex != INDEX_NONE)
+				{
+					return Result->ParameterValueName == ResultCommited->ParameterValueName;					
+				}
+				else
+				{
+					if (Result->ParameterRangeValueNames.IsValidIndex(RangeIndex) && ResultCommited->ParameterRangeValueNames.IsValidIndex(RangeIndex))
+					{
+						return Result->ParameterRangeValueNames[RangeIndex] == ResultCommited->ParameterRangeValueNames[RangeIndex];
+					}
+					else
+					{
+						return Result->ParameterRangeValueNames.Num() != ResultCommited->ParameterRangeValueNames.Num();
+					}
+				}
+			}
+			else
+			{
+				return Result != ResultCommited;
+			}
+		}
+		
+	case EMutableParameterType::Float:
+		{
+			const FCustomizableObjectFloatParameterValue* Result = Descriptor.GetFloatParameters().FindByPredicate([&](const FCustomizableObjectFloatParameterValue& Value)
+				{
+					return Value.ParameterName == ParamName;
+				});
+			
+			const FCustomizableObjectFloatParameterValue* ResultCommited = GetPrivate()->CommittedDescriptor.GetFloatParameters().FindByPredicate([&](const FCustomizableObjectFloatParameterValue& Value)
+				{
+					return Value.ParameterName == ParamName;
+				});
+			
+			if (Result && ResultCommited)
+			{
+				if (RangeIndex != INDEX_NONE)
+				{
+					return Result->ParameterValue == ResultCommited->ParameterValue;					
+				}
+				else
+				{
+					if (Result->ParameterRangeValues.IsValidIndex(RangeIndex) && ResultCommited->ParameterRangeValues.IsValidIndex(RangeIndex))
+					{
+						return Result->ParameterRangeValues[RangeIndex] == ResultCommited->ParameterRangeValues[RangeIndex];
+					}
+					else
+					{
+						return Result->ParameterRangeValues.Num() != ResultCommited->ParameterRangeValues.Num();
+					}
+				}
+			}
+			else
+			{
+				return Result != ResultCommited;
+			}
+		}
+		
+	case EMutableParameterType::Color:
+		{
+			const FCustomizableObjectVectorParameterValue* Result = Descriptor.GetVectorParameters().FindByPredicate([&](const FCustomizableObjectVectorParameterValue& Value)
+				{
+					return Value.ParameterName == ParamName;
+				});
+			
+			const FCustomizableObjectVectorParameterValue* ResultCommited = GetPrivate()->CommittedDescriptor.GetVectorParameters().FindByPredicate([&](const FCustomizableObjectVectorParameterValue& Value)
+				{
+					return Value.ParameterName == ParamName;
+				});
+			
+			if (Result && ResultCommited)
+			{
+				if (RangeIndex != INDEX_NONE)
+				{
+					return Result->ParameterValue == ResultCommited->ParameterValue;					
+				}
+				else
+				{
+					return false;
+				}
+			}
+			else
+			{
+				return Result != ResultCommited;
+			}
+		}
+		
+	default:
+		unimplemented();
+		return false;
+	}
+}
+
+
 void UCustomizableInstancePrivate::PostEditChangePropertyWithoutEditor()
 {
 	MUTABLE_CPUPROFILER_SCOPE(UCustomizableInstancePrivate::PostEditChangePropertyWithoutEditor);
@@ -2481,6 +2695,12 @@ void UCustomizableObjectInstance::SetRandomValuesFromStream(const FRandomStream&
 }
 
 
+void UCustomizableObjectInstance::SetDefaultValues()
+{
+	Descriptor.SetDefaultValues();
+}
+
+
 bool UCustomizableInstancePrivate::LoadParametersFromProfile(int32 ProfileIndex)
 {
 	UCustomizableObject* CustomizableObject = GetPublic()->GetCustomizableObject();
diff --git a/Engine/Plugins/Experimental/Mutable/Source/CustomizableObject/Private/MuCO/CustomizableObjectInstanceDescriptor.cpp b/Engine/Plugins/Experimental/Mutable/Source/CustomizableObject/Private/MuCO/CustomizableObjectInstanceDescriptor.cpp
--- a/Engine/Plugins/Experimental/Mutable/Source/CustomizableObject/Private/MuCO/CustomizableObjectInstanceDescriptor.cpp
+++ b/Engine/Plugins/Experimental/Mutable/Source/CustomizableObject/Private/MuCO/CustomizableObjectInstanceDescriptor.cpp
@@ -2335,6 +2335,63 @@ void FCustomizableObjectInstanceDescriptor::SetRandomValuesFromStream(const FRan
 }
 
 
+void FCustomizableObjectInstanceDescriptor::SetDefaultValues()
+{
+	check(CustomizableObject);
+
+	for (FCustomizableObjectBoolParameterValue& Value : BoolParameters)
+	{	
+		Value.ParameterValue =	CustomizableObject->GetBoolParameterDefaultValue(Value.ParameterName);
+	}
+
+	for (FCustomizableObjectIntParameterValue& Value : IntParameters)
+	{
+		const int32 ParameterIndex = CustomizableObject->FindParameter(Value.ParameterName);
+		Value.ParameterValueName = CustomizableObject->FindIntParameterValueName(ParameterIndex, CustomizableObject->GetIntParameterDefaultValue(Value.ParameterValueName));
+
+		for (FString& RangeValue : Value.ParameterRangeValueNames)
+		{
+			RangeValue = Value.ParameterValueName;
+		}
+	}
+
+	for (FCustomizableObjectFloatParameterValue& Value : FloatParameters)
+	{
+		Value.ParameterValue = CustomizableObject->GetFloatParameterDefaultValue(Value.ParameterName);
+
+		for (float& RangeValue : Value.ParameterRangeValues)
+		{
+			RangeValue = Value.ParameterValue;
+		}
+	}
+
+	for (FCustomizableObjectTextureParameterValue& Value : TextureParameters)
+	{
+		Value.ParameterValue = CustomizableObject->GetTextureParameterDefaultValue(Value.ParameterName);
+
+		for (FName& RangeValue : Value.ParameterRangeValues)
+		{
+			RangeValue = Value.ParameterValue;
+		}
+	}
+
+	for (FCustomizableObjectVectorParameterValue& Value : VectorParameters)
+	{
+		Value.ParameterValue = CustomizableObject->GetColorParameterDefaultValue(Value.ParameterName);
+	}
+
+	for (FCustomizableObjectProjectorParameterValue& Value : ProjectorParameters)
+	{
+		Value.Value = CustomizableObject->GetProjectorParameterDefaultValue(Value.ParameterName);
+
+		for (FCustomizableObjectProjector& RangeValue : Value.RangeValues)
+		{
+			RangeValue = Value.Value;
+		}
+	}
+}
+
+
 bool FCustomizableObjectInstanceDescriptor::CreateMultiLayerProjector(const FName& ProjectorParamName)
 {
 	if (!FMultilayerProjector::AreDescriptorParametersValid(*this, ProjectorParamName.ToString()))
diff --git a/Engine/Plugins/Experimental/Mutable/Source/CustomizableObject/Public/MuCO/CustomizableObjectInstanceDescriptor.h b/Engine/Plugins/Experimental/Mutable/Source/CustomizableObject/Public/MuCO/CustomizableObjectInstanceDescriptor.h
--- a/Engine/Plugins/Experimental/Mutable/Source/CustomizableObject/Public/MuCO/CustomizableObjectInstanceDescriptor.h
+++ b/Engine/Plugins/Experimental/Mutable/Source/CustomizableObject/Public/MuCO/CustomizableObjectInstanceDescriptor.h
@@ -296,6 +296,8 @@ struct CUSTOMIZABLEOBJECT_API FCustomizableObjectInstanceDescriptor
 	
 	void SetRandomValuesFromStream(const FRandomStream& InStream);
 
+	void SetDefaultValues();
+	
 	// ------------------------------------------------------------
 	// Multilayer Projectors
 	// ------------------------------------------------------------
diff --git a/Engine/Plugins/Experimental/Mutable/Source/CustomizableObject/Public/MuCO/CustomizableObjectParameterTypeDefinitions.h b/Engine/Plugins/Experimental/Mutable/Source/CustomizableObject/Public/MuCO/CustomizableObjectParameterTypeDefinitions.h
--- a/Engine/Plugins/Experimental/Mutable/Source/CustomizableObject/Public/MuCO/CustomizableObjectParameterTypeDefinitions.h
+++ b/Engine/Plugins/Experimental/Mutable/Source/CustomizableObject/Public/MuCO/CustomizableObjectParameterTypeDefinitions.h
@@ -248,6 +248,8 @@ struct CUSTOMIZABLEOBJECT_API FCustomizableObjectProjector
 	// Just for cylindrical projectors, in radians
 	UPROPERTY()
 	float Angle = 2.0f * PI;
+
+	bool operator==(const FCustomizableObjectProjector& Other) const = default;
 };
 
 
