[Mutable] Fix Texture Parameter Declarations not working.

#rnx

diff --git a/Engine/Plugins/Experimental/Mutable/Source/CustomizableObjectEditor/Private/MuCOE/CustomizableInstanceDetails.cpp b/Engine/Plugins/Experimental/Mutable/Source/CustomizableObjectEditor/Private/MuCOE/CustomizableInstanceDetails.cpp
--- a/Engine/Plugins/Experimental/Mutable/Source/CustomizableObjectEditor/Private/MuCOE/CustomizableInstanceDetails.cpp
+++ b/Engine/Plugins/Experimental/Mutable/Source/CustomizableObjectEditor/Private/MuCOE/CustomizableInstanceDetails.cpp
@@ -1230,16 +1230,21 @@ void FCustomizableInstanceDetails::GenerateTextureParameterOptions()
 
 void FCustomizableInstanceDetails::OnTextureParameterComboBoxSelectionChanged(TSharedPtr<FString> Selection, ESelectInfo::Type SelectInfo, FString ParamName)
 {
-	if (Selection.IsValid())
+	if (!Selection.IsValid())
 	{
-		FScopedTransaction LocalTransaction(FText::Format(LOCTEXT("OnTextureParameterSelected", "Set Texture Parameter: {0}"), FText::FromString(ParamName)));
-		CustomInstance->Modify();
-		CustomInstance->SetTextureParameterSelectedOption(ParamName, *Selection);
-		UpdateInstance();
-
-		// Non-continuous change: collect garbage.
-		GEngine->ForceGarbageCollection();
+		return;
 	}
+
+	const int32 Index = TextureParameterValueNames.Find(Selection);
+	check(Index != INDEX_NONE);
+	
+	FScopedTransaction LocalTransaction(FText::Format(LOCTEXT("OnTextureParameterSelected", "Set Texture Parameter: {0}"), FText::FromString(ParamName)));
+	CustomInstance->Modify();
+	CustomInstance->SetTextureParameterSelectedOption(ParamName, TextureParameterValues[Index].ToString());
+	UpdateInstance();
+
+	// Non-continuous change: collect garbage.
+	GEngine->ForceGarbageCollection();
 }
 
 
