[Mutable] Fix cook crash when preparing morph and cloth data for multiple platforms.

#rnx

diff --git a/Engine/Plugins/Experimental/Mutable/Source/CustomizableObject/Private/MuCO/CustomizableObject.cpp b/Engine/Plugins/Experimental/Mutable/Source/CustomizableObject/Private/MuCO/CustomizableObject.cpp
--- a/Engine/Plugins/Experimental/Mutable/Source/CustomizableObject/Private/MuCO/CustomizableObject.cpp
+++ b/Engine/Plugins/Experimental/Mutable/Source/CustomizableObject/Private/MuCO/CustomizableObject.cpp
@@ -2169,12 +2169,9 @@ void UCustomizableObjectBulk::PrepareBulkData(UCustomizableObject* InOuter, cons
 	{
 		uint64 SourceOffset = 0;
 		
-		constexpr bool bGetCooked = false;
-
 		TArray<FMutableStreamableBlock> RealTimeMorphTargetsBlocks;
 		
-		const TMap<uint32, FRealTimeMorphStreamable>& RealTimeMorphStreamables = 
-				CustomizableObject->GetPrivate()->GetModelResources(bGetCooked).RealTimeMorphStreamables;
+		const TMap<uint32, FRealTimeMorphStreamable>& RealTimeMorphStreamables = ModelResources.RealTimeMorphStreamables;
 
 		const int32 NumBlocks = RealTimeMorphTargetsBlocks.Num();
 		for (const TPair<uint32, FRealTimeMorphStreamable>& MorphStreamable : RealTimeMorphStreamables)
@@ -2204,12 +2201,9 @@ void UCustomizableObjectBulk::PrepareBulkData(UCustomizableObject* InOuter, cons
 	{
 		uint64 SourceOffset = 0;
 		
-		constexpr bool bGetCooked = false;
-
 		TArray<FMutableStreamableBlock> ClothingBlocks;
 		
-		const TMap<uint32, FClothingStreamable>& ClothingStreamables = 
-				CustomizableObject->GetPrivate()->GetModelResources(bGetCooked).ClothingStreamables;
+		const TMap<uint32, FClothingStreamable>& ClothingStreamables = ModelResources.ClothingStreamables;
 
 		const int32 NumBlocks = ClothingBlocks.Num();
 		for (const TPair<uint32, FClothingStreamable>& ClothStreamable : ClothingStreamables)
@@ -2360,7 +2354,6 @@ void UCustomizableObjectBulk::PrepareBulkData(UCustomizableObject* InOuter, cons
 		}
 	}
 }
-
 #endif // WITH_EDITOR
 
 //-------------------------------------------------------------------------------------------------
