[Mutable] Remove access to parameter arrays.

#rnx

diff --git a/Engine/Plugins/Experimental/Mutable/Source/CustomizableObject/Classes/MuCO/CustomizableObjectInstance.h b/Engine/Plugins/Experimental/Mutable/Source/CustomizableObject/Classes/MuCO/CustomizableObjectInstance.h
--- a/Engine/Plugins/Experimental/Mutable/Source/CustomizableObject/Classes/MuCO/CustomizableObjectInstance.h
+++ b/Engine/Plugins/Experimental/Mutable/Source/CustomizableObject/Classes/MuCO/CustomizableObjectInstance.h
@@ -286,43 +286,7 @@ public:
 	/** Return true if a skeletal mesh has been generated for any component of this instance. */
 	UFUNCTION(BlueprintCallable, Category = CustomizableObjectInstance)
 	bool HasAnySkeletalMesh() const;
-
-	/** Get the array of parameters in the instance that are of type "bool". */
-	UFUNCTION(BlueprintCallable, Category = CustomizableObjectInstance)
-	TArray<FCustomizableObjectBoolParameterValue>& GetBoolParameters();
-
-	const TArray<FCustomizableObjectBoolParameterValue>& GetBoolParameters() const;
-
-	/** Get the array of parameters in the instance that are of type "integer". */
-	UFUNCTION(BlueprintCallable, Category = CustomizableObjectInstance)
-	TArray<FCustomizableObjectIntParameterValue>& GetIntParameters();
-
-	const TArray<FCustomizableObjectIntParameterValue>& GetIntParameters() const;
-
-	/** Get the array of parameters in the instance that are of type "float". */
-	UFUNCTION(BlueprintCallable, Category = CustomizableObjectInstance)
-	TArray<FCustomizableObjectFloatParameterValue>& GetFloatParameters();
-
-	const TArray<FCustomizableObjectFloatParameterValue>& GetFloatParameters() const;
 	
-	/** Get the array of parameters in the instance that are of type "texture". */
-	UFUNCTION(BlueprintCallable, Category = CustomizableObjectInstance)
-	TArray<FCustomizableObjectTextureParameterValue>& GetTextureParameters();
-
-	const TArray<FCustomizableObjectTextureParameterValue>& GetTextureParameters() const;
-
-	/** Get the array of parameters in the instance that are of type "vector". */
-	UFUNCTION(BlueprintCallable, Category = CustomizableObjectInstance)
-	TArray<FCustomizableObjectVectorParameterValue>& GetVectorParameters();
-
-	const TArray<FCustomizableObjectVectorParameterValue>& GetVectorParameters() const;
-
-	/** Get the array of parameters in the instance that are of type "projector". */
-	UFUNCTION(BlueprintCallable, Category = CustomizableObjectInstance)
-	TArray<FCustomizableObjectProjectorParameterValue>& GetProjectorParameters();
-	
-	const TArray<FCustomizableObjectProjectorParameterValue>& GetProjectorParameters() const;
-
 	/** Return true if the instance has any parameters. */
 	UFUNCTION(BlueprintCallable, Category = CustomizableObjectInstance)
 	bool HasAnyParameters() const;
diff --git a/Engine/Plugins/Experimental/Mutable/Source/CustomizableObject/Private/MuCO/CustomizableObjectInstance.cpp b/Engine/Plugins/Experimental/Mutable/Source/CustomizableObject/Private/MuCO/CustomizableObjectInstance.cpp
--- a/Engine/Plugins/Experimental/Mutable/Source/CustomizableObject/Private/MuCO/CustomizableObjectInstance.cpp
+++ b/Engine/Plugins/Experimental/Mutable/Source/CustomizableObject/Private/MuCO/CustomizableObjectInstance.cpp
@@ -6727,77 +6727,6 @@ bool UCustomizableObjectInstance::HasAnySkeletalMesh() const
 }
 
 
-TArray<FCustomizableObjectBoolParameterValue>& UCustomizableObjectInstance::GetBoolParameters()
-{
-	return Descriptor.GetBoolParameters();
-}
-
-
-const TArray<FCustomizableObjectBoolParameterValue>& UCustomizableObjectInstance::GetBoolParameters() const
-{
-	return Descriptor.GetBoolParameters();
-}
-
-
-TArray<FCustomizableObjectIntParameterValue>& UCustomizableObjectInstance::GetIntParameters()
-{
-	return Descriptor.GetIntParameters();
-}
-
-
-const TArray<FCustomizableObjectIntParameterValue>& UCustomizableObjectInstance::GetIntParameters() const
-{
-	return Descriptor.GetIntParameters();
-}
-
-
-TArray<FCustomizableObjectFloatParameterValue>& UCustomizableObjectInstance::GetFloatParameters()
-{
-	return Descriptor.GetFloatParameters();
-}
-
-
-const TArray<FCustomizableObjectFloatParameterValue>& UCustomizableObjectInstance::GetFloatParameters() const
-{
-	return Descriptor.GetFloatParameters();
-}
-
-
-TArray<FCustomizableObjectTextureParameterValue>& UCustomizableObjectInstance::GetTextureParameters()
-{
-	return Descriptor.GetTextureParameters();
-}
-
-
-const TArray<FCustomizableObjectTextureParameterValue>& UCustomizableObjectInstance::GetTextureParameters() const
-{
-	return Descriptor.GetTextureParameters();
-}
-
-
-TArray<FCustomizableObjectVectorParameterValue>& UCustomizableObjectInstance::GetVectorParameters()
-{
-	return Descriptor.GetVectorParameters();
-}
-
-const TArray<FCustomizableObjectVectorParameterValue>& UCustomizableObjectInstance::GetVectorParameters() const
-{
-	return Descriptor.GetVectorParameters();
-}
-
-
-TArray<FCustomizableObjectProjectorParameterValue>& UCustomizableObjectInstance::GetProjectorParameters()
-{
-	return Descriptor.GetProjectorParameters();
-}
-
-
-const TArray<FCustomizableObjectProjectorParameterValue>& UCustomizableObjectInstance::GetProjectorParameters() const
-{
-	return Descriptor.GetProjectorParameters();
-}
-
-
 bool UCustomizableObjectInstance::HasAnyParameters() const
 {
 	return Descriptor.HasAnyParameters();	
diff --git a/Engine/Plugins/Experimental/Mutable/Source/CustomizableObjectEditor/Private/MuCOE/CustomizableInstanceDetails.cpp b/Engine/Plugins/Experimental/Mutable/Source/CustomizableObjectEditor/Private/MuCOE/CustomizableInstanceDetails.cpp
--- a/Engine/Plugins/Experimental/Mutable/Source/CustomizableObjectEditor/Private/MuCOE/CustomizableInstanceDetails.cpp
+++ b/Engine/Plugins/Experimental/Mutable/Source/CustomizableObjectEditor/Private/MuCOE/CustomizableInstanceDetails.cpp
@@ -1175,7 +1175,7 @@ void FCustomizableInstanceDetails::GenerateTextureParameterOptions()
 {
 	UCustomizableObject* CustomizableObject = CustomInstance->GetCustomizableObject();
 
-	for (const FCustomizableObjectTextureParameterValue& TextureParameter : CustomInstance->GetTextureParameters())
+	for (const FCustomizableObjectTextureParameterValue& TextureParameter : CustomInstance->GetPrivate()->GetDescriptor().GetTextureParameters())
 	{
 		// Get default values.
 		const FName& DefaultValue = CustomizableObject->GetTextureParameterDefaultValue(TextureParameter.ParameterName);
@@ -1407,7 +1407,7 @@ FReply FCustomizableInstanceDetails::OnProjectorCopyTransform(const FString Para
 
 	if ((ParameterIndexInObject >= 0) && (ProjectorParamIndex >= 0))
 	{
-		TArray<FCustomizableObjectProjectorParameterValue>& ProjectorParameters = CustomInstance->GetProjectorParameters();
+		TArray<FCustomizableObjectProjectorParameterValue>& ProjectorParameters = CustomInstance->GetPrivate()->GetDescriptor().GetProjectorParameters();
 		FCustomizableObjectProjector Value;
 
 		if (RangeIndex == -1)
@@ -1490,7 +1490,7 @@ TSharedRef<SWidget> FCustomizableInstanceDetails::GenerateMultidimensionalProjec
 	FString ParamName = CustomizableObject->GetParameterName(ParamIndexInObject);
 
 	const TSharedPtr<ICustomizableObjectInstanceEditor> Editor = GetEditorChecked();
-	TArray<FCustomizableObjectProjectorParameterValue>& ProjectorParameters = CustomInstance->GetProjectorParameters();
+	TArray<FCustomizableObjectProjectorParameterValue>& ProjectorParameters = CustomInstance->GetPrivate()->GetDescriptor().GetProjectorParameters();
 	const int32 ProjectorParamIndex = CustomInstance->FindProjectorParameterNameIndex(ParamName);
 
 	check(ProjectorParamIndex < ProjectorParameters.Num());
@@ -1827,9 +1827,9 @@ FReply FCustomizableInstanceDetails::OnProjectorLayerAdded(FString ParamName)
 	const FString TextureSwitchEnumParamName = ParamName + FMultilayerProjector::IMAGE_PARAMETER_POSTFIX;
 	CustomInstance->AddValueToIntRange(TextureSwitchEnumParamName);
 	check(CustomInstance->FindIntParameterNameIndex(TextureSwitchEnumParamName) != INDEX_NONE);
-	check(NumLayers == CustomInstance->GetIntParameters()[CustomInstance->FindIntParameterNameIndex(TextureSwitchEnumParamName)].ParameterRangeValueNames.Num());
+	check(NumLayers == CustomInstance->GetPrivate()->GetDescriptor().GetIntParameters()[CustomInstance->FindIntParameterNameIndex(TextureSwitchEnumParamName)].ParameterRangeValueNames.Num());
 
-	TArray<FCustomizableObjectFloatParameterValue>& FloatParameters = CustomInstance->GetFloatParameters();
+	TArray<FCustomizableObjectFloatParameterValue>& FloatParameters = CustomInstance->GetPrivate()->GetDescriptor().GetFloatParameters();
 
 	const FString NumLayersParamName = ParamName + FMultilayerProjector::NUM_LAYERS_PARAMETER_POSTFIX;
 	check(CustomInstance->FindFloatParameterNameIndex(NumLayersParamName) != INDEX_NONE);
@@ -1861,7 +1861,7 @@ FReply FCustomizableInstanceDetails::OnProjectorLayerRemoved(const FString Param
 
 	const int32 projectorParameterIndex = CustomInstance->FindProjectorParameterNameIndex(ParamName);
 	if (projectorParameterIndex == INDEX_NONE
-		|| CustomInstance->GetProjectorParameters()[projectorParameterIndex].RangeValues.Num() <= 0)
+		|| CustomInstance->GetPrivate()->GetDescriptor().GetProjectorParameters()[projectorParameterIndex].RangeValues.Num() <= 0)
 	{
 		return FReply::Handled();
 	}
@@ -1871,9 +1871,9 @@ FReply FCustomizableInstanceDetails::OnProjectorLayerRemoved(const FString Param
 	const FString TextureSwitchEnumParamName = ParamName + FMultilayerProjector::IMAGE_PARAMETER_POSTFIX;
 	CustomInstance->RemoveValueFromIntRange(TextureSwitchEnumParamName, RangeIndex);
 	check(CustomInstance->FindIntParameterNameIndex(TextureSwitchEnumParamName) != INDEX_NONE);
-	check(NumLayers == CustomInstance->GetIntParameters()[CustomInstance->FindIntParameterNameIndex(TextureSwitchEnumParamName)].ParameterRangeValueNames.Num());
+	check(NumLayers == CustomInstance->GetPrivate()->GetDescriptor().GetIntParameters()[CustomInstance->FindIntParameterNameIndex(TextureSwitchEnumParamName)].ParameterRangeValueNames.Num());
 
-	TArray<FCustomizableObjectFloatParameterValue>& FloatParameters = CustomInstance->GetFloatParameters();
+	TArray<FCustomizableObjectFloatParameterValue>& FloatParameters = CustomInstance->GetPrivate()->GetDescriptor().GetFloatParameters();
 
 	const FString NumLayersParamName = ParamName + FMultilayerProjector::NUM_LAYERS_PARAMETER_POSTFIX;
 	check(CustomInstance->FindFloatParameterNameIndex(NumLayersParamName) != INDEX_NONE);
diff --git a/Engine/Plugins/Experimental/Mutable/Source/CustomizableObjectEditor/Private/MuCOE/CustomizableObjectEditor.cpp b/Engine/Plugins/Experimental/Mutable/Source/CustomizableObjectEditor/Private/MuCOE/CustomizableObjectEditor.cpp
--- a/Engine/Plugins/Experimental/Mutable/Source/CustomizableObjectEditor/Private/MuCOE/CustomizableObjectEditor.cpp
+++ b/Engine/Plugins/Experimental/Mutable/Source/CustomizableObjectEditor/Private/MuCOE/CustomizableObjectEditor.cpp
@@ -661,7 +661,7 @@ void FCustomizableObjectEditor::CompileOnlySelectedObjectUserPressedButton()
 	{
 		Compiler.ClearAllCompileOnlySelectedOption();
 
-		for (const FCustomizableObjectIntParameterValue& IntParam : PreviewInstance->GetIntParameters())
+		for (const FCustomizableObjectIntParameterValue& IntParam : PreviewInstance->GetPrivate()->GetDescriptor().GetIntParameters())
 		{
 			Compiler.AddCompileOnlySelectedOption(IntParam.ParameterName, IntParam.ParameterValueName);
 		}
diff --git a/Engine/Plugins/Experimental/Mutable/Source/CustomizableObjectEditor/Private/MuCOE/SCustomizableInstanceProperties.cpp b/Engine/Plugins/Experimental/Mutable/Source/CustomizableObjectEditor/Private/MuCOE/SCustomizableInstanceProperties.cpp
--- a/Engine/Plugins/Experimental/Mutable/Source/CustomizableObjectEditor/Private/MuCOE/SCustomizableInstanceProperties.cpp
+++ b/Engine/Plugins/Experimental/Mutable/Source/CustomizableObjectEditor/Private/MuCOE/SCustomizableInstanceProperties.cpp
@@ -91,7 +91,7 @@ void SCustomizableInstanceProperties::Construct(const FArguments& InArgs)
 
 		// Store which Texture Parameter values can be selected.
 		// Get default values.
-		for (const FCustomizableObjectTextureParameterValue& TextureParameter : CustomInstance->GetTextureParameters())
+		for (const FCustomizableObjectTextureParameterValue& TextureParameter : CustomInstance->GetPrivate()->GetDescriptor().GetTextureParameters())
 		{
 			const FName Value = CustomizableObject->GetTextureParameterDefaultValue(TextureParameter.ParameterName);
 			
@@ -103,7 +103,7 @@ void SCustomizableInstanceProperties::Construct(const FArguments& InArgs)
 		}
 
 		// Selected parameter
-		for (const FCustomizableObjectTextureParameterValue& TextureParameter : CustomInstance->GetTextureParameters())
+		for (const FCustomizableObjectTextureParameterValue& TextureParameter : CustomInstance->GetPrivate()->GetDescriptor().GetTextureParameters())
 		{
 			const FName& Value = TextureParameter.ParameterValue;
 			
@@ -984,7 +984,7 @@ void SCustomizableInstanceProperties::AddParameter(int32 ParamIndexInObject)
 		}
 		else
 		{
-			TArray<FCustomizableObjectProjectorParameterValue>& ProjectorParameters = CustomInstance->GetProjectorParameters();
+			TArray<FCustomizableObjectProjectorParameterValue>& ProjectorParameters = CustomInstance->GetPrivate()->GetDescriptor().GetProjectorParameters();
 			const int32 ProjectorParamIndex = CustomInstance->FindProjectorParameterNameIndex(ParamName);
 			check(ProjectorParamIndex < ProjectorParameters.Num());
 
@@ -1315,7 +1315,7 @@ void SCustomizableInstanceProperties::Tick(const FGeometry& AllottedGeometry, co
 
 	// Check the slider values
 	// We do this because the slider "changed" event seems to be called many times per-frame
-	TArray<FCustomizableObjectFloatParameterValue>& FloatParameters = CustomInstance->GetFloatParameters();
+	TArray<FCustomizableObjectFloatParameterValue>& FloatParameters = CustomInstance->GetPrivate()->GetDescriptor().GetFloatParameters();
 	for (int SliderIndex = 0; SliderIndex < FloatSliders.Num() && bSlidersChanged; ++SliderIndex)
 	{
 		for (int i = 0; i < FloatParameters.Num(); ++i)
@@ -1450,7 +1450,7 @@ FReply SCustomizableInstanceProperties::OnResetAllParameters()
 {
 	CustomInstance->PreEditChange(nullptr);
 
-	TArray<FCustomizableObjectIntParameterValue>& IntParameters = CustomInstance->GetIntParameters();
+	TArray<FCustomizableObjectIntParameterValue>& IntParameters = CustomInstance->GetPrivate()->GetDescriptor().GetIntParameters();
 	
 	const UCustomizableObject* CustomObject = CustomInstance->GetCustomizableObject();
 	const int32 NumObjectParameter = CustomObject->GetParameterCount();
@@ -1599,7 +1599,7 @@ void SCustomizableInstanceProperties::OnIntParameterChanged(int32 Value, FString
 
 void SCustomizableInstanceProperties::OnIntParameterComboBoxChanged(TSharedPtr<FString> Selection, ESelectInfo::Type SelectInfo, FString ParamName)
 {
-	TArray<FCustomizableObjectIntParameterValue>& IntParameters = CustomInstance->GetIntParameters();
+	TArray<FCustomizableObjectIntParameterValue>& IntParameters = CustomInstance->GetPrivate()->GetDescriptor().GetIntParameters();
 	
 	const UCustomizableObject* CustomObject = CustomInstance->GetCustomizableObject();
 	for (int32 i = 0; i < IntParameters.Num(); ++i)
@@ -1644,7 +1644,7 @@ void SCustomizableInstanceProperties::OnTextureParameterComboBoxSelectionChanged
 {
 	CustomInstance->PreEditChange(nullptr);
 
-	auto& TextureParameters = CustomInstance->GetTextureParameters();
+	auto& TextureParameters = CustomInstance->GetPrivate()->GetDescriptor().GetTextureParameters();
 
 	for (int i = 0; i < TextureParameters.Num(); ++i)
 	{
@@ -1726,9 +1726,9 @@ FReply SCustomizableInstanceProperties::OnProjectorLayerAdded(FString ParamName)
 	const FString TextureSwitchEnumParamName = ParamName + FMultilayerProjector::IMAGE_PARAMETER_POSTFIX;
 	CustomInstance->AddValueToIntRange(TextureSwitchEnumParamName);
 	check(CustomInstance->FindIntParameterNameIndex(TextureSwitchEnumParamName) != INDEX_NONE);
-	check(NumLayers == CustomInstance->GetIntParameters()[CustomInstance->FindIntParameterNameIndex(TextureSwitchEnumParamName)].ParameterRangeValueNames.Num());
+	check(NumLayers == CustomInstance->GetPrivate()->GetDescriptor().GetIntParameters()[CustomInstance->FindIntParameterNameIndex(TextureSwitchEnumParamName)].ParameterRangeValueNames.Num());
 
-	TArray<FCustomizableObjectFloatParameterValue>& FloatParameters = CustomInstance->GetFloatParameters();
+	TArray<FCustomizableObjectFloatParameterValue>& FloatParameters = CustomInstance->GetPrivate()->GetDescriptor().GetFloatParameters();
 	
 	const FString NumLayersParamName = ParamName + FMultilayerProjector::NUM_LAYERS_PARAMETER_POSTFIX;
 	check(CustomInstance->FindFloatParameterNameIndex(NumLayersParamName) != INDEX_NONE);
@@ -1756,7 +1756,7 @@ FReply SCustomizableInstanceProperties::OnProjectorLayerRemoved(const FString Pa
 {
 	const int32 projectorParameterIndex = CustomInstance->FindProjectorParameterNameIndex(ParamName);
 	if (projectorParameterIndex == INDEX_NONE
-		|| CustomInstance->GetProjectorParameters()[projectorParameterIndex].RangeValues.Num() <= 0)
+		|| CustomInstance->GetPrivate()->GetDescriptor().GetProjectorParameters()[projectorParameterIndex].RangeValues.Num() <= 0)
 	{
 		return FReply::Handled();
 	}
@@ -1766,9 +1766,9 @@ FReply SCustomizableInstanceProperties::OnProjectorLayerRemoved(const FString Pa
 	const FString TextureSwitchEnumParamName = ParamName + FMultilayerProjector::IMAGE_PARAMETER_POSTFIX;
 	CustomInstance->RemoveValueFromIntRange(TextureSwitchEnumParamName, RangeIndex);
 	check(CustomInstance->FindIntParameterNameIndex(TextureSwitchEnumParamName) != INDEX_NONE);
-	check(NumLayers == CustomInstance->GetIntParameters()[CustomInstance->FindIntParameterNameIndex(TextureSwitchEnumParamName)].ParameterRangeValueNames.Num());
+	check(NumLayers == CustomInstance->GetPrivate()->GetDescriptor().GetIntParameters()[CustomInstance->FindIntParameterNameIndex(TextureSwitchEnumParamName)].ParameterRangeValueNames.Num());
 
-	TArray<FCustomizableObjectFloatParameterValue>& FloatParameters = CustomInstance->GetFloatParameters();
+	TArray<FCustomizableObjectFloatParameterValue>& FloatParameters = CustomInstance->GetPrivate()->GetDescriptor().GetFloatParameters();
 
 	const FString NumLayersParamName = ParamName + FMultilayerProjector::NUM_LAYERS_PARAMETER_POSTFIX;
 	check(CustomInstance->FindFloatParameterNameIndex(NumLayersParamName) != INDEX_NONE);
@@ -1812,7 +1812,7 @@ FReply SCustomizableInstanceProperties::OnProjectorSelectChanged(const FString P
 
 void SCustomizableInstanceProperties::OnProjectorTextureParameterComboBoxChanged(TSharedPtr<FString> Selection, ESelectInfo::Type SelectInfo, FString ParamName, int32 RangeIndex) const
 {
-	TArray<FCustomizableObjectIntParameterValue>& IntParameters = CustomInstance->GetIntParameters();
+	TArray<FCustomizableObjectIntParameterValue>& IntParameters = CustomInstance->GetPrivate()->GetDescriptor().GetIntParameters();
 
 	const UCustomizableObject* CustomObject = CustomInstance->GetCustomizableObject();
 	for (int32 i = 0; i < IntParameters.Num(); ++i)
@@ -1889,7 +1889,7 @@ FReply SCustomizableInstanceProperties::OnProjectorCopyTransform(const FString P
 
 	if ((ParameterIndexInObject >= 0) && (ProjectorParamIndex >= 0))
 	{
-		TArray<FCustomizableObjectProjectorParameterValue>& ProjectorParameters = CustomInstance->GetProjectorParameters();
+		TArray<FCustomizableObjectProjectorParameterValue>& ProjectorParameters = CustomInstance->GetPrivate()->GetDescriptor().GetProjectorParameters();
 		FCustomizableObjectProjector Value;
 
 		if (RangeIndex == -1)
diff --git a/Engine/Plugins/Experimental/Mutable/Source/CustomizableObjectEditor/Private/MuCOE/StatePerformanceTesting.cpp b/Engine/Plugins/Experimental/Mutable/Source/CustomizableObjectEditor/Private/MuCOE/StatePerformanceTesting.cpp
--- a/Engine/Plugins/Experimental/Mutable/Source/CustomizableObjectEditor/Private/MuCOE/StatePerformanceTesting.cpp
+++ b/Engine/Plugins/Experimental/Mutable/Source/CustomizableObjectEditor/Private/MuCOE/StatePerformanceTesting.cpp
@@ -348,7 +348,7 @@ void FRuntimeTest::LoadStateData(int32 StateIndex)
 		}
 		case EMutableParameterType::Bool:
 		{
-			TArray<FCustomizableObjectBoolParameterValue>& BoolParameters = Instance->GetBoolParameters();
+			TArray<FCustomizableObjectBoolParameterValue>& BoolParameters = Instance->GetPrivate()->GetDescriptor().GetBoolParameters();
 			FCustomizableObjectBoolParameterValue Parameter;
 			Parameter.ParameterName = RunTimeParameter;
 			const uint32 MaxIndex = BoolParameters.Num();
@@ -370,7 +370,7 @@ void FRuntimeTest::LoadStateData(int32 StateIndex)
 		}
 		case EMutableParameterType::Int:
 		{
-			TArray<FCustomizableObjectIntParameterValue>& IntParameters = Instance->GetIntParameters();
+			TArray<FCustomizableObjectIntParameterValue>& IntParameters = Instance->GetPrivate()->GetDescriptor().GetIntParameters();
 			FCustomizableObjectIntParameterValue Parameter;
 			Parameter.ParameterName = RunTimeParameter;
 			const uint32 MaxIndex = IntParameters.Num();
@@ -397,7 +397,7 @@ void FRuntimeTest::LoadStateData(int32 StateIndex)
 		}
 		case EMutableParameterType::Float:
 		{
-			TArray<FCustomizableObjectFloatParameterValue>& FloatParameters = Instance->GetFloatParameters();
+			TArray<FCustomizableObjectFloatParameterValue>& FloatParameters = Instance->GetPrivate()->GetDescriptor().GetFloatParameters();
 			FCustomizableObjectFloatParameterValue Parameter;
 			Parameter.ParameterName = RunTimeParameter;
 			const uint32 MaxIndex = FloatParameters.Num();
@@ -420,7 +420,7 @@ void FRuntimeTest::LoadStateData(int32 StateIndex)
 		}
 		case EMutableParameterType::Color:
 		{
-			TArray<FCustomizableObjectVectorParameterValue>& VectorParameters = Instance->GetVectorParameters();
+			TArray<FCustomizableObjectVectorParameterValue>& VectorParameters = Instance->GetPrivate()->GetDescriptor().GetVectorParameters();
 			FCustomizableObjectVectorParameterValue Parameter;
 			Parameter.ParameterName = RunTimeParameter;
 			const uint32 MaxIndex = VectorParameters.Num();
@@ -442,7 +442,7 @@ void FRuntimeTest::LoadStateData(int32 StateIndex)
 		}
 		case EMutableParameterType::Projector:
 		{
-			TArray<FCustomizableObjectVectorParameterValue>& VectorParameters = Instance->GetVectorParameters();
+			TArray<FCustomizableObjectVectorParameterValue>& VectorParameters = Instance->GetPrivate()->GetDescriptor().GetVectorParameters();
 			FCustomizableObjectVectorParameterValue Parameter;
 			Parameter.ParameterName = RunTimeParameter;
 			const uint32 MaxIndex = VectorParameters.Num();
@@ -465,7 +465,7 @@ void FRuntimeTest::LoadStateData(int32 StateIndex)
 		}
 		case EMutableParameterType::Texture:
 		{
-			TArray<FCustomizableObjectTextureParameterValue>& TextureParameters = Instance->GetTextureParameters();
+			TArray<FCustomizableObjectTextureParameterValue>& TextureParameters = Instance->GetPrivate()->GetDescriptor().GetTextureParameters();
 			FCustomizableObjectTextureParameterValue Parameter;
 			Parameter.ParameterName = RunTimeParameter;
 			const uint32 MaxIndex = TextureParameters.Num();
diff --git a/Engine/Plugins/Experimental/Mutable/Source/CustomizableObjectEditor/Private/MuCOE/StressTest.cpp b/Engine/Plugins/Experimental/Mutable/Source/CustomizableObjectEditor/Private/MuCOE/StressTest.cpp
--- a/Engine/Plugins/Experimental/Mutable/Source/CustomizableObjectEditor/Private/MuCOE/StressTest.cpp
+++ b/Engine/Plugins/Experimental/Mutable/Source/CustomizableObjectEditor/Private/MuCOE/StressTest.cpp
@@ -9,6 +9,7 @@
 #include "MuCO/UnrealPortabilityHelpers.h"
 #include "MuCOE/StatePerformanceTesting.h"
 #include "Engine/SkeletalMesh.h"
+#include "MuCO/CustomizableObjectInstancePrivate.h"
 #include "UObject/Package.h"
 
 class FSkeletalMeshLODRenderData;
@@ -492,12 +493,12 @@ void FRunningStressTest::FinishTest()
 
 void FRunningStressTest::CaptureInstanceParameters(UCustomizableObjectInstance* Instance, MostExpensiveInstanceData* Destination)
 {
-	Destination->BoolParameters = Instance->GetBoolParameters();
-	Destination->IntParameters = Instance->GetIntParameters();
-	Destination->FloatParameters = Instance->GetFloatParameters();
-	Destination->TextureParameters = Instance->GetTextureParameters();
-	Destination->VectorParameters = Instance->GetVectorParameters();
-	Destination->ProjectorParameters = Instance->GetProjectorParameters();
+	Destination->BoolParameters = Instance->GetPrivate()->GetDescriptor().GetBoolParameters();
+	Destination->IntParameters = Instance->GetPrivate()->GetDescriptor().GetIntParameters();
+	Destination->FloatParameters = Instance->GetPrivate()->GetDescriptor().GetFloatParameters();
+	Destination->TextureParameters = Instance->GetPrivate()->GetDescriptor().GetTextureParameters();
+	Destination->VectorParameters = Instance->GetPrivate()->GetDescriptor().GetVectorParameters();
+	Destination->ProjectorParameters = Instance->GetPrivate()->GetDescriptor().GetProjectorParameters();
 }
 
 
