[Mutable] Moved FCustomizableObjectInstanceDescriptor to Internal.

#rnx

diff --git a/Engine/Plugins/Experimental/Mutable/Source/CustomizableObject/Classes/MuCO/CustomizableObjectInstance.h b/Engine/Plugins/Experimental/Mutable/Source/CustomizableObject/Classes/MuCO/CustomizableObjectInstance.h
--- a/Engine/Plugins/Experimental/Mutable/Source/CustomizableObject/Classes/MuCO/CustomizableObjectInstance.h
+++ b/Engine/Plugins/Experimental/Mutable/Source/CustomizableObject/Classes/MuCO/CustomizableObjectInstance.h
@@ -228,10 +228,6 @@ class CUSTOMIZABLEOBJECT_API UCustomizableObjectInstance : public UObject
 public:
 	UCustomizableObjectInstance();
 	
-	const FCustomizableObjectInstanceDescriptor& GetDescriptor() const;
-	
-	void SetDescriptor(const FCustomizableObjectInstanceDescriptor& InDescriptor);
-
 	/** Broadcast when the Customizable Object Instance is updated. */
 	UPROPERTY(Transient, BlueprintAssignable, Category = CustomizableObjectInstance)
 	FObjectInstanceUpdatedDelegate UpdatedDelegate;
diff --git a/Engine/Plugins/Experimental/Mutable/Source/CustomizableObject/Internal/MuCO/CustomizableObjectInstanceDescriptor.h b/Engine/Plugins/Experimental/Mutable/Source/CustomizableObject/Internal/MuCO/CustomizableObjectInstanceDescriptor.h
new file mode 100644
--- /dev/null
+++ b/Engine/Plugins/Experimental/Mutable/Source/CustomizableObject/Internal/MuCO/CustomizableObjectInstanceDescriptor.h
@@ -0,0 +1,413 @@
+// Copyright Epic Games, Inc. All Rights Reserved.
+
+#pragma once
+
+#include "MuCO/MultilayerProjector.h"
+#include "MuCO/CustomizableObjectParameterTypeDefinitions.h"
+#include "Math/RandomStream.h"
+
+#include "CustomizableObjectInstanceDescriptor.generated.h"
+
+class UTexture2D;
+enum class ECustomizableObjectProjectorType : uint8;
+
+class FArchive;
+class UCustomizableInstancePrivate;
+class UCustomizableObject;
+class UCustomizableObjectInstance;
+class FDescriptorHash;
+class FMutableUpdateCandidate;
+
+typedef TMap<const UCustomizableObjectInstance*, FMutableUpdateCandidate> FMutableInstanceUpdateMap;
+
+namespace mu
+{
+	class Parameters;
+	
+	template<typename Type>
+	class Ptr;
+}
+
+
+/** Set of parameters + state that defines a CustomizableObjectInstance.
+ *
+ * This object has the same parameters + state interface as UCustomizableObjectInstance.
+ * UCustomizableObjectInstance must share the same interface. Any public methods added here should also end up in the Instance. */
+USTRUCT()
+struct CUSTOMIZABLEOBJECT_API FCustomizableObjectInstanceDescriptor
+{
+	GENERATED_BODY()
+
+	FCustomizableObjectInstanceDescriptor() = default;
+	
+	explicit FCustomizableObjectInstanceDescriptor(UCustomizableObject& Object);
+
+	/** Serialize this object. 
+	 *
+	 * Backwards compatibility is not guaranteed.
+ 	 * Multilayer Projectors not supported.
+	 *
+  	 * @param bUseCompactDescriptor If true it assumes the compiled objects are the same on both ends of the serialisation */
+	void SaveDescriptor(FArchive &Ar, bool bUseCompactDescriptor);
+
+	/** Deserialize this object. Does not support Multilayer Projectors! */
+
+	/** Deserialize this object.
+     *
+	 * Backwards compatibility is not guaranteed.
+	 * Multilayer Projectors not supported */
+	void LoadDescriptor(FArchive &Ar);
+
+	UCustomizableObject* GetCustomizableObject() const;
+
+	void SetCustomizableObject(UCustomizableObject& InCustomizableObject);
+	
+	bool GetBuildParameterRelevancy() const;
+	
+	void SetBuildParameterRelevancy(bool Value);
+	
+	/** Update all parameters to be up to date with the Mutable Core parameters. */
+	void ReloadParameters();
+    
+	int32 GetMinLod() const;
+
+	void SetMinLod(int32 InMinLOD);
+
+	int32 GetMaxLod() const { return MAX_int32; }; // DEPRECATED
+
+	void SetMaxLod(int32 InMaxLOD) {}; // DEPRECATED
+
+	void SetRequestedLODLevels(const TArray<uint16>& InRequestedLODLevels);
+
+	const TArray<uint16>& GetRequestedLODLevels() const;
+
+	// ------------------------------------------------------------
+	// Parameters
+	// ------------------------------------------------------------
+
+	TArray<FCustomizableObjectBoolParameterValue>& GetBoolParameters();
+
+	const TArray<FCustomizableObjectBoolParameterValue>& GetBoolParameters() const;
+	
+	TArray<FCustomizableObjectIntParameterValue>& GetIntParameters();
+
+	const TArray<FCustomizableObjectIntParameterValue>& GetIntParameters() const;
+
+	TArray<FCustomizableObjectFloatParameterValue>& GetFloatParameters();
+
+	const TArray<FCustomizableObjectFloatParameterValue>& GetFloatParameters() const;
+	
+	TArray<FCustomizableObjectTextureParameterValue>& GetTextureParameters();
+	
+	const TArray<FCustomizableObjectTextureParameterValue>& GetTextureParameters() const;
+
+	TArray<FCustomizableObjectVectorParameterValue>& GetVectorParameters();
+	
+	const TArray<FCustomizableObjectVectorParameterValue>& GetVectorParameters() const;
+
+	TArray<FCustomizableObjectProjectorParameterValue>& GetProjectorParameters();
+	
+	const TArray<FCustomizableObjectProjectorParameterValue>& GetProjectorParameters() const;
+	
+	/** Return true if there are any parameters. */
+	bool HasAnyParameters() const;
+
+	/** Gets the value of the int parameter with name "ParamName". */
+	const FString& GetIntParameterSelectedOption(const FString& ParamName, int32 RangeIndex = -1) const;
+
+	/** Sets the selected option of an int parameter by the option's name. */
+	void SetIntParameterSelectedOption(int32 IntParamIndex, const FString& SelectedOption, int32 RangeIndex = -1);
+
+	/** Sets the selected option of an int parameter, by the option's name */
+	void SetIntParameterSelectedOption(const FString& ParamName, const FString& SelectedOptionName, int32 RangeIndex = -1);
+
+	/** Gets the value of a float parameter with name "FloatParamName". */
+	float GetFloatParameterSelectedOption(const FString& FloatParamName, int32 RangeIndex = -1) const;
+
+	/** Sets the float value "FloatValue" of a float parameter with index "FloatParamIndex". */
+	void SetFloatParameterSelectedOption(const FString& FloatParamName, float FloatValue, int32 RangeIndex = -1);
+
+	/** Gets the value of a texture parameter with name "TextureParamName". */
+	FName GetTextureParameterSelectedOption(const FString& TextureParamName, int32 RangeIndex) const;
+
+	/** Sets the texture value "TextureValue" of a texture parameter with index "TextureParamIndex". */
+	void SetTextureParameterSelectedOption(const FString& TextureParamName, const FString& TextureValue, int32 RangeIndex);
+	
+	/** Gets the value of a color parameter with name "ColorParamName". */
+	FLinearColor GetColorParameterSelectedOption(const FString& ColorParamName) const;
+
+	/** Sets the color value "ColorValue" of a color parameter with index "ColorParamIndex". */
+	void SetColorParameterSelectedOption(const FString& ColorParamName, const FLinearColor& ColorValue);
+
+	/** Gets the value of the bool parameter with name "BoolParamName". */
+	bool GetBoolParameterSelectedOption(const FString& BoolParamName) const;
+
+	/** Sets the bool value "BoolValue" of a bool parameter with name "BoolParamName". */
+	void SetBoolParameterSelectedOption(const FString& BoolParamName, bool BoolValue);
+
+	/** Sets the vector value "VectorValue" of a bool parameter with index "VectorParamIndex". */
+	void SetVectorParameterSelectedOption(const FString& VectorParamName, const FLinearColor& VectorValue);
+
+	/** Sets the projector values of a projector parameter with index "ProjectorParamIndex". */
+	void SetProjectorValue(const FString& ProjectorParamName,
+		const FVector& Pos, const FVector& Direction, const FVector& Up, const FVector& Scale,
+		float Angle,
+		int32 RangeIndex = -1);
+
+	/** Set only the projector position. */
+	void SetProjectorPosition(const FString& ProjectorParamName, const FVector& Pos, int32 RangeIndex = -1);
+
+	/** Set only the projector direction. */
+	void SetProjectorDirection(const FString& ProjectorParamName, const FVector& Direction, int32 RangeIndex = -1);
+	
+	/** Set only the projector up vector. */
+	void SetProjectorUp(const FString& ProjectorParamName, const FVector& Up, int32 RangeIndex = -1);
+
+	/** Set only the projector scale. */
+	void SetProjectorScale(const FString& ProjectorParamName, const FVector& Scale, int32 RangeIndex = -1);
+
+	/** Set only the cylindrical projector angle. */
+	void SetProjectorAngle(const FString& ProjectorParamName, float Angle, int32 RangeIndex = -1);
+	
+	/** Get the projector values of a projector parameter with index "ProjectorParamIndex". */
+	void GetProjectorValue(const FString& ProjectorParamName,
+		FVector& OutPos, FVector& OutDirection, FVector& OutUp, FVector& OutScale,
+		float& OutAngle, ECustomizableObjectProjectorType& OutType,
+		int32 RangeIndex = -1) const;
+
+	/** Float version. See GetProjectorValue. */
+	void GetProjectorValueF(const FString& ProjectorParamName,
+		FVector3f& OutPos, FVector3f& OutDirection, FVector3f& OutUp, FVector3f& OutScale,
+		float& OutAngle, ECustomizableObjectProjectorType& OutType,
+		int32 RangeIndex = -1) const;
+
+	/** Get the current projector position for the parameter with the given name. */
+	FVector GetProjectorPosition(const FString & ParamName, int32 RangeIndex = -1) const;
+
+	/** Get the current projector direction vector for the parameter with the given name. */
+	FVector GetProjectorDirection(const FString & ParamName, int32 RangeIndex = -1) const;
+
+	/** Get the current projector up vector for the parameter with the given name. */
+	FVector GetProjectorUp(const FString & ParamName, int32 RangeIndex = -1) const;
+
+	/** Get the current projector scale for the parameter with the given name. */
+	FVector GetProjectorScale(const FString & ParamName, int32 RangeIndex = -1) const;
+
+	/** Get the current cylindrical projector angle for the parameter with the given name. */
+	float GetProjectorAngle(const FString& ParamName, int32 RangeIndex = -1) const;
+
+	/** Get the current projector type for the parameter with the given name. */
+	ECustomizableObjectProjectorType GetProjectorParameterType(const FString& ParamName, int32 RangeIndex = -1) const;
+
+	/** Get the current projector for the parameter with the given name. */
+	FCustomizableObjectProjector GetProjector(const FString& ParamName, int32 RangeIndex) const;
+	
+	/** Finds in IntParameters a parameter with name ParamName, returns the index if found, -1 otherwise. */
+	int32 FindIntParameterNameIndex(const FString& ParamName) const;
+
+	/** Finds in FloatParameters a parameter with name ParamName, returns the index if found, -1 otherwise. */
+	int32 FindFloatParameterNameIndex(const FString& ParamName) const;
+
+	/** Finds in TextureParameters a parameter with name ParamName, returns the index if found, -1 otherwise. */
+	int32 FindTextureParameterNameIndex(const FString& ParamName) const;
+
+	/** Finds in BoolParameters a parameter with name ParamName, returns the index if found, -1 otherwise. */
+	int32 FindBoolParameterNameIndex(const FString& ParamName) const;
+
+	/** Finds in VectorParameters a parameter with name ParamName, returns the index if found, -1 otherwise. */
+	int32 FindVectorParameterNameIndex(const FString& ParamName) const;
+
+	/** Finds in ProjectorParameters a parameter with name ParamName, returns the index if found, -1 otherwise. */
+	int32 FindProjectorParameterNameIndex(const FString& ParamName) const;
+
+	// Parameter Ranges
+
+	/** Gets the range of values of the projector with ParamName, returns -1 if the parameter does not exist. */
+	int32 GetProjectorValueRange(const FString& ParamName) const;
+
+	/** Gets the range of values of the int with ParamName, returns -1 if the parameter does not exist. */
+	int32 GetIntValueRange(const FString& ParamName) const;
+
+	/** Gets the range of values of the float with ParamName, returns -1 if the parameter does not exist. */
+	int32 GetFloatValueRange(const FString& ParamName) const;
+
+	/** Gets the range of values of the texture with ParamName, returns -1 if the parameter does not exist. */
+	int32 GetTextureValueRange(const FString& ParamName) const;
+
+	/** Increases the range of values of the integer with ParamName, returns the index of the new integer value, -1 otherwise.
+	 * The added value is initialized with the first integer option and is the last one of the range. */
+	int32 AddValueToIntRange(const FString& ParamName);
+
+	/** Increases the range of values of the float with ParamName, returns the index of the new float value, -1 otherwise.
+	 * The added value is initialized with 0.5f and is the last one of the range. */
+	int32 AddValueToFloatRange(const FString& ParamName);
+
+	/** Increases the range of values of the float with ParamName, returns the index of the new float value, -1 otherwise. 
+	 * The added value is not initialized. */
+	int32 AddValueToTextureRange(const FString& ParamName);
+
+	/** Increases the range of values of the projector with ParamName, returns the index of the new projector value, -1 otherwise.
+	 * The added value is initialized with the default projector as set up in the editor and is the last one of the range. */
+	int32 AddValueToProjectorRange(const FString& ParamName);
+
+	/** Remove the last of the integer range of values from the parameter ParamName, returns the index of the last valid integer, -1 if no values left. */
+	int32 RemoveValueFromIntRange(const FString& ParamName);
+
+	/** Remove the RangeIndex element of the integer range of values from the parameter ParamName, returns the index of the last valid integer, -1 if no values left. */
+	int32 RemoveValueFromIntRange(const FString& ParamName, int32 RangeIndex);
+
+	/** Remove the last of the float range of values from the parameter ParamName, returns the index of the last valid float, -1 if no values left. */
+	int32 RemoveValueFromFloatRange(const FString& ParamName);
+
+	/** Remove the RangeIndex element of the float range of values from the parameter ParamName, returns the index of the last valid float, -1 if no values left. */
+	int32 RemoveValueFromFloatRange(const FString& ParamName, int32 RangeIndex);
+
+	/** Remove the last of the texture range of values from the parameter ParamName, returns the index of the last valid float, -1 if no values left. */
+	int32 RemoveValueFromTextureRange(const FString& ParamName);
+
+	/** Remove the RangeIndex element of the texture range of values from the parameter ParamName, returns the index of the last valid float, -1 if no values left. */
+	int32 RemoveValueFromTextureRange(const FString& ParamName, int32 RangeIndex);
+
+	/** Remove the last of the projector range of values from the parameter ParamName, returns the index of the last valid projector, -1 if no values left. */
+	int32 RemoveValueFromProjectorRange(const FString& ParamName);
+
+	/** Remove the RangeIndex element of the projector range of values from the parameter ParamName, returns the index of the last valid projector, -1 if no values left. */
+	int32 RemoveValueFromProjectorRange(const FString& ParamName, int32 RangeIndex);
+
+	// ------------------------------------------------------------
+   	// States
+   	// ------------------------------------------------------------
+
+	/** Get the current optimization state. */
+	int32 GetState() const;
+
+	/** Get the current optimization state. */	
+	FString GetCurrentState() const;
+
+	/** Set the current optimization state. */
+	void SetState(int32 InState);
+
+	/** Set the current optimization state. */
+	void SetCurrentState(const FString& StateName);
+
+	// ------------------------------------------------------------
+
+	void SetRandomValues();
+	
+	void SetRandomValuesFromStream(const FRandomStream& InStream);
+
+	void SetDefaultValues();
+	
+	// ------------------------------------------------------------
+	// Multilayer Projectors
+	// ------------------------------------------------------------
+	
+	/** Given Multilayer Projector name, create a new Multilayer Projector Helper (if non-existent). See FMultilayerProjector.
+	 *
+	 * @return ture if successfully created (or was already created).
+	 */
+	bool CreateMultiLayerProjector(const FName& ProjectorParamName);
+	
+	/** Given Multilayer Projector name, remove a Multilayer Projector Helper. See FMultilayerProjector. */
+	void RemoveMultilayerProjector(const FName& ProjectorParamName);
+	
+	// Layers
+
+	/** See FMultilayerProjector::NumLayers. */
+	int32 MultilayerProjectorNumLayers(const FName& ProjectorParamName) const;
+
+	/** See FMultilayerProjector::CreateLayer. */
+	void MultilayerProjectorCreateLayer(const FName& ProjectorParamName, int32 Index);
+
+	/** See FMultilayerProjector::RemoveLayerAt. */
+	void MultilayerProjectorRemoveLayerAt(const FName& ProjectorParamName, int32 Index);
+
+	/** See FMultilayerProjector::GetLayer. */
+	FMultilayerProjectorLayer MultilayerProjectorGetLayer(const FName& ProjectorParamName, int32 Index) const;
+
+	/** See FMultilayerProjector::UpdateLayer. */
+	void MultilayerProjectorUpdateLayer(const FName& ProjectorParamName, int32 Index, const FMultilayerProjectorLayer& Layer);
+
+	// Virtual layers
+
+	/** See FMultilayerProjector::GetVirtualLayers. */
+	TArray<FName> MultilayerProjectorGetVirtualLayers(const FName& ProjectorParamName) const;
+	
+	/** See FMultilayerProjector::VirtualLayer. */
+	void MultilayerProjectorCreateVirtualLayer(const FName& ProjectorParamName, const FName& Id);
+
+	/** See FMultilayerProjector::FindOrCreateVirtualLayer. */
+	FMultilayerProjectorVirtualLayer MultilayerProjectorFindOrCreateVirtualLayer(const FName& ProjectorParamName, const FName& Id);
+
+	/** See FMultilayerProjector::RemoveVirtualLayer. */
+	void MultilayerProjectorRemoveVirtualLayer(const FName& ProjectorParamName, const FName& Id);
+
+	/** See FMultilayerProjector::GetVirtualLayer. */
+	FMultilayerProjectorVirtualLayer MultilayerProjectorGetVirtualLayer(const FName& ProjectorParamName, const FName& Id) const;
+
+	/** See FMultilayerProjector::UpdateVirtualLayer. */
+	void MultilayerProjectorUpdateVirtualLayer(const FName& ProjectorParamName, const FName& Id, const FMultilayerProjectorVirtualLayer& Layer);
+
+	/** Return a Mutable Core object containing all parameters. */
+	mu::Ptr<mu::Parameters> GetParameters() const;
+
+	FString ToString() const;
+	
+private:
+
+	UPROPERTY()
+	TObjectPtr<UCustomizableObject> CustomizableObject = nullptr;
+
+	UPROPERTY()
+	TArray<FCustomizableObjectBoolParameterValue> BoolParameters;
+
+	UPROPERTY()
+	TArray<FCustomizableObjectIntParameterValue> IntParameters;
+	
+	UPROPERTY()
+	TArray<FCustomizableObjectFloatParameterValue> FloatParameters;
+	
+	UPROPERTY()
+	TArray<FCustomizableObjectTextureParameterValue> TextureParameters;
+	
+	UPROPERTY()
+	TArray<FCustomizableObjectVectorParameterValue> VectorParameters;
+	
+	UPROPERTY()
+	TArray<FCustomizableObjectProjectorParameterValue> ProjectorParameters;
+
+	/** Mutable parameters optimization state. Transient UProperty to make it transactable. */
+	UPROPERTY(Transient)
+	int32 State = 0;
+	
+	/** If this is set to true, when updating the instance an additional step will be performed to calculate the list of instance parameters that are relevant for the current parameter values. */
+	bool bBuildParameterRelevancy = false;
+
+	/** These are the LODs Mutable can generate, they MUST NOT be used in an update (Mutable thread). */
+	int32 MinLOD = 0;
+
+	/** Array of RequestedLODs per component to generate, they MUST NOT be used in an update (Mutable thread). */
+	TArray<uint16> RequestedLODLevels;
+
+	/** Lookup of UCustomizableObjectDescriptor::IntParameters. */
+	TMap<FString, int32> IntParametersLookupTable;
+	
+	/** Multilayer Projector helpers. See FMultilayerProjector.*/
+	UPROPERTY()
+	TMap<FName, FMultilayerProjector> MultilayerProjectors;
+
+	void CreateParametersLookupTable();
+	
+	// Friends
+	friend FDescriptorHash;
+	friend UCustomizableObjectInstance;
+	friend UCustomizableInstancePrivate;
+	friend FMultilayerProjector;
+	friend FMutableUpdateCandidate;
+};
+
+
+#if UE_ENABLE_INCLUDE_ORDER_DEPRECATED_IN_5_2
+#include "MuCO/CustomizableObject.h"
+#include "MuCO/CustomizableObjectParameterTypeDefinitions.h"
+#endif
diff --git a/Engine/Plugins/Experimental/Mutable/Source/CustomizableObject/Internal/MuCO/CustomizableObjectInstancePrivate.h b/Engine/Plugins/Experimental/Mutable/Source/CustomizableObject/Internal/MuCO/CustomizableObjectInstancePrivate.h
--- a/Engine/Plugins/Experimental/Mutable/Source/CustomizableObject/Internal/MuCO/CustomizableObjectInstancePrivate.h
+++ b/Engine/Plugins/Experimental/Mutable/Source/CustomizableObject/Internal/MuCO/CustomizableObjectInstancePrivate.h
@@ -308,6 +308,8 @@ public:
 	void AdditionalAssetsAsyncLoaded(UE::Tasks::FTaskEvent CompletionEvent);
 
 	FCustomizableObjectInstanceDescriptor& GetDescriptor() const;
+
+	void SetDescriptor(const FCustomizableObjectInstanceDescriptor& InDescriptor);
 	
 	UCustomizableObjectInstance* GetPublic() const;
 
diff --git a/Engine/Plugins/Experimental/Mutable/Source/CustomizableObject/Internal/MuCO/CustomizableObjectSystemPrivate.h b/Engine/Plugins/Experimental/Mutable/Source/CustomizableObject/Internal/MuCO/CustomizableObjectSystemPrivate.h
--- a/Engine/Plugins/Experimental/Mutable/Source/CustomizableObject/Internal/MuCO/CustomizableObjectSystemPrivate.h
+++ b/Engine/Plugins/Experimental/Mutable/Source/CustomizableObject/Internal/MuCO/CustomizableObjectSystemPrivate.h
@@ -60,12 +60,8 @@ public:
 	/** Array of RequestedLODs per component to generate if this candidate is chosen */
 	TArray<uint16> RequestedLODLevels;
 
-	FMutableUpdateCandidate(UCustomizableObjectInstance* InCustomizableObjectInstance) : CustomizableObjectInstance(InCustomizableObjectInstance)
-	{
-		const FCustomizableObjectInstanceDescriptor& Descriptor = InCustomizableObjectInstance->GetDescriptor();
-		MinLOD = Descriptor.GetMinLod();
-		RequestedLODLevels = Descriptor.GetRequestedLODLevels();
-	}
+	FMutableUpdateCandidate(UCustomizableObjectInstance* InCustomizableObjectInstance);
+
 
 	FMutableUpdateCandidate(const UCustomizableObjectInstance* InCustomizableObjectInstance, const int32 InMinLOD,
 		const TArray<uint16>& InRequestedLODLevels) :
diff --git a/Engine/Plugins/Experimental/Mutable/Source/CustomizableObject/Private/MuCO/CustomizableObjectInstance.cpp b/Engine/Plugins/Experimental/Mutable/Source/CustomizableObject/Private/MuCO/CustomizableObjectInstance.cpp
--- a/Engine/Plugins/Experimental/Mutable/Source/CustomizableObject/Private/MuCO/CustomizableObjectInstance.cpp
+++ b/Engine/Plugins/Experimental/Mutable/Source/CustomizableObject/Private/MuCO/CustomizableObjectInstance.cpp
@@ -271,27 +271,21 @@ UCustomizableObjectInstance::UCustomizableObjectInstance()
 }
 
 
-const FCustomizableObjectInstanceDescriptor& UCustomizableObjectInstance::GetDescriptor() const
-{
-	return Descriptor;
-}
-
-
-void UCustomizableObjectInstance::SetDescriptor(const FCustomizableObjectInstanceDescriptor& InDescriptor)
+void UCustomizableInstancePrivate::SetDescriptor(const FCustomizableObjectInstanceDescriptor& InDescriptor)
 {
 	UCustomizableObject* InCustomizableObject = InDescriptor.GetCustomizableObject();
-	const bool bCustomizableObjectChanged = Descriptor.GetCustomizableObject() != InCustomizableObject;
+	const bool bCustomizableObjectChanged = GetPublic()->Descriptor.GetCustomizableObject() != InCustomizableObject;
 
 #if WITH_EDITOR
 	// Bind a lambda to the PostCompileDelegate and unbind from the previous object if any.
-	PrivateData->BindObjectDelegates(GetCustomizableObject(), InCustomizableObject);
+	BindObjectDelegates(GetPublic()->GetCustomizableObject(), InCustomizableObject);
 #endif
 
-	Descriptor = InDescriptor;
+	GetPublic()->Descriptor = InDescriptor;
 
 	if (bCustomizableObjectChanged)
 	{
-		PrivateData->InitCustomizableObjectData(InCustomizableObject);
+		InitCustomizableObjectData(InCustomizableObject);
 	}
 }
 
@@ -2336,7 +2330,7 @@ UCustomizableObjectInstance* UCustomizableObjectInstance::CloneStatic(UObject* O
 
 void UCustomizableObjectInstance::CopyParametersFromInstance(UCustomizableObjectInstance* Instance)
 {
-	SetDescriptor(Instance->GetDescriptor());
+	GetPrivate()->SetDescriptor(Instance->GetPrivate()->GetDescriptor());
 }
 
 
@@ -3043,6 +3037,14 @@ UCustomizableInstancePrivate* UCustomizableObjectInstance::GetPrivate() const
 }
 
 
+FMutableUpdateCandidate::FMutableUpdateCandidate(UCustomizableObjectInstance* InCustomizableObjectInstance): CustomizableObjectInstance(InCustomizableObjectInstance)
+{
+	const FCustomizableObjectInstanceDescriptor& Descriptor = InCustomizableObjectInstance->GetPrivate()->GetDescriptor();
+	MinLOD = Descriptor.GetMinLod();
+	RequestedLODLevels = Descriptor.GetRequestedLODLevels();
+}
+
+
 bool FMutableUpdateCandidate::HasBeenIssued() const
 {
 	return bHasBeenIssued;
@@ -5536,12 +5538,14 @@ FCustomizableObjectInstanceDescriptor& UCustomizableInstancePrivate::GetDescript
 	return GetPublic()->Descriptor;
 }
 
+
 const TArray<TObjectPtr<UMaterialInterface>>* UCustomizableObjectInstance::GetOverrideMaterials(int32 ComponentIndex) const
 {
 	FCustomizableInstanceComponentData* ComponentData = PrivateData->GetComponentData(ComponentIndex);
 	return ComponentData ? &ComponentData->OverrideMaterials : nullptr;
 }
 
+
 void UCustomizableInstancePrivate::AdditionalAssetsAsyncLoaded(UCustomizableObjectInstance* Public)
 {
 	MUTABLE_CPUPROFILER_SCOPE(UCustomizableInstancePrivate::AdditionalAssetsAsyncLoaded);
diff --git a/Engine/Plugins/Experimental/Mutable/Source/CustomizableObject/Private/MuCO/CustomizableObjectSystem.cpp b/Engine/Plugins/Experimental/Mutable/Source/CustomizableObject/Private/MuCO/CustomizableObjectSystem.cpp
--- a/Engine/Plugins/Experimental/Mutable/Source/CustomizableObject/Private/MuCO/CustomizableObjectSystem.cpp
+++ b/Engine/Plugins/Experimental/Mutable/Source/CustomizableObject/Private/MuCO/CustomizableObjectSystem.cpp
@@ -173,7 +173,7 @@ FUpdateContextPrivate::FUpdateContextPrivate(UCustomizableObjectInstance& InInst
 
 
 FUpdateContextPrivate::FUpdateContextPrivate(UCustomizableObjectInstance& InInstance) :
-	FUpdateContextPrivate(InInstance, InInstance.GetDescriptor())
+	FUpdateContextPrivate(InInstance, InInstance.GetPrivate()->GetDescriptor())
 {
 }
 
@@ -1282,7 +1282,7 @@ void UCustomizableObjectSystemPrivate::EnqueueUpdateSkeletalMesh(const TSharedRe
 				FString String = TEXT("DESCRIPTOR DEBUG PRINT\n");
 				String += "================================\n";				
 				String += FString::Printf(TEXT("=== DESCRIPTOR HASH ===\n%s\n"), *Context->GetCapturedDescriptorHash().ToString());
-				String += FString::Printf(TEXT("=== DESCRIPTOR ===\n%s"), *Instance->GetDescriptor().ToString());
+				String += FString::Printf(TEXT("=== DESCRIPTOR ===\n%s"), *Instance->GetPrivate()->GetDescriptor().ToString());
 				String += "================================";
 				
 				UE_LOG(LogMutable, Log, TEXT("%s"), *String);
@@ -3668,7 +3668,7 @@ int32 UCustomizableObjectSystem::TickInternal()
 			else if (LODUpdateCandidateFound)
 			{
 				UCustomizableObjectInstance* Instance = LODUpdateCandidateFound->CustomizableObjectInstance;
-				const FCustomizableObjectInstanceDescriptor& Descriptor = Instance->GetPrivate()->SkeletalMeshStatus == ESkeletalMeshStatus::NotGenerated ? Instance->GetDescriptor() : Instance->GetPrivate()->CommittedDescriptor;
+				const FCustomizableObjectInstanceDescriptor& Descriptor = Instance->GetPrivate()->SkeletalMeshStatus == ESkeletalMeshStatus::NotGenerated ? Instance->GetPrivate()->GetDescriptor() : Instance->GetPrivate()->CommittedDescriptor;
 
 				const TSharedRef<FUpdateContextPrivate> Context = MakeShared<FUpdateContextPrivate>(*Instance, Descriptor);
 
diff --git a/Engine/Plugins/Experimental/Mutable/Source/CustomizableObject/Public/MuCO/CustomizableObjectInstanceDescriptor.h b/Engine/Plugins/Experimental/Mutable/Source/CustomizableObject/Public/MuCO/CustomizableObjectInstanceDescriptor.h
deleted file mode 100644
