[Mutable] Async compile CustomizableObjects when cooking. It is currently disabled.
* Move compiler-related code to the Editor module.
* Add FCompilationRequest and convert the compiler into a queue that processes compilation requests
* Remove unused code
* Remove FCustomizableObjectCompilerBase

#rnx

diff --git a/Engine/Plugins/Experimental/Mutable/Source/CustomizableObject/Classes/MuCO/CustomizableObject.h b/Engine/Plugins/Experimental/Mutable/Source/CustomizableObject/Classes/MuCO/CustomizableObject.h
--- a/Engine/Plugins/Experimental/Mutable/Source/CustomizableObject/Classes/MuCO/CustomizableObject.h
+++ b/Engine/Plugins/Experimental/Mutable/Source/CustomizableObject/Classes/MuCO/CustomizableObject.h
@@ -3,6 +3,7 @@
 #pragma once
 
 #include "RHIDefinitions.h"
+#include "MuCO/CustomizableObjectCompilerTypes.h"
 #include "MuCO/CustomizableObjectClothingTypes.h"
 #include "MuCO/CustomizableObjectResourceData.h"
 #include "MuCO/CustomizableObjectResourceDataTypes.h"
@@ -13,9 +14,6 @@
 
 #include "CustomizableObject.generated.h"
 
-/** Index of the maximum optimization level when compiling CustomizableObjects */
-#define UE_MUTABLE_MAX_OPTIMIZATION			2
-
 class FReply;
 class FObjectPreSaveContext;
 class FText;
@@ -28,9 +26,6 @@ class UCustomizableObjectPrivate;
 struct FFrame;
 struct FStreamableHandle;
 template <typename FuncType> class TFunctionRef;
-enum class ECompilationStatePrivate : uint8;
-
-DECLARE_MULTICAST_DELEGATE(FPostCompileDelegate)
  
 // Forward declaration of Mutable classes necessary for the interface
 namespace mu
@@ -101,105 +96,6 @@ struct FProfileParameterDat
 	TArray<FCustomizableObjectProjectorParameterValue> ProjectorParameters;
 };
 
-UENUM()
-enum class ECustomizableObjectNumBoneInfluences : uint8
-{
-	// The enum values can be used as the real numeric value of number of bone influences
-	Four = 4,
-	// 
-	Eight = 8,
-	//
-	Twelve = 12 // This is essentially the same as "Unlimited", but UE ultimately limits to 12
-};
-
-UENUM()
-enum class ECustomizableObjectTextureCompression : uint8
-{
-	// Don't use texture compression
-	None = 0,
-	// Use Mutable's fast low-quality compression
-	Fast,
-	// Use Unreal's highest quality compression (100x slower to compress)
-	HighQuality
-};
-
-
-USTRUCT()
-struct FCompilationOptions
-{
-	GENERATED_USTRUCT_BODY()
-	
-	/** Enum to know what texture compression should be used. This compression is used only in manual compiles in editor. 
-	 *  When packaging, ECustomizableObjectTextureCompression::HighQuality is always used.
-	 */
-	UPROPERTY()
-	ECustomizableObjectTextureCompression TextureCompression = ECustomizableObjectTextureCompression::Fast;
-
-	// From 0 to UE_MUTABLE_MAX_OPTIMIZATION
-	UPROPERTY()
-	int32 OptimizationLevel = UE_MUTABLE_MAX_OPTIMIZATION;
-
-	// Use the disk to store intermediate compilation data. This slows down the object compilation
-	// but it may be necessary for huge objects.
-	UPROPERTY()
-	bool bUseDiskCompilation = false;
-
-	/** High limit of the size in bytes of the packaged data when cooking this object. 
-	* This limit is before any pak or filesystem compression. This limit will be broken if a single piece of data is bigger because data is not fragmented for packaging purposes.
-	*/
-	UPROPERTY()
-	uint64 PackagedDataBytesLimit = 256*1024*1024;
-
-	/** High (inclusive) limit of the size in bytes of a data block to be included into the compiled object directly instead of stored in a streamable file. */
-	UPROPERTY()
-	uint64 EmbeddedDataBytesLimit = 1024;
-
-	// Did we have the extra bones enabled when we compiled?
-	ECustomizableObjectNumBoneInfluences CustomizableObjectNumBoneInfluences = ECustomizableObjectNumBoneInfluences::Four;
-
-	// Compiling for cook
-	bool bIsCooking = false;
-
-	// This can be set for additional settings
-	const ITargetPlatform* TargetPlatform = nullptr;
-
-	// Used to prevent the current model from being stored in the editor's streamed data and cache when the compilation finishes
-	bool bDontUpdateStreamedDataAndCache = false;
-
-	// Used to enable the use of real time morph targets.
-	bool bRealTimeMorphTargetsEnabled = false;
-
-	// Used to enable the use of clothing.
-	bool bClothingEnabled = false;
-
-	// Used to enable 16 bit bone weights
-	bool b16BitBoneWeightsEnabled = false;
-
-	// Used to enable skin weight profiles.
-	bool bSkinWeightProfilesEnabled = false;
-
-	// Used to enable physics asset merge.
-	bool bPhysicsAssetMergeEnabled = false;
-
-	// Used to enable AnimBp override physics mainipualtion.  
-	bool bAnimBpPhysicsManipulationEnabled = false;
-
-	// Used to reduce the number of notifications when compiling objects
-	bool bSilentCompilation = true;
-
-	// Used to reduce texture size on higher mesh LODs. Only active if LOD strategy is set to Automatic LODs from Mesh
-	bool bUseLODAsBias = true;
-
-	/** Force a very big number on the mips to skip during compilation. Useful to debug special cooks of the data. */
-	bool bForceLargeLODBias = false;
-	int32 DebugBias = 0;
-
-	// Control image tiled generation
-	int32 ImageTiling = 0;
-
-};
-
-
 // A USTRUCT version of FMeshToMeshVertData in SkeletalMeshTypes.h
 // We are taking advantage of the padding data to store from which asset this data comes from
 // maintaining the same memory footprint than the original.
diff --git a/Engine/Plugins/Experimental/Mutable/Source/CustomizableObject/Classes/MuCO/CustomizableObjectSystem.h b/Engine/Plugins/Experimental/Mutable/Source/CustomizableObject/Classes/MuCO/CustomizableObjectSystem.h
--- a/Engine/Plugins/Experimental/Mutable/Source/CustomizableObject/Classes/MuCO/CustomizableObjectSystem.h
+++ b/Engine/Plugins/Experimental/Mutable/Source/CustomizableObject/Classes/MuCO/CustomizableObjectSystem.h
@@ -11,17 +11,11 @@
 #include "MuR/Image.h"
 #include "MuR/Types.h"
 
-#if WITH_EDITOR
-#include "Framework/Notifications/NotificationManager.h"
-#endif
-
 #include "CustomizableObjectSystem.generated.h"
 
 class IConsoleVariable;
-class FCustomizableObjectCompilerBase;
 class UCustomizableInstanceLODManagementBase;
 class ITargetPlatform;
-class SNotificationItem;
 class UCustomizableObject;
 class UDefaultImageProvider;
 class USkeletalMesh;
@@ -198,9 +192,6 @@ public:
 	virtual FString GetDesc() override;
 	// End UObject interface.
 
-	// Creates a new Customizable Object Compiler (Only does real work in editor builds). The caller is responsible for freeing the new compiler
-	static TSharedPtr<FCustomizableObjectCompilerBase> GetNewCompiler();
-
 	bool IsReplaceDiscardedWithReferenceMeshEnabled() const;
 	void SetReplaceDiscardedWithReferenceMeshEnabled(bool bIsEnabled);
 
@@ -327,10 +318,6 @@ public:
 	bool IsMutableAnimInfoDebuggingEnabled() const;
 
 #if WITH_EDITOR
-	void RecompileCustomizableObjectAsync(const FAssetData& InAssetData, const UCustomizableObject* InObject);
-	
-	void RecompileCustomizableObjects(const TArray<FAssetData>& InObjects);
-
 	// Get the maximum size a chunk can have on a specific platform. If unspecified return MUTABLE_STREAMED_DATA_MAXCHUNKSIZE
 	uint64 GetMaxChunkSizeForPlatform(const ITargetPlatform* InTargetPlatform);
 #endif
@@ -355,17 +342,6 @@ public:
 	UFUNCTION(BlueprintCallable, BlueprintPure, Category = CustomizableObjectSystem)
 	bool IsUpdating(const UCustomizableObjectInstance* Instance) const;
 
-private:
-	// TODO: Can we move this to the editor module?
-#if WITH_EDITOR
-	// Used to ask the user if they want to recompile uncompiled PIE COs
-	void OnPreBeginPIE(const bool bIsSimulatingInEditor);
-
-	void StartNextRecompile();
-	void TickRecompileCustomizableObjects();
-#endif
-
-public:
 	/** Set Mutable's working memory limit (bytes). Mutable will flush internal caches to try to keep its memory consumption below the WorkingMemory (i.e., it is not a hard limit).
 	 * The working memory limit will especially reduce the memory required to perform Instance Updates and Texture Streaming.
  	 * Notice that Mutable does not track all its memory (e.g., UObjects memory is no tracked).
diff --git a/Engine/Plugins/Experimental/Mutable/Source/CustomizableObject/Internal/MuCO/CustomizableObjectCompilerTypes.h b/Engine/Plugins/Experimental/Mutable/Source/CustomizableObject/Internal/MuCO/CustomizableObjectCompilerTypes.h
new file mode 100644
--- /dev/null
+++ b/Engine/Plugins/Experimental/Mutable/Source/CustomizableObject/Internal/MuCO/CustomizableObjectCompilerTypes.h
@@ -0,0 +1,171 @@
+// Copyright Epic Games, Inc. All Rights Reserved.
+
+#pragma once
+
+#include "CustomizableObjectCompilerTypes.generated.h"
+
+
+/** Index of the maximum optimization level when compiling CustomizableObjects */
+#define UE_MUTABLE_MAX_OPTIMIZATION			2
+
+class UCustomizableObject;
+
+DECLARE_MULTICAST_DELEGATE(FPostCompileDelegate)
+
+UENUM()
+enum class ECustomizableObjectNumBoneInfluences : uint8
+{
+	// The enum values can be used as the real numeric value of number of bone influences
+	Four = 4,
+	// 
+	Eight = 8,
+	//
+	Twelve = 12 // This is essentially the same as "Unlimited", but UE ultimately limits to 12
+};
+
+//#if WITH_EDITORONLY_DATA
+
+UENUM()
+enum class ECustomizableObjectTextureCompression : uint8
+{
+	// Don't use texture compression
+	None = 0,
+	// Use Mutable's fast low-quality compression
+	Fast,
+	// Use Unreal's highest quality compression (100x slower to compress)
+	HighQuality
+};
+
+
+USTRUCT()
+struct FCompilationOptions
+{
+	GENERATED_USTRUCT_BODY()
+
+	/** Enum to know what texture compression should be used. This compression is used only in manual compiles in editor.
+	 *  When packaging, ECustomizableObjectTextureCompression::HighQuality is always used.
+	 */
+	UPROPERTY()
+	ECustomizableObjectTextureCompression TextureCompression = ECustomizableObjectTextureCompression::Fast;
+
+	// From 0 to UE_MUTABLE_MAX_OPTIMIZATION
+	UPROPERTY()
+	int32 OptimizationLevel = UE_MUTABLE_MAX_OPTIMIZATION;
+
+	// Use the disk to store intermediate compilation data. This slows down the object compilation
+	// but it may be necessary for huge objects.
+	UPROPERTY()
+	bool bUseDiskCompilation = false;
+
+	/** High limit of the size in bytes of the packaged data when cooking this object.
+	* This limit is before any pak or filesystem compression. This limit will be broken if a single piece of data is bigger because data is not fragmented for packaging purposes.
+	*/
+	UPROPERTY()
+	uint64 PackagedDataBytesLimit = 256 * 1024 * 1024;
+
+	/** High (inclusive) limit of the size in bytes of a data block to be included into the compiled object directly instead of stored in a streamable file. */
+	UPROPERTY()
+	uint64 EmbeddedDataBytesLimit = 1024;
+
+	// Did we have the extra bones enabled when we compiled?
+	ECustomizableObjectNumBoneInfluences CustomizableObjectNumBoneInfluences = ECustomizableObjectNumBoneInfluences::Four;
+
+	// Compiling for cook
+	bool bIsCooking = false;
+
+	// This can be set for additional settings
+	const ITargetPlatform* TargetPlatform = nullptr;
+
+	// Used to enable the use of real time morph targets.
+	bool bRealTimeMorphTargetsEnabled = false;
+
+	// Used to enable the use of clothing.
+	bool bClothingEnabled = false;
+
+	// Used to enable 16 bit bone weights
+	bool b16BitBoneWeightsEnabled = false;
+
+	// Used to enable skin weight profiles.
+	bool bSkinWeightProfilesEnabled = false;
+
+	// Used to enable physics asset merge.
+	bool bPhysicsAssetMergeEnabled = false;
+
+	// Used to enable AnimBp override physics mainipualtion.  
+	bool bAnimBpPhysicsManipulationEnabled = false;
+
+	// Used to reduce the number of notifications when compiling objects
+	bool bSilentCompilation = true;
+
+	// Used to reduce texture size on higher mesh LODs. Only active if LOD strategy is set to Automatic LODs from Mesh
+	bool bUseLODAsBias = true;
+
+	/** Force a very big number on the mips to skip during compilation. Useful to debug special cooks of the data. */
+	bool bForceLargeLODBias = false;
+	int32 DebugBias = 0;
+
+	// Control image tiled generation
+	int32 ImageTiling = 0;
+
+};
+
+
+enum class ECompilationStatePrivate : uint8
+{
+	None,
+	InProgress,
+	Completed
+};
+
+
+enum class ECompilationResultPrivate : uint8
+{
+	Unknown, // Not compiled yet (compilation may be in progress).
+	Success, // No errors or warnings.
+	Errors, // At least have one error. Can have warnings.
+	Warnings, // Only warnings.
+};
+
+
+struct CUSTOMIZABLEOBJECT_API FCompilationRequest
+{
+	FCompilationRequest(UCustomizableObject& CustomizableObject, bool bAsync);
+
+	UCustomizableObject* GetCustomizableObject();
+
+	FCompilationOptions& GetCompileOptions();
+
+	bool IsAsyncCompilation() const;
+
+	void SetCompilationState(ECompilationStatePrivate InState, ECompilationResultPrivate InResult);
+
+	ECompilationStatePrivate GetCompilationState() const;
+	ECompilationResultPrivate GetCompilationResult() const;
+
+	TArray<FText>& GetWarnings();
+	TArray<FText>& GetErrors();
+
+	void SetParameterNamesToSelectedOptions(const TMap<FString, FString>& InParamNamesToSelectedOptions);
+	const TMap<FString, FString>& GetParameterNamesToSelectedOptions() const;
+
+	bool operator==(const FCompilationRequest& Other) const;
+
+private:
+
+	TWeakObjectPtr<UCustomizableObject> CustomizableObject;
+
+	FCompilationOptions Options;
+
+	ECompilationStatePrivate State = ECompilationStatePrivate::None;
+	ECompilationResultPrivate Result = ECompilationResultPrivate::Unknown;
+
+	bool bAsync = true;
+
+	// Stores the only option of an Int Param that should be compiled
+	TMap<FString, FString> ParamNamesToSelectedOptions;
+
+	TArray<FText> Warnings;
+	TArray<FText> Errors;
+};
+
+//#endif // WITH_EDITORONLY_DATA
\ No newline at end of file
diff --git a/Engine/Plugins/Experimental/Mutable/Source/CustomizableObject/Internal/MuCO/CustomizableObjectPrivate.h b/Engine/Plugins/Experimental/Mutable/Source/CustomizableObject/Internal/MuCO/CustomizableObjectPrivate.h
--- a/Engine/Plugins/Experimental/Mutable/Source/CustomizableObject/Internal/MuCO/CustomizableObjectPrivate.h
+++ b/Engine/Plugins/Experimental/Mutable/Source/CustomizableObject/Internal/MuCO/CustomizableObjectPrivate.h
@@ -3,6 +3,7 @@
 #pragma once
 
 #include "MuCO/CustomizableObject.h"
+#include "MuCO/CustomizableObjectCompilerTypes.h"
 #include "MuCO/StateMachine.h"
 #include "MuCO/CustomizableObjectUIData.h"
 #include "MuCO/CustomizableObjectIdentifier.h"
@@ -35,27 +36,6 @@ class UCustomizableObject;
 
 FGuid CUSTOMIZABLEOBJECT_API GenerateIdentifier(const UCustomizableObject& CustomizableObject);
 
-
-class CUSTOMIZABLEOBJECT_API FCustomizableObjectCompilerBase
-{
-public:
-
-	FCustomizableObjectCompilerBase() {};
-
-	// Ensure virtual destruction
-	virtual ~FCustomizableObjectCompilerBase() {};
-	
-	virtual void Compile(UCustomizableObject& Object, const FCompilationOptions& Options, bool bAsync) {};
-
-	/** Return true if no further ticks are required. In other words, all work has completed. */
-	virtual bool Tick(bool bBlocking) { return true; }
-	virtual void ForceFinishCompilation() {}
-
-	/** Provides the caller with the warning and error messages produced during compilation */
-	virtual void GetCompilationMessages(TArray<FText>& OutWarningMessages, TArray<FText>& OutErrorMessages) const = 0;
-};
-
-
 // Warning! MutableCompiledDataHeader must be the first data serialized in a stream
 struct MutableCompiledDataStreamHeader
 {
@@ -199,24 +179,6 @@ struct FCustomizableObjectStatusTypes
 using FCustomizableObjectStatus = FStateMachine<FCustomizableObjectStatusTypes>;
 
 
-
-enum class ECompilationStatePrivate : uint8
-{
-	None,
-	InProgress,
-	Completed
-};
-
-
-enum class ECompilationResultPrivate : uint8
-{
-	Unknown, // Not compiled yet (compilation may be in progress).
-	Success, // No errors or warnings.
-	Errors, // At least have one error. Can have warnings.
-	Warnings, // Only warnings.
-};
-
-
 USTRUCT()
 struct CUSTOMIZABLEOBJECT_API FMutableModelImageProperties
 {
@@ -724,8 +686,6 @@ class CUSTOMIZABLEOBJECT_API UCustomizableObjectPrivate : public UObject
 {
 	GENERATED_BODY()
 
-	
-
 	TSharedPtr<mu::Model, ESPMode::ThreadSafe> MutableModel;
 
 	/** Stores resources to be used by MutableModel In-Game. Cooked resources. */
@@ -883,8 +843,8 @@ public:
 	/** Unique Identifier - Deterministic. Used to locate Model and Streamable data on disk. Should not be modified. */
 	FGuid Identifier;
 
-	/** Cooked platform names. */
-	TArray<FString> CachedPlatformNames;
+	/** Cook requests. */
+	TArray<TSharedRef<FCompilationRequest>> CompileRequests;
 
 	/** List of external packages that if changed, a compilation is required.
 	 * Key is the package name. Value is the the UPackage::Guid, which is regenerated each time the packages is saved.
diff --git a/Engine/Plugins/Experimental/Mutable/Source/CustomizableObject/Internal/MuCO/CustomizableObjectSystemPrivate.h b/Engine/Plugins/Experimental/Mutable/Source/CustomizableObject/Internal/MuCO/CustomizableObjectSystemPrivate.h
--- a/Engine/Plugins/Experimental/Mutable/Source/CustomizableObject/Internal/MuCO/CustomizableObjectSystemPrivate.h
+++ b/Engine/Plugins/Experimental/Mutable/Source/CustomizableObject/Internal/MuCO/CustomizableObjectSystemPrivate.h
@@ -794,16 +794,7 @@ public:
 	// For async material loading
 	FStreamableManager StreamableManager;
 	
-#if WITH_EDITOR
-	TSharedPtr<FCustomizableObjectCompilerBase> RecompileCustomizableObjectsCompiler;
-	
-	TArray<FAssetData> ObjectsToRecompile;
-	uint32 TotalNumObjectsToRecompile = 0;
-	uint32 NumObjectsCompiled = 0;
-
-	/** Recompile progress bar handle */
-	FProgressNotificationHandle RecompileNotificationHandle;
-
+#if WITH_EDITORONLY_DATA
 	// Array to keep track of cached objects
 	TArray<FGuid> UncompiledCustomizableObjectIds;
 
diff --git a/Engine/Plugins/Experimental/Mutable/Source/CustomizableObject/Internal/MuCO/ICustomizableObjectEditorModule.h b/Engine/Plugins/Experimental/Mutable/Source/CustomizableObject/Internal/MuCO/ICustomizableObjectEditorModule.h
--- a/Engine/Plugins/Experimental/Mutable/Source/CustomizableObject/Internal/MuCO/ICustomizableObjectEditorModule.h
+++ b/Engine/Plugins/Experimental/Mutable/Source/CustomizableObject/Internal/MuCO/ICustomizableObjectEditorModule.h
@@ -5,7 +5,6 @@
 #include "Modules/ModuleManager.h"
 #include "Templates/SharedPointer.h"
 
-struct FBakingConfiguration;
 class FCustomizableObjectEditorLogger;
 class ICustomizableObjectDebugger;
 class ICustomizableObjectEditor;
@@ -15,8 +14,9 @@ class UCustomizableObject;
 class UCustomizableObjectPrivate;
 class UCustomizableObjectInstance;
 class FExtensibilityManager;
-class FCustomizableObjectCompilerBase;
 class FBakeOperationCompletedDelegate;
+struct FBakingConfiguration;
+struct FCompilationRequest;
 
 extern const FName CustomizableObjectEditorAppIdentifier;
 extern const FName CustomizableObjectInstanceEditorAppIdentifier;
@@ -27,7 +27,7 @@ extern const FName CustomizableObjectDebuggerAppIdentifier;
 /**
  * Customizable object editor module interface
  */
-class ICustomizableObjectEditorModule : public IModuleInterface
+class CUSTOMIZABLEOBJECT_API ICustomizableObjectEditorModule : public IModuleInterface
 {
 public:
 	static ICustomizableObjectEditorModule* Get()
@@ -52,9 +52,6 @@ public:
    	  * @return true if the compilation is out of date. */
 	virtual bool IsCompilationOutOfDate(const UCustomizableObject& Object, TArray<FName>* OutOfDatePackages = nullptr) const = 0;
 
-	/** Create a new compiler. */
-	virtual TSharedRef<FCustomizableObjectCompilerBase> CreateCompiler() const = 0;
-
 	/** See GraphTraversal::IsRootObject(...) */
 	virtual bool IsRootObject(const UCustomizableObject& Object) const = 0;
 	
@@ -63,6 +60,18 @@ public:
 	 * @param InTargetInstance The instance we want to bake
 	 * @param InBakingConfig Structure containing the configuration to be used for the baking
 	 */
-	virtual void BakeCustomizableObjectInstance(UCustomizableObjectInstance* InTargetInstance, const FBakingConfiguration& InBakingConfig)  = 0;
+	virtual void BakeCustomizableObjectInstance(UCustomizableObjectInstance* InTargetInstance, const FBakingConfiguration& InBakingConfig) = 0;
+
+	/** Request for a given customizable object to be compiled. Async compile requests will be queued and processed sequentially.
+	 * @param InCompilationRequest - Request to compile an object.
+	 * @param bForceRequest - Queue request even if already in the pending list. */
+	virtual void CompileCustomizableObject(const TSharedRef<FCompilationRequest>& InCompilationRequest, bool bForceRequest = false) = 0;
+	virtual void CompileCustomizableObjects(const TArray<TSharedRef<FCompilationRequest>>& InCompilationRequests, bool bForceRequests = false) = 0;
+
+	virtual int32 Tick(bool bBlocking) = 0;
+
+	/** Force finish current compile request and cancels all pending requests */
+	virtual void CancelCompileRequests() = 0;
+
 };
  
diff --git a/Engine/Plugins/Experimental/Mutable/Source/CustomizableObject/Private/MuCO/CustomizableObject.cpp b/Engine/Plugins/Experimental/Mutable/Source/CustomizableObject/Private/MuCO/CustomizableObject.cpp
--- a/Engine/Plugins/Experimental/Mutable/Source/CustomizableObject/Private/MuCO/CustomizableObject.cpp
+++ b/Engine/Plugins/Experimental/Mutable/Source/CustomizableObject/Private/MuCO/CustomizableObject.cpp
@@ -333,18 +333,35 @@ void UCustomizableObject::BeginCacheForCookedPlatformData(const ITargetPlatform*
 		return;
 	}
 
-	if (Private->CachedPlatformNames.Find(TargetPlatform->PlatformName()) == INDEX_NONE)
+	const TSharedRef<FCompilationRequest>* CompileRequest = GetPrivate()->CompileRequests.FindByPredicate(
+		[&TargetPlatform](const TSharedPtr<FCompilationRequest>& Request) { return Request->GetCompileOptions().TargetPlatform == TargetPlatform; });
+	
+	if (CompileRequest)
 	{
-		// Compile and save in the CachedPlatformsData map
-		GetPrivate()->CompileForTargetPlatform(TargetPlatform);
-		GetPrivate()->CachedPlatformNames.Add(TargetPlatform->PlatformName());
+		return;
 	}
+
+	// Compile and save in the CachedPlatformsData map
+	GetPrivate()->CompileForTargetPlatform(TargetPlatform);
 }
 
 
-bool UCustomizableObject::IsCachedCookedPlatformDataLoaded( const ITargetPlatform* TargetPlatform ) 
+bool UCustomizableObject::IsCachedCookedPlatformDataLoaded(const ITargetPlatform* TargetPlatform) 
 {
-	return !TargetPlatform || Private->CachedPlatformNames.Find(TargetPlatform->PlatformName()) != INDEX_NONE;
+	if (!TargetPlatform)
+	{
+		return true;
+	}
+
+	const TSharedRef<FCompilationRequest>* CompileRequest = GetPrivate()->CompileRequests.FindByPredicate(
+		[&TargetPlatform](const TSharedRef<FCompilationRequest>& Request) { return Request->GetCompileOptions().TargetPlatform == TargetPlatform; });
+	
+	if (CompileRequest)
+	{
+		return CompileRequest->Get().GetCompilationState() == ECompilationStatePrivate::Completed;
+	}
+
+	return true;
 }
 
 
@@ -832,32 +849,26 @@ void UCustomizableObjectPrivate::CompileForTargetPlatform(const ITargetPlatform*
 		return;
 	}
 
-	const ICustomizableObjectEditorModule* EditorModule = ICustomizableObjectEditorModule::Get();
-	if (!EditorModule)
-	{
-		return;
-	}
+	UCustomizableObject* CustomizableObject = GetPublic();
 
-	const bool bIsRootObject = EditorModule->IsRootObject(*GetPublic());
-	if (!bIsRootObject)
+	ICustomizableObjectEditorModule* EditorModule = ICustomizableObjectEditorModule::Get();
+	if (!EditorModule || !EditorModule->IsRootObject(*CustomizableObject))
 	{
 		return;
 	}
 
-	const TSharedPtr<FCustomizableObjectCompilerBase> Compiler = UCustomizableObjectSystem::GetInstance()->GetNewCompiler();
-	if (!Compiler)
-	{
-		return;
-	}
+	const bool bAsync = false; // TODO PERE
 
-	FCompilationOptions Options;
-	Options.OptimizationLevel = UE_MUTABLE_MAX_OPTIMIZATION;	// max optimization when packaging.
+	TSharedRef<FCompilationRequest> CompileRequest = MakeShared<FCompilationRequest>(*CustomizableObject, bAsync);
+	FCompilationOptions& Options = CompileRequest->GetCompileOptions();
+	Options.OptimizationLevel = UE_MUTABLE_MAX_OPTIMIZATION;	// Force max optimization when packaging.
 	Options.TextureCompression = ECustomizableObjectTextureCompression::HighQuality;
 	Options.bIsCooking = true;
 	Options.TargetPlatform = TargetPlatform;
 	Options.CustomizableObjectNumBoneInfluences = ICustomizableObjectModule::Get().GetNumBoneInfluences();
+	CompileRequests.Add(CompileRequest);
 
-	Compiler->Compile(*GetPublic(), Options, false);
+	EditorModule->CompileCustomizableObject(CompileRequest, true);
 }
 
 
@@ -909,19 +920,13 @@ bool UCustomizableObject::ConditionalAutoCompile()
 		return false;
 	}
 
-	// Sync/Async compilation
-	if (System->IsAutoCompilationSync())
-	{
-		if (const TSharedPtr<FCustomizableObjectCompilerBase> Compiler = UCustomizableObjectSystem::GetNewCompiler())
-		{		
-			Compiler->Compile(*this, this->CompileOptions, false);
-		}
-	}
-	else
+	ICustomizableObjectEditorModule* EditorModule = ICustomizableObjectEditorModule::Get();
+	if (ensure(EditorModule))
 	{
-		const FAssetRegistryModule& AssetRegistryModule = FModuleManager::GetModuleChecked<FAssetRegistryModule>(TEXT("AssetRegistry"));
-		const FAssetData AssetData = AssetRegistryModule.Get().GetAssetByObjectPath(UE_MUTABLE_OBJECTPATH(GetPathName()));
-		System->RecompileCustomizableObjectAsync(AssetData, this);
+		// Sync/Async compilation
+		TSharedRef<FCompilationRequest> CompileRequest = MakeShared<FCompilationRequest>(*this, !System->IsAutoCompilationSync());
+		CompileRequest->GetCompileOptions().bSilentCompilation = true;
+		EditorModule->CompileCustomizableObject(CompileRequest);
 	}
 
 	return IsCompiled();
@@ -2011,8 +2016,6 @@ TObjectPtr<UEdGraph>& UCustomizableObjectPrivate::GetSource() const
 //-------------------------------------------------------------------------------------------------
 //-------------------------------------------------------------------------------------------------
 //-------------------------------------------------------------------------------------------------
-
-
 void UCustomizableObjectBulk::PostLoad()
 {
 	UObject::PostLoad();
@@ -2641,6 +2644,73 @@ FArchive& operator<<(FArchive& Ar, FMutableRefSkeletalMeshData& Data)
 	return Ar;
 }
 
+
+//-------------------------------------------------------------------------------------------------
+//-------------------------------------------------------------------------------------------------
+//-------------------------------------------------------------------------------------------------
+FCompilationRequest::FCompilationRequest(UCustomizableObject& InCustomizableObject, bool bAsyncCompile)
+{
+	CustomizableObject = &InCustomizableObject;
+	Options = InCustomizableObject.CompileOptions;
+	bAsync = bAsyncCompile;
+}
+
+UCustomizableObject* FCompilationRequest::GetCustomizableObject()
+{
+	return CustomizableObject.Get();
+}
+
+FCompilationOptions& FCompilationRequest::GetCompileOptions()
+{
+	return Options;
+}
+
+bool FCompilationRequest::IsAsyncCompilation() const
+{
+	return bAsync;
+}
+
+void FCompilationRequest::SetCompilationState(ECompilationStatePrivate InState, ECompilationResultPrivate InResult)
+{
+	State = InState;
+	Result = InResult;
+}
+
+ECompilationStatePrivate FCompilationRequest::GetCompilationState() const
+{
+	return State;
+}
+
+ECompilationResultPrivate FCompilationRequest::GetCompilationResult() const
+{
+	return Result;
+}
+
+TArray<FText>& FCompilationRequest::GetWarnings()
+{
+	return Warnings;
+}
+
+TArray<FText>& FCompilationRequest::GetErrors()
+{
+	return Errors;
+}
+
+void FCompilationRequest::SetParameterNamesToSelectedOptions(const TMap<FString, FString>& InParamNamesToSelectedOptions)
+{
+	ParamNamesToSelectedOptions = InParamNamesToSelectedOptions;
+}
+
+const TMap<FString, FString>& FCompilationRequest::GetParameterNamesToSelectedOptions() const
+{
+	return ParamNamesToSelectedOptions;
+}
+
+bool FCompilationRequest::operator==(const FCompilationRequest& Other) const
+{
+	return CustomizableObject == Other.CustomizableObject && Options.TargetPlatform == Other.Options.TargetPlatform;
+};
+
 #endif // WITH_EDITORONLY_DATA
 
 #undef LOCTEXT_NAMESPACE
diff --git a/Engine/Plugins/Experimental/Mutable/Source/CustomizableObject/Private/MuCO/CustomizableObjectModule.cpp b/Engine/Plugins/Experimental/Mutable/Source/CustomizableObject/Private/MuCO/CustomizableObjectModule.cpp
--- a/Engine/Plugins/Experimental/Mutable/Source/CustomizableObject/Private/MuCO/CustomizableObjectModule.cpp
+++ b/Engine/Plugins/Experimental/Mutable/Source/CustomizableObject/Private/MuCO/CustomizableObjectModule.cpp
@@ -7,6 +7,7 @@
 #include "Kismet/GameplayStatics.h"
 #include "Misc/ConfigCacheIni.h"
 #include "MuCO/CustomizableObject.h"
+#include "MuCO/CustomizableObjectCompilerTypes.h"
 #include "MuCO/CustomizableObjectDGGUI.h"
 #include "MuCO/CustomizableObjectExtension.h"
 #include "MuCO/CustomizableObjectInstanceUsage.h"
diff --git a/Engine/Plugins/Experimental/Mutable/Source/CustomizableObject/Private/MuCO/CustomizableObjectSystem.cpp b/Engine/Plugins/Experimental/Mutable/Source/CustomizableObject/Private/MuCO/CustomizableObjectSystem.cpp
--- a/Engine/Plugins/Experimental/Mutable/Source/CustomizableObject/Private/MuCO/CustomizableObjectSystem.cpp
+++ b/Engine/Plugins/Experimental/Mutable/Source/CustomizableObject/Private/MuCO/CustomizableObjectSystem.cpp
@@ -3,14 +3,12 @@
 #include "MuCO/CustomizableObjectSystem.h"
 
 #include "Animation/Skeleton.h"
-#include "AssetRegistry/ARFilter.h"
 #include "AssetRegistry/AssetRegistryModule.h"
 #include "Engine/SkeletalMesh.h"
 #include "Engine/SkinnedAsset.h"
 #include "Engine/SkinnedAssetCommon.h"
 #include "Engine/SkeletalMeshLODSettings.h"
 #include "GameFramework/PlayerController.h"
-#include "Interfaces/ITargetPlatform.h"
 #include "MuCO/CustomizableInstanceLODManagement.h"
 #include "MuCO/CustomizableObjectInstancePrivate.h"
 #include "MuCO/CustomizableObjectPrivate.h"
@@ -29,18 +27,18 @@
 #include "UObject/UObjectIterator.h"
 #include "Widgets/Notifications/SNotificationList.h"
 #include "ContentStreaming.h"
-#include "MuCO/EditorImageProvider.h"
 #include "Components/SkeletalMeshComponent.h"
 #include "MuCO/CustomizableObjectSystemPrivate.h"
 #include "CustomizableObjectSettings.h"
-#include "MuCO/ICustomizableObjectEditorModule.h"
 
 #if WITH_EDITOR
 #include "Editor.h"
 #include "Logging/MessageLog.h"
 #include "Misc/ConfigCacheIni.h"
-#include "Misc/MessageDialog.h"
 #include "Engine/World.h"
+#include "Interfaces/ITargetPlatform.h"
+#include "MuCO/ICustomizableObjectEditorModule.h"
+#include "MuCO/EditorImageProvider.h"
 #else
 #include "Engine/Engine.h"
 #endif
@@ -593,21 +591,13 @@ void UCustomizableObjectSystem::InitSystem()
 
 void UCustomizableObjectSystem::BeginDestroy()
 {
-#if WITH_EDITOR
-	if (!IsRunningGame())
-	{
-		FEditorDelegates::PreBeginPIE.RemoveAll(this);
-	}
-
-#endif
-
 	// It could be null, for the default object.
 	if (Private)
 	{
 #if WITH_EDITOR
-		if (Private->RecompileCustomizableObjectsCompiler)
+		if (ICustomizableObjectEditorModule* EditorModule = ICustomizableObjectEditorModule::Get())
 		{
-			Private->RecompileCustomizableObjectsCompiler->ForceFinishCompilation();
+			EditorModule->CancelCompileRequests();
 		}
 #endif
 
@@ -653,19 +643,6 @@ FString UCustomizableObjectSystem::GetDesc()
 }
 
 
-TSharedPtr<FCustomizableObjectCompilerBase> UCustomizableObjectSystem::GetNewCompiler()
-{
-	if (const ICustomizableObjectEditorModule* Module = ICustomizableObjectEditorModule::Get())
-	{
-		return Module->CreateCompiler();
-	}
-	else
-	{
-		return nullptr;
-	}
-}
-
-
 int32 UCustomizableObjectSystemPrivate::EnableMutableAnimInfoDebugging = 0;
 
 static FAutoConsoleVariableRef CVarEnableMutableAnimInfoDebugging(
@@ -3705,7 +3682,6 @@ int32 UCustomizableObjectSystem::TickInternal()
 	}
 
 #if WITH_EDITOR
-	TickRecompileCustomizableObjects();
 	GetPrivate()->TickMutableThreadDependencies();
 #endif
 	
@@ -3733,7 +3709,11 @@ int32 UCustomizableObjectSystem::TickInternal()
 		RemainingTasks;
 
 #if WITH_EDITOR
-	RemainingWork += static_cast<int32>(GetPrivate()->RecompileCustomizableObjectsCompiler != nullptr); // Compiler only is valid if we are compiling a CO.
+	if (Private->bBlocking)
+	{
+		ICustomizableObjectEditorModule* EditorModule = ICustomizableObjectEditorModule::Get();
+		RemainingWork += EditorModule ? EditorModule->Tick(true) : 0;
+	}
 #endif
 	
 	return RemainingWork;
@@ -4033,170 +4013,6 @@ void UCustomizableObjectSystem::SetReleaseMutableTexturesImmediately(bool bRelea
 
 #if WITH_EDITOR
 
-void UCustomizableObjectSystem::OnPreBeginPIE(const bool bIsSimulatingInEditor)
-{
-	if (!EditorSettings.bCompileRootObjectsOnStartPIE || IsRunningGame())
-	{
-		return;
-	}
-	
-	// Find root customizable objects
-	FARFilter AssetRegistryFilter;
-	UE_MUTABLE_GET_CLASSPATHS(AssetRegistryFilter).Add(UE_MUTABLE_TOPLEVELASSETPATH(TEXT("/Script/CustomizableObject"), TEXT("CustomizableObject")));
-	AssetRegistryFilter.TagsAndValues.Add(FName("IsRoot"), FString::FromInt(1));
-
-	TArray<FAssetData> OutAssets;
-	const FAssetRegistryModule& AssetRegistryModule = FModuleManager::LoadModuleChecked<FAssetRegistryModule>("AssetRegistry");
-	AssetRegistryModule.Get().GetAssets(AssetRegistryFilter, OutAssets);
-
-	TArray<FAssetData> TempObjectsToRecompile;
-	for (const FAssetData& Asset : OutAssets)
-	{
-		// If it is referenced by PIE it should be loaded
-		if (!Asset.IsAssetLoaded())
-		{
-			continue;
-		}
-		
-		const UCustomizableObject* Object = Cast<UCustomizableObject>(Asset.GetAsset());
-		if (!Object || Object->IsCompiled() || Object->GetPrivate()->IsLocked() || Object->IsChildObject())
-		{
-			continue;
-		}
-		
-		// Add uncompiled objects to the objects to cook list
-		TempObjectsToRecompile.Add(Asset);
-	}
-
-	if (!TempObjectsToRecompile.IsEmpty())
-	{
-		const FText Msg = FText::FromString(TEXT("Warning: one or more Customizable Objects used in PIE are uncompiled.\n\nDo you want to compile them?"));
-		if (FMessageDialog::Open(EAppMsgType::OkCancel, Msg) == EAppReturnType::Ok)
-		{
-			GetPrivate()->ObjectsToRecompile.Empty(TempObjectsToRecompile.Num());
-			RecompileCustomizableObjects(TempObjectsToRecompile);
-		}
-	}
-}
-
-void UCustomizableObjectSystem::StartNextRecompile()
-{
-	if (GEngine)
-	{
-		GEngine->ForceGarbageCollection();
-	}
-
-	FAssetData Itr = GetPrivate()->ObjectsToRecompile.Pop();
-
-	if (UCustomizableObject* CustomizableObject = Cast<UCustomizableObject>(Itr.GetAsset()))
-	{
-		const FText UpdateMsg = FText::FromString(FString::Printf(TEXT("Compiling Customizable Objects:\n%s"), *CustomizableObject->GetName()));
-		FSlateNotificationManager::Get().UpdateProgressNotification(GetPrivate()->RecompileNotificationHandle, GetPrivate()->NumObjectsCompiled, GetPrivate()->TotalNumObjectsToRecompile, UpdateMsg);
-
-		// Use default options
-		FCompilationOptions Options = CustomizableObject->CompileOptions;
-		Options.bSilentCompilation = true;
-		check(GetPrivate()->RecompileCustomizableObjectsCompiler != nullptr);
-		GetPrivate()->RecompileCustomizableObjectsCompiler->Compile(*CustomizableObject, Options, !GetPrivate()->bBlocking);
-	}
-}
-
-void UCustomizableObjectSystem::RecompileCustomizableObjectAsync(const FAssetData& InAssetData,
-	const UCustomizableObject* InObject)
-{
-	if (!IsActive() || IsRunningGame())
-	{
-		return;
-	}
-	
-	if ((InObject && InObject->GetPrivate()->IsLocked()) || GetPrivate()->ObjectsToRecompile.Find((InAssetData)) != INDEX_NONE)
-	{
-		return;
-	}
-	
-	if (!GetPrivate()->ObjectsToRecompile.IsEmpty())
-	{
-		GetPrivate()->ObjectsToRecompile.Add(InAssetData);
-	}
-	else
-	{
-		RecompileCustomizableObjects({InAssetData});
-	}
-}
-
-void UCustomizableObjectSystem::RecompileCustomizableObjects(const TArray<FAssetData>& InObjects)
-{
-	if (!IsActive() || IsRunningGame())
-	{
-		return;
-	}
-
-	if (InObjects.Num())
-	{
-		if (!GetPrivate()->RecompileCustomizableObjectsCompiler)
-		{
-			GetPrivate()->RecompileCustomizableObjectsCompiler = GetNewCompiler();
-
-			if (!GetPrivate()->RecompileCustomizableObjectsCompiler)
-			{
-				return;
-			}
-		}
-
-		GetPrivate()->ObjectsToRecompile.Append(InObjects);
-
-		GetPrivate()->TotalNumObjectsToRecompile = GetPrivate()->ObjectsToRecompile.Num();
-		GetPrivate()->NumObjectsCompiled = 0;
-
-		if (GetPrivate()->RecompileNotificationHandle.IsValid())
-		{
-			++GetPrivate()->TotalNumObjectsToRecompile;
-			FSlateNotificationManager::Get().UpdateProgressNotification(GetPrivate()->RecompileNotificationHandle, GetPrivate()->NumObjectsCompiled, GetPrivate()->TotalNumObjectsToRecompile);
-		}
-		else
-		{
-			GetPrivate()->RecompileNotificationHandle = FSlateNotificationManager::Get().StartProgressNotification(FText::FromString(TEXT("Compiling Customizable Objects")), GetPrivate()->TotalNumObjectsToRecompile);
-			StartNextRecompile();
-		}
-	}
-}
-
-
-void UCustomizableObjectSystem::TickRecompileCustomizableObjects()
-{
-	bool bUpdated = false;
-	
-	if (GetPrivate()->RecompileCustomizableObjectsCompiler)
-	{
-		bUpdated = GetPrivate()->RecompileCustomizableObjectsCompiler->Tick(GetPrivate()->bBlocking);
-	}
-
-	if (bUpdated)
-	{
-		GetPrivate()->NumObjectsCompiled++;
-
-		if (!GetPrivate()->ObjectsToRecompile.IsEmpty())
-		{
-			StartNextRecompile();
-		}
-		else // All objects compiled, clean up
-		{
-			GetPrivate()->RecompileCustomizableObjectsCompiler.Reset();
-
-			// Remove progress bar
-			FSlateNotificationManager::Get().UpdateProgressNotification(GetPrivate()->RecompileNotificationHandle, GetPrivate()->NumObjectsCompiled, GetPrivate()->TotalNumObjectsToRecompile);
-			FSlateNotificationManager::Get().CancelProgressNotification(GetPrivate()->RecompileNotificationHandle);
-			GetPrivate()->RecompileNotificationHandle.Reset();
-
-			if (GEngine)
-			{
-				GEngine->ForceGarbageCollection();
-			}
-		}
-	}
-}
-
-
 void UCustomizableObjectSystem::SetWorkingMemory(int32 Bytes)
 {
 	WorkingMemory = Bytes;
@@ -4315,19 +4131,6 @@ void UCustomizableObjectSystemPrivate::OnMutableEnabledChanged(IConsoleVariable*
 			}			
 		}
 #endif // !UE_SERVER
-
-#if WITH_EDITOR
-		if (!IsRunningGame() && !FEditorDelegates::PreBeginPIE.IsBoundToObject(System))
-		{
-			FEditorDelegates::PreBeginPIE.AddUObject(System, &UCustomizableObjectSystem::OnPreBeginPIE);
-		}
-#endif
-	}
-	else
-	{
-#if WITH_EDITOR
-		FEditorDelegates::PreBeginPIE.RemoveAll(System);
-#endif
 	}
 }
 
diff --git a/Engine/Plugins/Experimental/Mutable/Source/CustomizableObjectEditor/Private/MuCOE/AssetDefinition_CustomizableObject.cpp b/Engine/Plugins/Experimental/Mutable/Source/CustomizableObjectEditor/Private/MuCOE/AssetDefinition_CustomizableObject.cpp
--- a/Engine/Plugins/Experimental/Mutable/Source/CustomizableObjectEditor/Private/MuCOE/AssetDefinition_CustomizableObject.cpp
+++ b/Engine/Plugins/Experimental/Mutable/Source/CustomizableObjectEditor/Private/MuCOE/AssetDefinition_CustomizableObject.cpp
@@ -180,9 +180,14 @@ namespace MenuExtension_CustomizableObject
 			return;
 		}
 
-		TArray<FAssetData> ObjectsToRecompile;
-
 		FAssetRegistryModule& AssetRegistryModule = FModuleManager::LoadModuleChecked<FAssetRegistryModule>("AssetRegistry");
+		if (AssetRegistryModule.Get().IsLoadingAssets())
+		{
+			return;
+		}
+
+		TArray<TSharedRef<FCompilationRequest>> CompilationRequests;
+
 
 		if (RecompileType == ERecompileCO::RCO_Selected)
 		{
@@ -190,12 +195,15 @@ namespace MenuExtension_CustomizableObject
 			check(Context);
 			
 			TArray<UCustomizableObject*> Objects = Context->LoadSelectedObjects<UCustomizableObject>();
+			CompilationRequests.Reserve(Objects.Num());
+
 			for (UCustomizableObject* Object : Objects)
 			{
-				FAssetData AssetData = AssetRegistryModule.Get().GetAssetByObjectPath(FSoftObjectPath(Object));
-				if (AssetData.IsValid())
+				if (Object)
 				{
-					ObjectsToRecompile.Add(AssetData);
+					TSharedRef<FCompilationRequest> NewRequest = MakeShared<FCompilationRequest>(*Object, true);
+					NewRequest->GetCompileOptions().bSilentCompilation = true;
+					CompilationRequests.Add(NewRequest);
 				}
 			}
 		}
@@ -223,27 +231,28 @@ namespace MenuExtension_CustomizableObject
 
 			if (FMessageDialog::Open(EAppMsgType::OkCancel, Msg) == EAppReturnType::Ok)
 			{
-				if (RecompileType == ERecompileCO::RCO_InMemory)
+				TArray<FAssetData> OutAssets;
+				AssetRegistryModule.Get().GetAssets(AssetRegistryFilter, OutAssets);
+
+				for (const FAssetData& Asset : OutAssets)
 				{
-					TArray<FAssetData> OutAssets;
-					AssetRegistryModule.Get().GetAssets(AssetRegistryFilter, OutAssets);
+					if (RecompileType == ERecompileCO::RCO_InMemory && !Asset.IsAssetLoaded())
+					{
+						continue;
+					}
 
-					for (const FAssetData& Asset : OutAssets)
+					UCustomizableObject* Object = Cast<UCustomizableObject>(Asset.GetAsset());
+					if (Object)
 					{
-						if (Asset.IsAssetLoaded())
-						{
-							ObjectsToRecompile.Add(Asset);
-						}
+						TSharedRef<FCompilationRequest> NewRequest = MakeShared<FCompilationRequest>(*Object, true);
+						NewRequest->GetCompileOptions().bSilentCompilation = true;
+						CompilationRequests.Add(NewRequest);
 					}
 				}
-				else
-				{
-					AssetRegistryModule.Get().GetAssets(AssetRegistryFilter, ObjectsToRecompile);
-				}
 			}
 		}
 
-		UCustomizableObjectSystem::GetInstance()->RecompileCustomizableObjects(ObjectsToRecompile);
+		ICustomizableObjectEditorModule::GetChecked().CompileCustomizableObjects(CompilationRequests);
 	}
 
 	
diff --git a/Engine/Plugins/Experimental/Mutable/Source/CustomizableObjectEditor/Private/MuCOE/CustomizableObjectCompileRunnable.cpp b/Engine/Plugins/Experimental/Mutable/Source/CustomizableObjectEditor/Private/MuCOE/CustomizableObjectCompileRunnable.cpp
--- a/Engine/Plugins/Experimental/Mutable/Source/CustomizableObjectEditor/Private/MuCOE/CustomizableObjectCompileRunnable.cpp
+++ b/Engine/Plugins/Experimental/Mutable/Source/CustomizableObjectEditor/Private/MuCOE/CustomizableObjectCompileRunnable.cpp
@@ -5,7 +5,7 @@
 #include "HAL/FileManager.h"
 #include "MuCO/UnrealMutableModelDiskStreamer.h"
 #include "MuCO/UnrealToMutableTextureConversionUtils.h"
-#include "MuCO/CustomizableObjectPrivate.h"
+#include "MuCO/CustomizableObject.h"
 #include "MuR/Model.h"
 #include "MuT/Compiler.h"
 #include "MuT/ErrorLog.h"
diff --git a/Engine/Plugins/Experimental/Mutable/Source/CustomizableObjectEditor/Private/MuCOE/CustomizableObjectCompileRunnable.h b/Engine/Plugins/Experimental/Mutable/Source/CustomizableObjectEditor/Private/MuCOE/CustomizableObjectCompileRunnable.h
--- a/Engine/Plugins/Experimental/Mutable/Source/CustomizableObjectEditor/Private/MuCOE/CustomizableObjectCompileRunnable.h
+++ b/Engine/Plugins/Experimental/Mutable/Source/CustomizableObjectEditor/Private/MuCOE/CustomizableObjectCompileRunnable.h
@@ -3,8 +3,8 @@
 #pragma once
 
 #include "HAL/Runnable.h"
-#include "MuCO/CustomizableObject.h"
 #include "MuCO/CustomizableObjectPrivate.h"
+#include "MuCO/CustomizableObjectCompilerTypes.h"
 #include "MuCO/UnrealToMutableTextureConversionUtils.h"
 #include "MuR/Ptr.h"
 #include "MuT/Node.h"
@@ -13,6 +13,7 @@
 #include <atomic>
 
 class ITargetPlatform;
+class UCustomizableObject;
 
 
 class FCustomizableObjectCompileRunnable : public FRunnable
diff --git a/Engine/Plugins/Experimental/Mutable/Source/CustomizableObjectEditor/Private/MuCOE/CustomizableObjectCompiler.cpp b/Engine/Plugins/Experimental/Mutable/Source/CustomizableObjectEditor/Private/MuCOE/CustomizableObjectCompiler.cpp
--- a/Engine/Plugins/Experimental/Mutable/Source/CustomizableObjectEditor/Private/MuCOE/CustomizableObjectCompiler.cpp
+++ b/Engine/Plugins/Experimental/Mutable/Source/CustomizableObjectEditor/Private/MuCOE/CustomizableObjectCompiler.cpp
@@ -9,6 +9,7 @@
 #include "Interfaces/ITargetPlatformManagerModule.h"
 #include "Materials/MaterialInterface.h"
 #include "MessageLogModule.h"
+#include "MuCO/CustomizableObject.h"
 #include "MuCO/CustomizableObjectPrivate.h"
 #include "MuCO/CustomizableObjectInstance.h"
 #include "MuCO/CustomizableObjectSystem.h"
@@ -43,68 +44,48 @@ class UTexture2D;
 UCustomizableObjectNodeObject* GetRootNode(UCustomizableObject* Object, bool &bOutMultipleBaseObjectsFound);
 
 
-
-FCustomizableObjectCompiler::FCustomizableObjectCompiler() : FCustomizableObjectCompilerBase()
-	, CompilationLaunchPending(false)
-	, PreloadingReferencerAssets(false)
-	, CurrentGAsyncLoadingTimeLimit(-1.0f)
-	, CompletedUnrealToMutableTask(0)
-	, MaxConvertToMutableTextureTime(0.2f)
+FCustomizableObjectCompiler::FCustomizableObjectCompiler()
 {
 	CustomizableObjectNumBoneInfluences = ICustomizableObjectModule::Get().GetNumBoneInfluences();
 }
 
 
-
 bool FCustomizableObjectCompiler::Tick(bool bBlocking)
 {
+	MUTABLE_CPUPROFILER_SCOPE(FCustomizableObjectCompiler::Tick);
+
 	bool bFinished = true;
 
-	if (bBlocking && AsynchronousStreamableHandlePtr)
+	if (TryPopCompileRequest())
 	{
 		bFinished = false;
-		AsynchronousStreamableHandlePtr->CancelHandle();
-		
-		UCustomizableObjectSystem::GetInstance()->GetPrivate()->StreamableManager.RequestSyncLoad(ArrayAssetToStream);
-		PreloadingReferencerAssetsCallback(false);
 	}
-	
-	if (PreloadingReferencerAssets)
+
+	if (AsynchronousStreamableHandlePtr)
 	{
 		bFinished = false;
+
+		if (bBlocking)
+		{
+			AsynchronousStreamableHandlePtr->CancelHandle();
+			UCustomizableObjectSystem::GetInstance()->GetPrivate()->StreamableManager.RequestSyncLoad(ArrayAssetToStream);
+			PreloadingReferencerAssetsCallback(false);
+		}
 	}
-	
+
 	if (CompileTask)
 	{
 		bFinished = false;
 		CompileTask->Tick();
-		
+
 		if (CompileTask->IsCompleted())
 		{
-			UE_LOG(LogMutable, Verbose, TEXT("PROFILE: [ %16.8f ] Finishing Compilation task for Object %s."), FPlatformTime::Seconds(), *CurrentObject->GetName());
-
 			FinishCompilationTask();
 
 			if (SaveDDTask.IsValid())
 			{
 				SaveCODerivedData();
 			}
-			else
-			{
-			    SetCompilationState(ECompilationStatePrivate::Completed, GetCompilationResult());
-
-				RemoveCompileNotification();
-
-				NotifyCompilationErrors();
-			}
-
-			UE_LOG(LogMutable, Verbose, TEXT("PROFILE: [ %16.8f ] Finished Compilation task."), FPlatformTime::Seconds());
-			UE_LOG(LogMutable, Verbose, TEXT("PROFILE: -----------------------------------------------------------"));
-
-			CleanCachedReferencers();
-			UpdateArrayGCProtect();
-
-			CompilationLogsContainer.ClearMessagesArray();
 		}
 	}
 
@@ -114,134 +95,155 @@ bool FCustomizableObjectCompiler::Tick(bool bBlocking)
 
 		if (SaveDDTask->IsCompleted())
 		{
-		    SetCompilationState(ECompilationStatePrivate::Completed, GetCompilationResult());
-		
 			FinishSavingDerivedDataTask();
-	
-			CurrentObject->GetPrivate()->PostCompile();
-
-			UE_LOG(LogMutable, Verbose, TEXT("PROFILE: [ %16.8f ] Finished Saving Derived Data task."), FPlatformTime::Seconds());
-			UE_LOG(LogMutable, Verbose, TEXT("PROFILE: -----------------------------------------------------------"));
 		}
 	}
-	
 
-	if (CompilationLaunchPending)
+	if (bFinished && CurrentRequest.IsValid())
 	{
-		bFinished = false;
+		bFinished = CompileRequests.IsEmpty();
 
-		CompilationLaunchPending = false;
-		LaunchMutableCompile();
+		CompleteRequest(ECompilationStatePrivate::Completed, GetCompilationResult());
+	}
+
+	if (CompileNotificationHandle.IsValid())
+	{
+		const int32 NumCompletedRequests = NumCompilationRequests - GetNumRemainingWork();
+		FSlateNotificationManager::Get().UpdateProgressNotification(CompileNotificationHandle, NumCompletedRequests, NumCompilationRequests);
 	}
 
 	return bFinished;
 }
 
 
+int32 FCustomizableObjectCompiler::GetNumRemainingWork() const
+{
+	return static_cast<int32>(CurrentRequest.IsValid()) + CompileRequests.Num();
+}
+
+
 void FCustomizableObjectCompiler::PreloadingReferencerAssetsCallback(bool bAsync)
 {
-	UE_LOG(LogMutable, Verbose, TEXT("PROFILE: [ %16.8f ] Preload asynchronously assets end."), FPlatformTime::Seconds());
+	check(IsInGameThread());
+
+	UpdateArrayGCProtect();
 
-	if (CurrentGAsyncLoadingTimeLimit != -1.0f)
+	if (AsynchronousStreamableHandlePtr)
 	{
-		static IConsoleVariable* CVAsyncLoadingTimeLimit = IConsoleManager::Get().FindConsoleVariable(TEXT("s.AsyncLoadingTimeLimit"));
-		if (CVAsyncLoadingTimeLimit)
-		{
-			CVAsyncLoadingTimeLimit->Set(CurrentGAsyncLoadingTimeLimit);
-		}
-		CurrentGAsyncLoadingTimeLimit = -1.0f;
+		AsynchronousStreamableHandlePtr.Reset();
 	}
 
-	UpdateArrayGCProtect();
-	PreloadingReferencerAssets = false;
-
+	UE_LOG(LogMutable, Verbose, TEXT("PROFILE: [ %16.8f ] Preload asynchronously assets end."), FPlatformTime::Seconds());
 	TRACE_END_REGION(UE_MUTABLE_PRELOAD_REGION);
 
 	CompileInternal(bAsync);
 }
 
 
-void FCustomizableObjectCompiler::Compile(UCustomizableObject& Object, const FCompilationOptions& InOptions, bool bAsync)
+void FCustomizableObjectCompiler::Compile(const TSharedRef<FCompilationRequest>& InCompileRequest)
 {
 	TRACE_BEGIN_REGION(UE_MUTABLE_COMPILE_REGION);
+	
+	check(IsInGameThread());
+	check(!CurrentRequest);
+	
+	CurrentRequest = InCompileRequest.ToSharedPtr();
+	CurrentObject = InCompileRequest->GetCustomizableObject();
+	CurrentOptions = InCompileRequest->GetCompileOptions();
+	CurrentOptions.CustomizableObjectNumBoneInfluences = CustomizableObjectNumBoneInfluences;
 
-	CurrentObject= Object;
+	if (!CurrentObject)
+	{
+		UE_LOG(LogMutable, Warning, TEXT("Failed to compile Customizable Object. Object is missing."));
+		CompleteRequest(ECompilationStatePrivate::Completed, ECompilationResultPrivate::Errors);
+		return;
+	}
 
 	if (CurrentObject->GetPrivate()->CompilationState == ECompilationStatePrivate::InProgress)
 	{
-		UE_LOG(LogMutable, Warning, TEXT("Failed to compile Customizable Object [%s]. Object already being compiled."), *Object.GetName());
+		UE_LOG(LogMutable, Warning, TEXT("Failed to compile Customizable Object [%s]. Object already being compiled."), *CurrentObject->GetName());
+		CurrentObject = nullptr; // Someone else is compiling the CO. Invalidate the CurrentObject pointer to avoid changing the state of the ongoing compilation.
+		CompleteRequest(ECompilationStatePrivate::Completed, ECompilationResultPrivate::Errors);
 		return;
 	}
-	
-	SetCompilationState(ECompilationStatePrivate::InProgress, ECompilationResultPrivate::Unknown);
 
 	if (!UCustomizableObjectSystem::IsActive())
 	{
-		UE_LOG(LogMutable, Warning, TEXT("Failed to compile Customizable Object [%s]. Mutable is disabled. To enable it set the CVar Mutable.Enabled to true."), *Object.GetName());
-		SetCompilationState(ECompilationStatePrivate::Completed, ECompilationResultPrivate::Errors);
+		UE_LOG(LogMutable, Warning, TEXT("Failed to compile Customizable Object [%s]. Mutable is disabled. To enable it set the CVar Mutable.Enabled to true."), *CurrentObject->GetName());
+		CompleteRequest(ECompilationStatePrivate::Completed, ECompilationResultPrivate::Errors);
 		return;
 	}
 
-	if (Object.VersionBridge && !Object.VersionBridge->GetClass()->ImplementsInterface(UCustomizableObjectVersionBridgeInterface::StaticClass()))
+	if (CurrentObject->VersionBridge && !CurrentObject->VersionBridge->GetClass()->ImplementsInterface(UCustomizableObjectVersionBridgeInterface::StaticClass()))
 	{
-		UE_LOG(LogMutable, Warning, TEXT("In Customizable Object [%s], the VersionBridge asset [%s] does not implement the required UCustomizableObjectVersionBridgeInterface."), 
-			*Object.GetName(), *Object.VersionBridge.GetName());
-		SetCompilationState(ECompilationStatePrivate::Completed, ECompilationResultPrivate::Errors);
+		UE_LOG(LogMutable, Warning, TEXT("In Customizable Object [%s], the VersionBridge asset [%s] does not implement the required UCustomizableObjectVersionBridgeInterface."),
+			*CurrentObject->GetName(), *CurrentObject->VersionBridge.GetName());
+		CompleteRequest(ECompilationStatePrivate::Completed, ECompilationResultPrivate::Errors);
 		return;
 	}
 
-	UCustomizableObjectSystem* System = UCustomizableObjectSystem::GetInstanceChecked();
-
-	UE_LOG(LogMutable, Verbose, TEXT("PROFILE: [ %16.8f ] Preload asynchronously assets start."), FPlatformTime::Seconds());
-
-	FString Message = FString::Printf(TEXT("Customizable Object %s is already being compiled or updated. Please wait a few seconds and try again."), *Object.GetName());
-	FNotificationInfo Info(LOCTEXT("CustomizableObjectBeingCompilerOrUpdated", "Customizable Object compile and/or update still in process. Please wait a few seconds and try again."));
-	Info.bFireAndForget = true;
-	Info.bUseThrobber = true;
-	Info.FadeOutDuration = 1.0f;
-	Info.ExpireDuration = 1.0f;
-
-	if (PreloadingReferencerAssets || CompilationLaunchPending || (Object.GetPrivate()->IsLocked()))
+	if (!CurrentOptions.bIsCooking && IsRunningCookCommandlet())
 	{
-		UE_LOG(LogMutable, Warning, TEXT("%s"), *Message);
-		FSlateNotificationManager::Get().AddNotification(Info);
-
+		UE_LOG(LogMutable, Display, TEXT("Editor compilation suspended for Customizable Object [%s]. Can not compile COs when the cook commandlet is running. "), *CurrentObject->GetName());
+		CompleteRequest(ECompilationStatePrivate::Completed, ECompilationResultPrivate::Errors);
 		return;
 	}
 
-	// Lock object during asynchronous asset loading to avoid instance/mip updates and reentrant compilations
-	bool LockResult = System->LockObject(&Object);
+	UCustomizableObjectSystem* System = UCustomizableObjectSystem::GetInstanceChecked();
+
+	check(!CurrentObject->GetPrivate()->IsLocked());
 
-	if (!LockResult)
+	// Lock object during asynchronous asset loading to avoid instance/mip updates and reentrant compilations
+	if (!System->LockObject(CurrentObject))
 	{
+		FString Message = FString::Printf(TEXT("Customizable Object %s is already being compiled or updated. Please wait a few seconds and try again."), *CurrentObject->GetName());
 		UE_LOG(LogMutable, Warning, TEXT("%s"), *Message);
+		
+		FNotificationInfo Info(LOCTEXT("CustomizableObjectBeingCompilerOrUpdated", "Customizable Object compile and/or update still in process. Please wait a few seconds and try again."));
+		Info.bFireAndForget = true;
+		Info.bUseThrobber = true;
+		Info.FadeOutDuration = 1.0f;
+		Info.ExpireDuration = 1.0f;
 		FSlateNotificationManager::Get().AddNotification(Info);
-		SetCompilationState(ECompilationStatePrivate::Completed, ECompilationResultPrivate::Errors);
+
+		CurrentObject = nullptr; // Someone else is compiling the CO. Invalidate the CurrentObject pointer to avoid changing the state of the ongoing compilation.
+		CompleteRequest(ECompilationStatePrivate::Completed, ECompilationResultPrivate::Errors);
 		return;
 	}
 
+	SetCompilationState(ECompilationStatePrivate::InProgress, ECompilationResultPrivate::Unknown);
+
 	// Now that we know for sure that the CO is locked and there are no pending updates of instances using the CO,
 	// destroy any live update instances, as they become invalid when recompiling the CO
 	for (TObjectIterator<UCustomizableObjectInstance> It; It; ++It)
 	{
 		UCustomizableObjectInstance* Instance = *It;
 		if (IsValid(Instance) &&
-			Instance->GetCustomizableObject() == &Object)
+			Instance->GetCustomizableObject() == CurrentObject)
 		{
 			Instance->DestroyLiveUpdateInstance();
 		}
 	}
 
-	CurrentObject = &Object;
-	Options = InOptions;
-
-	PreloadingReferencerAssets = true;
+	// Create and update compilation progress notification
+	const FText UpdateMsg = FText::FromString(FString::Printf(TEXT("Compiling Customizable Objects:\n%s"), *CurrentObject->GetName()));
+	if (!CompileNotificationHandle.IsValid())
+	{
+		CompileNotificationHandle = FSlateNotificationManager::Get().StartProgressNotification(UpdateMsg, NumCompilationRequests);
+	}
+	else
+	{
+		const int32 NumCompletedRequests = NumCompilationRequests - GetNumRemainingWork();
+		FSlateNotificationManager::Get().UpdateProgressNotification(CompileNotificationHandle, NumCompletedRequests, NumCompilationRequests, UpdateMsg);
+	}
+	
 	TRACE_BEGIN_REGION(UE_MUTABLE_PRELOAD_REGION);
+	UE_LOG(LogMutable, Verbose, TEXT("PROFILE: [ %16.8f ] Preload asynchronously assets start."), FPlatformTime::Seconds());
 
 	CleanCachedReferencers();
 	UpdateArrayGCProtect();
 	TArray<FName> ArrayReferenceNames;
-	AddCachedReferencers(*Object.GetOuter()->GetPathName(), ArrayReferenceNames);
+	AddCachedReferencers(*CurrentObject->GetOuter()->GetPathName(), ArrayReferenceNames);
 
 	ArrayAssetToStream.Empty();
 	for (FAssetData& Element : ArrayAssetData)
@@ -249,35 +251,27 @@ void FCustomizableObjectCompiler::Compile(UCustomizableObject& Object, const FCo
 		ArrayAssetToStream.Add(Element.GetSoftObjectPath());
 	}
 
-	MaxConvertToMutableTextureTime = ComputeAsyncLoadingTimeLimit() * 0.001f;
-
 	bool bAssetsLoaded = true;
 
-	// Customizations are marked as editoronly on load and are not packaged into the runtime game by default.
- 	// The ones that need to be kept will be copied into SoftObjectPath on the object during save.
-	FCookLoadScope CookLoadScope(ECookLoadType::EditorOnly);
+	const bool bAsync = InCompileRequest->IsAsyncCompilation();
 	if (ArrayAssetToStream.Num() > 0)
 	{
+		// Customizations are marked as editoronly on load and are not packaged into the runtime game by default.
+		// The ones that need to be kept will be copied into SoftObjectPath on the object during save.
+		FCookLoadScope CookLoadScope(ECookLoadType::EditorOnly);
+
 		FStreamableManager& Streamable = System->GetPrivate()->StreamableManager;
 
-		if (bAsync) 
+		if (bAsync && !CurrentOptions.bIsCooking)
 		{
 			AddCompileNotification(LOCTEXT("LoadingReferencerAssets", "Loading assets"));
 
-			static IConsoleVariable* CVAsyncLoadingTimeLimit = IConsoleManager::Get().FindConsoleVariable(TEXT("s.AsyncLoadingTimeLimit"));
-			if (CVAsyncLoadingTimeLimit)
-			{
-				CurrentGAsyncLoadingTimeLimit = CVAsyncLoadingTimeLimit->GetFloat();
-				CVAsyncLoadingTimeLimit->Set(MaxConvertToMutableTextureTime * 1000.0f);
-			}
-
 			AsynchronousStreamableHandlePtr = Streamable.RequestAsyncLoad(ArrayAssetToStream, FStreamableDelegate::CreateRaw(this, &FCustomizableObjectCompiler::PreloadingReferencerAssetsCallback, bAsync));
 			bAssetsLoaded = false;
 		}
 		else
 		{
 			Streamable.RequestSyncLoad(ArrayAssetToStream);
-			bAssetsLoaded = true;
 		}
 	}
 
@@ -288,6 +282,45 @@ void FCustomizableObjectCompiler::Compile(UCustomizableObject& Object, const FCo
 }
 
 
+void FCustomizableObjectCompiler::Compile(const TArray<TSharedRef<FCompilationRequest>>& InCompileRequests)
+{
+	NumCompilationRequests += InCompileRequests.Num();
+	CompileRequests.Append(InCompileRequests);
+}
+
+
+bool FCustomizableObjectCompiler::IsTickable() const
+{
+	return NumCompilationRequests > 0 || CurrentRequest;
+}
+
+
+void FCustomizableObjectCompiler::Tick(float InDeltaTime)
+{
+	MUTABLE_CPUPROFILER_SCOPE(FCustomizableObjectCompiler::Tick);
+	Tick();
+}
+
+
+TStatId FCustomizableObjectCompiler::GetStatId() const
+{
+	RETURN_QUICK_DECLARE_CYCLE_STAT(FCustomizableObjectCompiler, STATGROUP_Tickables);
+}
+
+
+void FCustomizableObjectCompiler::TickCook(float DeltaTime, bool bCookCompete)
+{
+	MUTABLE_CPUPROFILER_SCOPE(FCustomizableObjectCompiler::TickCook);
+	Tick();
+}
+
+
+bool FCustomizableObjectCompiler::IsRequestQueued(const TSharedRef<FCompilationRequest>& InCompileRequest) const
+{
+	return (CurrentRequest == InCompileRequest) || (CompileRequests.Find(InCompileRequest) != INDEX_NONE);
+}
+
+
 void FCustomizableObjectCompiler::AddReferencedObjects(FReferenceCollector& Collector)
 {
 	// While compilation takes place, no COs involved can be garbage-collected
@@ -448,34 +481,6 @@ void FCustomizableObjectCompiler::DisplayUnnamedNodeObjectWarning(FMutableGraphG
 }
 
 
-//void FCustomizableObjectCompiler::DisplayDiscardedPhysicsAssetSingleWarning(FMutableGraphGenerationContext& GenerationContext)
-//{
-//	if (GenerationContext.DiscartedPhysicsAssetMap.Num())
-//	{
-//		FString PhysicsAssetName;
-//		for (int32 ComponentIndex = 0; ComponentIndex < GenerationContext.NumMeshComponentsInRoot; ++ComponentIndex)
-//		{
-//			USkeletalMesh* RefSkeletalMesh = GenerationContext.ComponentInfos.IsValidIndex(ComponentIndex) ? GenerationContext.ComponentInfos[ComponentIndex].RefSkeletalMesh : nullptr;
-//			if (RefSkeletalMesh && GenerationContext.DiscartedPhysicsAssetMap.Find(RefSkeletalMesh->GetPhysicsAsset()))
-//			{
-//				PhysicsAssetName = RefSkeletalMesh->GetPhysicsAsset()->GetName();
-//				break;
-//			}
-//		}
-//		
-//		if(PhysicsAssetName.IsEmpty())
-//		{
-//			PhysicsAssetName = (*GenerationContext.DiscartedPhysicsAssetMap.FindKey(0))->GetName();
-//		}
-//
-//		FText Message = FText::Format( LOCTEXT("Discarted PhysicsAssets", "{0} and {1} other PhysicsAssets have been discarted because one or more Bodies have no corresponding bones in the SkeletalMesh.\nFor more infromation on the assets that need attention check the Output log."),
-//										FText::FromString(PhysicsAssetName), FText::AsNumber(GenerationContext.DiscartedPhysicsAssetMap.Num()));
-//		
-//		CompilerLog(Message, nullptr, EMessageSeverity::Warning, false);
-//	}
-//}
-
-
 void FCustomizableObjectCompiler::DisplayOrphanNodesWarning(FMutableGraphGenerationContext& GenerationContext)
 {
 	for (const TPair<FGeneratedKey, FGeneratedData>& It : GenerationContext.Generated)
@@ -539,9 +544,9 @@ mu::NodeObjectPtr FCustomizableObjectCompiler::GenerateMutableRoot(
 		GenerationContext.NoNameNodeObjectArray.AddUnique(Root);
 	}
 
-	if ((Object->MeshCompileType == EMutableCompileMeshType::Full) || Options.bIsCooking)
+	if ((Object->MeshCompileType == EMutableCompileMeshType::Full) || CurrentOptions.bIsCooking)
 	{
-		if (Root->ParentObject!=nullptr && Options.bIsCooking)
+		if (Root->ParentObject!=nullptr && CurrentOptions.bIsCooking)
 		{
 			// This happens while packaging.
 			return nullptr;
@@ -707,6 +712,33 @@ mu::NodeObjectPtr FCustomizableObjectCompiler::GenerateMutableRoot(
 	return MutableRoot;
 }
 
+void FCustomizableObjectCompiler::LaunchMutableCompile()
+{
+	AddCompileNotification(LOCTEXT("CustomizableObjectCompileInProgress", "Compiling"));
+
+	// Even for async build, we spawn a thread, so that we can set a large stack. 
+	// Thread names need to be unique, apparently.
+	static int ThreadCount = 0;
+	FString ThreadName = FString::Printf(TEXT("MutableCompile-%03d"), ++ThreadCount);
+	CompileThread = MakeShareable(FRunnableThread::Create(CompileTask.Get(), *ThreadName, 16 * 1024 * 1024, TPri_Normal));
+}
+
+
+void FCustomizableObjectCompiler::SaveCODerivedData()
+{
+	if (!SaveDDTask.IsValid())
+	{
+		return;
+	}
+
+	AddCompileNotification(LOCTEXT("SavingCustomizableObjectDerivedData", "Saving Data"));
+
+	// Even for async saving derived data.
+	static int SDDThreadCount = 0;
+	FString ThreadName = FString::Printf(TEXT("MutableSDD-%03d"), ++SDDThreadCount);
+	SaveDDThread = MakeShareable(FRunnableThread::Create(SaveDDTask.Get(), *ThreadName));
+}
+
 
 void FCustomizableObjectCompiler::AddCachedReferencers(const FName& PathName, TArray<FName>& ArrayReferenceNames)
 {
@@ -716,10 +748,10 @@ void FCustomizableObjectCompiler::AddCachedReferencers(const FName& PathName, TA
 
 	// Required to make compilations deterministic within editor runs.
 	ArrayReferenceNames.Sort([](const FName& A, const FName& B)
-	{
-		return A.LexicalLess(B);
-	});
-	
+		{
+			return A.LexicalLess(B);
+		});
+
 	FARFilter Filter;
 	for (const FName& ReferenceName : ArrayReferenceNames)
 	{
@@ -746,34 +778,6 @@ void FCustomizableObjectCompiler::AddCachedReferencers(const FName& PathName, TA
 }
 
 
-void FCustomizableObjectCompiler::LaunchMutableCompile()
-{
-	AddCompileNotification(LOCTEXT("CustomizableObjectCompileInProgress", "Compiling"));
-
-	// Even for async build, we spawn a thread, so that we can set a large stack. 
-	// Thread names need to be unique, apparently.
-	static int ThreadCount = 0;
-	FString ThreadName = FString::Printf(TEXT("MutableCompile-%03d"), ++ThreadCount);
-	CompileThread = MakeShareable(FRunnableThread::Create(CompileTask.Get(), *ThreadName, 16 * 1024 * 1024, TPri_Normal));
-}
-
-
-void FCustomizableObjectCompiler::SaveCODerivedData()
-{
-	if (!SaveDDTask.IsValid())
-	{
-		return;
-	}
-
-	AddCompileNotification(LOCTEXT("SavingCustomizableObjectDerivedData", "Saving Data"));
-
-	// Even for async saving derived data.
-	static int SDDThreadCount = 0;
-	FString ThreadName = FString::Printf(TEXT("MutableSDD-%03d"), ++SDDThreadCount);
-	SaveDDThread = MakeShareable(FRunnableThread::Create(SaveDDTask.Get(), *ThreadName));
-}
-
-
 void FCustomizableObjectCompiler::CleanCachedReferencers()
 {
 	ArrayAssetData.Empty();
@@ -812,21 +816,6 @@ FAssetData* FCustomizableObjectCompiler::GetCachedAssetData(const FString& Packa
 }
 
 
-float FCustomizableObjectCompiler::ComputeAsyncLoadingTimeLimit()
-{
-	float DeltaTimeSeconds = FApp::GetDeltaTime();
-
-	if (DeltaTimeSeconds == 0.0f)
-	{
-		return 100.0f;
-	}
-
-	// AsyncLoadingTimeLimit will depend on the speed of the computer, using delta time as reference
-	float Value = 2.0f / DeltaTimeSeconds;
-	Value = FMath::Clamp(Value, 75.0f, 250.0f);
-	return Value;
-}
-
 
 ECompilationResultPrivate FCustomizableObjectCompiler::GetCompilationResult() const
 {
@@ -845,14 +834,11 @@ ECompilationResultPrivate FCustomizableObjectCompiler::GetCompilationResult() co
 }
 
 
-void FCustomizableObjectCompiler::GetCompilationMessages(TArray<FText>& OutWarningMessages, TArray<FText>& OutErrorMessages) const
-{
-	CompilationLogsContainer.GetMessages(OutWarningMessages,OutErrorMessages);
-}
-
-
 void FCustomizableObjectCompiler::SetCompilationState(ECompilationStatePrivate State, ECompilationResultPrivate Result) const
 {
+	check(CurrentRequest);
+	CurrentRequest->SetCompilationState(State, Result);
+
 	if (CurrentObject)
 	{
 		CurrentObject->GetPrivate()->CompilationState = State;
@@ -865,60 +851,28 @@ void FCustomizableObjectCompiler::CompileInternal(bool bAsync)
 {
 	MUTABLE_CPUPROFILER_SCOPE(FCustomizableObjectCompiler::Compile)
 	
-	if (!CurrentObject)
-	{
-		SetCompilationState(ECompilationStatePrivate::Completed, ECompilationResultPrivate::Errors);
-		return;
-	}
-
-	if (bAsync && CompileTask.IsValid()) // Don't start compilation if there's a compilation running
-	{
-		UE_LOG(LogMutable, Log, TEXT("FCustomizableObjectCompiler::Compile An object is already being compiled."));
-		SetCompilationState(ECompilationStatePrivate::Completed, ECompilationResultPrivate::Errors);
-		return;
-	}
-
-	UE_LOG(LogMutable, Log, TEXT("PROFILE: [ %16.8f ] FCustomizableObjectCompiler::Compile start."), FPlatformTime::Seconds());
+	UE_LOG(LogMutable, Verbose, TEXT("PROFILE: [ %16.8f ] FCustomizableObjectCompiler::Compile start."), FPlatformTime::Seconds());
 	
 	// This is redundant but necessary to keep static analysis happy.
 	if (!CurrentObject)
 	{
-		SetCompilationState(ECompilationStatePrivate::Completed, ECompilationResultPrivate::Errors);
+		CompleteRequest(ECompilationStatePrivate::Completed, ECompilationResultPrivate::Errors);
 		return;
 	}
 
-	Options.CustomizableObjectNumBoneInfluences = CustomizableObjectNumBoneInfluences;
-	Options.bRealTimeMorphTargetsEnabled = CurrentObject->bEnableRealTimeMorphTargets;
-	Options.bClothingEnabled = CurrentObject->bEnableClothing;
-	Options.b16BitBoneWeightsEnabled = CurrentObject->bEnable16BitBoneWeights;
-	Options.bSkinWeightProfilesEnabled = CurrentObject->bEnableAltSkinWeightProfiles;
-	Options.bPhysicsAssetMergeEnabled = CurrentObject->bEnablePhysicsAssetMerge;
-	Options.bAnimBpPhysicsManipulationEnabled = CurrentObject->bEnableAnimBpPhysicsAssetsManipualtion;
-	Options.bUseDiskCompilation = CurrentObject->CompileOptions.bUseDiskCompilation;
-	Options.ImageTiling = CurrentObject->CompileOptions.ImageTiling;
 
-	if (!Options.bIsCooking && IsRunningCookCommandlet())
+	if (CurrentOptions.bIsCooking && CurrentOptions.TargetPlatform)
 	{
-		UE_LOG(LogMutable, Display, TEXT("Editor compilation suspended for Customizable Object [%s]. Can not compile COs when the cook commandlet is running. "), *CurrentObject->GetName());
-		SetCompilationState(ECompilationStatePrivate::Completed, ECompilationResultPrivate::Errors);
-        UCustomizableObjectSystem::GetInstance()->UnlockObject(CurrentObject);
-		return;
+		UE_LOG(LogMutable, Display, TEXT("Compiling Customizable Object %s for platform %s."), *CurrentObject->GetName(), *CurrentOptions.TargetPlatform->PlatformName());
 	}
-	
-	UE_LOG(LogMutable, Display, TEXT("Started Customizable Object Compile %s."), *CurrentObject->GetName());
 
-	if (Options.bIsCooking && Options.TargetPlatform)
+	if (CurrentOptions.bIsCooking && CurrentOptions.bForceLargeLODBias)
 	{
-		UE_LOG(LogMutable, Display, TEXT("Compiling Customizable Object %s for platform %s."), *CurrentObject->GetName(), *Options.TargetPlatform->PlatformName());
+		UE_LOG(LogMutable, Display, TEXT("Compiling Customizable Object with %d LODBias."), CurrentOptions.DebugBias);
 	}
 
-	if (Options.bIsCooking && Options.bForceLargeLODBias)
-	{
-		UE_LOG(LogMutable, Display, TEXT("Compiling Customizable Object with %d LODBias."), Options.DebugBias);
-	}
-
-	FMutableGraphGenerationContext GenerationContext(CurrentObject, this, Options);
-	GenerationContext.ParamNamesToSelectedOptions = ParamNamesToSelectedOptions;
+	FMutableGraphGenerationContext GenerationContext(CurrentObject, this, CurrentOptions);
+	GenerationContext.ParamNamesToSelectedOptions = CurrentRequest->GetParameterNamesToSelectedOptions();
 
 	// If we don't have the target platform yet (in editor) we need to get it
 	if (!GenerationContext.Options.TargetPlatform)
@@ -936,29 +890,14 @@ void FCustomizableObjectCompiler::CompileInternal(bool bAsync)
 	CompilationLogsContainer.ClearMessagesArray();
 
 	// Generate the mutable node expression
-	FText ErrorMessage;
+	FText ErrorMessage = FText(LOCTEXT("FailedToGenerateRoot", "Failed to generate the mutable node graph. Object not built."));
 	bool bIsRootObject = false;
 	mu::NodeObjectPtr MutableRoot = GenerateMutableRoot(CurrentObject, GenerationContext, ErrorMessage, bIsRootObject);
 
 	if (!MutableRoot)
 	{
-		if (!ErrorMessage.IsEmpty())
-		{
-			CompilerLog(ErrorMessage, nullptr);
-		}
-		else
-		{
-			CompilerLog(FText::FromString(TEXT("Failed to generate the mutable node graph. Object not built.")), nullptr);
-		}
-
-		SetCompilationState(ECompilationStatePrivate::Completed, GetCompilationResult());
-		
-		if (CurrentObject->GetPrivate()->IsLocked())
-		{
-			UCustomizableObjectSystem::GetInstance()->UnlockObject(CurrentObject);
-		}
-
-		RemoveCompileNotification();
+		CompilerLog(ErrorMessage);
+		CompleteRequest(ECompilationStatePrivate::Completed, ECompilationResultPrivate::Errors);
 	}
 	else
 	{
@@ -1183,7 +1122,7 @@ void FCustomizableObjectCompiler::CompileInternal(bool bAsync)
 			NewCompileTimeReferencedTextures.Add(Tex);
 		}
 
-		if (!ParamNamesToSelectedOptions.Num())
+		if (!CurrentRequest->GetParameterNamesToSelectedOptions().Num())
 		{
 			// Get possible objects used in the compilation that are not directly referenced.
 			// Due to this check being done also in PIE (to detect out of date compilations), it has to be performant. Therefore we are gathering a relaxed set.
@@ -1208,22 +1147,8 @@ void FCustomizableObjectCompiler::CompileInternal(bool bAsync)
 			CurrentObject->GetPrivate()->DirtyParticipatingObjects.Empty();
 		}
 
-		if (CompileTask.IsValid()) // Don't start compilation if there's a compilation running
-		{
-			// TODO : warning?
-			UCustomizableObjectSystem::GetInstance()->UnlockObject(CurrentObject);
-			return;
-		}
-
-		// Lock the object to prevent instance updates while compiling. It's ignored and returns false if it's already locked
-		// Will unlock in the FinishCompilation call.
-		if (!CurrentObject->GetPrivate()->IsLocked())
-		{
-			UCustomizableObjectSystem::GetInstance()->LockObject(CurrentObject);
-		}
-
 		CompileTask = MakeShareable(new FCustomizableObjectCompileRunnable(MutableRoot));
-		CompileTask->Options = Options;
+		CompileTask->Options = CurrentOptions;
 		CompileTask->ReferencedTextures = NewCompileTimeReferencedTextures;
 
 		if (!bAsync)
@@ -1236,30 +1161,14 @@ void FCustomizableObjectCompiler::CompileInternal(bool bAsync)
 			{
 				SaveDDTask->Init();
 				SaveDDTask->Run();
-				Model->GetPrivate()->UnloadRoms();
 				FinishSavingDerivedDataTask();
 			}
 
-			CleanCachedReferencers();
-			UpdateArrayGCProtect();
-
-			CurrentObject->GetPrivate()->PostCompile(); 
-
-			SetCompilationState(ECompilationStatePrivate::Completed, GetCompilationResult());
+			CompleteRequest(ECompilationStatePrivate::Completed, GetCompilationResult());
 		}
 		else
 		{
-			// If packaging, convert textures and launch Mutable compile thread
-			if (IsRunningCommandlet())
-			{
-				LaunchMutableCompile();
-			}
-			else
-			{
-				// If asynchronous compilation and not packaging, set CompilationLaunchPending to true and FCustomizableObjectCompiler::Tick()
-				// will do the remaining steps (convert textures asynchronously and launch the Mutable compile thread)
-				CompilationLaunchPending = true;
-			}
+			LaunchMutableCompile();
 		}
 	}
 
@@ -1268,8 +1177,6 @@ void FCustomizableObjectCompiler::CompileInternal(bool bAsync)
 		Node->ResetAttachedErrorData();
 	}
 
-	UE_LOG(LogMutable, Log, TEXT("PROFILE: [ %16.8f ] FCustomizableObjectCompiler::Compile end."), FPlatformTime::Seconds());
-
 	// Population Recompilation
 	if (MutableRoot)
 	{
@@ -1282,6 +1189,93 @@ void FCustomizableObjectCompiler::CompileInternal(bool bAsync)
 }
 
 
+void FCustomizableObjectCompiler::CompleteRequest(ECompilationStatePrivate State, ECompilationResultPrivate Result)
+{
+	check(IsInGameThread());
+	check(CurrentRequest);
+
+	ECompilationStatePrivate CurrentState = CurrentRequest->GetCompilationState();
+	SetCompilationState(State, Result);
+
+	if (CurrentState == ECompilationStatePrivate::InProgress && CurrentObject)
+	{
+		// Unlock the object so that instances can be updated
+		UCustomizableObjectSystem* System = UCustomizableObjectSystem::IsCreated() ? UCustomizableObjectSystem::GetInstance() : nullptr;
+		if (System && !System->HasAnyFlags(EObjectFlags::RF_BeginDestroyed))
+		{
+			System->UnlockObject(CurrentObject);
+		}
+
+		if (CurrentModel)
+		{
+			CurrentModel->GetPrivate()->UnloadRoms();
+		}
+
+		CurrentObject->GetPrivate()->SetModel(CurrentModel, GenerateIdentifier(*CurrentObject));
+
+		if (Result == ECompilationResultPrivate::Success || Result == ECompilationResultPrivate::Warnings)
+		{
+			CurrentObject->GetPrivate()->PostCompile();
+		}
+	}
+
+	// Remove referenced objects
+	CleanCachedReferencers();
+	UpdateArrayGCProtect();
+
+	// Notifications
+	RemoveCompileNotification();
+	NotifyCompilationErrors();
+
+	// Update compilation progress notification
+	if (CompileNotificationHandle.IsValid())
+	{
+		const int32 NumCompletedRequests = NumCompilationRequests - CompileRequests.Num();
+		FSlateNotificationManager::Get().UpdateProgressNotification(CompileNotificationHandle, NumCompletedRequests, NumCompilationRequests);
+
+		if (NumCompletedRequests == NumCompilationRequests)
+		{
+			// Remove progress bar
+			FSlateNotificationManager::Get().CancelProgressNotification(CompileNotificationHandle);
+			CompileNotificationHandle.Reset();
+			NumCompilationRequests = 0;
+		}
+	}
+
+	// Copy warnings and errors to the request
+	CompilationLogsContainer.GetMessages(CurrentRequest->GetWarnings(), CurrentRequest->GetErrors());
+	
+	// Clear Messages
+	CompilationLogsContainer.ClearMessageCounters();
+	CompilationLogsContainer.ClearMessagesArray();
+
+	if (GEngine)
+	{
+		GEngine->ForceGarbageCollection();
+	}
+
+	// Request completed, reset pointers and state
+	CurrentObject = nullptr;
+	CurrentRequest.Reset();
+	CurrentModel.Reset();
+
+	UE_LOG(LogMutable, Verbose, TEXT("PROFILE: [ %16.8f ] Completed compile request."), FPlatformTime::Seconds());
+	UE_LOG(LogMutable, Verbose, TEXT("PROFILE: -----------------------------------------------------------"));
+}
+
+
+bool FCustomizableObjectCompiler::TryPopCompileRequest()
+{
+	if (CurrentRequest.IsValid() || CompileRequests.IsEmpty())
+	{
+		return false;
+	}
+
+	Compile(CompileRequests.Pop());
+	return true;
+}
+
+
 mu::NodePtr FCustomizableObjectCompiler::Export(UCustomizableObject* Object, const FCompilationOptions& InCompilerOptions, 
 	TArray<TSoftObjectPtr<UTexture>>& OutRuntimeReferencedTextures,
 	TArray<FMutableSourceTextureData>& OutCompilerReferencedTextures )
@@ -1306,7 +1300,7 @@ mu::NodePtr FCustomizableObjectCompiler::Export(UCustomizableObject* Object, con
 	CompilerOptions.bAnimBpPhysicsManipulationEnabled = Object->bEnableAnimBpPhysicsAssetsManipualtion;
 
 	FMutableGraphGenerationContext GenerationContext(Object, this, CompilerOptions);
-	GenerationContext.ParamNamesToSelectedOptions = ParamNamesToSelectedOptions;
+	GenerationContext.ParamNamesToSelectedOptions = CurrentRequest->GetParameterNamesToSelectedOptions();
 
 	// Generate the mutable node expression
 	FText ErrorMsg;
@@ -1321,7 +1315,7 @@ mu::NodePtr FCustomizableObjectCompiler::Export(UCustomizableObject* Object, con
 		}
 		else
 		{
-			FCustomizableObjectCompiler::CompilerLog(LOCTEXT("FailedToGenerate","Failed to generate the mutable node graph. Object not built."), nullptr);
+			FCustomizableObjectCompiler::CompilerLog(LOCTEXT("FailedToExport","Failed to generate the mutable node graph. Object not built."), nullptr);
 		}
 		return nullptr;
 	}
@@ -1362,66 +1356,42 @@ void FCustomizableObjectCompiler::FinishCompilationTask()
 	check(CompileTask.IsValid());
 
 	UpdateCompilerLogData();
-	Model = CompileTask->Model;
+	CurrentModel = CompileTask->Model;
 
 	// Generate a map that using the resource id tells the offset and size of the resource inside the bulk data
 	// At this point it is assumed that all data goes into a single file.
-	if (Model)
+	if (CurrentModel)
 	{
 		// Always work with the ModelResources (Editor) when compiling. They'll be copied to the cooked version during PreSave.
 		FModelResources& ModelResources = CurrentObject->GetPrivate()->GetModelResources(false);
-		
-		const int32 NumStreamingFiles = Model->GetRomCount();
+
+		const int32 NumStreamingFiles = CurrentModel->GetRomCount();
 		ModelResources.HashToStreamableBlock.Empty(NumStreamingFiles);
 
 		uint64 Offset = 0;
 		for (int32 FileIndex = 0; FileIndex < NumStreamingFiles; ++FileIndex)
 		{
-			const uint32 ResourceId = Model->GetRomId(FileIndex);
-			const uint32 ResourceSize = Model->GetRomSize(FileIndex);
+			const uint32 ResourceId = CurrentModel->GetRomId(FileIndex);
+			const uint32 ResourceSize = CurrentModel->GetRomSize(FileIndex);
 
-			ModelResources.HashToStreamableBlock.Add(ResourceId, FMutableStreamableBlock{0, ResourceSize, Offset });
+			ModelResources.HashToStreamableBlock.Add(ResourceId, FMutableStreamableBlock{ 0, ResourceSize, Offset });
 			Offset += ResourceSize;
 		}
 	}
 
 	// Generate ParameterProperties and IntParameterLookUpTable
-	CurrentObject->GetPrivate()->UpdateParameterPropertiesFromModel(Model);
-	
-	// Reset all instances, as the parameters may need to be rebuilt.
-	for (TObjectIterator<UCustomizableObjectInstance> It; It; ++It)
-	{
-		UCustomizableObjectInstance* Instance = *It;
-		if (IsValid(Instance) &&
-			Instance->GetCustomizableObject() == CurrentObject)
-		{
-			Instance->SetObject(CurrentObject);
-		}
-	}
+	CurrentObject->GetPrivate()->UpdateParameterPropertiesFromModel(CurrentModel);
 
 	// Order matters
 	CompileThread.Reset();
 	CompileTask.Reset();
 
+	UE_LOG(LogMutable, Verbose, TEXT("PROFILE: [ %16.8f ] Finishing Compilation task for CO [%s]."), FPlatformTime::Seconds(), *CurrentObject->GetName());
 	TRACE_END_REGION(UE_MUTABLE_COMPILE_REGION);
 
-	if (!Options.bDontUpdateStreamedDataAndCache)
-	{
-		TRACE_BEGIN_REGION(UE_MUTABLE_SAVEDD_REGION);
-
-		SaveDDTask = MakeShareable(new FCustomizableObjectSaveDDRunnable(CurrentObject, Options, Model));
-	}
-	else
-	{
-		Model->GetPrivate()->UnloadRoms();
-
-		// when skipping the SaveDerivedData task unlock the object so that instances can be updated
-		UCustomizableObjectSystem::GetInstance()->UnlockObject(CurrentObject);
-
-		CurrentObject->GetPrivate()->SetModel(Model, GenerateIdentifier(*CurrentObject));
-	}
-
-	UE_LOG(LogMutable, Display, TEXT("Finished Customizable Object Compile %s."), *CurrentObject->GetName());
+	// Create  SaveDD task
+	TRACE_BEGIN_REGION(UE_MUTABLE_SAVEDD_REGION);
+	SaveDDTask = MakeShareable(new FCustomizableObjectSaveDDRunnable(CurrentObject, CurrentOptions, CurrentModel));
 }
 
 
@@ -1431,7 +1401,7 @@ void FCustomizableObjectCompiler::FinishSavingDerivedDataTask()
 
 	check(SaveDDTask.IsValid());
 
-	if (Options.bIsCooking)
+	if (CurrentOptions.bIsCooking)
 	{
 		CurrentObject->GetPrivate()->CachePlatformData(
 				SaveDDTask->GetTargetPlatform(), 
@@ -1440,31 +1410,26 @@ void FCustomizableObjectCompiler::FinishSavingDerivedDataTask()
 				SaveDDTask->MorphDataBytes,
 				SaveDDTask->ClothingDataBytes);
 	}
-	
-
-	Model->GetPrivate()->UnloadRoms();
-	
-	// Unlock the object so that instances can be updated
-	UCustomizableObjectSystem::GetInstance()->UnlockObject(CurrentObject);
-
-	CurrentObject->GetPrivate()->SetModel(Model, GenerateIdentifier(*CurrentObject));
 
 	// Order matters
 	SaveDDThread.Reset();
 	SaveDDTask.Reset();
-	
-	RemoveCompileNotification();
 
-	NotifyCompilationErrors();
+	UE_LOG(LogMutable, Verbose, TEXT("PROFILE: [ %16.8f ] Finished Saving Derived Data task for CO [%s]."), FPlatformTime::Seconds(), *CurrentObject->GetName());
 
 	TRACE_END_REGION(UE_MUTABLE_SAVEDD_REGION);
-
 }
 
 
 void FCustomizableObjectCompiler::ForceFinishCompilation()
 {
-	if (CompileTask.IsValid())
+	if (AsynchronousStreamableHandlePtr)
+	{
+		AsynchronousStreamableHandlePtr->CancelHandle();
+		AsynchronousStreamableHandlePtr.Reset();
+	}
+
+	else if (CompileTask.IsValid())
 	{
 		// Compilation needs game thread tasks every now and then. Wait for compilation to finish while
 		// giving execution time for these tasks.
@@ -1474,45 +1439,35 @@ void FCustomizableObjectCompiler::ForceFinishCompilation()
 			FTaskGraphInterface::Get().ProcessThreadUntilIdle(ENamedThreads::GameThread);
 		}
 
-		CleanCachedReferencers();
-		UpdateArrayGCProtect();
-
-		UCustomizableObjectSystem* System = UCustomizableObjectSystem::GetInstance();
+		// Order matters
+		CompileThread.Reset();
+		CompileTask.Reset();
 
-		// Check the system has not started BeginDestroy, it can crash during a GC during a compilation forced shutdown without this check
-		if (!System->HasAnyFlags(EObjectFlags::RF_BeginDestroyed))
-		{
-			System->UnlockObject(CurrentObject);
-		}
+		UE_LOG(LogMutable, Verbose, TEXT("Force Finish Compilation task for Object."));
+		TRACE_END_REGION(UE_MUTABLE_COMPILE_REGION);
 	}
 
 	else if (SaveDDTask.IsValid())
 	{
 		SaveDDThread->WaitForCompletion();
 
-		UCustomizableObjectSystem* System = UCustomizableObjectSystem::GetInstance();
+		// Order matters
+		SaveDDThread.Reset();
+		SaveDDTask.Reset();
 
-		if (!System->HasAnyFlags(EObjectFlags::RF_BeginDestroyed))
-		{
-			System->UnlockObject(CurrentObject);
-		}
+		UE_LOG(LogMutable, Verbose, TEXT("Forced Finish Saving Derived Data task."));
+		TRACE_END_REGION(UE_MUTABLE_SAVEDD_REGION);
 	}
 
-	RemoveCompileNotification();
-}
-
-
-void FCustomizableObjectCompiler::ForceFinishBeforeStartCompilation(UCustomizableObject* Object)
-{
-	CleanCachedReferencers();
-	UpdateArrayGCProtect();
-
-	if (Object && Object->GetPrivate()->IsLocked())
+	if (CurrentRequest)
 	{
-		UCustomizableObjectSystem::GetInstance()->UnlockObject(Object);
+		CompleteRequest(ECompilationStatePrivate::Completed, ECompilationResultPrivate::Errors);
 	}
+}
 
-	RemoveCompileNotification();
+void FCustomizableObjectCompiler::ClearCompileRequests()
+{
+	CompileRequests.Empty();
 }
 
 
@@ -1523,7 +1478,7 @@ void FCustomizableObjectCompiler::AddCompileNotification(const FText& Compilatio
 	FCustomizableObjectEditorLogger::CreateLog(Text)
 	.SubText(CompilationStep)
 	.Category(ELoggerCategory::Compilation)
-	.Notification(!Options.bSilentCompilation)
+	.Notification(!CurrentOptions.bSilentCompilation)
 	.CustomNotification()
 	.FixNotification()
 	.Log();
@@ -1571,7 +1526,7 @@ void FCustomizableObjectCompiler::NotifyCompilationErrors() const
 	FCustomizableObjectEditorLogger::CreateLog(Message)
 	.Category(ELoggerCategory::Compilation)
 	.Severity(Severity)
-	.Notification(!Options.bSilentCompilation || !NoWarningsOrErrors)
+	.Notification(!CurrentOptions.bSilentCompilation || !NoWarningsOrErrors)
 	.CustomNotification()
 	.Log();
 }
@@ -1631,16 +1586,4 @@ void FCustomizableObjectCompiler::UpdateCompilerLogData()
 	}
 }
 
-
-void FCustomizableObjectCompiler::AddCompileOnlySelectedOption(const FString& ParamName, const FString& OptionValue)
-{
-	ParamNamesToSelectedOptions.Add(ParamName, OptionValue);
-}
-
-
-void FCustomizableObjectCompiler::ClearAllCompileOnlySelectedOption()
-{
-	ParamNamesToSelectedOptions.Empty();
-}
-
 #undef LOCTEXT_NAMESPACE
diff --git a/Engine/Plugins/Experimental/Mutable/Source/CustomizableObjectEditor/Private/MuCOE/CustomizableObjectCompiler.h b/Engine/Plugins/Experimental/Mutable/Source/CustomizableObjectEditor/Private/MuCOE/CustomizableObjectCompiler.h
--- a/Engine/Plugins/Experimental/Mutable/Source/CustomizableObjectEditor/Private/MuCOE/CustomizableObjectCompiler.h
+++ b/Engine/Plugins/Experimental/Mutable/Source/CustomizableObjectEditor/Private/MuCOE/CustomizableObjectCompiler.h
@@ -4,29 +4,44 @@
 
 #include "AssetRegistry/AssetData.h"
 #include "MuCO/CustomizableObjectPrivate.h"
+#include "MuCO/CustomizableObjectCompilerTypes.h"
 #include "MuCO/UnrealToMutableTextureConversionUtils.h"
 #include "MuCOE/CompilationMessageCache.h"
 #include "MuCOE/GenerateMutableSource/GenerateMutableSource.h"
 #include "MuCOE/CustomizableObjectEditorLogger.h"
 #include "UObject/GCObject.h"
+#include "TickableEditorObject.h"
+
+#include "Framework/Notifications/NotificationManager.h"
 
 class FCustomizableObjectCompileRunnable;
 class FCustomizableObjectSaveDDRunnable;
 class FReferenceCollector;
 class FRunnableThread;
 class FText;
+class UCustomizableObject;
 class UCustomizableObjectNode;
 
-
-class FCustomizableObjectCompiler : public FCustomizableObjectCompilerBase, public FGCObject
+class FCustomizableObjectCompiler : public FTickableEditorObject, public FTickableCookObject, public FGCObject
 {
 public:
 
-	CUSTOMIZABLEOBJECTEDITOR_API FCustomizableObjectCompiler();
-	CUSTOMIZABLEOBJECTEDITOR_API virtual ~FCustomizableObjectCompiler() {}
+	FCustomizableObjectCompiler();
+	virtual ~FCustomizableObjectCompiler() {}
 
 	/** Check for pending compilation process. Returns true if an object has been updated. */
-	CUSTOMIZABLEOBJECTEDITOR_API virtual bool Tick(bool bBlocking) override;
+	bool Tick(bool bBlocking = false);
+	int32 GetNumRemainingWork() const;
+
+	// FTickableGameObject interface
+	virtual ETickableTickType GetTickableTickType() const override { return ETickableTickType::Conditional; };
+	virtual bool IsTickable() const override;
+	virtual void Tick( float InDeltaTime ) override;
+	virtual TStatId GetStatId() const override;
+
+	// FTickableCookObject interface
+	virtual void TickCook(float DeltaTime, bool bCookCompete) override;
+
 
 	/** Generate the Mutable Graph from the Unreal Graph. */
 	mu::NodePtr Export(UCustomizableObject* Object, const FCompilationOptions& Options, TArray<TSoftObjectPtr<UTexture>>& OutRuntimeReferencedTextures, TArray<FMutableSourceTextureData>& OutCompilerReferencedTextures);
@@ -37,7 +52,9 @@ public:
 
 	void FinishCompilationTask();
 	void FinishSavingDerivedDataTask();
-	virtual void ForceFinishCompilation() override;
+
+	void ForceFinishCompilation();
+	void ClearCompileRequests();
 
 	void AddCompileNotification(const FText& CompilationStep) const;
 	static void RemoveCompileNotification();
@@ -45,8 +62,11 @@ public:
 	/** Load required assets and compile.
 	 *
 	 * Loads assets which reference Object's package asynchronously before calling ProcessChildObjectsRecursively. */
-	virtual void Compile(UCustomizableObject& Object, const FCompilationOptions& Options, bool bAsync) override;
+	void Compile(const TSharedRef<FCompilationRequest>& InCompileRequest);
+	void Compile(const TArray<TSharedRef<FCompilationRequest>>& InCompileRequests);
 	
+	bool IsRequestQueued(const TSharedRef<FCompilationRequest>& InCompileRequest) const;
+
 	/** FSerializableObject interface */
 	virtual void AddReferencedObjects(FReferenceCollector& Collector) override;
 	virtual FString GetReferencerName() const override
@@ -57,17 +77,6 @@ public:
 	/** Simply add CO elements from ArrayAssetData to ArrayGCProtect when they've been loaded from ArrayAssetData */
 	void UpdateArrayGCProtect();
 
-	// This is used to restrict Group nodes to compile only the SelectedOption and discard all the others
-	void AddCompileOnlySelectedOption(const FString& ParamName, const FString& OptionValue);
-	// This clears all restrictions and let's a full compilation happen
-	void ClearAllCompileOnlySelectedOption();
-
-	// Function to clear the compiler if the editor was closed before starting the compilation 
-	void ForceFinishBeforeStartCompilation(UCustomizableObject* Object);
-
-	// Getter of AsynchronousStreamableHandlePtr
-	TSharedPtr<FStreamableHandle> GetAsynchronousStreamableHandlePtr() { return AsynchronousStreamableHandlePtr; }
-
 private:
 
 	// Object containing all error and warning logs raised during compilation.
@@ -77,67 +86,33 @@ private:
 	
 	void CompileInternal(bool bAsync = false);
 
+	void CompleteRequest(ECompilationStatePrivate State, ECompilationResultPrivate Result);
+	bool TryPopCompileRequest();
+
 	void PreloadingReferencerAssetsCallback(bool bAsync);
 	
 	void ProcessChildObjectsRecursively(UCustomizableObject* Object, FMutableGraphGenerationContext &GenerationContext);
-
-	//
-	FCompilationOptions Options;
-
-	//
-	TSharedPtr< FCustomizableObjectCompileRunnable > CompileTask;
-
-	//
-	TSharedPtr< FRunnableThread > CompileThread;
-
-	//
-	TSharedPtr< FCustomizableObjectSaveDDRunnable > SaveDDTask;
-
-	//
-	TSharedPtr< FRunnableThread > SaveDDThread;
-
-	// Cache configuration settings from ini files 
-	ECustomizableObjectNumBoneInfluences CustomizableObjectNumBoneInfluences = ECustomizableObjectNumBoneInfluences::Four;
-
-	// Protected from GC with FCustomizableObjectCompiler::AddReferencedObjects
-	TObjectPtr<UCustomizableObject> CurrentObject = nullptr;
-
-	/** Array where to put the names of the already processed child in ProcessChildObjectsRecursively */
-	TArray<FName> ArrayAlreadyProcessedChild;
-
-	TArray<FSoftObjectPath> ArrayAssetToStream;
 	
 	// Will output to Mutable Log the warning and error messages generated during the CO compilation
 	// and update the values of NumWarnings and NumErrors
 	void UpdateCompilerLogData();
 
-
 	/** If duplicated elements are found in each entry of ParameterNamesMap, a warning for
 	the parameters with repeated name will be generated */
-	void DisplayParameterWarning(struct FMutableGraphGenerationContext& GenerationContext);
+	void DisplayParameterWarning(FMutableGraphGenerationContext& GenerationContext);
 	
 	/** If duplicated node ids are found, usually due to duplicating CustomizableObjects Assets, a warning
 	for the nodes with repeated ids will be generated */
-	void DisplayDuplicatedNodeIdsWarning(struct FMutableGraphGenerationContext& GenerationContext);
+	void DisplayDuplicatedNodeIdsWarning(FMutableGraphGenerationContext& GenerationContext);
 	
 	/** Display warnings for unnamed node objects */
-	void DisplayUnnamedNodeObjectWarning(struct FMutableGraphGenerationContext& GenerationContext);
-	
-	/** Display warnings from discarded PhysicsAssets due to SkeletalBodySetups with no corresponding bones
-	in the SkeletalMesh's RefSkeleton */
-	//void DisplayDiscardedPhysicsAssetSingleWarning(struct FMutableGraphGenerationContext& GenerationContext);
+	void DisplayUnnamedNodeObjectWarning(FMutableGraphGenerationContext& GenerationContext);
 	
 	/** Display a warning for each node contains an orphan pin. */
-	void DisplayOrphanNodesWarning(struct FMutableGraphGenerationContext& GenerationContext);
+	void DisplayOrphanNodesWarning(FMutableGraphGenerationContext& GenerationContext);
 	
 	mu::NodeObjectPtr GenerateMutableRoot(UCustomizableObject* Object, FMutableGraphGenerationContext& GenerationContext, FText& ErrorMessage, bool& bOutIsRootObject);
 
-	/** Add to ArrayAssetData the FAssetData information of all referencers of static class type UCustomizableObject::StaticClass()
-	* that reference the package given by the PathName parameter
-	* @param PathName            [in]  path to the CO to be analyzed (for instance, CO->GetOuter()->GetPathName())
-	* @param ArrayReferenceNames [out] array with the package names which are referenced by the CO with PathName given as parameter
-	* @return nothing */
-	void AddCachedReferencers(const FName& PathName, TArray<FName>& ArrayReferenceNames);
 
 	/** Launches the compile task in another thread when compiling a CO in the editor
 	* @param bShowNotification [in] whether to show the compiling CO notification or not
@@ -150,8 +125,12 @@ private:
 	* @return nothing */
 	void SaveCODerivedData();
 
-	/** When compiling a CO in the editor, flag to know if there's a mutable task pending to be launched through LaunchMutableCompile */
-	bool CompilationLaunchPending;
+	/** Add to ArrayAssetData the FAssetData information of all referencers of static class type UCustomizableObject::StaticClass()
+	* that reference the package given by the PathName parameter
+	* @param PathName            [in]  path to the CO to be analyzed (for instance, CO->GetOuter()->GetPathName())
+	* @param ArrayReferenceNames [out] array with the package names which are referenced by the CO with PathName given as parameter
+	* @return nothing */
+	void AddCachedReferencers(const FName& PathName, TArray<FName>& ArrayReferenceNames);
 
 	/** Just used to clean ArrayAssetData
 	* @return nothing */
@@ -167,41 +146,48 @@ private:
 	* @return pointer to element if any found, nullptr otherwise */
 	FAssetData* GetCachedAssetData(const FString& PackageName);
 
-	/** Helper function to compute the value for Unreal Engine variable s.AsyncLoadingTimeLimit while asynchronous loading is used.
-	* Also assigned to MaxConvertToMutableTextureTime
-	* @return value in milliseconds to use for AsyncLoadingTimeLimit and MaxConvertToMutableTextureTime */
-	float ComputeAsyncLoadingTimeLimit();
-	
-    ECompilationResultPrivate GetCompilationResult() const;
-	
-public:
-	
-	virtual void GetCompilationMessages(TArray<FText>& OutWarningMessages, TArray<FText>& OutErrorMessages) const override;
-	
-private:
+	ECompilationResultPrivate GetCompilationResult() const;
+
+	/** Pointer to the Asynchronous Preloading process call back */
+	TSharedPtr<FStreamableHandle> AsynchronousStreamableHandlePtr;
+	TArray<FSoftObjectPath> ArrayAssetToStream;
+
+	/** Compile task and thread. */
+	TSharedPtr<FCustomizableObjectCompileRunnable> CompileTask;
+	TSharedPtr<FRunnableThread> CompileThread;
+
+	/** SaveDD task and thread. */
+	TSharedPtr<FCustomizableObjectSaveDDRunnable> SaveDDTask;
+	TSharedPtr<FRunnableThread> SaveDDThread;
+
+	// Cache configuration settings from ini files 
+	ECustomizableObjectNumBoneInfluences CustomizableObjectNumBoneInfluences = ECustomizableObjectNumBoneInfluences::Four;
+
+	/** Array where to put the names of the already processed child in ProcessChildObjectsRecursively */
+	TArray<FName> ArrayAlreadyProcessedChild;
+
 	/** Array with all the packages used to compile current Customizable Object */
 	TArray<FAssetData> ArrayAssetData;
 
 	/** Array used to protect from garbage collection those COs loaded asynchronously */
 	TArray<TObjectPtr<UCustomizableObject>> ArrayGCProtect;
 
-	/** Flag to know when asynchronous asset loading is being performed */
-	bool PreloadingReferencerAssets;
+	TSharedPtr<mu::Model, ESPMode::ThreadSafe> CurrentModel;
+
+	// Protected from GC with FCustomizableObjectCompiler::AddReferencedObjects
+	TObjectPtr<UCustomizableObject> CurrentObject = nullptr;
 
-	/** Copy of GAsyncLoadingTimeLimit while assets are loaded, previous value is restored after asset load */
-	float CurrentGAsyncLoadingTimeLimit;
+	FCompilationOptions CurrentOptions;
 
-	/** Counter to know how many of the textures in ArrayTextureUnrealToMutableTask have been converted from Unreal to Mutable */
-	int32 CompletedUnrealToMutableTask;
+	/** Current Compilation request. */
+	TSharedPtr<FCompilationRequest> CurrentRequest;
 
-	/** Time threshold used to prevent doing too much work in one tick */
-	float MaxConvertToMutableTextureTime;
+	/** Pending requests. */
+	TArray<TSharedRef<FCompilationRequest>> CompileRequests;
 
-	// Stores the only option of an Int Param that should be compiled
-	TMap<FString, FString> ParamNamesToSelectedOptions;
+	uint32 NumCompilationRequests = 0;
 
-	/** Pointer to the Asynchronous Preloading process call back */
-	TSharedPtr<FStreamableHandle> AsynchronousStreamableHandlePtr;
+	/** Compilation progress bar handle */
+	FProgressNotificationHandle CompileNotificationHandle;
 
-	TSharedPtr<mu::Model, ESPMode::ThreadSafe> Model;
 };
diff --git a/Engine/Plugins/Experimental/Mutable/Source/CustomizableObjectEditor/Private/MuCOE/CustomizableObjectEditor.cpp b/Engine/Plugins/Experimental/Mutable/Source/CustomizableObjectEditor/Private/MuCOE/CustomizableObjectEditor.cpp
--- a/Engine/Plugins/Experimental/Mutable/Source/CustomizableObjectEditor/Private/MuCOE/CustomizableObjectEditor.cpp
+++ b/Engine/Plugins/Experimental/Mutable/Source/CustomizableObjectEditor/Private/MuCOE/CustomizableObjectEditor.cpp
@@ -11,6 +11,7 @@
 #include "FileHelpers.h"
 #include "Framework/Commands/GenericCommands.h"
 #include "Framework/MultiBox/MultiBoxBuilder.h"
+#include "Framework/Notifications/NotificationManager.h"
 #include "GraphEditorActions.h"
 #include "HAL/PlatformApplicationMisc.h"
 #include "IDetailsView.h"
@@ -20,6 +21,7 @@
 #include "MuCO/CustomizableObjectSystem.h"
 #include "MuCO/CustomizableSkeletalComponent.h"
 #include "MuCO/UnrealPortabilityHelpers.h"
+#include "MuCO/CustomizableObjectCompilerTypes.h"
 #include "MuCOE/CustomizableObjectCustomSettings.h"
 #include "MuCOE/CustomizableObjectEditorActions.h"
 #include "MuCOE/CustomizableObjectEditorLogger.h"
@@ -39,6 +41,7 @@
 #include "MuCOE/Nodes/CustomizableObjectNodeMaterialVariation.h"
 #include "MuCOE/Nodes/CustomizableObjectNodeMeshClipMorph.h"
 #include "MuCOE/Nodes/CustomizableObjectNodeMeshClipWithMesh.h"
+#include "MuCOE/Nodes/CustomizableObjectNodeObject.h"
 #include "MuCOE/Nodes/CustomizableObjectNodeObjectGroup.h"
 #include "MuCOE/Nodes/CustomizableObjectNodeProjectorConstant.h"
 #include "MuCOE/Nodes/CustomizableObjectNodeSkeletalMesh.h"
@@ -172,13 +175,6 @@ FCustomizableObjectEditor::~FCustomizableObjectEditor()
 	PreviewCustomizableSkeletalComponents.Empty();
 	CustomizableObjectDetailsView.Reset();
 	GEditor->UnregisterForUndo(this);
-	Compiler.ForceFinishCompilation();
-
-	if (Compiler.GetAsynchronousStreamableHandlePtr().IsValid() && !Compiler.GetAsynchronousStreamableHandlePtr()->HasLoadCompleted())
-	{
-		Compiler.GetAsynchronousStreamableHandlePtr()->CancelHandle();
-		Compiler.ForceFinishBeforeStartCompilation(CustomizableObject);
-	}
 
 	FCoreUObjectDelegates::OnObjectModified.RemoveAll(this);
 
@@ -648,29 +644,6 @@ UCustomizableObjectInstance* FCustomizableObjectEditor::GetPreviewInstance()
 }
 
 
-void FCustomizableObjectEditor::CompileObjectUserPressedButton()
-{
-	Compiler.ClearAllCompileOnlySelectedOption();
-	CompileObject();
-}
-
-
-void FCustomizableObjectEditor::CompileOnlySelectedObjectUserPressedButton()
-{
-	if (PreviewInstance)
-	{
-		Compiler.ClearAllCompileOnlySelectedOption();
-
-		for (const FCustomizableObjectIntParameterValue& IntParam : PreviewInstance->GetPrivate()->GetDescriptor().GetIntParameters())
-		{
-			Compiler.AddCompileOnlySelectedOption(IntParam.ParameterName, IntParam.ParameterValueName);
-		}
-
-		CompileObject();
-	}
-}
-
-
 void FCustomizableObjectEditor::BindCommands()
 {
 	const FCustomizableObjectEditorCommands& Commands = FCustomizableObjectEditorCommands::Get();
@@ -679,13 +652,13 @@ void FCustomizableObjectEditor::BindCommands()
 	// Compile and options
 	ToolkitCommands->MapAction(
 		Commands.Compile,
-		FExecuteAction::CreateSP(this, &FCustomizableObjectEditor::CompileObjectUserPressedButton),
+		FExecuteAction::CreateSP(this, &FCustomizableObjectEditor::CompileObject, false),
 		FCanExecuteAction::CreateStatic(&UCustomizableObjectSystem::IsActive),
 		FIsActionChecked());
 
 	ToolkitCommands->MapAction(
 		Commands.CompileOnlySelected,
-		FExecuteAction::CreateSP(this, &FCustomizableObjectEditor::CompileOnlySelectedObjectUserPressedButton),
+		FExecuteAction::CreateSP(this, &FCustomizableObjectEditor::CompileObject, true),
 		FCanExecuteAction::CreateStatic(&UCustomizableObjectSystem::IsActive),
 		FIsActionChecked());
 
@@ -1487,14 +1460,11 @@ void FCustomizableObjectEditor::OnObjectModified(UObject* Object)
 }
 
 
-void FCustomizableObjectEditor::CompileObject()
+void FCustomizableObjectEditor::CompileObject(bool bOnlySelectedParameters)
 {
 	// Resetting viewport parameters
 	Viewport->SetDrawDefaultUVMaterial();
 
-	UE_LOG(LogMutable, Verbose, TEXT("PROFILE: -----------------------------------------------------------"));
-	UE_LOG(LogMutable, Verbose, TEXT("PROFILE: [ %16.8f ] FCustomizableObjectEditor::CompileObject start."), FPlatformTime::Seconds());
-
 	if (CustomizableObject->GetPrivate()->Status.Get() == FCustomizableObjectStatus::EState::Loading)
 	{
 		FNotificationInfo Info(LOCTEXT("CustomizableObjectCompileTryLater", "Please wait until Customizable Object is loaded"));
@@ -1506,14 +1476,29 @@ void FCustomizableObjectEditor::CompileObject()
 		return;
 	}
 
-	if (CustomizableObject->GetPrivate()->GetSource())
+	if (!CustomizableObject->GetPrivate()->GetSource())
 	{
-		FCompilationOptions Options = CustomizableObject->CompileOptions;
-		Options.bSilentCompilation = false;
-		Compiler.Compile(*CustomizableObject, Options, true);
+		return;
 	}
 
-	UE_LOG(LogMutable, Verbose, TEXT("PROFILE: [ %16.8f ] FCustomizableObjectEditor::CompileObject end."), FPlatformTime::Seconds());
+	TSharedRef<FCompilationRequest> CompileRequest = MakeShared<FCompilationRequest>(*CustomizableObject, true);
+	CompileRequest->GetCompileOptions().bSilentCompilation = false;
+
+	if (bOnlySelectedParameters && PreviewInstance)
+	{
+		const TArray<FCustomizableObjectIntParameterValue>& IntParameters = PreviewInstance->GetPrivate()->GetDescriptor().GetIntParameters();
+		TMap<FString, FString> ParamNamesToSelectedOptions;
+		ParamNamesToSelectedOptions.Reserve(IntParameters.Num());
+
+		for (const FCustomizableObjectIntParameterValue& IntParam : IntParameters)
+		{
+			ParamNamesToSelectedOptions.Add(IntParam.ParameterName, IntParam.ParameterValueName);
+		}
+
+		CompileRequest->SetParameterNamesToSelectedOptions(ParamNamesToSelectedOptions);
+	}
+
+	ICustomizableObjectEditorModule::GetChecked().CompileCustomizableObject(CompileRequest);
 }
 
 
@@ -1734,25 +1719,6 @@ void FCustomizableObjectEditor::NotifyPostChange( const FPropertyChangedEvent& P
 	}
 }
 
-
-bool FCustomizableObjectEditor::IsTickable() const
-{
-	return true;
-}
-
-
-void FCustomizableObjectEditor::Tick(float InDeltaTime)
-{
-	Compiler.Tick(false);
-}
-
-
-TStatId FCustomizableObjectEditor::GetStatId() const
-{
-	RETURN_QUICK_DECLARE_CYCLE_STAT(FCustomizableObjectEditor, STATGROUP_Tickables);
-}
-
-
 void FCustomizableObjectEditor::CopySelectedNodes()
 {
 	// Export the selected nodes and place the text on the clipboard
diff --git a/Engine/Plugins/Experimental/Mutable/Source/CustomizableObjectEditor/Private/MuCOE/CustomizableObjectEditor.h b/Engine/Plugins/Experimental/Mutable/Source/CustomizableObjectEditor/Private/MuCOE/CustomizableObjectEditor.h
--- a/Engine/Plugins/Experimental/Mutable/Source/CustomizableObjectEditor/Private/MuCOE/CustomizableObjectEditor.h
+++ b/Engine/Plugins/Experimental/Mutable/Source/CustomizableObjectEditor/Private/MuCOE/CustomizableObjectEditor.h
@@ -6,9 +6,7 @@
 #include "EditorUndoClient.h"
 #include "GraphEditor.h"
 #include "Misc/NotifyHook.h"
-#include "MuCOE/CustomizableObjectCompiler.h"
 #include "MuCOE/ICustomizableObjectEditor.h"
-#include "TickableEditorObject.h"
 #include "MuCO/CustomizableObjectPrivate.h"
 #include "Widgets/Input/SNumericDropDown.h"
 
@@ -156,7 +154,6 @@ class FCustomizableObjectEditor :
 	public ICustomizableObjectEditor, 
 	public FGCObject,
 	public FNotifyHook,
-	public FTickableEditorObject,
 	public FEditorUndoClient
 {
 public:
@@ -235,11 +232,6 @@ public:
 	// FNotifyHook interface
 	virtual void NotifyPostChange( const FPropertyChangedEvent& PropertyChangedEvent, FProperty* PropertyThatChanged ) override;
 
-	// FTickableGameObject interface
-	virtual bool IsTickable() const override;
-	virtual void Tick( float InDeltaTime ) override;
-	virtual TStatId GetStatId() const override;
-
 	// Delegates
 	void DeleteSelectedNodes();
 	bool CanDeleteNodes() const;
@@ -292,9 +284,7 @@ private:
 	TSharedPtr<FUICommandList> GraphEditorCommands;
 
 	// Compile the customizable object.
-	void CompileObject();
-	void CompileObjectUserPressedButton();
-	void CompileOnlySelectedObjectUserPressedButton();
+	void CompileObject(bool bOnlySelectedParameters);
 	
 	// Compile options menu callbacks
 	TSharedRef<SWidget> GenerateCompileOptionsMenuContent(TSharedRef<FUICommandList> InCommandList);
@@ -403,9 +393,6 @@ private:
 	TArray<TObjectPtr<UCustomizableSkeletalComponent>> PreviewCustomizableSkeletalComponents;
 	TArray<TObjectPtr<UDebugSkelMeshComponent>> PreviewSkeletalMeshComponents;
 
-	/** Object compiler */
-	FCustomizableObjectCompiler Compiler;
-
 	/** List of open tool panels; used to ensure only one exists at any one time */
 	TMap< FName, TWeakPtr<class SDockableTab> > SpawnedToolPanels;
 
diff --git a/Engine/Plugins/Experimental/Mutable/Source/CustomizableObjectEditor/Private/MuCOE/CustomizableObjectEditorFunctionLibrary.cpp b/Engine/Plugins/Experimental/Mutable/Source/CustomizableObjectEditor/Private/MuCOE/CustomizableObjectEditorFunctionLibrary.cpp
--- a/Engine/Plugins/Experimental/Mutable/Source/CustomizableObjectEditor/Private/MuCOE/CustomizableObjectEditorFunctionLibrary.cpp
+++ b/Engine/Plugins/Experimental/Mutable/Source/CustomizableObjectEditor/Private/MuCOE/CustomizableObjectEditorFunctionLibrary.cpp
@@ -4,51 +4,36 @@
 #include "MuCOE/CustomizableObjectCompiler.h"
 #include "MuCO/CustomizableObjectSystem.h"
 #include "MuCO/CustomizableObject.h"
+#include "MuCO/ICustomizableObjectEditorModule.h"
 
 ECustomizableObjectCompilationState UCustomizableObjectEditorFunctionLibrary::CompileCustomizableObjectSynchronously(
 	UCustomizableObject* CustomizableObject,
 	ECustomizableObjectOptimizationLevel InOptimizationLevel,
 	ECustomizableObjectTextureCompression InTextureCompression)
 {
-	if (!CustomizableObject)
-	{
-		UE_LOG(LogMutable, Warning, TEXT("Attempted to compile nullptr Customizable Object"));
-		return ECustomizableObjectCompilationState::Failed;
-	}
-
-	FString ObjectPath = CustomizableObject->GetPathName();
-	if (CustomizableObject->GetPrivate()->IsLocked())
-	{
-		// Take this if you need a hack:
-		// UCustomizableObjectSystem::GetInstance()->UnlockObject(CustomizableObject);
-		UE_LOG( LogMutable, Warning, 
-			TEXT("Attempted to compile %s Customizable Object when it is locked"),
-			*ObjectPath);
-		return ECustomizableObjectCompilationState::Failed;
-	}
-
 	// store package dirty state so that we can restore it - compile is not an edit:
 	const bool bPackageWasDirty = CustomizableObject->GetOutermost()->IsDirty();
 
 	const double StartTime = FPlatformTime::Seconds();
 
-	FCustomizableObjectCompiler Compiler;
-	FCompilationOptions Options = CustomizableObject->CompileOptions;
+	TSharedRef<FCompilationRequest> CompileRequest = MakeShared<FCompilationRequest>(*CustomizableObject, false);
+	FCompilationOptions& Options = CompileRequest->GetCompileOptions();
 	Options.OptimizationLevel = static_cast<int32>(InOptimizationLevel);
 	Options.TextureCompression = InTextureCompression;
 	Options.bSilentCompilation = false;
-	const bool bAsync = false;
-	Compiler.Compile(*CustomizableObject, Options, bAsync);
+	ICustomizableObjectEditorModule::GetChecked().CompileCustomizableObject(CompileRequest);
+
+	check(CompileRequest->GetCompilationState() == ECompilationStatePrivate::Completed);
 
 	CustomizableObject->GetOutermost()->SetDirtyFlag(bPackageWasDirty);
 
-	const bool bCompilationSuccess = CustomizableObject->GetPrivate()->CompilationResult == ECompilationResultPrivate::Success ||
-		CustomizableObject->GetPrivate()->CompilationResult == ECompilationResultPrivate::Warnings;
+	const bool bCompilationSuccess = CompileRequest->GetCompilationResult() == ECompilationResultPrivate::Success ||
+		CompileRequest->GetCompilationResult() == ECompilationResultPrivate::Warnings;
 	
 	const double CurrentTime = FPlatformTime::Seconds();
 	UE_LOG( LogMutable, Display,
 		TEXT("Synchronously Compiled %s %s in %f seconds"),
-		*ObjectPath, 
+		*GetPathNameSafe(CustomizableObject), 
 		bCompilationSuccess ? TEXT("successfully") : TEXT("unsuccessfully"),
 		CurrentTime - StartTime
 	);
diff --git a/Engine/Plugins/Experimental/Mutable/Source/CustomizableObjectEditor/Private/MuCOE/CustomizableObjectEditorModule.cpp b/Engine/Plugins/Experimental/Mutable/Source/CustomizableObjectEditor/Private/MuCOE/CustomizableObjectEditorModule.cpp
--- a/Engine/Plugins/Experimental/Mutable/Source/CustomizableObjectEditor/Private/MuCOE/CustomizableObjectEditorModule.cpp
+++ b/Engine/Plugins/Experimental/Mutable/Source/CustomizableObjectEditor/Private/MuCOE/CustomizableObjectEditorModule.cpp
@@ -7,10 +7,12 @@
 #include "ISettingsModule.h"
 #include "ISettingsSection.h"
 #include "MessageLogModule.h"
+#include "Misc/MessageDialog.h"
 #include "MuCO/CustomizableObjectSystem.h"		// For defines related to memory function replacements.
 #include "MuCO/CustomizableObjectInstanceUsage.h"
 #include "MuCO/CustomizableSkeletalMeshActor.h"
 #include "MuCO/ICustomizableObjectModule.h"		// For instance editor command utility function
+#include "MuCO/UnrealPortabilityHelpers.h"
 #include "MuCOE/CustomizableInstanceDetails.h"
 #include "MuCOE/CustomizableObjectCustomSettings.h"
 #include "MuCOE/CustomizableObjectCustomSettingsDetails.h"
@@ -65,6 +67,7 @@
 #include "Components/SkeletalMeshComponent.h"
 #include "MuCOE/GraphTraversal.h"
 #include "MuCOE/CustomizableObjectInstanceBaker.h"
+#include "Editor.h"
 
 class AActor;
 class FString;
@@ -249,11 +252,17 @@ void FCustomizableObjectEditorModule::StartupModule()
 		ShowOnScreenCompileWarnings();
 		return true;
 	});
+
+	FEditorDelegates::PreBeginPIE.AddRaw(this, &FCustomizableObjectEditorModule::OnPreBeginPIE);
 }
 
 
 void FCustomizableObjectEditorModule::ShutdownModule()
 {
+	FEditorDelegates::PreBeginPIE.RemoveAll(this);
+
+	check(Compiler.GetNumRemainingWork() == 0);
+
 	if( FModuleManager::Get().IsModuleLoaded( "PropertyEditor" ) )
 	{
 		FPropertyEditorModule& PropertyModule = FModuleManager::GetModuleChecked<FPropertyEditorModule>("PropertyEditor");
@@ -580,12 +589,6 @@ bool FCustomizableObjectEditorModule::IsCompilationOutOfDate(const UCustomizable
 }
 
 
-TSharedRef<FCustomizableObjectCompilerBase> FCustomizableObjectEditorModule::CreateCompiler() const
-{
-	return MakeShared<FCustomizableObjectCompiler>();
-}
-
-
 bool FCustomizableObjectEditorModule::IsRootObject(const UCustomizableObject& Object) const
 {
 	return GraphTraversal::IsRootObject(Object);
@@ -611,4 +614,132 @@ void FCustomizableObjectEditorModule::BakeCustomizableObjectInstance(UCustomizab
 }
 
 
+void CompileCustomizableObjectsSync(const TSharedRef<FCompilationRequest>& CompilationRequest)
+{
+	FCustomizableObjectCompiler* SyncCompiler = new FCustomizableObjectCompiler();
+	SyncCompiler->Compile(CompilationRequest);
+	delete SyncCompiler;
+}
+
+
+void FCustomizableObjectEditorModule::CompileCustomizableObject(const TSharedRef<FCompilationRequest>& InCompilationRequest, bool bForceRequest)
+{
+	if (IsRunningGame())
+	{
+		return;
+	}
+
+	if (InCompilationRequest->IsAsyncCompilation())
+	{
+		CompileCustomizableObjects({ InCompilationRequest }, bForceRequest);
+	}
+	else
+	{
+		CompileCustomizableObjectsSync(InCompilationRequest);
+	}
+}
+
+void FCustomizableObjectEditorModule::CompileCustomizableObjects(const TArray<TSharedRef<FCompilationRequest>>& InCompilationRequests, bool bForceRequests)
+{
+	check(IsInGameThread());
+
+	if (IsRunningGame())
+	{
+		return;
+	}
+
+	TArray<TSharedRef<FCompilationRequest>> FilteredAsyncRequests;
+	FilteredAsyncRequests.Reserve(InCompilationRequests.Num());
+
+	for (const TSharedRef<FCompilationRequest>& Request : InCompilationRequests)
+	{
+		const UCustomizableObject* CustomizableObject = Request->GetCustomizableObject();
+		if (!CustomizableObject)
+		{
+			continue;
+		}
+
+		if (!Request->IsAsyncCompilation())
+		{
+			CompileCustomizableObjectsSync(Request);
+		}
+
+		else if (bForceRequests ||
+			(!CustomizableObject->GetPrivate()->IsLocked() && !Compiler.IsRequestQueued(Request)))
+		{
+			FilteredAsyncRequests.Add(Request);
+		}
+	}
+
+	Compiler.Compile(FilteredAsyncRequests);
+}
+
+
+int32 FCustomizableObjectEditorModule::Tick(bool bBlocking)
+{
+	Compiler.Tick(bBlocking);
+	return Compiler.GetNumRemainingWork();
+}
+
+
+void FCustomizableObjectEditorModule::CancelCompileRequests()
+{
+	Compiler.ForceFinishCompilation();
+	Compiler.ClearCompileRequests();
+}
+
+
+void FCustomizableObjectEditorModule::OnPreBeginPIE(const bool bIsSimulatingInEditor)
+{
+	if (IsRunningGame() || !UCustomizableObjectSystem::IsActive())
+	{
+		return;
+	}
+
+	UCustomizableObjectSystem* System = UCustomizableObjectSystem::GetInstanceChecked();
+	if (!System->EditorSettings.bCompileRootObjectsOnStartPIE)
+	{
+		return;
+	}
+
+	// Find root customizable objects
+	FARFilter AssetRegistryFilter;
+	UE_MUTABLE_GET_CLASSPATHS(AssetRegistryFilter).Add(UE_MUTABLE_TOPLEVELASSETPATH(TEXT("/Script/CustomizableObject"), TEXT("CustomizableObject")));
+	AssetRegistryFilter.TagsAndValues.Add(FName("IsRoot"), FString::FromInt(1));
+
+	TArray<FAssetData> OutAssets;
+	const FAssetRegistryModule& AssetRegistryModule = FModuleManager::LoadModuleChecked<FAssetRegistryModule>("AssetRegistry");
+	AssetRegistryModule.Get().GetAssets(AssetRegistryFilter, OutAssets);
+
+	TArray<TSharedRef<FCompilationRequest>> Requests;
+	for (const FAssetData& Asset : OutAssets)
+	{
+		// If it is referenced by PIE it should be loaded
+		if (!Asset.IsAssetLoaded())
+		{
+			continue;
+		}
+
+		UCustomizableObject* Object = Cast<UCustomizableObject>(Asset.GetAsset());
+		if (!Object || Object->IsCompiled() || Object->GetPrivate()->IsLocked())
+		{
+			continue;
+		}
+
+		// Add uncompiled objects to the objects to cook list
+		TSharedRef<FCompilationRequest> NewRequest = MakeShared<FCompilationRequest>(*Object, true);
+		NewRequest->GetCompileOptions().bSilentCompilation = true;
+		Requests.Add(NewRequest);
+	}
+
+	if (!Requests.IsEmpty())
+	{
+		const FText Msg = FText::FromString(TEXT("Warning: one or more Customizable Objects used in PIE are uncompiled.\n\nDo you want to compile them?"));
+		if (FMessageDialog::Open(EAppMsgType::OkCancel, Msg) == EAppReturnType::Ok)
+		{
+			CompileCustomizableObjects(Requests);
+		}
+	}
+}
+
 #undef LOCTEXT_NAMESPACE
diff --git a/Engine/Plugins/Experimental/Mutable/Source/CustomizableObjectEditor/Private/MuCOE/CustomizableObjectEditorModule.h b/Engine/Plugins/Experimental/Mutable/Source/CustomizableObjectEditor/Private/MuCOE/CustomizableObjectEditorModule.h
--- a/Engine/Plugins/Experimental/Mutable/Source/CustomizableObjectEditor/Private/MuCOE/CustomizableObjectEditorModule.h
+++ b/Engine/Plugins/Experimental/Mutable/Source/CustomizableObjectEditor/Private/MuCOE/CustomizableObjectEditorModule.h
@@ -3,11 +3,12 @@
 #pragma once
 
 #include "MuCO/ICustomizableObjectEditorModule.h"
+#include "MuCO/CustomizableObjectCompilerTypes.h"
 
+#include "MuCOE/CustomizableObjectCompiler.h"
 #include "MuCOE/CustomizableObjectEditorLogger.h"
 #include "Toolkits/AssetEditorToolkit.h"
 
-class FCustomizableObjectCompilerBase;
 struct FBakingConfiguration;
 class USkeletalMeshComponent;
 class FPropertyEditorModule;
@@ -38,16 +39,22 @@ public:
 	// ICustomizableObjectEditorModule interface
 	virtual FCustomizableObjectEditorLogger& GetLogger() override;
 	virtual bool IsCompilationOutOfDate(const UCustomizableObject& Object, TArray<FName>* OutOfDatePackages) const override;
-	virtual TSharedRef<FCustomizableObjectCompilerBase> CreateCompiler() const override;
 	virtual bool IsRootObject(const UCustomizableObject& Object) const override;
 	virtual void BakeCustomizableObjectInstance(UCustomizableObjectInstance* InTargetInstance, const FBakingConfiguration& InBakingConfig) override;
 
+	/** Request for a given customizable object to be compiled. Async compile requests will be queued and processed sequentially. 
+	 * @param InCompilationRequest - Request to compile an object. 
+	 * @param bForceRequest - Queue request even if already in the pending list. */
+	virtual void CompileCustomizableObject(const TSharedRef<FCompilationRequest>& InCompilationRequest, bool bForceRequest = false) override;
+	virtual void CompileCustomizableObjects(const TArray<TSharedRef<FCompilationRequest>>& InCompilationRequests, bool bForceRequests = false) override;
+
+	virtual int32 Tick(bool bBlocking) override;
+
+	virtual void CancelCompileRequests() override;
+
 	virtual TSharedPtr<FExtensibilityManager> GetCustomizableObjectEditorToolBarExtensibilityManager() override { return CustomizableObjectEditor_ToolBarExtensibilityManager; }
 	virtual TSharedPtr<FExtensibilityManager> GetCustomizableObjectEditorMenuExtensibilityManager() override { return CustomizableObjectEditor_MenuExtensibilityManager; }
 
-	bool HandleSettingsSaved();
-	void RegisterSettings();
-
 private:	
 	TSharedPtr<FExtensibilityManager> CustomizableObjectEditor_ToolBarExtensibilityManager;
 	TSharedPtr<FExtensibilityManager> CustomizableObjectEditor_MenuExtensibilityManager;
@@ -59,6 +66,8 @@ private:
 	void RegisterCustomDetails(FPropertyEditorModule& PropertyModule, const UClass* Class, FOnGetDetailCustomizationInstance DetailLayoutDelegate);
 
 	FCustomizableObjectEditorLogger Logger;
+	
+	FCustomizableObjectCompiler Compiler;
 
 	// Command to look for Customizable Object Instance in the player pawn of the current world and open its Customizable Object Instance Editor
 	IConsoleCommand* LaunchCOIECommand = nullptr;
@@ -67,6 +76,12 @@ private:
 	
 	static void OpenCOIE(const TArray<FString>& Arguments);
 
+	// Used to ask the user if they want to recompile uncompiled PIE COs
+	void OnPreBeginPIE(const bool bIsSimulatingInEditor);
+
 	/** Register the COI factory */
 	void RegisterFactory();
+
+	bool HandleSettingsSaved();
+	void RegisterSettings();
 };
\ No newline at end of file
diff --git a/Engine/Plugins/Experimental/Mutable/Source/CustomizableObjectEditor/Private/MuCOE/CustomizableObjectInstanceBakingUtils.cpp b/Engine/Plugins/Experimental/Mutable/Source/CustomizableObjectEditor/Private/MuCOE/CustomizableObjectInstanceBakingUtils.cpp
--- a/Engine/Plugins/Experimental/Mutable/Source/CustomizableObjectEditor/Private/MuCOE/CustomizableObjectInstanceBakingUtils.cpp
+++ b/Engine/Plugins/Experimental/Mutable/Source/CustomizableObjectEditor/Private/MuCOE/CustomizableObjectInstanceBakingUtils.cpp
@@ -3,11 +3,14 @@
 #include "MuCOE/CustomizableObjectInstanceBakingUtils.h"
 
 #include "MuCOE/CustomizableObjectEditor.h"
+#include "MuCOE/CustomizableObjectEditorLogger.h"
 
 #include "Misc/MessageDialog.h"
 #include "FileHelpers.h"
 #include "ObjectTools.h"
 #include "UnrealBakeHelpers.h"
+#include "Animation/Skeleton.h"
+#include "Engine/SkeletalMesh.h"
 #include "Engine/SkinnedAssetCommon.h"
 #include "Materials/MaterialInstanceConstant.h"
 #include "Materials/MaterialInstanceDynamic.h"
diff --git a/Engine/Plugins/Experimental/Mutable/Source/CustomizableObjectEditor/Private/MuCOE/CustomizableObjectLayout.cpp b/Engine/Plugins/Experimental/Mutable/Source/CustomizableObjectEditor/Private/MuCOE/CustomizableObjectLayout.cpp
--- a/Engine/Plugins/Experimental/Mutable/Source/CustomizableObjectEditor/Private/MuCOE/CustomizableObjectLayout.cpp
+++ b/Engine/Plugins/Experimental/Mutable/Source/CustomizableObjectEditor/Private/MuCOE/CustomizableObjectLayout.cpp
@@ -3,6 +3,7 @@
 #include "MuCOE/CustomizableObjectLayout.h"
 
 #include "Engine/StaticMesh.h"
+#include "MuCO/CustomizableObjectCompilerTypes.h"
 #include "MuCOE/CustomizableObjectCompiler.h"
 #include "MuCOE/GenerateMutableSource/GenerateMutableSourceMesh.h"
 #include "MuCOE/GraphTraversal.h"
@@ -94,11 +95,11 @@ void UCustomizableObjectLayout::GenerateBlocksFromUVs()
 	if (Node && Mesh)
 	{
 		//Creating a GenerationContext
-		TUniquePtr<FCustomizableObjectCompiler> Compiler( new FCustomizableObjectCompiler() );
+		FCustomizableObjectCompiler Compiler;
 		UCustomizableObject* Object = Node->GetGraphEditor()->GetCustomizableObject();
 		FCompilationOptions Options = Object->CompileOptions;
 	
-		FMutableGraphGenerationContext GenerationContext(Object, Compiler.Get(), Options);
+		FMutableGraphGenerationContext GenerationContext(Object, &Compiler, Options);
 	
 		//Transforming skeletalmesh to mutable mesh
 		mu::MeshPtr	MutableMesh = nullptr;
@@ -126,8 +127,6 @@ void UCustomizableObjectLayout::GenerateBlocksFromUVs()
 			Layout = mu::NodeLayoutBlocks::GenerateLayoutBlocks(MutableMesh, UVChannel, GridSize.X, GridSize.Y);
 		}
 	
-		Compiler.Reset();
-	
 		if (Layout)
 		{
 			Blocks.Empty();
diff --git a/Engine/Plugins/Experimental/Mutable/Source/CustomizableObjectEditor/Private/MuCOE/CustomizableObjectNodeObjectGroupDetails.cpp b/Engine/Plugins/Experimental/Mutable/Source/CustomizableObjectEditor/Private/MuCOE/CustomizableObjectNodeObjectGroupDetails.cpp
--- a/Engine/Plugins/Experimental/Mutable/Source/CustomizableObjectEditor/Private/MuCOE/CustomizableObjectNodeObjectGroupDetails.cpp
+++ b/Engine/Plugins/Experimental/Mutable/Source/CustomizableObjectEditor/Private/MuCOE/CustomizableObjectNodeObjectGroupDetails.cpp
@@ -8,6 +8,7 @@
 #include "MuCOE/CustomizableObjectEditor.h"
 #include "MuCOE/CustomizableObjectGraph.h"
 #include "MuCOE/GraphTraversal.h"
+#include "MuCOE/Nodes/CustomizableObjectNodeObject.h"
 #include "MuCOE/Nodes/CustomizableObjectNodeObjectGroup.h"
 #include "PropertyCustomizationHelpers.h"
 #include "Widgets/Input/STextComboBox.h"
diff --git a/Engine/Plugins/Experimental/Mutable/Source/CustomizableObjectEditor/Private/MuCOE/GenerateMutableSource/GenerateMutableSource.h b/Engine/Plugins/Experimental/Mutable/Source/CustomizableObjectEditor/Private/MuCOE/GenerateMutableSource/GenerateMutableSource.h
--- a/Engine/Plugins/Experimental/Mutable/Source/CustomizableObjectEditor/Private/MuCOE/GenerateMutableSource/GenerateMutableSource.h
+++ b/Engine/Plugins/Experimental/Mutable/Source/CustomizableObjectEditor/Private/MuCOE/GenerateMutableSource/GenerateMutableSource.h
@@ -8,6 +8,7 @@
 #include "Engine/SkeletalMesh.h"
 #include "MuCO/CustomizableObject.h"
 #include "MuCO/CustomizableObjectIdentifier.h"
+#include "MuCO/CustomizableObjectCompilerTypes.h"
 #include "MuCOE/ExtensionDataCompilerInterface.h"
 #include "MuCOE/Nodes/CustomizableObjectNodeMaterialBase.h"
 #include "MuCOE/Nodes/CustomizableObjectNodeObject.h"
diff --git a/Engine/Plugins/Experimental/Mutable/Source/CustomizableObjectEditor/Private/MuCOE/Nodes/CustomizableObjectNode.cpp b/Engine/Plugins/Experimental/Mutable/Source/CustomizableObjectEditor/Private/MuCOE/Nodes/CustomizableObjectNode.cpp
--- a/Engine/Plugins/Experimental/Mutable/Source/CustomizableObjectEditor/Private/MuCOE/Nodes/CustomizableObjectNode.cpp
+++ b/Engine/Plugins/Experimental/Mutable/Source/CustomizableObjectEditor/Private/MuCOE/Nodes/CustomizableObjectNode.cpp
@@ -9,6 +9,7 @@
 #include "MuCOE/CustomizableObjectPin.h"
 #include "MuCOE/EdGraphSchema_CustomizableObject.h"
 #include "MuCOE/GraphTraversal.h"
+#include "MuCOE/Nodes/CustomizableObjectNodeObject.h"
 #include "MuCOE/Nodes/CustomizableObjectNodeObjectGroup.h"
 #include "MuCOE/RemapPins/CustomizableObjectNodeRemapPinsByName.h"
 #include "MuCOE/RemapPins/CustomizableObjectNodeRemapPinsByPosition.h"
diff --git a/Engine/Plugins/Experimental/Mutable/Source/CustomizableObjectEditor/Private/MuCOE/Nodes/CustomizableObjectNodeSkeletalMesh.cpp b/Engine/Plugins/Experimental/Mutable/Source/CustomizableObjectEditor/Private/MuCOE/Nodes/CustomizableObjectNodeSkeletalMesh.cpp
--- a/Engine/Plugins/Experimental/Mutable/Source/CustomizableObjectEditor/Private/MuCOE/Nodes/CustomizableObjectNodeSkeletalMesh.cpp
+++ b/Engine/Plugins/Experimental/Mutable/Source/CustomizableObjectEditor/Private/MuCOE/Nodes/CustomizableObjectNodeSkeletalMesh.cpp
@@ -2,6 +2,7 @@
 
 #include "MuCOE/Nodes/CustomizableObjectNodeSkeletalMesh.h"
 
+#include "Engine/SkeletalMesh.h"
 #include "AssetThumbnail.h"
 #include "ISinglePropertyView.h"
 #include "Modules/ModuleManager.h"
diff --git a/Engine/Plugins/Experimental/Mutable/Source/CustomizableObjectEditor/Private/MuCOE/Nodes/CustomizableObjectNodeTable.cpp b/Engine/Plugins/Experimental/Mutable/Source/CustomizableObjectEditor/Private/MuCOE/Nodes/CustomizableObjectNodeTable.cpp
--- a/Engine/Plugins/Experimental/Mutable/Source/CustomizableObjectEditor/Private/MuCOE/Nodes/CustomizableObjectNodeTable.cpp
+++ b/Engine/Plugins/Experimental/Mutable/Source/CustomizableObjectEditor/Private/MuCOE/Nodes/CustomizableObjectNodeTable.cpp
@@ -3,6 +3,7 @@
 #include "MuCOE/Nodes/CustomizableObjectNodeTable.h"
 
 #include "Engine/StaticMesh.h"
+#include "Engine/SkeletalMesh.h"
 #include "Engine/Texture2DArray.h"
 #include "Engine/UserDefinedStruct.h"
 #include "Kismet2/StructureEditorUtils.h"
diff --git a/Engine/Plugins/Experimental/Mutable/Source/CustomizableObjectEditor/Private/MuCOE/SCustomizableObjectEditorTagExplorer.cpp b/Engine/Plugins/Experimental/Mutable/Source/CustomizableObjectEditor/Private/MuCOE/SCustomizableObjectEditorTagExplorer.cpp
--- a/Engine/Plugins/Experimental/Mutable/Source/CustomizableObjectEditor/Private/MuCOE/SCustomizableObjectEditorTagExplorer.cpp
+++ b/Engine/Plugins/Experimental/Mutable/Source/CustomizableObjectEditor/Private/MuCOE/SCustomizableObjectEditorTagExplorer.cpp
@@ -8,6 +8,7 @@
 #include "MuCOE/CustomizableObjectEditor.h"
 #include "MuCOE/CustomizableObjectGraph.h"
 #include "MuCOE/GraphTraversal.h"
+#include "MuCOE/Nodes/CustomizableObjectNodeObject.h"
 #include "MuCOE/Nodes/CustomizableObjectNodeExtendMaterial.h"
 #include "MuCOE/Nodes/CustomizableObjectNodeMaterial.h"
 #include "MuCOE/Nodes/CustomizableObjectNodeMaterialVariation.h"
diff --git a/Engine/Plugins/Experimental/Mutable/Source/CustomizableObjectEditor/Private/MuCOE/SCustomizableObjectEditorTextureAnalyzer.cpp b/Engine/Plugins/Experimental/Mutable/Source/CustomizableObjectEditor/Private/MuCOE/SCustomizableObjectEditorTextureAnalyzer.cpp
--- a/Engine/Plugins/Experimental/Mutable/Source/CustomizableObjectEditor/Private/MuCOE/SCustomizableObjectEditorTextureAnalyzer.cpp
+++ b/Engine/Plugins/Experimental/Mutable/Source/CustomizableObjectEditor/Private/MuCOE/SCustomizableObjectEditorTextureAnalyzer.cpp
@@ -2,6 +2,7 @@
 
 #include "MuCOE/SCustomizableObjectEditorTextureAnalyzer.h"
 
+#include "Engine/SkeletalMesh.h"
 #include "Engine/SkinnedAssetCommon.h"
 #include "Framework/Views/ITypedTableView.h"
 #include "IPropertyTable.h"
diff --git a/Engine/Plugins/Experimental/Mutable/Source/CustomizableObjectEditor/Private/MuCOE/SMutableObjectViewer.cpp b/Engine/Plugins/Experimental/Mutable/Source/CustomizableObjectEditor/Private/MuCOE/SMutableObjectViewer.cpp
--- a/Engine/Plugins/Experimental/Mutable/Source/CustomizableObjectEditor/Private/MuCOE/SMutableObjectViewer.cpp
+++ b/Engine/Plugins/Experimental/Mutable/Source/CustomizableObjectEditor/Private/MuCOE/SMutableObjectViewer.cpp
@@ -6,6 +6,7 @@
 #include "Framework/Views/TableViewMetadata.h"
 #include "Interfaces/ITargetPlatform.h"
 #include "Interfaces/ITargetPlatformManagerModule.h"
+#include "MuCO/CustomizableObject.h"
 #include "MuCOE/CustomizableObjectCompileRunnable.h"
 #include "MuCOE/CustomizableObjectEditorStyle.h"
 #include "MuCOE/SMutableCodeViewer.h"
@@ -104,6 +105,7 @@ void SMutableObjectViewer::Construct(const FArguments& InArgs, UCustomizableObje
 void SMutableObjectViewer::GenerateMutableGraphPressed()
 {
 	// Convert from Unreal graph to Mutable graph.
+	FCustomizableObjectCompiler Compiler;
 	TArray<TSoftObjectPtr<UTexture>> RuntimeTextures;
 	TArray<FMutableSourceTextureData> CompilerTextures;
 	mu::Ptr<mu::Node> RootNode = Compiler.Export(CustomizableObject, CompileOptions, RuntimeTextures, CompilerTextures);
@@ -137,6 +139,7 @@ void SMutableObjectViewer::CompileMutableCodePressed()
 		{
 			CompileOptions.DebugBias = Bias;
 
+			FCustomizableObjectCompiler Compiler;
 			RuntimeTextures.Empty();
 			CompilerTextures.Empty();
 			mu::NodePtr RootNode = Compiler.Export(CustomizableObject, CompileOptions, RuntimeTextures, CompilerTextures);
@@ -157,6 +160,7 @@ void SMutableObjectViewer::CompileMutableCodePressed()
 	}
 
 	// Convert from Unreal graph to Mutable graph.
+	FCustomizableObjectCompiler Compiler;
 	RuntimeTextures.Empty();
 	CompilerTextures.Empty();
 	mu::NodePtr RootNode = Compiler.Export(CustomizableObject, CompileOptions, RuntimeTextures, CompilerTextures);
diff --git a/Engine/Plugins/Experimental/Mutable/Source/CustomizableObjectEditor/Private/MuCOE/SMutableObjectViewer.h b/Engine/Plugins/Experimental/Mutable/Source/CustomizableObjectEditor/Private/MuCOE/SMutableObjectViewer.h
--- a/Engine/Plugins/Experimental/Mutable/Source/CustomizableObjectEditor/Private/MuCOE/SMutableObjectViewer.h
+++ b/Engine/Plugins/Experimental/Mutable/Source/CustomizableObjectEditor/Private/MuCOE/SMutableObjectViewer.h
@@ -2,7 +2,7 @@
 
 #pragma once
 
-#include "MuCOE/CustomizableObjectCompiler.h"
+#include "MuCO/CustomizableObjectCompilerTypes.h"
 #include "Widgets/SCompoundWidget.h"
 #include "Widgets/Input/SNumericDropDown.h"
 
@@ -11,6 +11,7 @@ namespace ESelectInfo { enum Type : int; }
 template <typename ItemType> class STreeView;
 
 // Forward declarations
+class UCustomizableObject;
 class FMutableObjectTreeElement;
 class FReferenceCollector;
 class FTabManager;
@@ -51,9 +52,6 @@ private:
 	TArray< TSharedPtr<FString> > DebugPlatformStrings;
 	TSharedPtr<STextComboBox> DebugPlatformCombo;
 
-	/** Object compiler. */
-	FCustomizableObjectCompiler Compiler;
-
 	/** UI callbacks */
 	void GenerateMutableGraphPressed();
 	void CompileMutableCodePressed();
diff --git a/Engine/Plugins/Experimental/Mutable/Source/CustomizableObjectEditor/Private/MuCOE/Thumbnails/CustomizableObjectInstanceThumbnailRenderer.cpp b/Engine/Plugins/Experimental/Mutable/Source/CustomizableObjectEditor/Private/MuCOE/Thumbnails/CustomizableObjectInstanceThumbnailRenderer.cpp
--- a/Engine/Plugins/Experimental/Mutable/Source/CustomizableObjectEditor/Private/MuCOE/Thumbnails/CustomizableObjectInstanceThumbnailRenderer.cpp
+++ b/Engine/Plugins/Experimental/Mutable/Source/CustomizableObjectEditor/Private/MuCOE/Thumbnails/CustomizableObjectInstanceThumbnailRenderer.cpp
@@ -5,6 +5,7 @@
 #include "CanvasTypes.h"
 #include "MuCO/CustomizableObjectInstance.h"
 #include "MuCO/CustomizableSkeletalComponent.h"
+#include "MuCO/CustomizableObjectCompilerTypes.h"
 #include "MuCOE/CustomizableObjectCompiler.h"
 #include "SceneView.h"
 #include "TextureResource.h"
diff --git a/Engine/Plugins/Experimental/Mutable/Source/CustomizableObjectEditor/Private/MuCOE/Thumbnails/CustomizableObjectThumbnailRenderer.cpp b/Engine/Plugins/Experimental/Mutable/Source/CustomizableObjectEditor/Private/MuCOE/Thumbnails/CustomizableObjectThumbnailRenderer.cpp
--- a/Engine/Plugins/Experimental/Mutable/Source/CustomizableObjectEditor/Private/MuCOE/Thumbnails/CustomizableObjectThumbnailRenderer.cpp
+++ b/Engine/Plugins/Experimental/Mutable/Source/CustomizableObjectEditor/Private/MuCOE/Thumbnails/CustomizableObjectThumbnailRenderer.cpp
@@ -5,6 +5,7 @@
 #include "CanvasTypes.h"
 #include "MuCO/CustomizableSkeletalComponent.h"
 #include "MuCO/UnrealPortabilityHelpers.h"
+#include "MuCO/CustomizableObjectCompilerTypes.h"
 #include "MuCOE/CustomizableObjectCompiler.h"
 #include "SceneView.h"
 #include "TextureResource.h"
diff --git a/Engine/Plugins/Experimental/Mutable/Source/MutableValidation/Private/MuV/AssetValidator_CustomizableObjects.cpp b/Engine/Plugins/Experimental/Mutable/Source/MutableValidation/Private/MuV/AssetValidator_CustomizableObjects.cpp
--- a/Engine/Plugins/Experimental/Mutable/Source/MutableValidation/Private/MuV/AssetValidator_CustomizableObjects.cpp
+++ b/Engine/Plugins/Experimental/Mutable/Source/MutableValidation/Private/MuV/AssetValidator_CustomizableObjects.cpp
@@ -13,6 +13,8 @@
 #include "MuCO/CustomizableObject.h"
 #include "MuCO/CustomizableObjectPrivate.h"
 #include "MuCO/CustomizableObjectSystem.h"
+#include "MuCO/CustomizableObjectCompilerTypes.h"
+#include "MuCO/ICustomizableObjectEditorModule.h"
 #include "MuCOE/GraphTraversal.h"
 #include "UObject/NameTypes.h"
 #include "UObject/Object.h"
@@ -88,9 +90,6 @@ EDataValidationResult UAssetValidator_CustomizableObjects::IsCustomizableObjectV
 		OnPostCOValidationHandle = FEditorDelegates::OnPostAssetValidation.AddStatic(OnPostCOsValidation);	
 	}
 	
-	// Request a compiler to be able to locate the root and to compile it
-	const TSharedPtr<FCustomizableObjectCompilerBase> Compiler = UCustomizableObjectSystem::GetNewCompiler();
-	
 	// Find out which is the root for this CO (it may be itself but that is OK)
 	UCustomizableObject* RootObject = GetRootObject(InCustomizableObject);
 	check (RootObject);
@@ -127,27 +126,27 @@ EDataValidationResult UAssetValidator_CustomizableObjects::IsCustomizableObjectV
 	TArray<FText> CachedValidationWarnings;
 	
 	// Map with all the possible compilation states. We use this so at the end we can know if any of those states was returned by any of the compilation runs
-	TMap<ECompilationResultPrivate,bool>PossibleEndCompilationStates;
-	PossibleEndCompilationStates.Add(ECompilationResultPrivate::Unknown, false);
-	PossibleEndCompilationStates.Add(ECompilationResultPrivate::Success, false);
-	PossibleEndCompilationStates.Add(ECompilationResultPrivate::Errors, false);
-	PossibleEndCompilationStates.Add(ECompilationResultPrivate::Warnings, false);
-	PossibleEndCompilationStates.Shrink();
+	TMap<ECompilationResultPrivate,bool>PossibleEndCompilationResults;
+	PossibleEndCompilationResults.Add(ECompilationResultPrivate::Unknown, false);
+	PossibleEndCompilationResults.Add(ECompilationResultPrivate::Success, false);
+	PossibleEndCompilationResults.Add(ECompilationResultPrivate::Errors, false);
+	PossibleEndCompilationResults.Add(ECompilationResultPrivate::Warnings, false);
+	PossibleEndCompilationResults.Shrink();
 	
 	// Iterate over the compilation options that we want to test and perform the compilation
 	for	(const FCompilationOptions& Options : CompilationOptionsToTest)
 	{
 		// Run Sync compilation -> Warning : Potentially long operation -------------
-		Compiler->Compile(*RootObject, Options, false);
+		TSharedRef<FCompilationRequest> CompileRequest = MakeShared<FCompilationRequest>(*RootObject, false);
+		ICustomizableObjectEditorModule::GetChecked().CompileCustomizableObject(CompileRequest);
 		// --------------------------------------------------------------------------
 		
 		// Get compilation errors and warnings
-		TArray<FText> CompilationErrors;
-		TArray<FText> CompilationWarnings;
-		Compiler->GetCompilationMessages(CompilationWarnings, CompilationErrors);
+		const TArray<FText>& CompilationErrors = CompileRequest->GetErrors();
+		const TArray<FText>& CompilationWarnings = CompileRequest->GetWarnings();
 		
 		// Cache the messages returned by the compiler
-		for ( const FText& FoundError : CompilationErrors)
+		for (const FText& FoundError : CompilationErrors)
 		{
 			// Add message if not already present
 			if (!CachedValidationErrors.ContainsByPredicate([&FoundError](const FText& ArrayEntry)
@@ -156,7 +155,7 @@ EDataValidationResult UAssetValidator_CustomizableObjects::IsCustomizableObjectV
 				CachedValidationErrors.Add(FoundError);
 			}
 		}
-		for ( const FText& FoundWarning : CompilationWarnings)
+		for (const FText& FoundWarning : CompilationWarnings)
 		{
 			if (!CachedValidationWarnings.ContainsByPredicate([&FoundWarning](const FText& ArrayEntry)
 				{ return FoundWarning.EqualTo(ArrayEntry);}))
@@ -166,8 +165,7 @@ EDataValidationResult UAssetValidator_CustomizableObjects::IsCustomizableObjectV
 		}
 
 		// Flag the array with end results to have the current output as true since it was produced by this execution
-		ECompilationResultPrivate CompilationEndResult = RootObject->GetPrivate()->CompilationResult;
-		bool* Value = PossibleEndCompilationStates.Find(CompilationEndResult);
+		bool* Value = PossibleEndCompilationResults.Find(CompileRequest->GetCompilationResult());
 		check (Value);		// If this fails it may mean we are getting a compilation state we are not considering. 
 		*Value = true;
 	}
@@ -192,12 +190,12 @@ EDataValidationResult UAssetValidator_CustomizableObjects::IsCustomizableObjectV
 	
 	// Return informed guess about what the validation state of this object should be
 
-	// If it contains invalid states then notify about it too:
-	// ECustomizableObjectCompilationState::None would mean the resource is locked (and should not be)
-	check (*PossibleEndCompilationStates.Find(ECompilationResultPrivate::Unknown) == false);
+	// If it contains invalid results then notify about it too:
+	// ECompilationResultPrivate::None would mean the resource is locked (and should not be)
+	check (*PossibleEndCompilationResults.Find(ECompilationResultPrivate::Unknown) == false);
 	
 	// If one or more tests failed to ran then the result must be invalid
-	if (*PossibleEndCompilationStates.Find(ECompilationResultPrivate::Errors) == true)
+	if (*PossibleEndCompilationResults.Find(ECompilationResultPrivate::Errors) == true)
 	{
 		// Early CO compilation error (before starting mutable compilation) -> Output is invalid
 		Result = EDataValidationResult::Invalid;
diff --git a/Engine/Plugins/Experimental/Mutable/Source/MutableValidation/Private/MuV/CustomizableObjectCompilationUtility.cpp b/Engine/Plugins/Experimental/Mutable/Source/MutableValidation/Private/MuV/CustomizableObjectCompilationUtility.cpp
--- a/Engine/Plugins/Experimental/Mutable/Source/MutableValidation/Private/MuV/CustomizableObjectCompilationUtility.cpp
+++ b/Engine/Plugins/Experimental/Mutable/Source/MutableValidation/Private/MuV/CustomizableObjectCompilationUtility.cpp
@@ -6,6 +6,8 @@
 #include "MuCO/CustomizableObject.h"
 #include "MuCO/CustomizableObjectPrivate.h"
 #include "MuCO/CustomizableObjectSystem.h"
+#include "MuCO/CustomizableObjectCompilerTypes.h"
+#include "MuCO/ICustomizableObjectEditorModule.h"
 #include "Interfaces/ITargetPlatform.h"
 #include "MuCOE/CustomizableObjectEditorFunctionLibrary.h"
 
@@ -53,17 +55,11 @@ bool FCustomizableObjectCompilationUtility::CompileCustomizableObject(UCustomiza
 		UE_LOG(LogMutable, Log, TEXT("(string) model_compile_platform_name : %s "), *CompilationOptions.TargetPlatform->PlatformName());
 	}
 
+	// Compilation completed successfully
+	bool bCompilationSuccess = false;
 	
 	// Run and wait for the compilation to be completed ----------------------------------------------------------------------------------------------
 	{
-		// Bind our OnCompilationFinished() method to the compilation completion callback to get notified once it is done
-		FPostCompileDelegate& PostCompilationDelegate = CustomizableObject->GetPostCompileDelegate();
-		PostCompilationDelegate.AddSP(this, &FCustomizableObjectCompilationUtility::OnCompilationFinished);
-		
-		// Create the compiler that will take care of the ASYNC compilation
-		UCustomizableObjectSystem* CustomizableObjectSystem = UCustomizableObjectSystem::GetInstanceChecked();
-		TSharedPtr<FCustomizableObjectCompilerBase> Compiler = CustomizableObjectSystem->GetNewCompiler();
-
 		// Get the memory usage before staring the compilation
 #if ENABLE_LOW_LEVEL_MEM_TRACKER
 		FLowLevelMemTracker& LowLevelMemoryTracker = FLowLevelMemTracker::Get();
@@ -89,19 +85,19 @@ bool FCustomizableObjectCompilationUtility::CompileCustomizableObject(UCustomiza
 		mu::FGlobalMemoryCounter::Zero();
 		
 		UE_LOG(LogMutable,Display,TEXT("Compiling Customizable Object..."));
-		check(Compiler);
-		Compiler->Compile(*CustomizableObject, CompilationOptions, true);
+		TSharedRef<FCompilationRequest> CompileRequest = MakeShared<FCompilationRequest>(*InCustomizableObject, true);
+		CompileRequest->GetCompileOptions() = CompilationOptions;
+		ICustomizableObjectEditorModule::GetChecked().CompileCustomizableObject(CompileRequest);
 		
 		// Wait while the compilation takes place
 		const double CompilationStartSeconds = FPlatformTime::Seconds();
-		bIsTheCompilationInProcess = true;
-		while (bIsTheCompilationInProcess)
+		while (CompileRequest->GetCompilationState() != ECompilationStatePrivate::Completed)
 		{
 			// Tick the engine
 			CommandletHelpers::TickEngine();
 
 			// todo: Will this be required in the future? It feels odd having to tick the compiler manually when we are already ticking the engine itself
-			Compiler->Tick(false);
+			ICustomizableObjectEditorModule::GetChecked().Tick(false);
 			
 			// Cache the peak value found during the compilation of the CO
 #if ENABLE_LOW_LEVEL_MEM_TRACKER
@@ -120,6 +116,10 @@ bool FCustomizableObjectCompilationUtility::CompileCustomizableObject(UCustomiza
 			}
 		}
 		
+
+		bCompilationSuccess = CompileRequest->GetCompilationResult() == ECompilationResultPrivate::Success ||
+			CompileRequest->GetCompilationResult() == ECompilationResultPrivate::Warnings;
+
 		// Report the time we took to run the compilation
 		const double CompilationEndSeconds = FPlatformTime::Seconds() - CompilationStartSeconds;
 		UE_LOG(LogMutable, Display, TEXT("The compilation of the %s CO model took %f seconds."), *InCustomizableObject->GetName(), CompilationEndSeconds);
@@ -128,7 +128,8 @@ bool FCustomizableObjectCompilationUtility::CompileCustomizableObject(UCustomiza
 			UE_LOG(LogMutable, Log, TEXT("(double) model_compile_time_ms : %f "), CompilationEndSeconds * 1000);
 
 			// Also report compilation end status
-			const ECustomizableObjectCompilationState CompilationEndResult = bWasCoCompilationSuccessful ? ECustomizableObjectCompilationState::Completed : ECustomizableObjectCompilationState::Failed;
+			const ECustomizableObjectCompilationState CompilationEndResult = bCompilationSuccess ? ECustomizableObjectCompilationState::Completed
+				: ECustomizableObjectCompilationState::Failed;
 			UE_LOG(LogMutable, Log, TEXT("(string) model_compile_end_state : %s "), *UEnum::GetValueAsString(CompilationEndResult));
 		}
 		
@@ -169,25 +170,5 @@ bool FCustomizableObjectCompilationUtility::CompileCustomizableObject(UCustomiza
 	}
 	
 	// Return the success state of the operation
-	return bWasCoCompilationSuccessful;
-}
-
-
-void FCustomizableObjectCompilationUtility::OnCompilationFinished()
-{
-	UE_LOG(LogMutable,Display,TEXT("Compilation of CO completed!"));
-
-	// Store the success end state for the operation.
-	bWasCoCompilationSuccessful =	CustomizableObject->GetPrivate()->CompilationResult == ECompilationResultPrivate::Success ||
-									CustomizableObject->GetPrivate()->CompilationResult == ECompilationResultPrivate::Warnings;
-	if (bWasCoCompilationSuccessful)
-	{
-		UE_LOG(LogMutable, Display, TEXT("The compilation of the %s model was succesfull."), *CustomizableObject->GetName());
-	}
-	else
-	{
-		UE_LOG(LogMutable, Error, TEXT("The compilation of the %s model failed."), *CustomizableObject->GetName());
-	}
-	
-	bIsTheCompilationInProcess = false;
-}
+	return bCompilationSuccess;
+}
\ No newline at end of file
diff --git a/Engine/Plugins/Experimental/Mutable/Source/MutableValidation/Private/MuV/CustomizableObjectCompilationUtility.h b/Engine/Plugins/Experimental/Mutable/Source/MutableValidation/Private/MuV/CustomizableObjectCompilationUtility.h
--- a/Engine/Plugins/Experimental/Mutable/Source/MutableValidation/Private/MuV/CustomizableObjectCompilationUtility.h
+++ b/Engine/Plugins/Experimental/Mutable/Source/MutableValidation/Private/MuV/CustomizableObjectCompilationUtility.h
@@ -1,11 +1,11 @@
 ﻿// Copyright Epic Games, Inc. All Rights Reserved.
 #pragma once
 
-#include "MuCO/CustomizableObject.h"
 #include "UObject/StrongObjectPtr.h"
 #include "Templates/SharedPointer.h"
 
 // Forward declarations
+class UCustomizableObject;
 struct FCompilationOptions;
 
 /**
@@ -27,14 +27,5 @@ public:
 private:
 
 	/** The CO that is currently being compiled.  */
-	TStrongObjectPtr<UCustomizableObject> CustomizableObject = nullptr;
-
-	/** Variable being used to know when the compilation is running and when it is not. */
-	bool bIsTheCompilationInProcess = false;
-
-	/** Variable holding the success state of the compilation operation. */
-	bool bWasCoCompilationSuccessful = false;
-	
-	/** Callback invoked once the compilation of the target CO has been completed. It will set a value for bWasCoCompilationSuccessful. */
-	void OnCompilationFinished();
+	TStrongObjectPtr<UCustomizableObject> CustomizableObject;
 };
\ No newline at end of file
diff --git a/Engine/Plugins/Experimental/Mutable/Source/MutableValidation/Private/MuV/CustomizableObjectValidationCommandlet.cpp b/Engine/Plugins/Experimental/Mutable/Source/MutableValidation/Private/MuV/CustomizableObjectValidationCommandlet.cpp
--- a/Engine/Plugins/Experimental/Mutable/Source/MutableValidation/Private/MuV/CustomizableObjectValidationCommandlet.cpp
+++ b/Engine/Plugins/Experimental/Mutable/Source/MutableValidation/Private/MuV/CustomizableObjectValidationCommandlet.cpp
@@ -11,9 +11,11 @@
 #include "MuCO/CustomizableObjectInstancePrivate.h"
 #include "MuCO/CustomizableObjectPrivate.h"
 #include "MuCO/CustomizableObjectSystem.h"
+#include "MuCO/CustomizableObjectCompilerTypes.h"
 #include "MuR/Model.h"
 #include "Interfaces/ITargetPlatformManagerModule.h"
 #include "Interfaces/ITargetPlatform.h"
+#include "Serialization/MemoryWriter.h"
 
 
 int32 UCustomizableObjectValidationCommandlet::Main(const FString& Params)
