[Mutable] Moved UCustomizableObjectInstanceUserData header from Internal to Public.

This header should be public.

#rnx

diff --git a/Engine/Plugins/Experimental/Mutable/Source/CustomizableObject/Internal/MuCO/CustomizableObjectInstanceAssetUserData.h b/Engine/Plugins/Experimental/Mutable/Source/CustomizableObject/Internal/MuCO/CustomizableObjectInstanceAssetUserData.h
deleted file mode 100644
diff --git a/Engine/Plugins/Experimental/Mutable/Source/CustomizableObject/Public/MuCO/CustomizableObjectInstanceAssetUserData.h b/Engine/Plugins/Experimental/Mutable/Source/CustomizableObject/Public/MuCO/CustomizableObjectInstanceAssetUserData.h
new file mode 100644
--- /dev/null
+++ b/Engine/Plugins/Experimental/Mutable/Source/CustomizableObject/Public/MuCO/CustomizableObjectInstanceAssetUserData.h
@@ -0,0 +1,37 @@
+// Copyright Epic Games, Inc. All Rights Reserved.
+
+#pragma once
+
+#include "Engine/AssetUserData.h"
+#include "GameplayTagContainer.h"
+
+#include "CustomizableObjectInstanceAssetUserData.generated.h"
+
+
+/** Additional data attached to Skeletal Meshes. */
+UCLASS(BlueprintType)
+class CUSTOMIZABLEOBJECT_API UCustomizableObjectInstanceUserData : public UAssetUserData
+{
+	GENERATED_BODY()
+
+public:
+
+	/** Return the list of tags for this instance. */
+	UFUNCTION(BlueprintCallable, Category = CustomizableObjectInstance)
+	const FGameplayTagContainer& GetAnimationGameplayTags() const
+	{
+		return AnimBPGameplayTags;
+	};
+
+	/** Sets the list of tags for this instance. */
+	UFUNCTION(BlueprintCallable, Category = CustomizableObjectInstance)
+	void SetAnimationGameplayTags(const FGameplayTagContainer& InstanceTags)
+	{
+		AnimBPGameplayTags = InstanceTags;
+	};
+
+private:
+
+	UPROPERTY()
+	FGameplayTagContainer AnimBPGameplayTags;
+};
