[Mutable] Split FParameterUIData into Parameter and State structures.

- Created FMutableStateData.
- Renamed FParameterUIData to MutableParameterData.
- Now only Metadata structures are public. Data structures that hold Metadata are private.
- Created CO getters to access Metadata structures.
- Removed trivial constructors.
- Removed unused operators.
- Minor code improvements.

#rnx

diff --git a/Engine/Plugins/Experimental/Mutable/Config/DefaultMutable.ini b/Engine/Plugins/Experimental/Mutable/Config/DefaultMutable.ini
--- a/Engine/Plugins/Experimental/Mutable/Config/DefaultMutable.ini
+++ b/Engine/Plugins/Experimental/Mutable/Config/DefaultMutable.ini
@@ -1,8 +1,7 @@
 ; Settings for the Mutable plugin
 
 [CoreRedirects]
-; None yet
-
++StructRedirects=(OldName="/Script/CustomizableObject.ParameterUIData",NewName="/Script/CustomizableObject.MutableParameterData")
 
 [Features]
 
diff --git a/Engine/Plugins/Experimental/Mutable/Source/CustomizableObject/Classes/MuCO/CustomizableObject.h b/Engine/Plugins/Experimental/Mutable/Source/CustomizableObject/Classes/MuCO/CustomizableObject.h
--- a/Engine/Plugins/Experimental/Mutable/Source/CustomizableObject/Classes/MuCO/CustomizableObject.h
+++ b/Engine/Plugins/Experimental/Mutable/Source/CustomizableObject/Classes/MuCO/CustomizableObject.h
@@ -623,23 +623,20 @@ public:
 	UFUNCTION(BlueprintCallable, Category = CustomizableObject)
 	FString GetStateParameterName(const FString& StateName, int32 ParameterIndex) const;
 	FString GetStateParameterName(int32 StateIndex, int32 ParameterIndex) const;
-
-	/** Return the metadata associated to an object state by name. */
-	UFUNCTION(BlueprintCallable, Category = CustomizableObject)
-	FParameterUIData GetStateUIMetadata(const FString& StateName) const;
-
-	/** Return the metadata associated to an object state by state index. */
+	
+	/** Return the metadata associated to a parameter. */
 	UFUNCTION(BlueprintCallable, Category = CustomizableObject)
-	FParameterUIData GetStateUIMetadataFromIndex(int32 StateIndex) const;
+	FMutableParamUIMetadata GetParameterUIMetadata(const FString& ParamName) const;
 
-	/** Return the metadata associated to an object parameter by parameter name. */
+	/** Return the metadata associated to an int parameter option. */
 	UFUNCTION(BlueprintCallable, Category = CustomizableObject)
-	FParameterUIData GetParameterUIMetadata(const FString& ParamName) const;
-
-	/** Return the metadata associated to an object parameter by parameter index (from 0 to GetParameterCount-1). */
+	FMutableParamUIMetadata GetIntParameterOptionUIMetadata(const FString& ParamName, const FString& OptionName) const;
+	
+	/** Return the metadata associated to a state. */
 	UFUNCTION(BlueprintCallable, Category = CustomizableObject)
-	FParameterUIData GetParameterUIMetadataFromIndex(int32 ParamIndex) const;
+	FMutableStateUIMetadata GetStateUIMetadata(const FString& StateName) const;
 
+	
 private:
 	/** Textures marked as low priority will generate defaulted resident mips (if texture streaming is enabled).
 	  * Generating defaulted resident mips greatly reduce initial generation times. */
diff --git a/Engine/Plugins/Experimental/Mutable/Source/CustomizableObject/Internal/MuCO/CustomizableObjectCustomVersion.h b/Engine/Plugins/Experimental/Mutable/Source/CustomizableObject/Internal/MuCO/CustomizableObjectCustomVersion.h
--- a/Engine/Plugins/Experimental/Mutable/Source/CustomizableObject/Internal/MuCO/CustomizableObjectCustomVersion.h
+++ b/Engine/Plugins/Experimental/Mutable/Source/CustomizableObject/Internal/MuCO/CustomizableObjectCustomVersion.h
@@ -114,6 +114,8 @@ struct CUSTOMIZABLEOBJECT_API FCustomizableObjectCustomVersion
 
 		ProjectorNodesDefaultValueFix,
 
+		StateUIMetadata,
+
 		// -----<new versions can be added above this line>-------------------------------------------------
 		VersionPlusOne,
 		LatestVersion = VersionPlusOne - 1
diff --git a/Engine/Plugins/Experimental/Mutable/Source/CustomizableObject/Internal/MuCO/CustomizableObjectPrivate.h b/Engine/Plugins/Experimental/Mutable/Source/CustomizableObject/Internal/MuCO/CustomizableObjectPrivate.h
--- a/Engine/Plugins/Experimental/Mutable/Source/CustomizableObject/Internal/MuCO/CustomizableObjectPrivate.h
+++ b/Engine/Plugins/Experimental/Mutable/Source/CustomizableObject/Internal/MuCO/CustomizableObjectPrivate.h
@@ -556,6 +556,70 @@ struct FMorphTargetVertexData
 static_assert(sizeof(FMorphTargetVertexData) == sizeof(FVector3f)*2 + sizeof(uint32)); // Make sure no padding is present.
 template<> struct TCanBulkSerialize<FMorphTargetVertexData> { enum { Value = true }; };
 
+
+USTRUCT()
+struct FIntegerParameterUIData
+{
+	GENERATED_BODY()
+
+	UPROPERTY()
+	FMutableParamUIMetadata ParamUIMetadata;
+
+	friend FArchive& operator<<(FArchive& Ar, FIntegerParameterUIData& Struct);
+};
+
+
+USTRUCT()
+struct FMutableParameterData
+{
+	GENERATED_BODY()
+
+	UPROPERTY()
+	FMutableParamUIMetadata ParamUIMetadata;
+
+	/** Parameter type */
+	UPROPERTY()
+	EMutableParameterType Type = EMutableParameterType::None;
+
+	/** In the case of an integer parameter, store here all options */
+	UPROPERTY()
+	TMap<FString, FIntegerParameterUIData> ArrayIntegerParameterOption;
+
+	/** How are the different options selected (one, one or none, etc...) */
+	UPROPERTY()
+	ECustomizableObjectGroupType IntegerParameterGroupType = ECustomizableObjectGroupType::COGT_ONE_OR_NONE;
+	
+	friend FArchive& operator<<(FArchive& Ar, FMutableParameterData& Struct);
+};
+
+
+USTRUCT()
+struct FMutableStateData
+{
+	GENERATED_BODY()
+
+	UPROPERTY()
+	FMutableStateUIMetadata StateUIMetadata;
+
+	/** In this mode instances and their temp data will be reused between updates. It will be much faster but spend as much as ten times the memory.
+	 * Useful for customization lockers with few characters that are going to have their parameters changed many times, not for in-game */
+	UPROPERTY()
+	bool bLiveUpdateMode = false;
+
+	/** If this is enabled, texture streaming won't be used for this state, and full images will be generated when an instance is first updated. */
+	UPROPERTY()
+	bool bDisableTextureStreaming = false;
+
+	UPROPERTY()
+	bool bReuseInstanceTextures = false;
+
+	UPROPERTY()
+	TMap<FString, FString> ForcedParameterValues;
+
+	friend FArchive& operator<<(FArchive& Ar, FMutableStateData& Struct);
+};
+
+
 // Referenced materials, skeletons, passthrough textures...
 USTRUCT()
 struct FModelResources
@@ -614,11 +678,11 @@ struct FModelResources
 
 	/** Parameter UI metadata information for all the dependencies of this Customizable Object. */
 	UPROPERTY()
-	TMap<FString, FParameterUIData> ParameterUIDataMap;
+	TMap<FString, FMutableParameterData> ParameterUIDataMap;
 
 	/** State UI metadata information for all the dependencies of this Customizable Object */
 	UPROPERTY()
-	TMap<FString, FParameterUIData> StateUIDataMap;
+	TMap<FString, FMutableStateData> StateUIDataMap;
 
 	UPROPERTY()
 	TMap<uint32, FRealTimeMorphStreamable> RealTimeMorphStreamables;
@@ -727,10 +791,6 @@ public:
 
 	int32 FindIntParameterValue(int32 ParamIndex, const FString& Value) const;
 
-	FParameterUIData GetStateUIMetadataFromIndex(int32 StateIndex) const;
-
-	FParameterUIData GetStateUIMetadata(const FString& StateName) const;
-
 	FString GetStateName(int32 StateIndex) const;
 
 #if WITH_EDITORONLY_DATA
@@ -893,6 +953,6 @@ public:
 	// This is a manual version number for the binary blobs in this asset.
 	// Increasing it invalidates all the previously compiled models.
 	// Warning: If while merging code both versions have changed, take the highest+1.
-	static constexpr int32 CurrentSupportedVersion = 448;
+	static constexpr int32 CurrentSupportedVersion = 449;
 };
 
diff --git a/Engine/Plugins/Experimental/Mutable/Source/CustomizableObject/Private/MuCO/CustomizableObject.cpp b/Engine/Plugins/Experimental/Mutable/Source/CustomizableObject/Private/MuCO/CustomizableObject.cpp
--- a/Engine/Plugins/Experimental/Mutable/Source/CustomizableObject/Private/MuCO/CustomizableObject.cpp
+++ b/Engine/Plugins/Experimental/Mutable/Source/CustomizableObject/Private/MuCO/CustomizableObject.cpp
@@ -43,6 +43,8 @@
 #endif
 
 
+#include "MuCO/CustomizableObjectCustomVersion.h"
+
 #include UE_INLINE_GENERATED_CPP_BY_NAME(CustomizableObject)
 
 #define LOCTEXT_NAMESPACE "CustomizableObject"
@@ -596,6 +598,7 @@ void UCustomizableObjectPrivate::SaveCompiledData(FArchive& MemoryWriter, bool b
 	}
 }
 
+
 void UCustomizableObjectPrivate::LoadCompiledData(FArchive& MemoryReader, const ITargetPlatform* InTargetPlatform, bool bIsCooking)
 {
 	TSharedPtr<mu::Model, ESPMode::ThreadSafe> LoadedModel;
@@ -1522,43 +1525,36 @@ FString UCustomizableObject::FindIntParameterValueName(int32 ParamIndex, int32 P
 }
 
 
-FParameterUIData UCustomizableObject::GetParameterUIMetadata(const FString& ParamName) const
-{
-	const FParameterUIData* ParameterUIData = Private->GetModelResources().ParameterUIDataMap.Find(ParamName);
-
-	return ParameterUIData ? *ParameterUIData : FParameterUIData();
-}
-
-
-FParameterUIData UCustomizableObject::GetParameterUIMetadataFromIndex(int32 ParamIndex) const
-{
-	return GetParameterUIMetadata(GetParameterName(ParamIndex));
-}
-
-
-FParameterUIData UCustomizableObject::GetStateUIMetadata(const FString& StateName) const
+FMutableParamUIMetadata UCustomizableObject::GetParameterUIMetadata(const FString& ParamName) const
 {
-	return GetPrivate()->GetStateUIMetadata(StateName);
+	const FMutableParameterData* ParameterData = Private->GetModelResources().ParameterUIDataMap.Find(ParamName);
+	return ParameterData ? ParameterData->ParamUIMetadata : FMutableParamUIMetadata();
 }
 
 
-FParameterUIData UCustomizableObjectPrivate::GetStateUIMetadata(const FString& StateName) const
+FMutableParamUIMetadata UCustomizableObject::GetIntParameterOptionUIMetadata(const FString& ParamName, const FString& OptionName) const
 {
-	const FParameterUIData* StateUIData = GetModelResources().StateUIDataMap.Find(StateName);
-
-	return StateUIData ? *StateUIData : FParameterUIData();
-}
-
+	const int32 ParameterIndex = FindParameter(ParamName);
+	if (ParameterIndex == INDEX_NONE)
+	{
+		return {};
+	}
+	
+	const FMutableParameterData* ParameterData = Private->GetModelResources().ParameterUIDataMap.Find(ParamName);
+	if (!ParameterData)
+	{
+		return {};
+	}
 
-FParameterUIData UCustomizableObject::GetStateUIMetadataFromIndex(int32 StateIndex) const
-{
-	return GetPrivate()->GetStateUIMetadataFromIndex(StateIndex);
+	const FIntegerParameterUIData* IntegerParameterUIData = ParameterData->ArrayIntegerParameterOption.Find(OptionName);
+	return IntegerParameterUIData ? IntegerParameterUIData->ParamUIMetadata : FMutableParamUIMetadata();
 }
 
 
-FParameterUIData UCustomizableObjectPrivate::GetStateUIMetadataFromIndex(int32 StateIndex) const
+FMutableStateUIMetadata UCustomizableObject::GetStateUIMetadata(const FString& StateName) const
 {
-	return GetStateUIMetadata(GetStateName(StateIndex));
+	const FMutableStateData* StateData = Private->GetModelResources().StateUIDataMap.Find(StateName);
+	return StateData ? StateData->StateUIMetadata : FMutableStateUIMetadata();
 }
 
 
@@ -1730,10 +1726,16 @@ bool UCustomizableObject::IsParameterMultidimensional(const int32& InParamIndex)
 }
 
 
-void UCustomizableObjectPrivate::ApplyStateForcedValuesToParameters( int32 State, mu::Parameters* Parameters)
+void UCustomizableObjectPrivate::ApplyStateForcedValuesToParameters(int32 State, mu::Parameters* Parameters)
 {
-	FParameterUIData StateMetaData = GetStateUIMetadataFromIndex(State);
-	for (const TPair<FString, FString>& ForcedParameter : StateMetaData.ForcedParameterValues)
+	const FString& ParameterName = GetPublic()->GetParameterName(State);
+	const FMutableStateData* StateData = GetModelResources().StateUIDataMap.Find(ParameterName);
+	if (!StateData)
+	{
+		return;
+	}
+	
+	for (const TPair<FString, FString>& ForcedParameter : StateData->ForcedParameterValues)
 	{
 		int32 ForcedParameterIndex = FindParameter(ForcedParameter.Key);
 		if (ForcedParameterIndex == INDEX_NONE)
@@ -1773,7 +1775,6 @@ void UCustomizableObjectPrivate::ApplyStateForcedValuesToParameters( int32 State
 			}
 		}
 	}
-
 }
 
 
@@ -1874,6 +1875,37 @@ void FSkeletonCache::Add(const TArray<uint16>& Key, USkeleton* Value)
 }
 
 
+FArchive& operator<<(FArchive& Ar, FIntegerParameterUIData& Struct)
+{
+	Ar << Struct.ParamUIMetadata;
+	
+	return Ar;
+}
+
+
+FArchive& operator<<(FArchive& Ar, FMutableParameterData& Struct)
+{
+	Ar << Struct.ParamUIMetadata;
+	Ar << Struct.Type;
+	Ar << Struct.ArrayIntegerParameterOption;
+	Ar << Struct.IntegerParameterGroupType;
+	
+	return Ar;
+}
+
+
+FArchive& operator<<(FArchive& Ar, FMutableStateData& Struct)
+{
+	Ar << Struct.StateUIMetadata;
+	Ar << Struct.bLiveUpdateMode;
+	Ar << Struct.bDisableTextureStreaming;
+	Ar << Struct.bReuseInstanceTextures;
+	Ar << Struct.ForcedParameterValues;
+
+	return Ar;
+}
+
+
 UCustomizableObjectPrivate::UCustomizableObjectPrivate()
 {
 #if WITH_EDITOR
diff --git a/Engine/Plugins/Experimental/Mutable/Source/CustomizableObject/Private/MuCO/CustomizableObjectInstance.cpp b/Engine/Plugins/Experimental/Mutable/Source/CustomizableObject/Private/MuCO/CustomizableObjectInstance.cpp
--- a/Engine/Plugins/Experimental/Mutable/Source/CustomizableObject/Private/MuCO/CustomizableObjectInstance.cpp
+++ b/Engine/Plugins/Experimental/Mutable/Source/CustomizableObject/Private/MuCO/CustomizableObjectInstance.cpp
@@ -50,6 +50,7 @@
 #include "PhysicsEngine/PhysicsConstraintTemplate.h"
 #include "PhysicsEngine/SkeletalBodySetup.h"
 #include "Hash/CityHash.h"
+#include "MuCO/CustomizableObjectCustomVersion.h"
 
 #include UE_INLINE_GENERATED_CPP_BY_NAME(CustomizableObjectInstance)
 
diff --git a/Engine/Plugins/Experimental/Mutable/Source/CustomizableObject/Private/MuCO/CustomizableObjectSystem.cpp b/Engine/Plugins/Experimental/Mutable/Source/CustomizableObject/Private/MuCO/CustomizableObjectSystem.cpp
--- a/Engine/Plugins/Experimental/Mutable/Source/CustomizableObject/Private/MuCO/CustomizableObjectSystem.cpp
+++ b/Engine/Plugins/Experimental/Mutable/Source/CustomizableObject/Private/MuCO/CustomizableObjectSystem.cpp
@@ -986,8 +986,7 @@ void UCustomizableObjectSystemPrivate::GetMipStreamingConfig(const UCustomizable
 
 	// From user-controlled per-state flag?
 	const FString CurrentState = Instance.GetCurrentState();
-	const FParameterUIData* State = Instance.GetCustomizableObject()->GetPrivate()->GetModelResources().StateUIDataMap.Find(CurrentState);	
-	if (State)
+	if (const FMutableStateData* State = Instance.GetCustomizableObject()->GetPrivate()->GetModelResources().StateUIDataMap.Find(CurrentState))
 	{
 		bOutNeverStream = State->bDisableTextureStreaming;
 	}
@@ -3255,7 +3254,7 @@ namespace impl
 		SystemPrivateData->ExtensionDataStreamer->SetActiveObject(CustomizableObject);
 
 		FString StateName = CandidateInstance->GetCustomizableObject()->GetStateName(CandidateInstance->GetPrivate()->GetState());
-		const FParameterUIData* StateData = CandidateInstance->GetCustomizableObject()->GetPrivate()->GetModelResources().StateUIDataMap.Find(StateName);
+		const FMutableStateData* StateData = CandidateInstance->GetCustomizableObject()->GetPrivate()->GetModelResources().StateUIDataMap.Find(StateName);
 
 		Operation->bLiveUpdateMode = false;
 
diff --git a/Engine/Plugins/Experimental/Mutable/Source/CustomizableObject/Private/MuCO/CustomizableObjectUIData.cpp b/Engine/Plugins/Experimental/Mutable/Source/CustomizableObject/Private/MuCO/CustomizableObjectUIData.cpp
new file mode 100644
--- /dev/null
+++ b/Engine/Plugins/Experimental/Mutable/Source/CustomizableObject/Private/MuCO/CustomizableObjectUIData.cpp
@@ -0,0 +1,81 @@
+﻿// Copyright Epic Games, Inc. All Rights Reserved.
+
+#pragma once
+
+#include "MuCO/CustomizableObjectUIData.h"
+
+#include "MuCO/CustomizableObjectCustomVersion.h"
+
+
+FArchive& operator<<(FArchive& Ar, FMutableUIMetadata& Metadata)
+{
+	Ar.UsingCustomVersion(FCustomizableObjectCustomVersion::GUID);
+
+	Ar << Metadata.ObjectFriendlyName;
+	Ar << Metadata.UISectionName;
+	Ar << Metadata.UIOrder;
+
+	if (Ar.IsLoading())
+	{
+		FString StringRef;
+		Ar << StringRef;
+		Metadata.UIThumbnail = TSoftObjectPtr<UTexture2D>(FSoftObjectPath(StringRef));
+	}
+	else
+	{
+		FString StringRef = Metadata.UIThumbnail.ToSoftObjectPath().ToString();
+		Ar << StringRef;
+	}
+
+	Ar << Metadata.ExtraInformation;
+
+	if (Ar.IsLoading())
+	{
+		int32 NumReferencedAssets = 0;
+		Ar << NumReferencedAssets;
+		Metadata.ExtraAssets.Empty(NumReferencedAssets);
+
+		for (int32 i = 0; i < NumReferencedAssets; ++i)
+		{
+			FString Key, StringRef;
+			Ar << Key;
+			Ar << StringRef;
+
+			Metadata.ExtraAssets.Add(Key, TSoftObjectPtr<UObject>(FSoftObjectPath(StringRef)));
+		}
+	}
+	else
+	{
+		int32 NumReferencedAssets = Metadata.ExtraAssets.Num();
+		Ar << NumReferencedAssets;
+
+		for (TPair<FString, TSoftObjectPtr<UObject>>& AssetPair : Metadata.ExtraAssets)
+		{
+			FString StringRef = AssetPair.Value.ToSoftObjectPath().ToString();
+			Ar << AssetPair.Key;
+			Ar << StringRef;
+		}
+	}
+
+
+	return Ar;
+}
+
+
+FArchive& operator<<(FArchive& Ar, FMutableParamUIMetadata& Struct)
+{
+	Ar << *static_cast<FMutableUIMetadata*>(&Struct);
+	
+	Ar << Struct.MinimumValue;
+	Ar << Struct.MaximumValue;
+
+	return Ar;
+}
+
+
+FArchive& operator<<(FArchive& Ar, FMutableStateUIMetadata& Struct)
+{
+	Ar << *static_cast<FMutableUIMetadata*>(&Struct);
+
+	return Ar;
+}
\ No newline at end of file
diff --git a/Engine/Plugins/Experimental/Mutable/Source/CustomizableObject/Public/MuCO/CustomizableObjectUIData.h b/Engine/Plugins/Experimental/Mutable/Source/CustomizableObject/Public/MuCO/CustomizableObjectUIData.h
--- a/Engine/Plugins/Experimental/Mutable/Source/CustomizableObject/Public/MuCO/CustomizableObjectUIData.h
+++ b/Engine/Plugins/Experimental/Mutable/Source/CustomizableObject/Public/MuCO/CustomizableObjectUIData.h
@@ -2,10 +2,9 @@
 
 #pragma once
 
+#include "InstancedStruct.h"
 #include "Engine/Texture2D.h"
 #include "MuCO/CustomizableObjectParameterTypeDefinitions.h"
-#include "MuCO/CustomizableObjectCustomVersion.h"
-#include "MuR/System.h"
 
 #include "CustomizableObjectUIData.generated.h"
 
@@ -13,27 +12,23 @@ class UTexture2D;
 
 
 USTRUCT(BlueprintType)
-struct FMutableParamUIMetadata
+struct FMutableUIMetadata
 {
 	GENERATED_BODY()
 
-	FMutableParamUIMetadata()
-		: UIOrder(0)
-	{}
-
 	/** This is the name to be shown in UI */
 	UPROPERTY(BlueprintReadWrite, EditAnywhere, Category = UI)
 	FString ObjectFriendlyName;
 
-	/** This is the name of the section where the parameter will be placed in UI */
+	/** This is the name of the section where the object will be placed in UI */
 	UPROPERTY(BlueprintReadWrite, EditAnywhere, Category = UI)
 	FString UISectionName;
 
-	/** This is the order of the parameter inside its section */
+	/** This is the order of the object inside its section */
 	UPROPERTY(BlueprintReadWrite, EditAnywhere, Category = UI)
-	int32 UIOrder;
+	int32 UIOrder = 0;
 
-	/** Thumnbail for UI */
+	/** Thumbnail for UI */
 	UPROPERTY(BlueprintReadWrite, EditAnywhere, Category = UI)
 	TSoftObjectPtr<UTexture2D> UIThumbnail;
 
@@ -45,281 +40,34 @@ struct FMutableParamUIMetadata
 	UPROPERTY(BlueprintReadWrite, EditAnywhere, Category = UI)
 	TMap<FString, TSoftObjectPtr<UObject>> ExtraAssets;
 
-	UPROPERTY(BlueprintReadWrite, EditAnywhere, Category = UI)
-	float MinimumValue = 0.0f;
-
-	UPROPERTY(BlueprintReadWrite, EditAnywhere, Category = UI)
-	float MaximumValue = 1.0f;
-
-	bool operator ==(const FMutableParamUIMetadata& Other) const
-	{
-		if (ObjectFriendlyName != Other.ObjectFriendlyName || UISectionName != Other.UISectionName || UIOrder != Other.UIOrder || UIThumbnail != Other.UIThumbnail)
-		{
-			return false;
-		}
-
-		if (!ExtraInformation.OrderIndependentCompareEqual(Other.ExtraInformation))
-		{
-			return false;
-		}
-
-		if (!ExtraAssets.OrderIndependentCompareEqual(Other.ExtraAssets))
-		{
-			return false;
-		}
-
-		return true;
-	}
-
-	bool operator !=(const FMutableParamUIMetadata& Other) const
-	{
-		return !(*this == Other);
-	}
-
-	friend FArchive& operator <<(FArchive& Ar, FMutableParamUIMetadata& Metadata)
-	{
-		Ar.UsingCustomVersion(FCustomizableObjectCustomVersion::GUID);
-
-		Ar << Metadata.ObjectFriendlyName;
-		Ar << Metadata.UISectionName;
-		Ar << Metadata.UIOrder;
-
-		//Ar << Metadata.UIThumbnail;
-		if (Ar.IsLoading())
-		{
-			FString StringRef;
-			Ar << StringRef;
-			Metadata.UIThumbnail = TSoftObjectPtr<UTexture2D>(FSoftObjectPath(StringRef));
-		}
-		else
-		{
-			FString StringRef = Metadata.UIThumbnail.ToSoftObjectPath().ToString();
-			Ar << StringRef;
-		}
-
-		Ar << Metadata.ExtraInformation;
-
-		//Ar << Metadata.ExtraAssets;
-		if (Ar.IsLoading())
-		{
-			int32 NumReferencedAssets = 0;
-			Ar << NumReferencedAssets;
-			Metadata.ExtraAssets.Empty(NumReferencedAssets);
-
-			for (int32 i = 0; i < NumReferencedAssets; ++i)
-			{
-				FString Key, StringRef;
-				Ar << Key;
-				Ar << StringRef;
-
-				Metadata.ExtraAssets.Add(Key, TSoftObjectPtr<UObject>(FSoftObjectPath(StringRef)));
-			}
-		}
-		else
-		{
-			int32 NumReferencedAssets = Metadata.ExtraAssets.Num();
-			Ar << NumReferencedAssets;
-
-			for (TPair<FString, TSoftObjectPtr<UObject>>& AssetPair : Metadata.ExtraAssets)
-			{
-				FString StringRef = AssetPair.Value.ToSoftObjectPath().ToString();
-				Ar << AssetPair.Key;
-				Ar << StringRef;
-			}
-		}
-
-		// Update structure
-		if (Ar.CustomVer(FCustomizableObjectCustomVersion::GUID) >= FCustomizableObjectCustomVersion::BeforeCustomVersionWasAdded)
-		{
-			Ar << Metadata.MinimumValue;
-			Ar << Metadata.MaximumValue;
-		}
-
-		return Ar;
-	}
-
-#if WITH_EDITOR
-
-	/** Only called in the BeginCacheForCookPlatform to include new references in the final package */
-	void LoadResources()
-	{
-		UIThumbnail.LoadSynchronous();
-
-		for (const TPair<FString, TSoftObjectPtr<UObject>>& a : ExtraAssets)
-		{
-			a.Value.LoadSynchronous();
-		}
-	}
-
-#endif
+	friend FArchive& operator<<(FArchive& Ar, FMutableUIMetadata& Struct);
 };
 
+
 USTRUCT(BlueprintType)
-struct FIntegerParameterUIData
+struct FMutableParamUIMetadata : public FMutableUIMetadata
 {
 	GENERATED_BODY()
 
-	FIntegerParameterUIData()
-	{
-
-	}
-
-	FIntegerParameterUIData(
-		const FString& ParamName,
-		const FMutableParamUIMetadata& InParamFriendlyName) :
-		  Name(ParamName)
-		, ParamUIMetadata(InParamFriendlyName)
-	{
-
-	}
-
-	/** Integer parameter option name */
-	UPROPERTY(BlueprintReadWrite, Category = UI)
-	FString Name;
-
-	UPROPERTY(BlueprintReadWrite, Category = UI, meta = (DisplayName = "Parameter UI Metadata"))
-	FMutableParamUIMetadata ParamUIMetadata;
-
-	bool operator ==(const FIntegerParameterUIData& Other) const
-	{
-		if (Name != Other.Name || ParamUIMetadata != Other.ParamUIMetadata)
-		{
-			return false;
-		}
-
-		return true;
-	}
-
-	bool operator !=(const FIntegerParameterUIData& Other) const
-	{
-		return !(*this == Other);
-	}
-
-	friend FArchive& operator <<(FArchive& Ar, FIntegerParameterUIData& UIData)
-	{
-		Ar << UIData.Name;
-		Ar << UIData.ParamUIMetadata;
-
-		return Ar;
-	}
-
-#if WITH_EDITOR
+	UPROPERTY(BlueprintReadWrite, EditAnywhere, Category = UI)
+	float MinimumValue = 0.0f;
 
-	/** Only called in the BeginCacheForCookPlatform to include new references in the final package */
-	void LoadResources()
-	{
-		ParamUIMetadata.LoadResources();
-	}
+	UPROPERTY(BlueprintReadWrite, EditAnywhere, Category = UI)
+	float MaximumValue = 1.0f;
 
-#endif
+	friend FArchive& operator<<(FArchive& Ar, FMutableParamUIMetadata& Struct);
 };
 
 
 USTRUCT(BlueprintType)
-struct FParameterUIData
+struct FMutableStateUIMetadata : public FMutableUIMetadata
 {
 	GENERATED_BODY()
 
-	FParameterUIData() 
-	{
-	}
-
-	FParameterUIData(
-		const FString& ParamName,
-		const FMutableParamUIMetadata& InParamUIMetadata,
-		EMutableParameterType ParamType) :
-		Name(ParamName),
-		ParamUIMetadata(InParamUIMetadata),
-		Type(ParamType)
-	{
-
-	}
-
-	/** Parameter name */
-	UPROPERTY(BlueprintReadWrite, Category = UI)
-	FString Name;
-
-	UPROPERTY(BlueprintReadWrite, Category = UI, meta = (DisplayName = "Parameter UI Metadata"))
-	FMutableParamUIMetadata ParamUIMetadata;
-
-	/** Parameter type, using uint8 since the enum in declared in the class it is used */
-	UPROPERTY(BlueprintReadWrite, Category = UI)
-	EMutableParameterType Type = EMutableParameterType::None;
-
-	/** In the case of an integer parameter, store here all options */
-	UPROPERTY(BlueprintReadWrite, Category = UI)
-	TArray<FIntegerParameterUIData> ArrayIntegerParameterOption;
-
-	/** In the case of an integer parameter, how are the different options selected (one, one or none, etc...) */
-	UPROPERTY(BlueprintReadWrite, Category = UI)
-	ECustomizableObjectGroupType IntegerParameterGroupType = ECustomizableObjectGroupType::COGT_ONE_OR_NONE;
-
-	/** Not really relevant for UI, but apparently bDontCompressRuntimeTextures_DEPRECATED was used to decided some texture properties at runtime. */
-	UPROPERTY()
-	ETextureCompressionStrategy TextureCompressionStrategy = ETextureCompressionStrategy::None;
-
-	/** If this is enabled, texture streaming won't be used for this state, and full images will be generated when an instance is first updated. */
-	UPROPERTY()
-	bool bDisableTextureStreaming = false;
-
-	/** In this mode instances and their temp data will be reused between updates. It will be much faster but spend as much as ten times the memory.
-	    Useful for customization lockers with few characters that are going to have their parameters changed many times, not for in-game */
-	UPROPERTY(BlueprintReadWrite, Category = CustomizableObject)
-	bool bLiveUpdateMode = false;
-
-	UPROPERTY(BlueprintReadWrite, Category = CustomizableObject)
-	bool bReuseInstanceTextures = false;
-
-	UPROPERTY(BlueprintReadWrite, Category = CustomizableObject)
-	TMap<FString, FString> ForcedParameterValues;
-
-	bool operator ==(const FParameterUIData& Other) const
-	{
-		if (Name != Other.Name || ParamUIMetadata != Other.ParamUIMetadata || Type != Other.Type 
-			|| ArrayIntegerParameterOption != Other.ArrayIntegerParameterOption || IntegerParameterGroupType != Other.IntegerParameterGroupType 
-			|| bLiveUpdateMode != Other.bLiveUpdateMode || bReuseInstanceTextures != Other.bReuseInstanceTextures
-			|| bDisableTextureStreaming != Other.bDisableTextureStreaming
-			|| !ForcedParameterValues.OrderIndependentCompareEqual(Other.ForcedParameterValues)
-			)
-		{
-			return false;
-		}
-
-		return true;
-	}
-
-	friend FArchive& operator <<(FArchive& Ar, FParameterUIData& UIData)
-	{
-		Ar << UIData.Name;
-		Ar << UIData.ParamUIMetadata;
-		Ar << UIData.Type;
-		Ar << UIData.ArrayIntegerParameterOption;
-		Ar << UIData.IntegerParameterGroupType;
-		Ar << UIData.TextureCompressionStrategy;
-		Ar << UIData.bDisableTextureStreaming;
-		Ar << UIData.bLiveUpdateMode;
-		Ar << UIData.bReuseInstanceTextures;
-		Ar << UIData.ForcedParameterValues;
-
-		return Ar;
-	}
-
-#if WITH_EDITOR
-
-	/** Only called in the BeginCacheForCookPlatform to include new references in the final package */
-	void LoadResources()
-	{
-		ParamUIMetadata.LoadResources();
-
-		for (FIntegerParameterUIData& intParam : ArrayIntegerParameterOption)
-		{
-			intParam.LoadResources();
-		}
-	}
-
-#endif
+	friend FArchive& operator<<(FArchive& Ar, FMutableStateUIMetadata& Struct);
 };
 
+
 #if UE_ENABLE_INCLUDE_ORDER_DEPRECATED_IN_5_2
 #include "MuCO/UnrealPortabilityHelpers.h"
 #endif
diff --git a/Engine/Plugins/Experimental/Mutable/Source/CustomizableObjectEditor/Classes/MuCOE/Nodes/CustomizableObjectNodeObject.h b/Engine/Plugins/Experimental/Mutable/Source/CustomizableObjectEditor/Classes/MuCOE/Nodes/CustomizableObjectNodeObject.h
--- a/Engine/Plugins/Experimental/Mutable/Source/CustomizableObjectEditor/Classes/MuCOE/Nodes/CustomizableObjectNodeObject.h
+++ b/Engine/Plugins/Experimental/Mutable/Source/CustomizableObjectEditor/Classes/MuCOE/Nodes/CustomizableObjectNodeObject.h
@@ -58,6 +58,8 @@ struct CUSTOMIZABLEOBJECTEDITOR_API FCustomizableObjectState
 {
 	GENERATED_USTRUCT_BODY()
 
+	friend UCustomizableObjectNodeObject;
+	
 	UPROPERTY(EditAnywhere, Category=CustomizableObject)
 	FString Name;
 
@@ -93,9 +95,12 @@ struct CUSTOMIZABLEOBJECTEDITOR_API FCustomizableObjectState
 	TMap<FString, FString> ForcedParameterValues;
 
 	UPROPERTY(EditAnywhere, Category = UI, meta = (DisplayName = "State UI Metadata"))
-	FMutableParamUIMetadata StateUIMetadata;
+	FMutableStateUIMetadata UIMetadata;
 
+private:
 	// Deprecated
+	UPROPERTY()
+	FMutableParamUIMetadata StateUIMetadata_DEPRECATED;
 	
 	/** This is now TextureCompressionStrategy.  */
 	UPROPERTY()
diff --git a/Engine/Plugins/Experimental/Mutable/Source/CustomizableObjectEditor/Private/MuCOE/CustomizableInstanceDetails.cpp b/Engine/Plugins/Experimental/Mutable/Source/CustomizableObjectEditor/Private/MuCOE/CustomizableInstanceDetails.cpp
--- a/Engine/Plugins/Experimental/Mutable/Source/CustomizableObjectEditor/Private/MuCOE/CustomizableInstanceDetails.cpp
+++ b/Engine/Plugins/Experimental/Mutable/Source/CustomizableObjectEditor/Private/MuCOE/CustomizableInstanceDetails.cpp
@@ -626,8 +626,8 @@ bool FCustomizableInstanceDetails::GenerateParametersView(IDetailCategoryBuilder
 
 			if (CustomInstance->IsParameterRelevant(ParameterSortInfo.ParamIndexInObject) && IsVisible(ParameterSortInfo.ParamIndexInObject))
 			{
-				ParameterSortInfo.ParamUIOrder = CustomizableObject->GetParameterUIMetadataFromIndex(ParameterSortInfo.ParamIndexInObject).ParamUIMetadata.UIOrder;
 				ParameterSortInfo.ParamName = CustomizableObject->GetParameterName(ParameterSortInfo.ParamIndexInObject);
+				ParameterSortInfo.ParamUIOrder = CustomizableObject->GetParameterUIMetadata(ParameterSortInfo.ParamName).UIOrder;
 				ParametersTree.Add(ParameterSortInfo);
 			}
 		}
@@ -684,8 +684,8 @@ bool FCustomizableInstanceDetails::GenerateParametersView(IDetailCategoryBuilder
 			ParameterSortInfo.ParamIndexInObject = ParamIndexInObject;
 			if ((!CustomInstance->GetPrivate()->bShowOnlyRelevantParameters || CustomInstance->IsParameterRelevant(ParameterSortInfo.ParamIndexInObject)) && IsVisible(ParameterSortInfo.ParamIndexInObject))
 			{
-				ParameterSortInfo.ParamUIOrder = CustomizableObject->GetParameterUIMetadataFromIndex(ParameterSortInfo.ParamIndexInObject).ParamUIMetadata.UIOrder;
 				ParameterSortInfo.ParamName = CustomizableObject->GetParameterName(ParameterSortInfo.ParamIndexInObject);
+				ParameterSortInfo.ParamUIOrder = CustomizableObject->GetParameterUIMetadata(ParameterSortInfo.ParamName).UIOrder;
 				ParametersTree.Add(ParameterSortInfo);
 			}
 		}
@@ -782,12 +782,10 @@ void FCustomizableInstanceDetails::RecursivelyAddParamAndChildren(const int32 Pa
 void FCustomizableInstanceDetails::FillChildrenMap(int32 ParamIndexInObject)
 {
 	const UCustomizableObject* CustomizableObject = CustomInstance->GetCustomizableObject();
-	FString ParamName = CustomizableObject->GetParameterName(ParamIndexInObject);
-	FParameterUIData UIData = CustomizableObject->GetParameterUIMetadataFromIndex(ParamIndexInObject);
-
-	const FString* ParentName = UIData.ParamUIMetadata.ExtraInformation.Find(FString("__ParentParamName"));
+	const FString& ParamName = CustomizableObject->GetParameterName(ParamIndexInObject);
+	FMutableParamUIMetadata UIMetadata = CustomizableObject->GetParameterUIMetadata(ParamName);
 
-	if (ParentName)
+	if (const FString* ParentName = UIMetadata.ExtraInformation.Find(FString("__ParentParamName")))
 	{
 		ParamChildren.Add(*ParentName, ParamIndexInObject);
 		ParamHasParent.Add(ParamIndexInObject, true);
@@ -798,18 +796,18 @@ void FCustomizableInstanceDetails::FillChildrenMap(int32 ParamIndexInObject)
 bool FCustomizableInstanceDetails::IsVisible(int32 ParamIndexInObject)
 {
 	const UCustomizableObject* CustomizableObject = CustomInstance->GetCustomizableObject();
-	FString ParamName = CustomizableObject->GetParameterName(ParamIndexInObject);
-	FParameterUIData UIData = CustomizableObject->GetParameterUIMetadataFromIndex(ParamIndexInObject);
-	FString* ParentName = UIData.ParamUIMetadata.ExtraInformation.Find(FString("__ParentParamName"));
+	const FString& ParamName = CustomizableObject->GetParameterName(ParamIndexInObject);
+	FMutableParamUIMetadata UIMetadata = CustomizableObject->GetParameterUIMetadata(ParamName);
+	const FString* ParentName = UIMetadata.ExtraInformation.Find(FString("__ParentParamName"));
 
-	bool IsAProjectorParam = ParamName.EndsWith(FMultilayerProjector::NUM_LAYERS_PARAMETER_POSTFIX)
+	const bool IsAProjectorParam = ParamName.EndsWith(FMultilayerProjector::NUM_LAYERS_PARAMETER_POSTFIX)
 		|| (ParamName.EndsWith(FMultilayerProjector::IMAGE_PARAMETER_POSTFIX) && CustomizableObject->IsParameterMultidimensional(ParamIndexInObject))
 		|| (ParamName.EndsWith(FMultilayerProjector::OPACITY_PARAMETER_POSTFIX) && CustomizableObject->IsParameterMultidimensional(ParamIndexInObject))
 		|| (ParamName.EndsWith(FMultilayerProjector::POSE_PARAMETER_POSTFIX));
 
 	if (!IsAProjectorParam && ParentName && CustomInstance->GetPrivate()->bShowOnlyRelevantParameters)
 	{
-		FString* Value = UIData.ParamUIMetadata.ExtraInformation.Find(FString("__DisplayWhenParentValueEquals"));
+		const FString* Value = UIMetadata.ExtraInformation.Find(FString("__DisplayWhenParentValueEquals"));
 
 		if (Value && CustomizableObject->FindParameter(*ParentName) != INDEX_NONE && CustomInstance->GetIntParameterSelectedOption(*ParentName) != *Value)
 		{
@@ -833,8 +831,9 @@ IDetailGroup* FCustomizableInstanceDetails::GenerateParameterSection(const int32
 {
 	const UCustomizableObject* CustomizableObject = CustomInstance->GetCustomizableObject();
 
-	FMutableParamUIMetadata UIData = CustomizableObject->GetParameterUIMetadataFromIndex(ParamIndexInObject).ParamUIMetadata;
-	FString SectionName = UIData.UISectionName.IsEmpty() ? "Miscellaneous" : UIData.UISectionName;
+	const FString& ParamName = CustomizableObject->GetParameterName(ParamIndexInObject);
+	const FMutableParamUIMetadata UIMetadata = CustomizableObject->GetParameterUIMetadata(ParamName);
+	const FString SectionName = UIMetadata.UISectionName.IsEmpty() ? "Miscellaneous" : UIMetadata.UISectionName;
 	IDetailGroup* CurrentSection = nullptr;
 
 	if (GeneratedSections.Contains(SectionName))
@@ -1015,8 +1014,8 @@ void FCustomizableInstanceDetails::OnIntParameterComboBoxChanged(TSharedPtr<FStr
 TSharedRef<SWidget> FCustomizableInstanceDetails::GenerateFloatWidget(const int32 ParamIndexInObject)
 {
 	const UCustomizableObject* CustomizableObject = CustomInstance->GetCustomizableObject();
-	const FMutableParamUIMetadata& UIMetadata = CustomizableObject->GetParameterUIMetadataFromIndex(ParamIndexInObject).ParamUIMetadata;
-	FString ParamName = CustomizableObject->GetParameterName(ParamIndexInObject);
+	const FString& ParamName = CustomizableObject->GetParameterName(ParamIndexInObject);
+	const FMutableParamUIMetadata& UIMetadata = CustomizableObject->GetParameterUIMetadata(ParamName);
 
 	if (const TSoftObjectPtr<UObject>* FloatDecoratorAsset = UIMetadata.ExtraAssets.Find(UIMetadataKeyWords::FloatDecoratorName))
 	{
diff --git a/Engine/Plugins/Experimental/Mutable/Source/CustomizableObjectEditor/Private/MuCOE/GenerateMutableSource/GenerateMutableSource.cpp b/Engine/Plugins/Experimental/Mutable/Source/CustomizableObjectEditor/Private/MuCOE/GenerateMutableSource/GenerateMutableSource.cpp
--- a/Engine/Plugins/Experimental/Mutable/Source/CustomizableObjectEditor/Private/MuCOE/GenerateMutableSource/GenerateMutableSource.cpp
+++ b/Engine/Plugins/Experimental/Mutable/Source/CustomizableObjectEditor/Private/MuCOE/GenerateMutableSource/GenerateMutableSource.cpp
@@ -981,14 +981,14 @@ mu::NodeObjectPtr GenerateMutableSource(const UEdGraphPin * Pin, FMutableGraphGe
 			ObjectNode->SetStateProperties(StateIndex, State.TextureCompressionStrategy, State.bBuildOnlyFirstLOD, GenerationContext.FirstLODAvailable, NumExtraLODsToBuildAfterFirstLOD);
 
 			// UI Data
-			FParameterUIData ParameterUIData(State.Name, State.StateUIMetadata, EMutableParameterType::None);
-			ParameterUIData.TextureCompressionStrategy = State.TextureCompressionStrategy;
-			ParameterUIData.bDisableTextureStreaming = State.bDisableTextureStreaming;
-			ParameterUIData.bLiveUpdateMode = State.bLiveUpdateMode;
-			ParameterUIData.bReuseInstanceTextures = State.bReuseInstanceTextures;
-			ParameterUIData.ForcedParameterValues = State.ForcedParameterValues;
+			FMutableStateData StateUIData;
+			StateUIData.StateUIMetadata = State.UIMetadata;
+			StateUIData.bDisableTextureStreaming = State.bDisableTextureStreaming;
+			StateUIData.bLiveUpdateMode = State.bLiveUpdateMode;
+			StateUIData.bReuseInstanceTextures = State.bReuseInstanceTextures;
+			StateUIData.ForcedParameterValues = State.ForcedParameterValues;
 
-			GenerationContext.StateUIDataMap.Add(State.Name, ParameterUIData);
+			GenerationContext.StateUIDataMap.Add(State.Name, StateUIData);
 		}
 
 		// Update the current automatic LOD policy
@@ -1238,8 +1238,7 @@ mu::NodeObjectPtr GenerateMutableSource(const UEdGraphPin * Pin, FMutableGraphGe
 		int32 ChildIndex = 0;
 
 		// UI data
-		FParameterUIData ParameterUIData(
-			TypedNodeGroup->GroupName,
+		FMutableParameterData ParameterUIData(
 			TypedNodeGroup->ParamUIMetadata,
 			EMutableParameterType::Int);
 
@@ -1271,17 +1270,16 @@ mu::NodeObjectPtr GenerateMutableSource(const UEdGraphPin * Pin, FMutableGraphGe
 					LeftSplit.Split(".", &LeftSplit, nullptr);
 					GenerationContext.CustomizableObjectPathMap.Add(CustomizableObjectNodeObject->Identifier.ToString(), LeftSplit);
 					GenerationContext.GroupNodeMap.Add(CustomizableObjectNodeObject->Identifier.ToString(), FCustomizableObjectIdPair(TypedNodeGroup->GroupName, ChildNode->GetName()));
-					ParameterUIData.ArrayIntegerParameterOption.Add(FIntegerParameterUIData(
+					ParameterUIData.ArrayIntegerParameterOption.Add(
 						CustomizableObjectNodeObject->ObjectName,
-						CustomizableObjectNodeObject->ParamUIMetadata));
+				FIntegerParameterUIData(CustomizableObjectNodeObject->ParamUIMetadata));
 
 					if (TypedNodeGroup->GroupType == ECustomizableObjectGroupType::COGT_TOGGLE)
 					{
 						GenerationContext.AddParameterNameUnique(CustomizableObjectNodeObject, CustomizableObjectNodeObject->ObjectName);
 						
 						// UI Data is only relevant when the group node is set to Toggle
-						GenerationContext.ParameterUIDataMap.Add(CustomizableObjectNodeObject->ObjectName, FParameterUIData(
-							CustomizableObjectNodeObject->ObjectName,
+						GenerationContext.ParameterUIDataMap.Add(CustomizableObjectNodeObject->ObjectName, FMutableParameterData(
 							CustomizableObjectNodeObject->ParamUIMetadata,
 							EMutableParameterType::Int));
 					}
@@ -1344,7 +1342,7 @@ mu::NodeObjectPtr GenerateMutableSource(const UEdGraphPin * Pin, FMutableGraphGe
 							{
 								ParameterUIData.ParamUIMetadata.ExtraInformation.Add(FString("CollapseUnderParent"));
 
-								FParameterUIData ParentParameterUIData;
+								FMutableParameterData ParentParameterUIData;
 								ParentParameterUIData.ParamUIMetadata.ExtraInformation.Add(FString("__HasCollapsibleChildren"));
 								GenerationContext.ParameterUIDataMap.Add(*ParentParamName, ParentParameterUIData);
 							}
@@ -1380,9 +1378,9 @@ mu::NodeObjectPtr GenerateMutableSource(const UEdGraphPin * Pin, FMutableGraphGe
 					LeftSplit.Split(".", &LeftSplit, nullptr);
 					GenerationContext.CustomizableObjectPathMap.Add(CustomizableObjectNodeObject->Identifier.ToString(), LeftSplit);
 					GenerationContext.GroupNodeMap.Add(CustomizableObjectNodeObject->Identifier.ToString(), FCustomizableObjectIdPair(TypedNodeGroup->GroupName, ChildNode->GetName()));
-					ParameterUIData.ArrayIntegerParameterOption.Add(FIntegerParameterUIData(
+					ParameterUIData.ArrayIntegerParameterOption.Add(
 						CustomizableObjectNodeObject->ObjectName,
-						CustomizableObjectNodeObject->ParamUIMetadata));
+						FIntegerParameterUIData(CustomizableObjectNodeObject->ParamUIMetadata));
 
 					if (CustomizableObjectNodeObject->ObjectName.IsEmpty())
 					{
@@ -1394,8 +1392,7 @@ mu::NodeObjectPtr GenerateMutableSource(const UEdGraphPin * Pin, FMutableGraphGe
 						GenerationContext.AddParameterNameUnique(CustomizableObjectNodeObject, CustomizableObjectNodeObject->ObjectName);
 
 						// UI Data is only relevant when the group node is set to Toggle
-						GenerationContext.ParameterUIDataMap.Add(CustomizableObjectNodeObject->ObjectName, FParameterUIData(
-							CustomizableObjectNodeObject->ObjectName,
+						GenerationContext.ParameterUIDataMap.Add(CustomizableObjectNodeObject->ObjectName, FMutableParameterData(
 							CustomizableObjectNodeObject->ParamUIMetadata,
 							EMutableParameterType::Int));
 					}
@@ -1418,7 +1415,7 @@ mu::NodeObjectPtr GenerateMutableSource(const UEdGraphPin * Pin, FMutableGraphGe
 			ChildIndex++;
 		}
 
-		const FParameterUIData* ChildFilledUIData = GenerationContext.ParameterUIDataMap.Find(TypedNodeGroup->GroupName);
+		const FMutableParameterData* ChildFilledUIData = GenerationContext.ParameterUIDataMap.Find(TypedNodeGroup->GroupName);
 		if (ChildFilledUIData && ChildFilledUIData->ParamUIMetadata.ExtraInformation.Find(FString("__HasCollapsibleChildren")))
 		{
 			// Some child param filled the HasCollapsibleChildren UI info, refill it so it's not lost
@@ -1427,16 +1424,15 @@ mu::NodeObjectPtr GenerateMutableSource(const UEdGraphPin * Pin, FMutableGraphGe
 
 		if (TypedNodeGroup->GroupType == ECustomizableObjectGroupType::COGT_TOGGLE)
 		{
-			for (const FIntegerParameterUIData& BooleanParam : ParameterUIData.ArrayIntegerParameterOption)
+			for (const TTuple<FString, FIntegerParameterUIData>& BooleanParam : ParameterUIData.ArrayIntegerParameterOption)
 			{
-				FParameterUIData ParameterUIDataBoolean(
-					BooleanParam.Name,
-					BooleanParam.ParamUIMetadata,
+				FMutableParameterData ParameterUIDataBoolean(
+					BooleanParam.Value.ParamUIMetadata,
 					EMutableParameterType::Bool);
 
 				ParameterUIDataBoolean.ParamUIMetadata.ExtraInformation = ParameterUIData.ParamUIMetadata.ExtraInformation;
 
-				GenerationContext.ParameterUIDataMap.Add(BooleanParam.Name, ParameterUIDataBoolean);
+				GenerationContext.ParameterUIDataMap.Add(BooleanParam.Key, ParameterUIDataBoolean);
 			}
 		}
 		else
diff --git a/Engine/Plugins/Experimental/Mutable/Source/CustomizableObjectEditor/Private/MuCOE/GenerateMutableSource/GenerateMutableSource.h b/Engine/Plugins/Experimental/Mutable/Source/CustomizableObjectEditor/Private/MuCOE/GenerateMutableSource/GenerateMutableSource.h
--- a/Engine/Plugins/Experimental/Mutable/Source/CustomizableObjectEditor/Private/MuCOE/GenerateMutableSource/GenerateMutableSource.h
+++ b/Engine/Plugins/Experimental/Mutable/Source/CustomizableObjectEditor/Private/MuCOE/GenerateMutableSource/GenerateMutableSource.h
@@ -45,7 +45,7 @@ class UPhysicsAsset;
 class UTexture2D;
 struct FAnimBpOverridePhysicsAssetsInfo;
 struct FMutableGraphGenerationContext;
-struct FParameterUIData;
+struct FMutableParameterData;
 struct FMutableRefSkeletalMeshData;
 struct FMutableRefSocket;
 struct FMutableSkinWeightProfileInfo;
@@ -702,8 +702,8 @@ struct FMutableGraphGenerationContext
 	TArray<const UCustomizableObjectNode*> NoNameNodeObjectArray;
 	TMap<FString, FCustomizableObjectIdPair> GroupNodeMap;
 	TMap<FString, FString> CustomizableObjectPathMap;
-	TMap<FString, FParameterUIData> ParameterUIDataMap;
-	TMap<FString, FParameterUIData> StateUIDataMap;
+	TMap<FString, FMutableParameterData> ParameterUIDataMap;
+	TMap<FString, FMutableStateData> StateUIDataMap;
 	TMultiMap<const UCustomizableObjectNodeObjectGroup*, FGroupProjectorTempData> ProjectorGroupMap;
 	//TMap<UPhysicsAsset*, uint32> DiscartedPhysicsAssetMap;
 
diff --git a/Engine/Plugins/Experimental/Mutable/Source/CustomizableObjectEditor/Private/MuCOE/GenerateMutableSource/GenerateMutableSourceColor.cpp b/Engine/Plugins/Experimental/Mutable/Source/CustomizableObjectEditor/Private/MuCOE/GenerateMutableSource/GenerateMutableSourceColor.cpp
--- a/Engine/Plugins/Experimental/Mutable/Source/CustomizableObjectEditor/Private/MuCOE/GenerateMutableSource/GenerateMutableSourceColor.cpp
+++ b/Engine/Plugins/Experimental/Mutable/Source/CustomizableObjectEditor/Private/MuCOE/GenerateMutableSource/GenerateMutableSourceColor.cpp
@@ -73,8 +73,7 @@ mu::NodeColourPtr GenerateMutableSourceColor(const UEdGraphPin* Pin, FMutableGra
 		ColorNode->SetUid(GenerationContext.GetNodeIdUnique(Node).ToString());
 		ColorNode->SetDefaultValue(TypedNodeColorParam->DefaultValue);
 
-		GenerationContext.ParameterUIDataMap.Add(TypedNodeColorParam->ParameterName, FParameterUIData(
-			TypedNodeColorParam->ParameterName,
+		GenerationContext.ParameterUIDataMap.Add(TypedNodeColorParam->ParameterName, FMutableParameterData(
 			TypedNodeColorParam->ParamUIMetadata,
 			EMutableParameterType::Color));
 	}
diff --git a/Engine/Plugins/Experimental/Mutable/Source/CustomizableObjectEditor/Private/MuCOE/GenerateMutableSource/GenerateMutableSourceFloat.cpp b/Engine/Plugins/Experimental/Mutable/Source/CustomizableObjectEditor/Private/MuCOE/GenerateMutableSource/GenerateMutableSourceFloat.cpp
--- a/Engine/Plugins/Experimental/Mutable/Source/CustomizableObjectEditor/Private/MuCOE/GenerateMutableSource/GenerateMutableSourceFloat.cpp
+++ b/Engine/Plugins/Experimental/Mutable/Source/CustomizableObjectEditor/Private/MuCOE/GenerateMutableSource/GenerateMutableSourceFloat.cpp
@@ -102,8 +102,7 @@ mu::NodeScalarPtr GenerateMutableSourceFloat(const UEdGraphPin* Pin, FMutableGra
 		ScalarNode->SetUid(GenerationContext.GetNodeIdUnique(Node).ToString());
 		ScalarNode->SetDefaultValue(FloatParameterNode->DefaultValue);
 
-		GenerationContext.ParameterUIDataMap.Add(FloatParameterNode->ParameterName, FParameterUIData(
-			FloatParameterNode->ParameterName,
+		GenerationContext.ParameterUIDataMap.Add(FloatParameterNode->ParameterName, FMutableParameterData(
 			FloatParameterNode->ParamUIMetadata,
 			EMutableParameterType::Float));
 	}
@@ -123,16 +122,16 @@ mu::NodeScalarPtr GenerateMutableSourceFloat(const UEdGraphPin* Pin, FMutableGra
 		EnumParameterNode->SetValueCount(NumSelectors);
 		EnumParameterNode->SetDefaultValueIndex(DefaultValue);
 
-		FParameterUIData ParameterUIData(EnumParamNode->ParameterName, EnumParamNode->ParamUIMetadata, EMutableParameterType::Int);
+		FMutableParameterData ParameterUIData(EnumParamNode->ParamUIMetadata, EMutableParameterType::Int);
 		ParameterUIData.IntegerParameterGroupType = ECustomizableObjectGroupType::COGT_ONE;
 
 		for (int SelectorIndex = 0; SelectorIndex < NumSelectors; ++SelectorIndex)
 		{
 			EnumParameterNode->SetValue(SelectorIndex, (float)SelectorIndex, EnumParamNode->Values[SelectorIndex].Name);
 
-			ParameterUIData.ArrayIntegerParameterOption.Add(FIntegerParameterUIData(
+			ParameterUIData.ArrayIntegerParameterOption.Add(
 				EnumParamNode->Values[SelectorIndex].Name,
-				EnumParamNode->Values[SelectorIndex].ParamUIMetadata));
+				FIntegerParameterUIData(EnumParamNode->Values[SelectorIndex].ParamUIMetadata));
 		}
 
 		Result = EnumParameterNode;
diff --git a/Engine/Plugins/Experimental/Mutable/Source/CustomizableObjectEditor/Private/MuCOE/GenerateMutableSource/GenerateMutableSourceGroupProjector.cpp b/Engine/Plugins/Experimental/Mutable/Source/CustomizableObjectEditor/Private/MuCOE/GenerateMutableSource/GenerateMutableSourceGroupProjector.cpp
--- a/Engine/Plugins/Experimental/Mutable/Source/CustomizableObjectEditor/Private/MuCOE/GenerateMutableSource/GenerateMutableSourceGroupProjector.cpp
+++ b/Engine/Plugins/Experimental/Mutable/Source/CustomizableObjectEditor/Private/MuCOE/GenerateMutableSource/GenerateMutableSourceGroupProjector.cpp
@@ -356,10 +356,8 @@ bool GenerateMutableSourceGroupProjector(const UEdGraphPin* Pin, FMutableGraphGe
 			GenerationContext.AddParameterNameUnique(originalGroup, NodeScalarParamName);
 			NodeScalarParam->SetDefaultValue(0.f);
 
-			GenerationContext.ParameterUIDataMap.Add(NodeScalarParamName, FParameterUIData(
-				NodeScalarParamName,
-				ProjParamNode->ParamUIMetadata,
-				EMutableParameterType::Int));
+			GenerationContext.ParameterUIDataMap.Add(NodeScalarParamName,
+				FMutableParameterData(ProjParamNode->ParamUIMetadata, EMutableParameterType::Int));
 
 			mu::Ptr<mu::NodeRangeFromScalar> NodeRangeFromScalar = new mu::NodeRangeFromScalar;
 			NodeRangeFromScalar->SetSize(NodeScalarParam);
@@ -377,11 +375,10 @@ bool GenerateMutableSourceGroupProjector(const UEdGraphPin* Pin, FMutableGraphGe
 			PoseEnumParameterNode->SetDefaultValueIndex(0);
 			GroupProjectorTempData.PoseOptionsParameter = PoseEnumParameterNode;
 
-			GenerationContext.ParameterUIDataMap.Add(PoseNodeEnumParamName, FParameterUIData(
-				PoseNodeEnumParamName,
-				ProjParamNode->ParamUIMetadata,
-				EMutableParameterType::Int));
-
+			
+			GenerationContext.ParameterUIDataMap.Add(PoseNodeEnumParamName,
+				FMutableParameterData(ProjParamNode->ParamUIMetadata, EMutableParameterType::Int));
+			
 			mu::NodeScalarParameterPtr OpacityParameterNode = new mu::NodeScalarParameter;
 			FString OpacityParameterNodeName = ProjParamNode->ParameterName + FMultilayerProjector::OPACITY_PARAMETER_POSTFIX;
 			OpacityParameterNode->SetName(OpacityParameterNodeName);
@@ -394,11 +391,14 @@ bool GenerateMutableSourceGroupProjector(const UEdGraphPin* Pin, FMutableGraphGe
 			OpacityMetadata.ObjectFriendlyName = FString("Opacity");
 			GroupProjectorTempData.NodeOpacityParameter = OpacityParameterNode;
 
-			GenerationContext.ParameterUIDataMap.Add(OpacityParameterNodeName, FParameterUIData(
-				OpacityParameterNodeName,
-				OpacityMetadata,
-				EMutableParameterType::Float));
-
+			{
+				FMutableParameterData ParameterUIData;
+				ParameterUIData.ParamUIMetadata = OpacityMetadata;
+				ParameterUIData.Type = EMutableParameterType::Float;
+			
+				GenerationContext.ParameterUIDataMap.Add(OpacityParameterNodeName, ParameterUIData);
+			}
+			
 			if (ArrayOptionImage.Num() == 0)
 			{
 				FString msg = FString::Printf(TEXT("The group projection node must have at least one option image connected to a texture or at least one valid element in Option Images Data Table."));
@@ -449,22 +449,26 @@ bool GenerateMutableSourceGroupProjector(const UEdGraphPin* Pin, FMutableGraphGe
 			EnumParameterNode->SetRangeCount(1);
 			EnumParameterNode->SetRange(0, NodeRangeFromScalar);
 
-			FParameterUIData ParameterUIData(NodeEnumParamName, ProjParamNode->ParamUIMetadata, EMutableParameterType::Int);
-			ParameterUIData.IntegerParameterGroupType = ECustomizableObjectGroupType::COGT_ONE;
-			ParameterUIData.ParamUIMetadata.ExtraInformation.Add(FString("UseThumbnails"));
-
-			for (int ImageIndex = 0; ImageIndex < ArrayOptionImage.Num(); ++ImageIndex)
 			{
-				EnumParameterNode->SetValue(ImageIndex, (float)ImageIndex, ArrayOptionImage[ImageIndex].OptionName);
+				FMutableParameterData ParameterUIData;
+				ParameterUIData.ParamUIMetadata = ProjParamNode->ParamUIMetadata;
+				ParameterUIData.Type = EMutableParameterType::Int;
+				ParameterUIData.IntegerParameterGroupType = ECustomizableObjectGroupType::COGT_ONE;
+				ParameterUIData.ParamUIMetadata.ExtraInformation.Add(FString("UseThumbnails"));
 
-				FMutableParamUIMetadata optionMetadata = ParameterUIData.ParamUIMetadata;
-				optionMetadata.UIThumbnail = ArrayOptionImage[ImageIndex].OptionImage;
-				ParameterUIData.ArrayIntegerParameterOption.Add(FIntegerParameterUIData(
-					ArrayOptionImage[ImageIndex].OptionName,
-					optionMetadata));
-			}
+				for (int ImageIndex = 0; ImageIndex < ArrayOptionImage.Num(); ++ImageIndex)
+				{
+					EnumParameterNode->SetValue(ImageIndex, (float)ImageIndex, ArrayOptionImage[ImageIndex].OptionName);
 
-			GenerationContext.ParameterUIDataMap.Add(NodeEnumParamName, ParameterUIData);
+					FMutableParamUIMetadata OptionMetadata = ParameterUIData.ParamUIMetadata;
+					OptionMetadata.UIThumbnail = ArrayOptionImage[ImageIndex].OptionImage;
+					ParameterUIData.ArrayIntegerParameterOption.Add(
+						ArrayOptionImage[ImageIndex].OptionName,
+						FIntegerParameterUIData(OptionMetadata));
+				}
+
+				GenerationContext.ParameterUIDataMap.Add(NodeEnumParamName, ParameterUIData);
+			}
 
 			mu::NodeImageSwitchPtr SwitchNode = new mu::NodeImageSwitch;
 			SwitchNode->SetParameter(EnumParameterNode);
diff --git a/Engine/Plugins/Experimental/Mutable/Source/CustomizableObjectEditor/Private/MuCOE/GenerateMutableSource/GenerateMutableSourceImage.cpp b/Engine/Plugins/Experimental/Mutable/Source/CustomizableObjectEditor/Private/MuCOE/GenerateMutableSource/GenerateMutableSourceImage.cpp
--- a/Engine/Plugins/Experimental/Mutable/Source/CustomizableObjectEditor/Private/MuCOE/GenerateMutableSource/GenerateMutableSourceImage.cpp
+++ b/Engine/Plugins/Experimental/Mutable/Source/CustomizableObjectEditor/Private/MuCOE/GenerateMutableSource/GenerateMutableSourceImage.cpp
@@ -145,8 +145,7 @@ mu::NodeImagePtr GenerateMutableSourceImage(const UEdGraphPin* Pin, FMutableGrap
 			TextureNode->SetDefaultValue(FName(TypedNodeParam->DefaultValue->GetPathName()));			
 		}
 
-		GenerationContext.ParameterUIDataMap.Add(TypedNodeParam->ParameterName, FParameterUIData(
-			TypedNodeParam->ParameterName,
+		GenerationContext.ParameterUIDataMap.Add(TypedNodeParam->ParameterName, FMutableParameterData(
 			TypedNodeParam->ParamUIMetadata,
 			EMutableParameterType::Texture));
 
diff --git a/Engine/Plugins/Experimental/Mutable/Source/CustomizableObjectEditor/Private/MuCOE/GenerateMutableSource/GenerateMutableSourceProjector.cpp b/Engine/Plugins/Experimental/Mutable/Source/CustomizableObjectEditor/Private/MuCOE/GenerateMutableSource/GenerateMutableSourceProjector.cpp
--- a/Engine/Plugins/Experimental/Mutable/Source/CustomizableObjectEditor/Private/MuCOE/GenerateMutableSource/GenerateMutableSourceProjector.cpp
+++ b/Engine/Plugins/Experimental/Mutable/Source/CustomizableObjectEditor/Private/MuCOE/GenerateMutableSource/GenerateMutableSourceProjector.cpp
@@ -130,8 +130,7 @@ mu::NodeProjectorPtr GenerateMutableSourceProjector(const UEdGraphPin* Pin, FMut
 			check(false);
 		}
 
-		GenerationContext.ParameterUIDataMap.Add(TypedNodeParam->ParameterName, FParameterUIData(
-			TypedNodeParam->ParameterName,
+		GenerationContext.ParameterUIDataMap.Add(TypedNodeParam->ParameterName, FMutableParameterData(
 			TypedNodeParam->ParamUIMetadata,
 			EMutableParameterType::Projector));
 	}
diff --git a/Engine/Plugins/Experimental/Mutable/Source/CustomizableObjectEditor/Private/MuCOE/GenerateMutableSource/GenerateMutableSourceTable.cpp b/Engine/Plugins/Experimental/Mutable/Source/CustomizableObjectEditor/Private/MuCOE/GenerateMutableSource/GenerateMutableSourceTable.cpp
--- a/Engine/Plugins/Experimental/Mutable/Source/CustomizableObjectEditor/Private/MuCOE/GenerateMutableSource/GenerateMutableSourceTable.cpp
+++ b/Engine/Plugins/Experimental/Mutable/Source/CustomizableObjectEditor/Private/MuCOE/GenerateMutableSource/GenerateMutableSourceTable.cpp
@@ -858,9 +858,9 @@ void GenerateTableParameterUIData(const UDataTable* DataTable, const UCustomizab
 		// Getting Table and row names to access the information
 		TArray<FName> RowNames = GetRowsToCompile(*DataTable, *TableNode, GenerationContext);
 
-		FParameterUIData ParameterUIData(TableNode->ParameterName, TableNode->ParamUIMetadata, EMutableParameterType::Int);
+		FMutableParameterData ParameterUIData(TableNode->ParamUIMetadata, EMutableParameterType::Int);
 		ParameterUIData.IntegerParameterGroupType = TableNode->bAddNoneOption ? ECustomizableObjectGroupType::COGT_ONE_OR_NONE : ECustomizableObjectGroupType::COGT_ONE;
-		FParameterUIData& ParameterUIDataRef = GenerationContext.ParameterUIDataMap.Add(TableNode->ParameterName, ParameterUIData);
+		FMutableParameterData& ParameterUIDataRef = GenerationContext.ParameterUIDataMap.Add(TableNode->ParameterName, ParameterUIData);
 
 		if (TableNode->ParamUIMetadataColumn.IsNone())
 		{
@@ -893,7 +893,7 @@ void GenerateTableParameterUIData(const UDataTable* DataTable, const UCustomizab
 				if (uint8* CellData = GetCellData(RowNames[NameIndex], *DataTable, *ColumnProperty))
 				{
 					FMutableParamUIMetadata Value = *(FMutableParamUIMetadata*)CellData;
-					ParameterUIDataRef.ArrayIntegerParameterOption.Add(FIntegerParameterUIData(RowNames[NameIndex].ToString(), Value));
+					ParameterUIDataRef.ArrayIntegerParameterOption.Add(RowNames[NameIndex].ToString(), FIntegerParameterUIData(Value));
 				}
 			}
 		}
diff --git a/Engine/Plugins/Experimental/Mutable/Source/CustomizableObjectEditor/Private/MuCOE/Nodes/CustomizableObjectNode.cpp b/Engine/Plugins/Experimental/Mutable/Source/CustomizableObjectEditor/Private/MuCOE/Nodes/CustomizableObjectNode.cpp
--- a/Engine/Plugins/Experimental/Mutable/Source/CustomizableObjectEditor/Private/MuCOE/Nodes/CustomizableObjectNode.cpp
+++ b/Engine/Plugins/Experimental/Mutable/Source/CustomizableObjectEditor/Private/MuCOE/Nodes/CustomizableObjectNode.cpp
@@ -3,6 +3,7 @@
 #include "MuCOE/Nodes/CustomizableObjectNode.h"
 
 #include "Containers/Queue.h"
+#include "MuCO/CustomizableObjectCustomVersion.h"
 #include "MuCOE/CustomizableObjectEditor.h"
 #include "MuCOE/CustomizableObjectEditorLogger.h"
 #include "MuCOE/CustomizableObjectGraph.h"
diff --git a/Engine/Plugins/Experimental/Mutable/Source/CustomizableObjectEditor/Private/MuCOE/Nodes/CustomizableObjectNodeColorSwitch.cpp b/Engine/Plugins/Experimental/Mutable/Source/CustomizableObjectEditor/Private/MuCOE/Nodes/CustomizableObjectNodeColorSwitch.cpp
--- a/Engine/Plugins/Experimental/Mutable/Source/CustomizableObjectEditor/Private/MuCOE/Nodes/CustomizableObjectNodeColorSwitch.cpp
+++ b/Engine/Plugins/Experimental/Mutable/Source/CustomizableObjectEditor/Private/MuCOE/Nodes/CustomizableObjectNodeColorSwitch.cpp
@@ -2,6 +2,7 @@
 
 #include "MuCOE/Nodes/CustomizableObjectNodeColorSwitch.h"
 
+#include "MuCO/CustomizableObjectCustomVersion.h"
 #include "MuCOE/EdGraphSchema_CustomizableObject.h"
 
 #define LOCTEXT_NAMESPACE "CustomizableObjectEditor"
diff --git a/Engine/Plugins/Experimental/Mutable/Source/CustomizableObjectEditor/Private/MuCOE/Nodes/CustomizableObjectNodeColorVariation.cpp b/Engine/Plugins/Experimental/Mutable/Source/CustomizableObjectEditor/Private/MuCOE/Nodes/CustomizableObjectNodeColorVariation.cpp
--- a/Engine/Plugins/Experimental/Mutable/Source/CustomizableObjectEditor/Private/MuCOE/Nodes/CustomizableObjectNodeColorVariation.cpp
+++ b/Engine/Plugins/Experimental/Mutable/Source/CustomizableObjectEditor/Private/MuCOE/Nodes/CustomizableObjectNodeColorVariation.cpp
@@ -2,6 +2,7 @@
 
 #include "MuCOE/Nodes/CustomizableObjectNodeColorVariation.h"
 
+#include "MuCO/CustomizableObjectCustomVersion.h"
 #include "MuCOE/CustomizableObjectEditor_Deprecated.h"
 #include "MuCOE/EdGraphSchema_CustomizableObject.h"
 
diff --git a/Engine/Plugins/Experimental/Mutable/Source/CustomizableObjectEditor/Private/MuCOE/Nodes/CustomizableObjectNodeCurve.cpp b/Engine/Plugins/Experimental/Mutable/Source/CustomizableObjectEditor/Private/MuCOE/Nodes/CustomizableObjectNodeCurve.cpp
--- a/Engine/Plugins/Experimental/Mutable/Source/CustomizableObjectEditor/Private/MuCOE/Nodes/CustomizableObjectNodeCurve.cpp
+++ b/Engine/Plugins/Experimental/Mutable/Source/CustomizableObjectEditor/Private/MuCOE/Nodes/CustomizableObjectNodeCurve.cpp
@@ -3,6 +3,7 @@
 #include "MuCOE/Nodes/CustomizableObjectNodeCurve.h"
 
 #include "Curves/CurveBase.h"
+#include "MuCO/CustomizableObjectCustomVersion.h"
 #include "MuCOE/EdGraphSchema_CustomizableObject.h"
 #include "MuCOE/RemapPins/CustomizableObjectNodeRemapPinsByPosition.h"
 #include "MuCOE/UnrealEditorPortabilityHelpers.h"
diff --git a/Engine/Plugins/Experimental/Mutable/Source/CustomizableObjectEditor/Private/MuCOE/Nodes/CustomizableObjectNodeEditMaterial.cpp b/Engine/Plugins/Experimental/Mutable/Source/CustomizableObjectEditor/Private/MuCOE/Nodes/CustomizableObjectNodeEditMaterial.cpp
--- a/Engine/Plugins/Experimental/Mutable/Source/CustomizableObjectEditor/Private/MuCOE/Nodes/CustomizableObjectNodeEditMaterial.cpp
+++ b/Engine/Plugins/Experimental/Mutable/Source/CustomizableObjectEditor/Private/MuCOE/Nodes/CustomizableObjectNodeEditMaterial.cpp
@@ -3,6 +3,7 @@
 #include "MuCOE/Nodes/CustomizableObjectNodeEditMaterial.h"
 
 #include "MaterialTypes.h"
+#include "MuCO/CustomizableObjectCustomVersion.h"
 #include "MuCOE/CustomizableObjectEditor_Deprecated.h"
 #include "MuCOE/CustomizableObjectLayout.h"
 #include "MuCOE/EdGraphSchema_CustomizableObject.h"
diff --git a/Engine/Plugins/Experimental/Mutable/Source/CustomizableObjectEditor/Private/MuCOE/Nodes/CustomizableObjectNodeExtendMaterial.cpp b/Engine/Plugins/Experimental/Mutable/Source/CustomizableObjectEditor/Private/MuCOE/Nodes/CustomizableObjectNodeExtendMaterial.cpp
--- a/Engine/Plugins/Experimental/Mutable/Source/CustomizableObjectEditor/Private/MuCOE/Nodes/CustomizableObjectNodeExtendMaterial.cpp
+++ b/Engine/Plugins/Experimental/Mutable/Source/CustomizableObjectEditor/Private/MuCOE/Nodes/CustomizableObjectNodeExtendMaterial.cpp
@@ -3,6 +3,7 @@
 #include "MuCOE/Nodes/CustomizableObjectNodeExtendMaterial.h"
 
 #include "MaterialTypes.h"
+#include "MuCO/CustomizableObjectCustomVersion.h"
 #include "MuCOE/CustomizableObjectEditor_Deprecated.h"
 #include "MuCOE/CustomizableObjectGraph.h"
 #include "MuCOE/EdGraphSchema_CustomizableObject.h"
diff --git a/Engine/Plugins/Experimental/Mutable/Source/CustomizableObjectEditor/Private/MuCOE/Nodes/CustomizableObjectNodeExternalPin.cpp b/Engine/Plugins/Experimental/Mutable/Source/CustomizableObjectEditor/Private/MuCOE/Nodes/CustomizableObjectNodeExternalPin.cpp
--- a/Engine/Plugins/Experimental/Mutable/Source/CustomizableObjectEditor/Private/MuCOE/Nodes/CustomizableObjectNodeExternalPin.cpp
+++ b/Engine/Plugins/Experimental/Mutable/Source/CustomizableObjectEditor/Private/MuCOE/Nodes/CustomizableObjectNodeExternalPin.cpp
@@ -2,6 +2,7 @@
 
 #include "MuCOE/Nodes/CustomizableObjectNodeExternalPin.h"
 
+#include "MuCO/CustomizableObjectCustomVersion.h"
 #include "MuCOE/CustomizableObjectGraph.h"
 #include "MuCOE/CustomizableObjectPin.h"
 #include "MuCOE/EdGraphSchema_CustomizableObject.h"
diff --git a/Engine/Plugins/Experimental/Mutable/Source/CustomizableObjectEditor/Private/MuCOE/Nodes/CustomizableObjectNodeFloatParameter.cpp b/Engine/Plugins/Experimental/Mutable/Source/CustomizableObjectEditor/Private/MuCOE/Nodes/CustomizableObjectNodeFloatParameter.cpp
--- a/Engine/Plugins/Experimental/Mutable/Source/CustomizableObjectEditor/Private/MuCOE/Nodes/CustomizableObjectNodeFloatParameter.cpp
+++ b/Engine/Plugins/Experimental/Mutable/Source/CustomizableObjectEditor/Private/MuCOE/Nodes/CustomizableObjectNodeFloatParameter.cpp
@@ -2,6 +2,7 @@
 
 #include "MuCOE/Nodes/CustomizableObjectNodeFloatParameter.h"
 
+#include "MuCO/CustomizableObjectCustomVersion.h"
 #include "MuCOE/EdGraphSchema_CustomizableObject.h"
 #include "MuCOE/UnrealEditorPortabilityHelpers.h"
 
diff --git a/Engine/Plugins/Experimental/Mutable/Source/CustomizableObjectEditor/Private/MuCOE/Nodes/CustomizableObjectNodeFloatSwitch.cpp b/Engine/Plugins/Experimental/Mutable/Source/CustomizableObjectEditor/Private/MuCOE/Nodes/CustomizableObjectNodeFloatSwitch.cpp
--- a/Engine/Plugins/Experimental/Mutable/Source/CustomizableObjectEditor/Private/MuCOE/Nodes/CustomizableObjectNodeFloatSwitch.cpp
+++ b/Engine/Plugins/Experimental/Mutable/Source/CustomizableObjectEditor/Private/MuCOE/Nodes/CustomizableObjectNodeFloatSwitch.cpp
@@ -2,6 +2,7 @@
 
 #include "MuCOE/Nodes/CustomizableObjectNodeFloatSwitch.h"
 
+#include "MuCO/CustomizableObjectCustomVersion.h"
 #include "MuCOE/EdGraphSchema_CustomizableObject.h"
 
 #define LOCTEXT_NAMESPACE "CustomizableObjectEditor"
diff --git a/Engine/Plugins/Experimental/Mutable/Source/CustomizableObjectEditor/Private/MuCOE/Nodes/CustomizableObjectNodeFloatVariation.cpp b/Engine/Plugins/Experimental/Mutable/Source/CustomizableObjectEditor/Private/MuCOE/Nodes/CustomizableObjectNodeFloatVariation.cpp
--- a/Engine/Plugins/Experimental/Mutable/Source/CustomizableObjectEditor/Private/MuCOE/Nodes/CustomizableObjectNodeFloatVariation.cpp
+++ b/Engine/Plugins/Experimental/Mutable/Source/CustomizableObjectEditor/Private/MuCOE/Nodes/CustomizableObjectNodeFloatVariation.cpp
@@ -2,6 +2,7 @@
 
 #include "MuCOE/Nodes/CustomizableObjectNodeFloatVariation.h"
 
+#include "MuCO/CustomizableObjectCustomVersion.h"
 #include "MuCOE/CustomizableObjectEditor_Deprecated.h"
 #include "MuCOE/EdGraphSchema_CustomizableObject.h"
 
diff --git a/Engine/Plugins/Experimental/Mutable/Source/CustomizableObjectEditor/Private/MuCOE/Nodes/CustomizableObjectNodeGroupProjectorParameter.cpp b/Engine/Plugins/Experimental/Mutable/Source/CustomizableObjectEditor/Private/MuCOE/Nodes/CustomizableObjectNodeGroupProjectorParameter.cpp
--- a/Engine/Plugins/Experimental/Mutable/Source/CustomizableObjectEditor/Private/MuCOE/Nodes/CustomizableObjectNodeGroupProjectorParameter.cpp
+++ b/Engine/Plugins/Experimental/Mutable/Source/CustomizableObjectEditor/Private/MuCOE/Nodes/CustomizableObjectNodeGroupProjectorParameter.cpp
@@ -2,6 +2,7 @@
 
 #include "MuCOE/Nodes/CustomizableObjectNodeGroupProjectorParameter.h"
 
+#include "MuCO/CustomizableObjectCustomVersion.h"
 #include "MuCOE/EdGraphSchema_CustomizableObject.h"
 #include "MuCOE/UnrealEditorPortabilityHelpers.h"
 
diff --git a/Engine/Plugins/Experimental/Mutable/Source/CustomizableObjectEditor/Private/MuCOE/Nodes/CustomizableObjectNodeLayoutBlocks.cpp b/Engine/Plugins/Experimental/Mutable/Source/CustomizableObjectEditor/Private/MuCOE/Nodes/CustomizableObjectNodeLayoutBlocks.cpp
--- a/Engine/Plugins/Experimental/Mutable/Source/CustomizableObjectEditor/Private/MuCOE/Nodes/CustomizableObjectNodeLayoutBlocks.cpp
+++ b/Engine/Plugins/Experimental/Mutable/Source/CustomizableObjectEditor/Private/MuCOE/Nodes/CustomizableObjectNodeLayoutBlocks.cpp
@@ -2,6 +2,7 @@
 
 #include "MuCOE/Nodes/CustomizableObjectNodeLayoutBlocks.h"
 
+#include "MuCO/CustomizableObjectCustomVersion.h"
 #include "MuCOE/CustomizableObjectLayout.h"
 #include "MuCOE/EdGraphSchema_CustomizableObject.h"
 #include "MuCOE/GraphTraversal.h"
diff --git a/Engine/Plugins/Experimental/Mutable/Source/CustomizableObjectEditor/Private/MuCOE/Nodes/CustomizableObjectNodeMaterial.cpp b/Engine/Plugins/Experimental/Mutable/Source/CustomizableObjectEditor/Private/MuCOE/Nodes/CustomizableObjectNodeMaterial.cpp
--- a/Engine/Plugins/Experimental/Mutable/Source/CustomizableObjectEditor/Private/MuCOE/Nodes/CustomizableObjectNodeMaterial.cpp
+++ b/Engine/Plugins/Experimental/Mutable/Source/CustomizableObjectEditor/Private/MuCOE/Nodes/CustomizableObjectNodeMaterial.cpp
@@ -19,6 +19,7 @@
 #include "ObjectEditorUtils.h"
 #include "PropertyCustomizationHelpers.h"
 #include "Modules/ModuleManager.h"
+#include "MuCO/CustomizableObjectCustomVersion.h"
 #include "MuCOE/CustomizableObjectEditorLogger.h"
 #include "MuCOE/Nodes/CustomizableObjectNodeExternalPin.h"
 
diff --git a/Engine/Plugins/Experimental/Mutable/Source/CustomizableObjectEditor/Private/MuCOE/Nodes/CustomizableObjectNodeMaterialVariation.cpp b/Engine/Plugins/Experimental/Mutable/Source/CustomizableObjectEditor/Private/MuCOE/Nodes/CustomizableObjectNodeMaterialVariation.cpp
--- a/Engine/Plugins/Experimental/Mutable/Source/CustomizableObjectEditor/Private/MuCOE/Nodes/CustomizableObjectNodeMaterialVariation.cpp
+++ b/Engine/Plugins/Experimental/Mutable/Source/CustomizableObjectEditor/Private/MuCOE/Nodes/CustomizableObjectNodeMaterialVariation.cpp
@@ -2,6 +2,7 @@
 
 #include "MuCOE/Nodes/CustomizableObjectNodeMaterialVariation.h"
 
+#include "MuCO/CustomizableObjectCustomVersion.h"
 #include "MuCOE/CustomizableObjectEditor_Deprecated.h"
 #include "MuCOE/EdGraphSchema_CustomizableObject.h"
 
diff --git a/Engine/Plugins/Experimental/Mutable/Source/CustomizableObjectEditor/Private/MuCOE/Nodes/CustomizableObjectNodeMeshClipMorph.cpp b/Engine/Plugins/Experimental/Mutable/Source/CustomizableObjectEditor/Private/MuCOE/Nodes/CustomizableObjectNodeMeshClipMorph.cpp
--- a/Engine/Plugins/Experimental/Mutable/Source/CustomizableObjectEditor/Private/MuCOE/Nodes/CustomizableObjectNodeMeshClipMorph.cpp
+++ b/Engine/Plugins/Experimental/Mutable/Source/CustomizableObjectEditor/Private/MuCOE/Nodes/CustomizableObjectNodeMeshClipMorph.cpp
@@ -2,6 +2,7 @@
 
 #include "MuCOE/Nodes/CustomizableObjectNodeMeshClipMorph.h"
 
+#include "MuCO/CustomizableObjectCustomVersion.h"
 #include "MuCOE/EdGraphSchema_CustomizableObject.h"
 #include "MuCOE/ICustomizableObjectEditor.h"
 #include "MuCOE/UnrealEditorPortabilityHelpers.h"
diff --git a/Engine/Plugins/Experimental/Mutable/Source/CustomizableObjectEditor/Private/MuCOE/Nodes/CustomizableObjectNodeMeshMorph.cpp b/Engine/Plugins/Experimental/Mutable/Source/CustomizableObjectEditor/Private/MuCOE/Nodes/CustomizableObjectNodeMeshMorph.cpp
--- a/Engine/Plugins/Experimental/Mutable/Source/CustomizableObjectEditor/Private/MuCOE/Nodes/CustomizableObjectNodeMeshMorph.cpp
+++ b/Engine/Plugins/Experimental/Mutable/Source/CustomizableObjectEditor/Private/MuCOE/Nodes/CustomizableObjectNodeMeshMorph.cpp
@@ -3,6 +3,7 @@
 #include "MuCOE/Nodes/CustomizableObjectNodeMeshMorph.h"
 
 #include "Engine/SkeletalMesh.h"
+#include "MuCO/CustomizableObjectCustomVersion.h"
 #include "MuCOE/EdGraphSchema_CustomizableObject.h"
 #include "MuCOE/GraphTraversal.h"
 #include "MuCOE/ICustomizableObjectEditor.h"
diff --git a/Engine/Plugins/Experimental/Mutable/Source/CustomizableObjectEditor/Private/MuCOE/Nodes/CustomizableObjectNodeMeshReshape.cpp b/Engine/Plugins/Experimental/Mutable/Source/CustomizableObjectEditor/Private/MuCOE/Nodes/CustomizableObjectNodeMeshReshape.cpp
--- a/Engine/Plugins/Experimental/Mutable/Source/CustomizableObjectEditor/Private/MuCOE/Nodes/CustomizableObjectNodeMeshReshape.cpp
+++ b/Engine/Plugins/Experimental/Mutable/Source/CustomizableObjectEditor/Private/MuCOE/Nodes/CustomizableObjectNodeMeshReshape.cpp
@@ -2,6 +2,7 @@
 
 #include "MuCOE/Nodes/CustomizableObjectNodeMeshReshape.h"
 
+#include "MuCO/CustomizableObjectCustomVersion.h"
 #include "MuCOE/EdGraphSchema_CustomizableObject.h"
 #include "MuCOE/Nodes/CustomizableObjectNodeMeshReshapeCommon.h"
 
diff --git a/Engine/Plugins/Experimental/Mutable/Source/CustomizableObjectEditor/Private/MuCOE/Nodes/CustomizableObjectNodeMeshSwitch.cpp b/Engine/Plugins/Experimental/Mutable/Source/CustomizableObjectEditor/Private/MuCOE/Nodes/CustomizableObjectNodeMeshSwitch.cpp
--- a/Engine/Plugins/Experimental/Mutable/Source/CustomizableObjectEditor/Private/MuCOE/Nodes/CustomizableObjectNodeMeshSwitch.cpp
+++ b/Engine/Plugins/Experimental/Mutable/Source/CustomizableObjectEditor/Private/MuCOE/Nodes/CustomizableObjectNodeMeshSwitch.cpp
@@ -2,6 +2,7 @@
 
 #include "MuCOE/Nodes/CustomizableObjectNodeMeshSwitch.h"
 
+#include "MuCO/CustomizableObjectCustomVersion.h"
 #include "MuCOE/EdGraphSchema_CustomizableObject.h"
 
 #define LOCTEXT_NAMESPACE "CustomizableObjectEditor"
diff --git a/Engine/Plugins/Experimental/Mutable/Source/CustomizableObjectEditor/Private/MuCOE/Nodes/CustomizableObjectNodeMeshVariation.cpp b/Engine/Plugins/Experimental/Mutable/Source/CustomizableObjectEditor/Private/MuCOE/Nodes/CustomizableObjectNodeMeshVariation.cpp
--- a/Engine/Plugins/Experimental/Mutable/Source/CustomizableObjectEditor/Private/MuCOE/Nodes/CustomizableObjectNodeMeshVariation.cpp
+++ b/Engine/Plugins/Experimental/Mutable/Source/CustomizableObjectEditor/Private/MuCOE/Nodes/CustomizableObjectNodeMeshVariation.cpp
@@ -2,6 +2,7 @@
 
 #include "MuCOE/Nodes/CustomizableObjectNodeMeshVariation.h"
 
+#include "MuCO/CustomizableObjectCustomVersion.h"
 #include "MuCOE/CustomizableObjectEditor_Deprecated.h"
 #include "MuCOE/EdGraphSchema_CustomizableObject.h"
 
diff --git a/Engine/Plugins/Experimental/Mutable/Source/CustomizableObjectEditor/Private/MuCOE/Nodes/CustomizableObjectNodeMorphMaterial.cpp b/Engine/Plugins/Experimental/Mutable/Source/CustomizableObjectEditor/Private/MuCOE/Nodes/CustomizableObjectNodeMorphMaterial.cpp
--- a/Engine/Plugins/Experimental/Mutable/Source/CustomizableObjectEditor/Private/MuCOE/Nodes/CustomizableObjectNodeMorphMaterial.cpp
+++ b/Engine/Plugins/Experimental/Mutable/Source/CustomizableObjectEditor/Private/MuCOE/Nodes/CustomizableObjectNodeMorphMaterial.cpp
@@ -3,6 +3,7 @@
 #include "MuCOE/Nodes/CustomizableObjectNodeMorphMaterial.h"
 
 #include "Engine/SkeletalMesh.h"
+#include "MuCO/CustomizableObjectCustomVersion.h"
 #include "MuCOE/EdGraphSchema_CustomizableObject.h"
 #include "MuCOE/GraphTraversal.h"
 #include "MuCOE/ICustomizableObjectEditor.h"
diff --git a/Engine/Plugins/Experimental/Mutable/Source/CustomizableObjectEditor/Private/MuCOE/Nodes/CustomizableObjectNodeObject.cpp b/Engine/Plugins/Experimental/Mutable/Source/CustomizableObjectEditor/Private/MuCOE/Nodes/CustomizableObjectNodeObject.cpp
--- a/Engine/Plugins/Experimental/Mutable/Source/CustomizableObjectEditor/Private/MuCOE/Nodes/CustomizableObjectNodeObject.cpp
+++ b/Engine/Plugins/Experimental/Mutable/Source/CustomizableObjectEditor/Private/MuCOE/Nodes/CustomizableObjectNodeObject.cpp
@@ -15,6 +15,7 @@
 #include "Misc/UObjectToken.h"
 #include "Logging/MessageLog.h"
 #include "Containers/Queue.h"
+#include "MuCO/CustomizableObjectCustomVersion.h"
 
 class UCustomizableObjectNodeRemapPins;
 
@@ -78,6 +79,19 @@ void UCustomizableObjectNodeObject::BackwardsCompatibleFixup()
 			s.bDisableTextureStreaming = s.TextureCompressionStrategy != ETextureCompressionStrategy::None;
 		}
 	}
+
+	if (CustomizableObjectCustomVersion < FCustomizableObjectCustomVersion::StateUIMetadata)
+	{
+		for (FCustomizableObjectState& State : States)
+		{
+			State.UIMetadata.ObjectFriendlyName = State.StateUIMetadata_DEPRECATED.ObjectFriendlyName;
+			State.UIMetadata.UISectionName = State.StateUIMetadata_DEPRECATED.UISectionName;
+			State.UIMetadata.UIOrder = State.StateUIMetadata_DEPRECATED.UIOrder;
+			State.UIMetadata.UIThumbnail = State.StateUIMetadata_DEPRECATED.UIThumbnail;
+			State.UIMetadata.ExtraInformation = State.StateUIMetadata_DEPRECATED.ExtraInformation;
+			State.UIMetadata.ExtraAssets = State.StateUIMetadata_DEPRECATED.ExtraAssets;
+		}
+	}
 }
 
 
diff --git a/Engine/Plugins/Experimental/Mutable/Source/CustomizableObjectEditor/Private/MuCOE/Nodes/CustomizableObjectNodeObjectGroup.cpp b/Engine/Plugins/Experimental/Mutable/Source/CustomizableObjectEditor/Private/MuCOE/Nodes/CustomizableObjectNodeObjectGroup.cpp
--- a/Engine/Plugins/Experimental/Mutable/Source/CustomizableObjectEditor/Private/MuCOE/Nodes/CustomizableObjectNodeObjectGroup.cpp
+++ b/Engine/Plugins/Experimental/Mutable/Source/CustomizableObjectEditor/Private/MuCOE/Nodes/CustomizableObjectNodeObjectGroup.cpp
@@ -2,6 +2,7 @@
 
 #include "MuCOE/Nodes/CustomizableObjectNodeObjectGroup.h"
 
+#include "MuCO/CustomizableObjectCustomVersion.h"
 #include "MuCOE/EdGraphSchema_CustomizableObject.h"
 #include "MuCOE/GraphTraversal.h"
 #include "MuCOE/Nodes/CustomizableObjectNodeObject.h"
diff --git a/Engine/Plugins/Experimental/Mutable/Source/CustomizableObjectEditor/Private/MuCOE/Nodes/CustomizableObjectNodePassThroughTexture.cpp b/Engine/Plugins/Experimental/Mutable/Source/CustomizableObjectEditor/Private/MuCOE/Nodes/CustomizableObjectNodePassThroughTexture.cpp
--- a/Engine/Plugins/Experimental/Mutable/Source/CustomizableObjectEditor/Private/MuCOE/Nodes/CustomizableObjectNodePassThroughTexture.cpp
+++ b/Engine/Plugins/Experimental/Mutable/Source/CustomizableObjectEditor/Private/MuCOE/Nodes/CustomizableObjectNodePassThroughTexture.cpp
@@ -2,6 +2,7 @@
 
 #include "MuCOE/Nodes/CustomizableObjectNodePassThroughTexture.h"
 
+#include "MuCO/CustomizableObjectCustomVersion.h"
 #include "MuCOE/EdGraphSchema_CustomizableObject.h"
 
 #define LOCTEXT_NAMESPACE "CustomizableObjectEditor"
diff --git a/Engine/Plugins/Experimental/Mutable/Source/CustomizableObjectEditor/Private/MuCOE/Nodes/CustomizableObjectNodeProjectorConstant.cpp b/Engine/Plugins/Experimental/Mutable/Source/CustomizableObjectEditor/Private/MuCOE/Nodes/CustomizableObjectNodeProjectorConstant.cpp
--- a/Engine/Plugins/Experimental/Mutable/Source/CustomizableObjectEditor/Private/MuCOE/Nodes/CustomizableObjectNodeProjectorConstant.cpp
+++ b/Engine/Plugins/Experimental/Mutable/Source/CustomizableObjectEditor/Private/MuCOE/Nodes/CustomizableObjectNodeProjectorConstant.cpp
@@ -2,6 +2,7 @@
 
 #include "MuCOE/Nodes/CustomizableObjectNodeProjectorConstant.h"
 
+#include "MuCO/CustomizableObjectCustomVersion.h"
 #include "MuCOE/EdGraphSchema_CustomizableObject.h"
 #include "MuCOE/UnrealEditorPortabilityHelpers.h"
 
diff --git a/Engine/Plugins/Experimental/Mutable/Source/CustomizableObjectEditor/Private/MuCOE/Nodes/CustomizableObjectNodeProjectorParameter.cpp b/Engine/Plugins/Experimental/Mutable/Source/CustomizableObjectEditor/Private/MuCOE/Nodes/CustomizableObjectNodeProjectorParameter.cpp
--- a/Engine/Plugins/Experimental/Mutable/Source/CustomizableObjectEditor/Private/MuCOE/Nodes/CustomizableObjectNodeProjectorParameter.cpp
+++ b/Engine/Plugins/Experimental/Mutable/Source/CustomizableObjectEditor/Private/MuCOE/Nodes/CustomizableObjectNodeProjectorParameter.cpp
@@ -2,6 +2,7 @@
 
 #include "MuCOE/Nodes/CustomizableObjectNodeProjectorParameter.h"
 
+#include "MuCO/CustomizableObjectCustomVersion.h"
 #include "MuCOE/EdGraphSchema_CustomizableObject.h"
 #include "MuCOE/UnrealEditorPortabilityHelpers.h"
 
diff --git a/Engine/Plugins/Experimental/Mutable/Source/CustomizableObjectEditor/Private/MuCOE/Nodes/CustomizableObjectNodeRemoveMeshBlocks.cpp b/Engine/Plugins/Experimental/Mutable/Source/CustomizableObjectEditor/Private/MuCOE/Nodes/CustomizableObjectNodeRemoveMeshBlocks.cpp
--- a/Engine/Plugins/Experimental/Mutable/Source/CustomizableObjectEditor/Private/MuCOE/Nodes/CustomizableObjectNodeRemoveMeshBlocks.cpp
+++ b/Engine/Plugins/Experimental/Mutable/Source/CustomizableObjectEditor/Private/MuCOE/Nodes/CustomizableObjectNodeRemoveMeshBlocks.cpp
@@ -2,6 +2,7 @@
 
 #include "MuCOE/Nodes/CustomizableObjectNodeRemoveMeshBlocks.h"
 
+#include "MuCO/CustomizableObjectCustomVersion.h"
 #include "MuCOE/CustomizableObjectLayout.h"
 #include "MuCOE/EdGraphSchema_CustomizableObject.h"
 #include "MuCOE/ICustomizableObjectEditor.h"
diff --git a/Engine/Plugins/Experimental/Mutable/Source/CustomizableObjectEditor/Private/MuCOE/Nodes/CustomizableObjectNodeSkeletalMesh.cpp b/Engine/Plugins/Experimental/Mutable/Source/CustomizableObjectEditor/Private/MuCOE/Nodes/CustomizableObjectNodeSkeletalMesh.cpp
--- a/Engine/Plugins/Experimental/Mutable/Source/CustomizableObjectEditor/Private/MuCOE/Nodes/CustomizableObjectNodeSkeletalMesh.cpp
+++ b/Engine/Plugins/Experimental/Mutable/Source/CustomizableObjectEditor/Private/MuCOE/Nodes/CustomizableObjectNodeSkeletalMesh.cpp
@@ -6,6 +6,7 @@
 #include "AssetThumbnail.h"
 #include "ISinglePropertyView.h"
 #include "Modules/ModuleManager.h"
+#include "MuCO/CustomizableObjectCustomVersion.h"
 #include "MuCOE/CustomizableObjectEditor.h"
 #include "MuCOE/CustomizableObjectEditorStyle.h"
 #include "MuCOE/CustomizableObjectLayout.h"
diff --git a/Engine/Plugins/Experimental/Mutable/Source/CustomizableObjectEditor/Private/MuCOE/Nodes/CustomizableObjectNodeStaticMesh.cpp b/Engine/Plugins/Experimental/Mutable/Source/CustomizableObjectEditor/Private/MuCOE/Nodes/CustomizableObjectNodeStaticMesh.cpp
--- a/Engine/Plugins/Experimental/Mutable/Source/CustomizableObjectEditor/Private/MuCOE/Nodes/CustomizableObjectNodeStaticMesh.cpp
+++ b/Engine/Plugins/Experimental/Mutable/Source/CustomizableObjectEditor/Private/MuCOE/Nodes/CustomizableObjectNodeStaticMesh.cpp
@@ -5,6 +5,7 @@
 #include "AssetThumbnail.h"
 #include "ISinglePropertyView.h"
 #include "Modules/ModuleManager.h"
+#include "MuCO/CustomizableObjectCustomVersion.h"
 #include "MuCOE/CustomizableObjectEditor.h"
 #include "MuCOE/CustomizableObjectEditorStyle.h"
 #include "MuCOE/CustomizableObjectLayout.h"
diff --git a/Engine/Plugins/Experimental/Mutable/Source/CustomizableObjectEditor/Private/MuCOE/Nodes/CustomizableObjectNodeSwitchBase.cpp b/Engine/Plugins/Experimental/Mutable/Source/CustomizableObjectEditor/Private/MuCOE/Nodes/CustomizableObjectNodeSwitchBase.cpp
--- a/Engine/Plugins/Experimental/Mutable/Source/CustomizableObjectEditor/Private/MuCOE/Nodes/CustomizableObjectNodeSwitchBase.cpp
+++ b/Engine/Plugins/Experimental/Mutable/Source/CustomizableObjectEditor/Private/MuCOE/Nodes/CustomizableObjectNodeSwitchBase.cpp
@@ -2,6 +2,7 @@
 
 #include "MuCOE/Nodes/CustomizableObjectNodeSwitchBase.h"
 
+#include "MuCO/CustomizableObjectCustomVersion.h"
 #include "MuCOE/EdGraphSchema_CustomizableObject.h"
 #include "MuCOE/GraphTraversal.h"
 #include "MuCOE/Nodes/CustomizableObjectNodeEnumParameter.h"
diff --git a/Engine/Plugins/Experimental/Mutable/Source/CustomizableObjectEditor/Private/MuCOE/Nodes/CustomizableObjectNodeTable.cpp b/Engine/Plugins/Experimental/Mutable/Source/CustomizableObjectEditor/Private/MuCOE/Nodes/CustomizableObjectNodeTable.cpp
--- a/Engine/Plugins/Experimental/Mutable/Source/CustomizableObjectEditor/Private/MuCOE/Nodes/CustomizableObjectNodeTable.cpp
+++ b/Engine/Plugins/Experimental/Mutable/Source/CustomizableObjectEditor/Private/MuCOE/Nodes/CustomizableObjectNodeTable.cpp
@@ -24,6 +24,7 @@
 #include "Animation/AnimInstance.h"
 #include "AssetRegistry/AssetRegistryModule.h"
 #include "AssetRegistry/IAssetRegistry.h"
+#include "MuCO/CustomizableObjectCustomVersion.h"
 
 class ICustomizableObjectEditor;
 class UCustomizableObjectNodeRemapPins;
diff --git a/Engine/Plugins/Experimental/Mutable/Source/CustomizableObjectEditor/Private/MuCOE/Nodes/CustomizableObjectNodeTexture.cpp b/Engine/Plugins/Experimental/Mutable/Source/CustomizableObjectEditor/Private/MuCOE/Nodes/CustomizableObjectNodeTexture.cpp
--- a/Engine/Plugins/Experimental/Mutable/Source/CustomizableObjectEditor/Private/MuCOE/Nodes/CustomizableObjectNodeTexture.cpp
+++ b/Engine/Plugins/Experimental/Mutable/Source/CustomizableObjectEditor/Private/MuCOE/Nodes/CustomizableObjectNodeTexture.cpp
@@ -4,6 +4,7 @@
 
 #include "ISinglePropertyView.h"
 #include "Modules/ModuleManager.h"
+#include "MuCO/CustomizableObjectCustomVersion.h"
 #include "MuCOE/CustomizableObjectEditorStyle.h"
 #include "MuCOE/EdGraphSchema_CustomizableObject.h"
 #include "Widgets/Images/SImage.h"
diff --git a/Engine/Plugins/Experimental/Mutable/Source/CustomizableObjectEditor/Private/MuCOE/Nodes/CustomizableObjectNodeTextureBinarise.cpp b/Engine/Plugins/Experimental/Mutable/Source/CustomizableObjectEditor/Private/MuCOE/Nodes/CustomizableObjectNodeTextureBinarise.cpp
--- a/Engine/Plugins/Experimental/Mutable/Source/CustomizableObjectEditor/Private/MuCOE/Nodes/CustomizableObjectNodeTextureBinarise.cpp
+++ b/Engine/Plugins/Experimental/Mutable/Source/CustomizableObjectEditor/Private/MuCOE/Nodes/CustomizableObjectNodeTextureBinarise.cpp
@@ -2,6 +2,7 @@
 
 #include "MuCOE/Nodes/CustomizableObjectNodeTextureBinarise.h"
 
+#include "MuCO/CustomizableObjectCustomVersion.h"
 #include "MuCOE/EdGraphSchema_CustomizableObject.h"
 
 class UCustomizableObjectNodeRemapPins;
diff --git a/Engine/Plugins/Experimental/Mutable/Source/CustomizableObjectEditor/Private/MuCOE/Nodes/CustomizableObjectNodeTextureColourMap.cpp b/Engine/Plugins/Experimental/Mutable/Source/CustomizableObjectEditor/Private/MuCOE/Nodes/CustomizableObjectNodeTextureColourMap.cpp
--- a/Engine/Plugins/Experimental/Mutable/Source/CustomizableObjectEditor/Private/MuCOE/Nodes/CustomizableObjectNodeTextureColourMap.cpp
+++ b/Engine/Plugins/Experimental/Mutable/Source/CustomizableObjectEditor/Private/MuCOE/Nodes/CustomizableObjectNodeTextureColourMap.cpp
@@ -2,6 +2,7 @@
 
 #include "MuCOE/Nodes/CustomizableObjectNodeTextureColourMap.h"
 
+#include "MuCO/CustomizableObjectCustomVersion.h"
 #include "MuCOE/EdGraphSchema_CustomizableObject.h"
 
 class UCustomizableObjectNodeRemapPins;
diff --git a/Engine/Plugins/Experimental/Mutable/Source/CustomizableObjectEditor/Private/MuCOE/Nodes/CustomizableObjectNodeTextureFromChannels.cpp b/Engine/Plugins/Experimental/Mutable/Source/CustomizableObjectEditor/Private/MuCOE/Nodes/CustomizableObjectNodeTextureFromChannels.cpp
--- a/Engine/Plugins/Experimental/Mutable/Source/CustomizableObjectEditor/Private/MuCOE/Nodes/CustomizableObjectNodeTextureFromChannels.cpp
+++ b/Engine/Plugins/Experimental/Mutable/Source/CustomizableObjectEditor/Private/MuCOE/Nodes/CustomizableObjectNodeTextureFromChannels.cpp
@@ -2,6 +2,7 @@
 
 #include "MuCOE/Nodes/CustomizableObjectNodeTextureFromChannels.h"
 
+#include "MuCO/CustomizableObjectCustomVersion.h"
 #include "MuCOE/EdGraphSchema_CustomizableObject.h"
 
 class UCustomizableObjectNodeRemapPins;
diff --git a/Engine/Plugins/Experimental/Mutable/Source/CustomizableObjectEditor/Private/MuCOE/Nodes/CustomizableObjectNodeTextureFromColor.cpp b/Engine/Plugins/Experimental/Mutable/Source/CustomizableObjectEditor/Private/MuCOE/Nodes/CustomizableObjectNodeTextureFromColor.cpp
--- a/Engine/Plugins/Experimental/Mutable/Source/CustomizableObjectEditor/Private/MuCOE/Nodes/CustomizableObjectNodeTextureFromColor.cpp
+++ b/Engine/Plugins/Experimental/Mutable/Source/CustomizableObjectEditor/Private/MuCOE/Nodes/CustomizableObjectNodeTextureFromColor.cpp
@@ -2,6 +2,7 @@
 
 #include "MuCOE/Nodes/CustomizableObjectNodeTextureFromColor.h"
 
+#include "MuCO/CustomizableObjectCustomVersion.h"
 #include "MuCOE/EdGraphSchema_CustomizableObject.h"
 
 class UCustomizableObjectNodeRemapPins;
diff --git a/Engine/Plugins/Experimental/Mutable/Source/CustomizableObjectEditor/Private/MuCOE/Nodes/CustomizableObjectNodeTextureInvert.cpp b/Engine/Plugins/Experimental/Mutable/Source/CustomizableObjectEditor/Private/MuCOE/Nodes/CustomizableObjectNodeTextureInvert.cpp
--- a/Engine/Plugins/Experimental/Mutable/Source/CustomizableObjectEditor/Private/MuCOE/Nodes/CustomizableObjectNodeTextureInvert.cpp
+++ b/Engine/Plugins/Experimental/Mutable/Source/CustomizableObjectEditor/Private/MuCOE/Nodes/CustomizableObjectNodeTextureInvert.cpp
@@ -2,6 +2,7 @@
 
 #include "MuCOE/Nodes/CustomizableObjectNodeTextureInvert.h"
 
+#include "MuCO/CustomizableObjectCustomVersion.h"
 #include "MuCOE/EdGraphSchema_CustomizableObject.h"
 
 class UCustomizableObjectNodeRemapPins;
diff --git a/Engine/Plugins/Experimental/Mutable/Source/CustomizableObjectEditor/Private/MuCOE/Nodes/CustomizableObjectNodeTextureLayer.cpp b/Engine/Plugins/Experimental/Mutable/Source/CustomizableObjectEditor/Private/MuCOE/Nodes/CustomizableObjectNodeTextureLayer.cpp
--- a/Engine/Plugins/Experimental/Mutable/Source/CustomizableObjectEditor/Private/MuCOE/Nodes/CustomizableObjectNodeTextureLayer.cpp
+++ b/Engine/Plugins/Experimental/Mutable/Source/CustomizableObjectEditor/Private/MuCOE/Nodes/CustomizableObjectNodeTextureLayer.cpp
@@ -2,6 +2,7 @@
 
 #include "MuCOE/Nodes/CustomizableObjectNodeTextureLayer.h"
 
+#include "MuCO/CustomizableObjectCustomVersion.h"
 #include "MuCOE/EdGraphSchema_CustomizableObject.h"
 #include "MuCOE/UnrealEditorPortabilityHelpers.h"
 
diff --git a/Engine/Plugins/Experimental/Mutable/Source/CustomizableObjectEditor/Private/MuCOE/Nodes/CustomizableObjectNodeTextureParameter.cpp b/Engine/Plugins/Experimental/Mutable/Source/CustomizableObjectEditor/Private/MuCOE/Nodes/CustomizableObjectNodeTextureParameter.cpp
--- a/Engine/Plugins/Experimental/Mutable/Source/CustomizableObjectEditor/Private/MuCOE/Nodes/CustomizableObjectNodeTextureParameter.cpp
+++ b/Engine/Plugins/Experimental/Mutable/Source/CustomizableObjectEditor/Private/MuCOE/Nodes/CustomizableObjectNodeTextureParameter.cpp
@@ -2,6 +2,7 @@
 
 #include "MuCOE/Nodes/CustomizableObjectNodeTextureParameter.h"
 
+#include "MuCO/CustomizableObjectCustomVersion.h"
 #include "MuCOE/EdGraphSchema_CustomizableObject.h"
 
 class UCustomizableObjectNodeRemapPins;
diff --git a/Engine/Plugins/Experimental/Mutable/Source/CustomizableObjectEditor/Private/MuCOE/Nodes/CustomizableObjectNodeTextureProject.cpp b/Engine/Plugins/Experimental/Mutable/Source/CustomizableObjectEditor/Private/MuCOE/Nodes/CustomizableObjectNodeTextureProject.cpp
--- a/Engine/Plugins/Experimental/Mutable/Source/CustomizableObjectEditor/Private/MuCOE/Nodes/CustomizableObjectNodeTextureProject.cpp
+++ b/Engine/Plugins/Experimental/Mutable/Source/CustomizableObjectEditor/Private/MuCOE/Nodes/CustomizableObjectNodeTextureProject.cpp
@@ -2,6 +2,7 @@
 
 #include "MuCOE/Nodes/CustomizableObjectNodeTextureProject.h"
 
+#include "MuCO/CustomizableObjectCustomVersion.h"
 #include "MuCOE/EdGraphSchema_CustomizableObject.h"
 
 class UCustomizableObjectNodeRemapPins;
diff --git a/Engine/Plugins/Experimental/Mutable/Source/CustomizableObjectEditor/Private/MuCOE/Nodes/CustomizableObjectNodeTextureSwitch.cpp b/Engine/Plugins/Experimental/Mutable/Source/CustomizableObjectEditor/Private/MuCOE/Nodes/CustomizableObjectNodeTextureSwitch.cpp
--- a/Engine/Plugins/Experimental/Mutable/Source/CustomizableObjectEditor/Private/MuCOE/Nodes/CustomizableObjectNodeTextureSwitch.cpp
+++ b/Engine/Plugins/Experimental/Mutable/Source/CustomizableObjectEditor/Private/MuCOE/Nodes/CustomizableObjectNodeTextureSwitch.cpp
@@ -2,6 +2,7 @@
 
 #include "MuCOE/Nodes/CustomizableObjectNodeTextureSwitch.h"
 
+#include "MuCO/CustomizableObjectCustomVersion.h"
 #include "MuCOE/EdGraphSchema_CustomizableObject.h"
 
 #define LOCTEXT_NAMESPACE "CustomizableObjectEditor"
diff --git a/Engine/Plugins/Experimental/Mutable/Source/CustomizableObjectEditor/Private/MuCOE/Nodes/CustomizableObjectNodeTextureToChannels.cpp b/Engine/Plugins/Experimental/Mutable/Source/CustomizableObjectEditor/Private/MuCOE/Nodes/CustomizableObjectNodeTextureToChannels.cpp
--- a/Engine/Plugins/Experimental/Mutable/Source/CustomizableObjectEditor/Private/MuCOE/Nodes/CustomizableObjectNodeTextureToChannels.cpp
+++ b/Engine/Plugins/Experimental/Mutable/Source/CustomizableObjectEditor/Private/MuCOE/Nodes/CustomizableObjectNodeTextureToChannels.cpp
@@ -2,6 +2,7 @@
 
 #include "MuCOE/Nodes/CustomizableObjectNodeTextureToChannels.h"
 
+#include "MuCO/CustomizableObjectCustomVersion.h"
 #include "MuCOE/EdGraphSchema_CustomizableObject.h"
 
 class UCustomizableObjectNodeRemapPins;
diff --git a/Engine/Plugins/Experimental/Mutable/Source/CustomizableObjectEditor/Private/MuCOE/Nodes/CustomizableObjectNodeTextureTransform.cpp b/Engine/Plugins/Experimental/Mutable/Source/CustomizableObjectEditor/Private/MuCOE/Nodes/CustomizableObjectNodeTextureTransform.cpp
--- a/Engine/Plugins/Experimental/Mutable/Source/CustomizableObjectEditor/Private/MuCOE/Nodes/CustomizableObjectNodeTextureTransform.cpp
+++ b/Engine/Plugins/Experimental/Mutable/Source/CustomizableObjectEditor/Private/MuCOE/Nodes/CustomizableObjectNodeTextureTransform.cpp
@@ -2,6 +2,7 @@
 
 #include "MuCOE/Nodes/CustomizableObjectNodeTextureTransform.h"
 
+#include "MuCO/CustomizableObjectCustomVersion.h"
 #include "MuCOE/EdGraphSchema_CustomizableObject.h"
 
 class UCustomizableObjectNodeRemapPins;
diff --git a/Engine/Plugins/Experimental/Mutable/Source/CustomizableObjectEditor/Private/MuCOE/Nodes/CustomizableObjectNodeVariation.cpp b/Engine/Plugins/Experimental/Mutable/Source/CustomizableObjectEditor/Private/MuCOE/Nodes/CustomizableObjectNodeVariation.cpp
--- a/Engine/Plugins/Experimental/Mutable/Source/CustomizableObjectEditor/Private/MuCOE/Nodes/CustomizableObjectNodeVariation.cpp
+++ b/Engine/Plugins/Experimental/Mutable/Source/CustomizableObjectEditor/Private/MuCOE/Nodes/CustomizableObjectNodeVariation.cpp
@@ -2,6 +2,7 @@
 
 #include "MuCOE/Nodes/CustomizableObjectNodeVariation.h"
 
+#include "MuCO/CustomizableObjectCustomVersion.h"
 #include "MuCOE/EdGraphSchema_CustomizableObject.h"
 
 class UCustomizableObjectNodeRemapPins;
diff --git a/Engine/Plugins/Experimental/Mutable/Source/CustomizableObjectEditor/Private/MuCOE/SCustomizableInstanceProperties.cpp b/Engine/Plugins/Experimental/Mutable/Source/CustomizableObjectEditor/Private/MuCOE/SCustomizableInstanceProperties.cpp
--- a/Engine/Plugins/Experimental/Mutable/Source/CustomizableObjectEditor/Private/MuCOE/SCustomizableInstanceProperties.cpp
+++ b/Engine/Plugins/Experimental/Mutable/Source/CustomizableObjectEditor/Private/MuCOE/SCustomizableInstanceProperties.cpp
@@ -463,8 +463,8 @@ void SCustomizableInstanceProperties::ResetParamBox()
 				ParameterSortInfo.PropertyCOIndex = CustomizableObject->GetStateParameterIndex(CustomInstance->GetPrivate()->GetState(), ParamIndexInState);
 				if (CustomInstance->IsParameterRelevant(ParameterSortInfo.PropertyCOIndex))
 				{
-					ParameterSortInfo.UIOrder = CustomizableObject->GetParameterUIMetadataFromIndex(ParameterSortInfo.PropertyCOIndex).ParamUIMetadata.UIOrder;
 					ParameterSortInfo.ParameterName = CustomizableObject->GetParameterName(ParameterSortInfo.PropertyCOIndex);
+					ParameterSortInfo.UIOrder = CustomizableObject->GetParameterUIMetadata(ParameterSortInfo.ParameterName).UIOrder;
 					ParamIndexesInState.Add(ParameterSortInfo);
 				}
 			}
@@ -488,8 +488,8 @@ void SCustomizableInstanceProperties::ResetParamBox()
 				ParameterSortInfo.PropertyCOIndex = ParamIndexInObject;
 				if (!CustomInstance->GetPrivate()->bShowOnlyRelevantParameters || CustomInstance->IsParameterRelevant(ParameterSortInfo.PropertyCOIndex))
 				{
-					ParameterSortInfo.UIOrder = CustomizableObject->GetParameterUIMetadataFromIndex(ParameterSortInfo.PropertyCOIndex).ParamUIMetadata.UIOrder;
 					ParameterSortInfo.ParameterName = CustomizableObject->GetParameterName(ParameterSortInfo.PropertyCOIndex);
+					ParameterSortInfo.UIOrder = CustomizableObject->GetParameterUIMetadata(ParameterSortInfo.ParameterName).UIOrder;
 					ParamIndexesInObject.Add(ParameterSortInfo);
 				}
 			}
@@ -534,12 +534,10 @@ void SCustomizableInstanceProperties::RecursivelyAddParamAndChildren(int32 Param
 void SCustomizableInstanceProperties::FillChildrenMap(int32 ParamIndexInObject)
 {
 	const UCustomizableObject* CustomizableObject = CustomInstance->GetCustomizableObject();
-	FString ParamName = CustomizableObject->GetParameterName(ParamIndexInObject);
-	FParameterUIData UIData = CustomizableObject->GetParameterUIMetadataFromIndex(ParamIndexInObject);
-
-	const FString* ParentName = UIData.ParamUIMetadata.ExtraInformation.Find(FString("__ParentParamName"));
+	const FString ParamName = CustomizableObject->GetParameterName(ParamIndexInObject);
+	FMutableParamUIMetadata UIMetadata = CustomizableObject->GetParameterUIMetadata(ParamName);
 
-	if (ParentName)
+	if (const FString* ParentName = UIMetadata.ExtraInformation.Find(FString("__ParentParamName")))
 	{
 		ParamChildren.Add(*ParentName, ParamIndexInObject);
 		ParamHasParent.Add(ParamIndexInObject, true);
@@ -554,14 +552,14 @@ void SCustomizableInstanceProperties::AddParameter(int32 ParamIndexInObject)
 
 	TSharedPtr<ICustomizableObjectInstanceEditor> Editor = GetEditorChecked();
 	
-	FString ParamName = CustomizableObject->GetParameterName(ParamIndexInObject);
+	const FString& ParamName = CustomizableObject->GetParameterName(ParamIndexInObject);
 
 	TSharedPtr<SVerticalBox> ActualParamBox = ParamBox;
 	TSharedPtr<SHorizontalBox> ParameterBox;
 
-	FParameterUIData UIData = CustomizableObject->GetParameterUIMetadataFromIndex(ParamIndexInObject);
-	FString* ParentName = UIData.ParamUIMetadata.ExtraInformation.Find(FString("__ParentParamName"));
-	bool bHasCollapsibleChildren = UIData.ParamUIMetadata.ExtraInformation.Find(FString("__HasCollapsibleChildren")) != nullptr;
+	FMutableParamUIMetadata UIMetadata = CustomizableObject->GetParameterUIMetadata(ParamName);
+	FString* ParentName = UIMetadata.ExtraInformation.Find(FString("__ParentParamName"));
+	bool bHasCollapsibleChildren = UIMetadata.ExtraInformation.Find(FString("__HasCollapsibleChildren")) != nullptr;
 
 	bool bHideParam = ParamName.EndsWith(FMultilayerProjector::NUM_LAYERS_PARAMETER_POSTFIX)
 		|| (ParamName.EndsWith(FMultilayerProjector::IMAGE_PARAMETER_POSTFIX) && CustomizableObject->IsParameterMultidimensional(ParamIndexInObject))
@@ -587,7 +585,7 @@ void SCustomizableInstanceProperties::AddParameter(int32 ParamIndexInObject)
 
 	if (ParentName)
 	{
-		if (UIData.ParamUIMetadata.ExtraInformation.Find(FString("CollapseUnderParent")))
+		if (UIMetadata.ExtraInformation.Find(FString("CollapseUnderParent")))
 		{
 			const FMutableParamExpandableArea* ParentExpandableAreaPtr = ParamNameToExpandableAreaMap.Find(*ParentName);
 
@@ -599,7 +597,7 @@ void SCustomizableInstanceProperties::AddParameter(int32 ParamIndexInObject)
 
 		if (CustomInstance->GetPrivate()->bShowOnlyRelevantParameters)
 		{
-			FString *Value = UIData.ParamUIMetadata.ExtraInformation.Find(FString("__DisplayWhenParentValueEquals"));
+			FString *Value = UIMetadata.ExtraInformation.Find(FString("__DisplayWhenParentValueEquals"));
 			if (Value && CustomInstance->GetIntParameterSelectedOption(*ParentName) != *Value)
 			{
 				return;
@@ -715,8 +713,8 @@ void SCustomizableInstanceProperties::AddParameter(int32 ParamIndexInObject)
 			[
 				SAssignNew(Slider, SSpinBox<float>)
 				.Value(this, &SCustomizableInstanceProperties::GetSliderValue, SliderIndex)
-				.MinValue(CustomizableObject->GetParameterUIMetadataFromIndex(ParamIndexInObject).ParamUIMetadata.MinimumValue)
-				.MaxValue(CustomizableObject->GetParameterUIMetadataFromIndex(ParamIndexInObject).ParamUIMetadata.MaximumValue)
+				.MinValue(CustomizableObject->GetParameterUIMetadata(ParamName).MinimumValue)
+				.MaxValue(CustomizableObject->GetParameterUIMetadata(ParamName).MaximumValue)
 				.OnValueChanged(this, &SCustomizableInstanceProperties::OnFloatParameterChanged, SliderIndex)
 				.OnBeginSliderMovement(this, &SCustomizableInstanceProperties::OnFloatParameterSliderBegin)
 				.OnEndSliderMovement(this, &SCustomizableInstanceProperties::OnFloatParameterSliderEnd)
