[Mutable] Fix missing implicit FMutableParameterData and FIntegerParameterUIData constructors on some compilers.

#rnx

diff --git a/Engine/Plugins/Experimental/Mutable/Source/CustomizableObject/Internal/MuCO/CustomizableObjectPrivate.h b/Engine/Plugins/Experimental/Mutable/Source/CustomizableObject/Internal/MuCO/CustomizableObjectPrivate.h
--- a/Engine/Plugins/Experimental/Mutable/Source/CustomizableObject/Internal/MuCO/CustomizableObjectPrivate.h
+++ b/Engine/Plugins/Experimental/Mutable/Source/CustomizableObject/Internal/MuCO/CustomizableObjectPrivate.h
@@ -558,10 +558,14 @@ template<> struct TCanBulkSerialize<FMorphTargetVertexData> { enum { Value = tru
 
 
 USTRUCT()
-struct FIntegerParameterUIData
+struct CUSTOMIZABLEOBJECT_API FIntegerParameterUIData
 {
 	GENERATED_BODY()
 
+	FIntegerParameterUIData() = default;
+	
+	FIntegerParameterUIData(const FMutableParamUIMetadata& InParamUIMetadata);
+	
 	UPROPERTY()
 	FMutableParamUIMetadata ParamUIMetadata;
 
@@ -570,10 +574,14 @@ struct FIntegerParameterUIData
 
 
 USTRUCT()
-struct FMutableParameterData
+struct CUSTOMIZABLEOBJECT_API FMutableParameterData
 {
 	GENERATED_BODY()
 
+	FMutableParameterData() = default;
+	
+	FMutableParameterData(const FMutableParamUIMetadata& InParamUIMetadata, EMutableParameterType InType);
+
 	UPROPERTY()
 	FMutableParamUIMetadata ParamUIMetadata;
 
@@ -594,7 +602,7 @@ struct FMutableParameterData
 
 
 USTRUCT()
-struct FMutableStateData
+struct CUSTOMIZABLEOBJECT_API FMutableStateData
 {
 	GENERATED_BODY()
 
diff --git a/Engine/Plugins/Experimental/Mutable/Source/CustomizableObject/Private/MuCO/CustomizableObject.cpp b/Engine/Plugins/Experimental/Mutable/Source/CustomizableObject/Private/MuCO/CustomizableObject.cpp
--- a/Engine/Plugins/Experimental/Mutable/Source/CustomizableObject/Private/MuCO/CustomizableObject.cpp
+++ b/Engine/Plugins/Experimental/Mutable/Source/CustomizableObject/Private/MuCO/CustomizableObject.cpp
@@ -2442,6 +2442,19 @@ bool FMutableSkinWeightProfileInfo::operator==(const FMutableSkinWeightProfileIn
 }
 
 
+FIntegerParameterUIData::FIntegerParameterUIData(const FMutableParamUIMetadata& InParamUIMetadata)
+{
+	ParamUIMetadata = InParamUIMetadata;
+}
+
+
+FMutableParameterData::FMutableParameterData(const FMutableParamUIMetadata& InParamUIMetadata, EMutableParameterType InType)
+{
+	ParamUIMetadata = InParamUIMetadata;
+	Type = InType;
+}
+
+
 #if WITH_EDITORONLY_DATA
 FArchive& operator<<(FArchive& Ar, FMutableRemappedBone& RemappedBone)
 {
