[Mutable] Fix ConditionalAutoCompile requesting multiple compilations for the same object.

#rnx

diff --git a/Engine/Plugins/Experimental/Mutable/Source/CustomizableObjectEditor/Private/MuCOE/CustomizableObjectCompiler.cpp b/Engine/Plugins/Experimental/Mutable/Source/CustomizableObjectEditor/Private/MuCOE/CustomizableObjectCompiler.cpp
--- a/Engine/Plugins/Experimental/Mutable/Source/CustomizableObjectEditor/Private/MuCOE/CustomizableObjectCompiler.cpp
+++ b/Engine/Plugins/Experimental/Mutable/Source/CustomizableObjectEditor/Private/MuCOE/CustomizableObjectCompiler.cpp
@@ -317,7 +317,11 @@ void FCustomizableObjectCompiler::TickCook(float DeltaTime, bool bCookCompete)
 
 bool FCustomizableObjectCompiler::IsRequestQueued(const TSharedRef<FCompilationRequest>& InCompileRequest) const
 {
-	return (CurrentRequest == InCompileRequest) || (CompileRequests.Find(InCompileRequest) != INDEX_NONE);
+	return (CurrentRequest == InCompileRequest) || 
+		(CompileRequests.ContainsByPredicate([&InCompileRequest](const TSharedRef<FCompilationRequest>& Other)
+		{
+			return InCompileRequest.Get() == Other.Get(); // Compare the content of the request not the ref
+		}));
 }
 
 
