[Mutable] Removed pragma once on CustomizableObjectUIData source file.

#rnx

diff --git a/Engine/Plugins/Experimental/Mutable/Source/CustomizableObject/Private/MuCO/CustomizableObjectUIData.cpp b/Engine/Plugins/Experimental/Mutable/Source/CustomizableObject/Private/MuCO/CustomizableObjectUIData.cpp
--- a/Engine/Plugins/Experimental/Mutable/Source/CustomizableObject/Private/MuCO/CustomizableObjectUIData.cpp
+++ b/Engine/Plugins/Experimental/Mutable/Source/CustomizableObject/Private/MuCO/CustomizableObjectUIData.cpp
@@ -1,7 +1,5 @@
 ﻿// Copyright Epic Games, Inc. All Rights Reserved.
 
-#pragma once
-
 #include "MuCO/CustomizableObjectUIData.h"
 
 #include "MuCO/CustomizableObjectCustomVersion.h"
