[Backout] - CL33259189 (backed out due to cook error)
Original CL Desc
-----------------------------------------------------------------
[Mutable] Fix missing implicit FMutableParameterData and FIntegerParameterUIData constructors on some compilers.

#rnx

diff --git a/Engine/Plugins/Experimental/Mutable/Source/CustomizableObject/Internal/MuCO/CustomizableObjectPrivate.h b/Engine/Plugins/Experimental/Mutable/Source/CustomizableObject/Internal/MuCO/CustomizableObjectPrivate.h
--- a/Engine/Plugins/Experimental/Mutable/Source/CustomizableObject/Internal/MuCO/CustomizableObjectPrivate.h
+++ b/Engine/Plugins/Experimental/Mutable/Source/CustomizableObject/Internal/MuCO/CustomizableObjectPrivate.h
@@ -558,14 +558,10 @@ template<> struct TCanBulkSerialize<FMorphTargetVertexData> { enum { Value = tru
 
 
 USTRUCT()
-struct CUSTOMIZABLEOBJECT_API FIntegerParameterUIData
+struct FIntegerParameterUIData
 {
 	GENERATED_BODY()
 
-	FIntegerParameterUIData() = default;
-	
-	FIntegerParameterUIData(const FMutableParamUIMetadata& InParamUIMetadata);
-	
 	UPROPERTY()
 	FMutableParamUIMetadata ParamUIMetadata;
 
@@ -574,14 +570,10 @@ struct CUSTOMIZABLEOBJECT_API FIntegerParameterUIData
 
 
 USTRUCT()
-struct CUSTOMIZABLEOBJECT_API FMutableParameterData
+struct FMutableParameterData
 {
 	GENERATED_BODY()
 
-	FMutableParameterData() = default;
-	
-	FMutableParameterData(const FMutableParamUIMetadata& InParamUIMetadata, EMutableParameterType InType);
-
 	UPROPERTY()
 	FMutableParamUIMetadata ParamUIMetadata;
 
@@ -602,7 +594,7 @@ struct CUSTOMIZABLEOBJECT_API FMutableParameterData
 
 
 USTRUCT()
-struct CUSTOMIZABLEOBJECT_API FMutableStateData
+struct FMutableStateData
 {
 	GENERATED_BODY()
 
diff --git a/Engine/Plugins/Experimental/Mutable/Source/CustomizableObject/Private/MuCO/CustomizableObject.cpp b/Engine/Plugins/Experimental/Mutable/Source/CustomizableObject/Private/MuCO/CustomizableObject.cpp
--- a/Engine/Plugins/Experimental/Mutable/Source/CustomizableObject/Private/MuCO/CustomizableObject.cpp
+++ b/Engine/Plugins/Experimental/Mutable/Source/CustomizableObject/Private/MuCO/CustomizableObject.cpp
@@ -2442,19 +2442,6 @@ bool FMutableSkinWeightProfileInfo::operator==(const FMutableSkinWeightProfileIn
 }
 
 
-FIntegerParameterUIData::FIntegerParameterUIData(const FMutableParamUIMetadata& InParamUIMetadata)
-{
-	ParamUIMetadata = InParamUIMetadata;
-}
-
-
-FMutableParameterData::FMutableParameterData(const FMutableParamUIMetadata& InParamUIMetadata, EMutableParameterType InType)
-{
-	ParamUIMetadata = InParamUIMetadata;
-	Type = InType;
-}
-
-
 #if WITH_EDITORONLY_DATA
 FArchive& operator<<(FArchive& Ar, FMutableRemappedBone& RemappedBone)
 {
