Fix UMG widgets in Mutable plugin which were broken by refactor of parameter metadata structs.

#changelist validated

#virtualized

diff --git a/Engine/Plugins/Experimental/Mutable/Source/CustomizableObject/Classes/MuCO/CustomizableObject.h b/Engine/Plugins/Experimental/Mutable/Source/CustomizableObject/Classes/MuCO/CustomizableObject.h
--- a/Engine/Plugins/Experimental/Mutable/Source/CustomizableObject/Classes/MuCO/CustomizableObject.h
+++ b/Engine/Plugins/Experimental/Mutable/Source/CustomizableObject/Classes/MuCO/CustomizableObject.h
@@ -632,6 +632,10 @@ public:
 	UFUNCTION(BlueprintCallable, Category = CustomizableObject)
 	FMutableParamUIMetadata GetIntParameterOptionUIMetadata(const FString& ParamName, const FString& OptionName) const;
 	
+	/** Returns the group type of the given integer parameter */
+	UFUNCTION(BlueprintCallable, Category = CustomizableObject)
+	ECustomizableObjectGroupType GetIntParameterGroupType(const FString& ParamName) const;
+
 	/** Return the metadata associated to a state. */
 	UFUNCTION(BlueprintCallable, Category = CustomizableObject)
 	FMutableStateUIMetadata GetStateUIMetadata(const FString& StateName) const;
diff --git a/Engine/Plugins/Experimental/Mutable/Source/CustomizableObject/Private/MuCO/CustomizableObject.cpp b/Engine/Plugins/Experimental/Mutable/Source/CustomizableObject/Private/MuCO/CustomizableObject.cpp
--- a/Engine/Plugins/Experimental/Mutable/Source/CustomizableObject/Private/MuCO/CustomizableObject.cpp
+++ b/Engine/Plugins/Experimental/Mutable/Source/CustomizableObject/Private/MuCO/CustomizableObject.cpp
@@ -1550,6 +1550,22 @@ FMutableParamUIMetadata UCustomizableObject::GetIntParameterOptionUIMetadata(con
 	return IntegerParameterUIData ? IntegerParameterUIData->ParamUIMetadata : FMutableParamUIMetadata();
 }
 
+ECustomizableObjectGroupType UCustomizableObject::GetIntParameterGroupType(const FString& ParamName) const
+{
+	const int32 ParameterIndex = FindParameter(ParamName);
+	if (ParameterIndex == INDEX_NONE)
+	{
+		return ECustomizableObjectGroupType::COGT_TOGGLE;
+	}		
+	
+	const FMutableParameterData* ParameterData = Private->GetModelResources().ParameterUIDataMap.Find(ParamName);
+	if (!ParameterData)
+	{
+		return ECustomizableObjectGroupType::COGT_TOGGLE;
+	}
+
+	return ParameterData->IntegerParameterGroupType;
+}
 
 FMutableStateUIMetadata UCustomizableObject::GetStateUIMetadata(const FString& StateName) const
 {
