[mutable] Stable vertex IDs based on Mesh prefix.

- Add support for 64 bit ints in mesh buffers.
- Vertex IDs can be stored as implicit (prefix+index in mesh), relative (prefix+32-bit ID) or explicit (64-bit ID).
- Update all mesh operations to handle the new vertex id format.
- Always generate Vertex IDs, since they are free at generation time, and simplifies logic.
- Remove unused face buffers.
- Split morph buffers to put vertex ids in their own.

#rnx

diff --git a/Engine/Plugins/Experimental/Mutable/Source/CustomizableObject/Internal/MuCO/CustomizableObjectPrivate.h b/Engine/Plugins/Experimental/Mutable/Source/CustomizableObject/Internal/MuCO/CustomizableObjectPrivate.h
--- a/Engine/Plugins/Experimental/Mutable/Source/CustomizableObject/Internal/MuCO/CustomizableObjectPrivate.h
+++ b/Engine/Plugins/Experimental/Mutable/Source/CustomizableObject/Internal/MuCO/CustomizableObjectPrivate.h
@@ -960,7 +960,14 @@ public:
 
 	// This is a manual version number for the binary blobs in this asset.
 	// Increasing it invalidates all the previously compiled models.
-	// Warning: If while merging code both versions have changed, take the highest+1.
-	static constexpr int32 CurrentSupportedVersion = 449;
+	enum ECustomizableObjectVersions
+	{
+		FirstEnumeratedVersion = 450,
+
+		DeterminisiticMeshVertexIds,
+
+		LastCustomizableObjectVersion
+	};
+	static constexpr int32 CurrentSupportedVersion = ECustomizableObjectVersions::LastCustomizableObjectVersion;
 };
 
diff --git a/Engine/Plugins/Experimental/Mutable/Source/CustomizableObjectEditor/Private/MuCOE/GenerateMutableSource/GenerateMutableSource.cpp b/Engine/Plugins/Experimental/Mutable/Source/CustomizableObjectEditor/Private/MuCOE/GenerateMutableSource/GenerateMutableSource.cpp
--- a/Engine/Plugins/Experimental/Mutable/Source/CustomizableObjectEditor/Private/MuCOE/GenerateMutableSource/GenerateMutableSource.cpp
+++ b/Engine/Plugins/Experimental/Mutable/Source/CustomizableObjectEditor/Private/MuCOE/GenerateMutableSource/GenerateMutableSource.cpp
@@ -49,6 +49,7 @@
 #include "MuT/NodeObjectGroup.h"
 #include "MuT/NodeObjectNew.h"
 #include "MuT/NodeSurfaceEdit.h"
+#include "MuR/Mesh.h"
 #include "PhysicsEngine/PhysicsAsset.h"
 #include "PlatformInfo.h"
 #include "Math/NumericLimits.h"
@@ -144,7 +145,7 @@ void FMutableGraphGenerationContext::AddParticipatingObject(const UObject& InObj
 }
 
 
-mu::MeshPtr FMutableGraphGenerationContext::FindGeneratedMesh( const FGeneratedMeshData::FKey& Key )
+mu::Ptr<mu::Mesh> FMutableGraphGenerationContext::FindGeneratedMesh( const FGeneratedMeshData::FKey& Key )
 {
 	for (const FGeneratedMeshData& d : GeneratedMeshes)
 	{
@@ -780,11 +781,11 @@ mu::NodeMeshApplyPosePtr CreateNodeMeshApplyPose(FMutableGraphGenerationContext&
 {
 	check(ArrayBoneName.Num() == ArrayTransform.Num());
 
-	mu::MeshPtr MutableMesh = new mu::Mesh();
+	mu::Ptr<mu::Mesh> MutableMesh = new mu::Mesh();
 	mu::NodeMeshConstantPtr PoseNodeMesh = new mu::NodeMeshConstant;
 	PoseNodeMesh->SetValue(MutableMesh);
 
-	mu::SkeletonPtr MutableSkeleton = new mu::Skeleton;
+	mu::Ptr<mu::Skeleton> MutableSkeleton = new mu::Skeleton;
 	MutableMesh->SetSkeleton(MutableSkeleton);
 	MutableMesh->SetBonePoseCount(ArrayBoneName.Num());
 	MutableSkeleton->SetBoneCount(ArrayBoneName.Num());
@@ -1939,7 +1940,7 @@ mu::Ptr<mu::Image> GenerateImageConstant(UTexture* Texture, FMutableGraphGenerat
 }
 
 
-void AddSocketTagsToMesh(const USkeletalMesh* SourceMesh, mu::MeshPtr MutableMesh, FMutableGraphGenerationContext& GenerationContext)
+void AddSocketTagsToMesh(const USkeletalMesh* SourceMesh, mu::Ptr<mu::Mesh> MutableMesh, FMutableGraphGenerationContext& GenerationContext)
 {
 	for (int32 SocketIndex = 0; SocketIndex < SourceMesh->NumSockets(); ++SocketIndex)
 	{
diff --git a/Engine/Plugins/Experimental/Mutable/Source/CustomizableObjectEditor/Private/MuCOE/GenerateMutableSource/GenerateMutableSourceMesh.cpp b/Engine/Plugins/Experimental/Mutable/Source/CustomizableObjectEditor/Private/MuCOE/GenerateMutableSource/GenerateMutableSourceMesh.cpp
--- a/Engine/Plugins/Experimental/Mutable/Source/CustomizableObjectEditor/Private/MuCOE/GenerateMutableSource/GenerateMutableSourceMesh.cpp
+++ b/Engine/Plugins/Experimental/Mutable/Source/CustomizableObjectEditor/Private/MuCOE/GenerateMutableSource/GenerateMutableSourceMesh.cpp
@@ -1742,7 +1742,6 @@ mu::MeshPtr ConvertSkeletalMeshToMutable(const USkeletalMesh* InSkeletalMesh, co
 		const uint32 IndexCount = MeshSection.NumTriangles * 3;
 		MutableMesh->GetIndexBuffers().SetBufferCount(1);
 		MutableMesh->GetIndexBuffers().SetElementCount(IndexCount);
-		MutableMesh->GetFaceBuffers().SetElementCount(MeshSection.NumTriangles);
 
 		using namespace mu;
 
@@ -1785,12 +1784,12 @@ mu::MeshPtr ConvertSkeletalMeshToMutable(const USkeletalMesh* InSkeletalMesh, co
 	}
 
 	// Ensure Surface Data
-	mu::MESH_SURFACE MeshSurface;
-	MeshSurface.m_vertexCount = MutableMesh->m_VertexBuffers.GetElementCount();
-	MeshSurface.m_indexCount = MutableMesh->m_IndexBuffers.GetElementCount();
+	mu::FMeshSurface MeshSurface;
+	MeshSurface.VertexCount = MutableMesh->VertexBuffers.GetElementCount();
+	MeshSurface.IndexCount = MutableMesh->IndexBuffers.GetElementCount();
 	MeshSurface.BoneMapCount = MutableMesh->BoneMap.Num();
 	MeshSurface.bCastShadow = MeshSection.bCastShadow;
-	MutableMesh->m_surfaces.Add(MeshSurface);
+	MutableMesh->Surfaces.Add(MeshSurface);
 
 	if (!bIgnorePhysics && InSkeletalMesh->GetPhysicsAsset() && MutableMesh->GetSkeleton() && GenerationContext.Options.bPhysicsAssetMergeEnabled)
 	{
@@ -2181,7 +2180,6 @@ mu::MeshPtr ConvertStaticMeshToMutable(const UStaticMesh* StaticMesh, int32 LODI
 		int IndexCount = StaticMesh->GetRenderData()->LODResources[LODIndex].Sections[SectionIndex].NumTriangles * 3;
 		MutableMesh->GetIndexBuffers().SetBufferCount(1);
 		MutableMesh->GetIndexBuffers().SetElementCount(IndexCount);
-		MutableMesh->GetFaceBuffers().SetElementCount(IndexCount / 3);
 
 		using namespace mu;
 		const int ElementSize = 2;
diff --git a/Engine/Plugins/Experimental/Mutable/Source/CustomizableObjectEditor/Private/MuCOE/MutableMeshPreviewUtils.cpp b/Engine/Plugins/Experimental/Mutable/Source/CustomizableObjectEditor/Private/MuCOE/MutableMeshPreviewUtils.cpp
--- a/Engine/Plugins/Experimental/Mutable/Source/CustomizableObjectEditor/Private/MuCOE/MutableMeshPreviewUtils.cpp
+++ b/Engine/Plugins/Experimental/Mutable/Source/CustomizableObjectEditor/Private/MuCOE/MutableMeshPreviewUtils.cpp
@@ -177,7 +177,7 @@ namespace MutableMeshPreviewUtils
 
 			bool bOutSuccess = false;
 			mu::MeshFormat(FormatResult.get(), InMutableMesh.get(), FormattedMutableMesh.get(),
-				false, true, true, false, false,
+				false, true, true, false,
 				bOutSuccess);
 		
 			return FormatResult;
diff --git a/Engine/Plugins/Experimental/Mutable/Source/CustomizableObjectEditor/Private/MuCOE/SMutableGraphViewer.cpp b/Engine/Plugins/Experimental/Mutable/Source/CustomizableObjectEditor/Private/MuCOE/SMutableGraphViewer.cpp
--- a/Engine/Plugins/Experimental/Mutable/Source/CustomizableObjectEditor/Private/MuCOE/SMutableGraphViewer.cpp
+++ b/Engine/Plugins/Experimental/Mutable/Source/CustomizableObjectEditor/Private/MuCOE/SMutableGraphViewer.cpp
@@ -478,9 +478,9 @@ void SMutableGraphViewer::GetChildrenForInfo(TSharedPtr<FMutableGraphTreeElement
 	{
 		mu::NodeMeshConstant* MeshConstantVar = StaticCast<mu::NodeMeshConstant*>(ParentNode);
 		mu::NodeMeshConstant::Private* Private = MeshConstantVar->GetPrivate();
-		for (int32 LayoutIndex = 0; LayoutIndex < Private->m_layouts.Num(); LayoutIndex++)
+		for (int32 LayoutIndex = 0; LayoutIndex < Private->Layouts.Num(); LayoutIndex++)
 		{
-			AddChildFunc(Private->m_layouts[LayoutIndex].get(), FString::Printf(TEXT("LAYOUT [%d]"), LayoutIndex));
+			AddChildFunc(Private->Layouts[LayoutIndex].get(), FString::Printf(TEXT("LAYOUT [%d]"), LayoutIndex));
 		}
 	}
 
diff --git a/Engine/Plugins/Experimental/Mutable/Source/CustomizableObjectEditor/Private/MuCOE/SMutableMeshViewer.cpp b/Engine/Plugins/Experimental/Mutable/Source/CustomizableObjectEditor/Private/MuCOE/SMutableMeshViewer.cpp
--- a/Engine/Plugins/Experimental/Mutable/Source/CustomizableObjectEditor/Private/MuCOE/SMutableMeshViewer.cpp
+++ b/Engine/Plugins/Experimental/Mutable/Source/CustomizableObjectEditor/Private/MuCOE/SMutableMeshViewer.cpp
@@ -273,6 +273,7 @@ TSharedRef<SWidget> SMutableMeshViewer::GenerateDataTableSlates()
 	const FText VerticesCountTitle = LOCTEXT("VerticesCountTitle", "Vertex count : ");
 	const FText FacesCountTitle = LOCTEXT("FacesCountTitle", "Face count : ");
 	const FText BonesCountTitle = LOCTEXT("BonesCountTitle", "Bone count : ");
+	const FText VertexIdPrefixTitle = LOCTEXT("VertexIdPrefixTitle", "Vertex ID prefix : ");
 	const FText BuffersTitle = LOCTEXT("BuffersTitle", "Buffers");
 	
 	return SNew(SScrollBox)
@@ -350,7 +351,6 @@ TSharedRef<SWidget> SMutableMeshViewer::GenerateDataTableSlates()
 				[
 					SNew(SHorizontalBox)
 
-					// Faces title
 					+ SHorizontalBox::Slot()
 					.AutoWidth()
 					[
@@ -358,7 +358,6 @@ TSharedRef<SWidget> SMutableMeshViewer::GenerateDataTableSlates()
 						Text(BonesCountTitle)
 					]
 
-					// Faces Value
 					+ SHorizontalBox::Slot()
 					.AutoWidth()
 					[
@@ -366,6 +365,28 @@ TSharedRef<SWidget> SMutableMeshViewer::GenerateDataTableSlates()
 						Text(this, &SMutableMeshViewer::GetBoneCount)
 					]
 				]
+
+				// Vertex ID
+				+ SVerticalBox::Slot().
+				Padding(0, SimpleSpacing).
+				AutoHeight()
+				[
+					SNew(SHorizontalBox)
+
+						+ SHorizontalBox::Slot()
+						.AutoWidth()
+						[
+							SNew(STextBlock).
+								Text(VertexIdPrefixTitle)
+						]
+
+						+ SHorizontalBox::Slot()
+						.AutoWidth()
+						[
+							SNew(STextBlock).
+								Text(this, &SMutableMeshViewer::GetVertexIdPrefix)
+						]
+				]
 			]
 
 			// Buffers Data --------------------------------------------------------------
@@ -420,23 +441,6 @@ TSharedRef<SWidget> SMutableMeshViewer::GenerateDataTableSlates()
 					// ---------------------------------
 				]
 				
-				+ SVerticalBox::Slot()
-				  .Padding(IndentationSpace, 6)
-				  .AutoHeight()
-				[
-					SNew(SVerticalBox)
-				
-					// List of Face buffers ----------
-					+ SVerticalBox::Slot()
-					.AutoHeight()
-					[
-						GenerateBuffersListView(
-							FaceBuffersSlateView,
-							FaceBuffers,
-							FText(LOCTEXT("FaceBufferType", "Face")))
-					]
-					// ---------------------------------
-				]
 			]
 
 			// Bones data ----------------------------------------------------------------
@@ -567,12 +571,10 @@ void SMutableMeshViewer::OnMeshChanged()
 	// Cache the data accessible from the mu::Mesh to be later used by the UI
 	FillTargetBufferSetDataArray(MutableMesh->GetVertexBuffers(), VertexBuffers, VertexBuffersSlateView);
 	FillTargetBufferSetDataArray(MutableMesh->GetIndexBuffers(), IndexBuffers, IndexBuffersSlateView);
-	FillTargetBufferSetDataArray(MutableMesh->GetFaceBuffers(), FaceBuffers, FaceBuffersSlateView);
 
 	// Restore the widths of the columns each time the mesh gets changed.
 	VertexBuffersSlateView->GetHeaderRow()->ResetColumnWidths();
 	IndexBuffersSlateView->GetHeaderRow()->ResetColumnWidths();
-	FaceBuffersSlateView->GetHeaderRow()->ResetColumnWidths();
 }
 
 
@@ -688,6 +690,10 @@ FText SMutableMeshViewer::GetBoneCount() const
 	return FText::AsNumber(MutableMesh && MutableMesh->GetSkeleton() ? MutableMesh->GetSkeleton()->GetBoneCount() : 0);
 }
 
+FText SMutableMeshViewer::GetVertexIdPrefix() const
+{
+	return FText::AsNumber(MutableMesh ? MutableMesh->VertexIDPrefix : 0);
+}
 
 
 TSharedRef<ITableRow> SMutableMeshViewer::OnGenerateBufferRow(TSharedPtr<FBufferElement, ESPMode::ThreadSafe> InBuffer,
diff --git a/Engine/Plugins/Experimental/Mutable/Source/CustomizableObjectEditor/Private/MuCOE/SMutableMeshViewer.h b/Engine/Plugins/Experimental/Mutable/Source/CustomizableObjectEditor/Private/MuCOE/SMutableMeshViewer.h
--- a/Engine/Plugins/Experimental/Mutable/Source/CustomizableObjectEditor/Private/MuCOE/SMutableMeshViewer.h
+++ b/Engine/Plugins/Experimental/Mutable/Source/CustomizableObjectEditor/Private/MuCOE/SMutableMeshViewer.h
@@ -98,7 +98,6 @@ private:
 	
 	TSharedPtr<SListView<TSharedPtr<FBufferElement>>> VertexBuffersSlateView;
 	TSharedPtr<SListView<TSharedPtr<FBufferElement>>> IndexBuffersSlateView;
-	TSharedPtr<SListView<TSharedPtr<FBufferElement>>> FaceBuffersSlateView;
 
 private:
 
@@ -124,6 +123,7 @@ private:
 	FText GetVertexCount() const;
 	FText GetFaceCount() const;
 	FText GetBoneCount() const;
+	FText GetVertexIdPrefix() const;
 
 	/*
 	 * Elements used to feed the buffers list (index and buffer channels as an internal list)
@@ -131,7 +131,6 @@ private:
 	
 	TArray<TSharedPtr<FBufferElement>> VertexBuffers;
 	TArray<TSharedPtr<FBufferElement>> IndexBuffers;
-	TArray<TSharedPtr<FBufferElement>> FaceBuffers;
 	
 	/** 
 	* Fills the provided TArray with the buffer definitions generated from the mutable buffers found on the provided mutable buffer set
diff --git a/Engine/Plugins/Experimental/Mutable/Source/MutableRuntime/Internal/MuR/ConvertData.h b/Engine/Plugins/Experimental/Mutable/Source/MutableRuntime/Internal/MuR/ConvertData.h
--- a/Engine/Plugins/Experimental/Mutable/Source/MutableRuntime/Internal/MuR/ConvertData.h
+++ b/Engine/Plugins/Experimental/Mutable/Source/MutableRuntime/Internal/MuR/ConvertData.h
@@ -709,6 +709,47 @@ namespace mu
 			break;
 		}
 
+		case MBF_UINT64:
+		{
+			uint64* pTypedResult = reinterpret_cast<uint64*>(pResult);
+
+			switch (sourceFormat)
+			{
+			case MBF_UINT8:
+			{
+				const uint8* pTypedSource = reinterpret_cast<const uint8*>(pSource);
+				pTypedResult[channel] = pTypedSource[channel];
+				break;
+			}
+
+			case MBF_UINT16:
+			{
+				const uint16* pTypedSource = reinterpret_cast<const uint16*>(pSource);
+				pTypedResult[channel] = pTypedSource[channel];
+				break;
+			}
+
+			case MBF_UINT32:
+			{
+				const uint32* pTypedSource = reinterpret_cast<const uint32*>(pSource);
+				pTypedResult[channel] = pTypedSource[channel];
+				break;
+			}
+
+			case MBF_UINT64:
+			{
+				const uint64* pTypedSource = reinterpret_cast<const uint64*>(pSource);
+				pTypedResult[channel] = pTypedSource[channel];
+				break;
+			}
+
+			default:
+				checkf(false, TEXT("Conversion not implemented."));
+				break;
+			}
+			break;
+		}
+
 		case MBF_INT8:
 		{
             int8* pTypedResult = reinterpret_cast<int8*>( pResult );
diff --git a/Engine/Plugins/Experimental/Mutable/Source/MutableRuntime/Internal/MuR/MeshPrivate.h b/Engine/Plugins/Experimental/Mutable/Source/MutableRuntime/Internal/MuR/MeshPrivate.h
--- a/Engine/Plugins/Experimental/Mutable/Source/MutableRuntime/Internal/MuR/MeshPrivate.h
+++ b/Engine/Plugins/Experimental/Mutable/Source/MutableRuntime/Internal/MuR/MeshPrivate.h
@@ -242,13 +242,20 @@ namespace mu
 
 			return res;
 		}
-		
-        uint32_t GetAsUINT32() const
-        {
-            uint32_t res;
-            ConvertData( 0, &res, MBF_UINT32, ptr(), m_format );
-            return res;
-        }
+
+		uint32 GetAsUINT32() const
+		{
+			uint32 res;
+			ConvertData(0, &res, MBF_UINT32, ptr(), m_format);
+			return res;
+		}
+
+		uint64 GetAsUINT64() const
+		{
+			uint64 res;
+			ConvertData(0, &res, MBF_UINT64, ptr(), m_format);
+			return res;
+		}
 
 		void SetFromUINT32(uint32_t v)
 		{
@@ -510,14 +517,21 @@ namespace mu
 
 			return res;
 		}
-		
-        uint32_t GetAsUINT32() const
-        {
-            uint32_t res = 0;
-            ConvertData( 0, &res, MBF_UINT32, ptr(), m_format );
-            return res;
-        }
- 
+
+		uint32 GetAsUINT32() const
+		{
+			uint32 res = 0;
+			ConvertData(0, &res, MBF_UINT32, ptr(), m_format);
+			return res;
+		}
+
+		uint64 GetAsUINT64() const
+		{
+			uint64 res;
+			ConvertData(0, &res, MBF_UINT64, ptr(), m_format);
+			return res;
+		}
+
 		void GetAsInt32Vec(int32* Data, int32 Count) const
 		{
 			for (int32 c = 0; c < FMath::Min(m_components, Count); ++c)
@@ -622,9 +636,129 @@ namespace mu
 	};
 
 
-	//---------------------------------------------------------------------------------------------
-	// Optimised mesh formats that are identified in some operations to chose a faster version
-	//---------------------------------------------------------------------------------------------
+	/** */
+	class MUTABLERUNTIME_API MeshVertexIdIteratorConst
+	{
+	private:
+
+		/** Current Id that the iterator is pointing at. */
+		int32 CurrentIdIndex=0;
+
+		/** */
+		Ptr<const Mesh> Mesh;
+
+		/** Buffer iterator in case there is an actual Id buffer. */
+		UntypedMeshBufferIteratorConst BufferIterator;
+
+	public:
+
+		inline MeshVertexIdIteratorConst()
+		{
+		}
+
+		inline MeshVertexIdIteratorConst( const mu::Mesh* InMesh )
+		{
+			if (!InMesh)
+			{
+				return;
+			}
+
+			Mesh = InMesh;
+			if (Mesh)
+			{
+				BufferIterator = UntypedMeshBufferIteratorConst(Mesh->VertexBuffers, MBS_VERTEXINDEX, 0 );
+			}
+		}		
+
+		inline void operator++()
+		{
+			if (!Mesh)
+			{
+				return;
+			}
+
+			CurrentIdIndex++;			
+
+			if (BufferIterator.ptr())
+			{
+				BufferIterator++;
+			}
+		}
+
+		inline void operator++(int32)
+		{
+			CurrentIdIndex++;
+
+			if (BufferIterator.ptr())
+			{
+				BufferIterator++;
+			}
+		}
+
+		inline void operator+=(int32 c)
+		{
+			CurrentIdIndex+=c;
+
+			if (BufferIterator.ptr())
+			{
+				BufferIterator+=c;
+			}
+		}
+
+		inline MeshVertexIdIteratorConst operator+(int c) const
+		{
+			MeshVertexIdIteratorConst res = *this;
+			res += c;
+			return res;
+		}
+
+		bool IsValid() const
+		{
+			return 
+				Mesh 
+				&&
+				Mesh->HasVertexIds()
+				&&
+				(CurrentIdIndex < Mesh->GetVertexCount())
+				;
+		}
+
+		uint64 Get() const
+		{
+			check(Mesh);
+			check(CurrentIdIndex < Mesh->GetVertexCount());
+
+			// Is it implicit?
+			if (Mesh->VertexIDPrefix && !BufferIterator.ptr())
+			{
+				// The id is just prefix and index
+				return (uint64(Mesh->VertexIDPrefix) << 32) | uint64(CurrentIdIndex);
+			}
+
+			// Is it relative?
+			else if (BufferIterator.GetFormat() == MBF_UINT32)
+			{
+				// There is a buffer storing IDs without prefix because it is the same for all vertices.
+				uint32 RelativeId = BufferIterator.GetAsUINT32();
+				return (uint64(Mesh->VertexIDPrefix) << 32) | uint64(RelativeId);
+			}
+
+			// Is it explicit?
+			else if (BufferIterator.GetFormat() == MBF_UINT64)
+			{
+				check( Mesh->VertexIDPrefix==0 );
+				uint64 Id = BufferIterator.GetAsUINT64();
+				return Id;
+			}
+
+			check(false);
+			return 0;
+		}
+
+	};
+
+
+	/** Optimised mesh formats that are identified in some operations to chose a faster version. */
 	typedef enum
 	{
 		SMF_NONE,
diff --git a/Engine/Plugins/Experimental/Mutable/Source/MutableRuntime/Internal/MuR/OpMeshFormat.h b/Engine/Plugins/Experimental/Mutable/Source/MutableRuntime/Internal/MuR/OpMeshFormat.h
--- a/Engine/Plugins/Experimental/Mutable/Source/MutableRuntime/Internal/MuR/OpMeshFormat.h
+++ b/Engine/Plugins/Experimental/Mutable/Source/MutableRuntime/Internal/MuR/OpMeshFormat.h
@@ -22,7 +22,6 @@ namespace mu
 			bool bKeepSystemBuffers,
 			bool bFormatVertices,
 			bool bFormatIndices,
-			bool bFormatFaces,
 			bool bIgnoreMissingChannels,
 			bool& bOutSuccess
 		);
diff --git a/Engine/Plugins/Experimental/Mutable/Source/MutableRuntime/Internal/MuR/OpMeshMerge.h b/Engine/Plugins/Experimental/Mutable/Source/MutableRuntime/Internal/MuR/OpMeshMerge.h
--- a/Engine/Plugins/Experimental/Mutable/Source/MutableRuntime/Internal/MuR/OpMeshMerge.h
+++ b/Engine/Plugins/Experimental/Mutable/Source/MutableRuntime/Internal/MuR/OpMeshMerge.h
@@ -21,8 +21,7 @@ namespace mu
 	//---------------------------------------------------------------------------------------------
 	//! Merge two meshes into one new mesh
 	//---------------------------------------------------------------------------------------------
-	inline void MeshMerge(Mesh* Result, const Mesh* pFirst, const Mesh* pSecond, bool bMergeSurfaces, 
-		FMeshMergeScratchMeshes& ScratchMeshes)
+	inline void MeshMerge(Mesh* Result, const Mesh* pFirst, const Mesh* pSecond, bool bMergeSurfaces, FMeshMergeScratchMeshes& ScratchMeshes)
 	{
 		MUTABLE_CPUPROFILER_SCOPE(MeshMerge);
 
@@ -327,91 +326,26 @@ namespace mu
 			}
 		}
 
-		// Faces
-		//-----------------
-		{
-			MUTABLE_CPUPROFILER_SCOPE(Faces);
-
-			const int32 FirstCount = pFirst->GetFaceBuffers().GetElementCount();
-			const int32 SecondCount = pSecond->GetFaceBuffers().GetElementCount();
-			Result->GetFaceBuffers().SetElementCount(FirstCount + SecondCount);
-			Result->GetFaceBuffers().SetBufferCount(pFirst->GetFaceBuffers().GetBufferCount());
-
-			// Merge only the buffers present in the first mesh
-			for (int32 b = 0; b < Result->GetFaceBuffers().GetBufferCount(); ++b)
-			{
-				const FMeshBuffer& first = pFirst->GetFaceBuffers().m_buffers[b];
-
-				FMeshBuffer& result = Result->GetFaceBuffers().m_buffers[b];
-				result.m_channels = first.m_channels;
-				result.m_elementSize = first.m_elementSize;
-				result.m_data.SetNum(result.m_elementSize * (FirstCount + SecondCount));
-
-				EMeshBufferSemantic semantic = MBS_NONE;
-				int semanticIndex = 0;
-				EMeshBufferFormat type = MBF_NONE;
-				int components = 0;
-				int offset = 0;
-				pFirst->GetFaceBuffers().GetChannel
-				(b, 0, &semantic, &semanticIndex, &type, &components, &offset);
-
-				if (FirstCount)
-				{
-					FMemory::Memcpy(&result.m_data[0],
-						&first.m_data[0],
-						result.m_elementSize * FirstCount
-					);
-				}
-
-				if (SecondCount)
-				{
-					// Find in the second mesh
-					int otherBuffer = -1;
-					int otherChannel = -1;
-					pSecond->GetFaceBuffers().FindChannel
-					(semantic, 0, &otherBuffer, &otherChannel);
-
-					if (otherBuffer >= 0)
-					{
-						const FMeshBuffer& second =
-							pSecond->GetFaceBuffers().m_buffers[otherBuffer];
-						check(first.m_channels == second.m_channels);
-
-						// Raw copy
-						FMemory::Memcpy(&result.m_data[result.m_elementSize * FirstCount],
-							&second.m_data[0],
-							result.m_elementSize * SecondCount);
-					}
-					else
-					{
-						// Fill with zeroes
-						FMemory::Memzero(&result.m_data[result.m_elementSize * FirstCount],
-							result.m_elementSize * SecondCount);
-					}
-				}
-			}
-		}
-
 
 		// Layouts
 		//-----------------
 		{
 			MUTABLE_CPUPROFILER_SCOPE(Layouts);
 
-			Result->m_layouts.SetNum(pFirst->m_layouts.Num());
-			for (int i = 0; i < pFirst->m_layouts.Num(); ++i)
+			Result->Layouts.SetNum(pFirst->Layouts.Num());
+			for (int i = 0; i < pFirst->Layouts.Num(); ++i)
 			{
-				const Layout* pF = pFirst->m_layouts[i].get();
+				const Layout* pF = pFirst->Layouts[i].get();
 				LayoutPtr pR = pF->Clone();
 
-				if (i < pSecond->m_layouts.Num())
+				if (i < pSecond->Layouts.Num())
 				{
-					const Layout* pS = pSecond->m_layouts[i].get();
+					const Layout* pS = pSecond->Layouts[i].get();
 
 					pR->m_blocks.Append(pS->m_blocks);
 				}
 
-				Result->m_layouts[i] = pR;
+				Result->Layouts[i] = pR;
 			}
 		}
 
@@ -524,13 +458,13 @@ namespace mu
 					bRemapBoneIndices = bRemapBoneIndices || BoneMapIndex != SecondBoneMapIndex;
 				}
 
-				MESH_SURFACE& NewSurface = Result->m_surfaces.AddDefaulted_GetRef();
-				NewSurface.m_vertexCount = pFirst->GetVertexCount() + pSecond->GetVertexCount();
-				NewSurface.m_indexCount = pFirst->GetIndexCount() + pSecond->GetIndexCount();
+				FMeshSurface& NewSurface = Result->Surfaces.AddDefaulted_GetRef();
+				NewSurface.VertexCount = pFirst->GetVertexCount() + pSecond->GetVertexCount();
+				NewSurface.IndexCount = pFirst->GetIndexCount() + pSecond->GetIndexCount();
 				NewSurface.BoneMapCount = Result->BoneMap.Num();
 				
 				//All merged surfaces will have the same bCastShadow value. Decided by the first merged mesh
-				NewSurface.bCastShadow = pFirst->m_surfaces.Last().bCastShadow;
+				NewSurface.bCastShadow = pFirst->Surfaces.Last().bCastShadow;
 			}
 			else
 			{
@@ -538,19 +472,19 @@ namespace mu
 				Result->BoneMap.Append(pSecond->BoneMap);
 
 				// Add pFirst surfaces
-				Result->m_surfaces = pFirst->m_surfaces;
+				Result->Surfaces = pFirst->Surfaces;
 
 				const int32 FirstVertexIndex = pFirst->GetVertexCount();
 				const int32 FirstIndexIndex = pFirst->GetIndexCount();
 
-				check(pSecond->m_surfaces.Num() == 1);
-				MESH_SURFACE& NewSurface = Result->m_surfaces.Add_GetRef(pSecond->m_surfaces[0]);
-				NewSurface.m_firstVertex += FirstVertexIndex;
-				NewSurface.m_firstIndex += FirstIndexIndex;
+				check(pSecond->Surfaces.Num() == 1);
+				FMeshSurface& NewSurface = Result->Surfaces.Add_GetRef(pSecond->Surfaces[0]);
+				NewSurface.FirstVertex += FirstVertexIndex;
+				NewSurface.FirstIndex += FirstIndexIndex;
 				NewSurface.BoneMapIndex += NumFirstBonesInBoneMap;
 			}
 
-			for (const MESH_SURFACE& Surface : Result->m_surfaces)
+			for (const FMeshSurface& Surface : Result->Surfaces)
 			{
 				MaxNumBonesInBoneMaps = FMath::Max(MaxNumBonesInBoneMaps, Surface.BoneMapCount);
 			}
@@ -753,15 +687,20 @@ namespace mu
             const int32 FirstCount = pFirst->GetVertexBuffers().GetElementCount();
 			const int32 SecondCount = pSecond->GetVertexBuffers().GetElementCount();
 
+			// TODO: fast path should be per-buffer
 			// TODO: when formats match, which at runtime should be always.
 			bool bFastPath = pFirst->GetVertexBuffers().HasSameFormat( pSecond->GetVertexBuffers() );
 
+			bool bNeedsExplicitVertexIds = pFirst->VertexIDPrefix != pSecond->VertexIDPrefix;
+			bFastPath = bFastPath && !bNeedsExplicitVertexIds;
+
 			// Check if the format of the BoneIndex buffer has to change
 			bool bChangeBoneIndicesFormat = false;
 			EMeshBufferFormat BoneIndexFormat = MaxNumBonesInBoneMaps > MAX_uint8 ? MBF_UINT16 : MBF_UINT8;
 
 			// Iterate all vertex buffers to check if we need to format bone indices
 			{
+				// TODO: deduplicate code. move all to cpp.
 				{
 					// Check if pFirst requires a reformat of the bone index buffers
 					const FMeshBufferSet& VertexBuffers = pFirst->GetVertexBuffers();
@@ -820,7 +759,7 @@ namespace mu
 
 				// Expand component counts in vertex channels of the format mesh
 				int32 vbcount = pFirst->GetVertexBuffers().m_buffers.Num();
-				Result->GetVertexBuffers().SetBufferCount( (int)vbcount );
+				Result->GetVertexBuffers().SetBufferCount( vbcount );
 
 				for ( int32 vb = 0; vb<vbcount; ++vb )
 				{
@@ -834,8 +773,8 @@ namespace mu
 					bool resetOffsets = false;
 					for ( int32 c=0; c<result.m_channels.Num(); ++c )
 					{
-						int sb = -1;
-						int sc = -1;
+						int32 sb = -1;
+						int32 sc = -1;
 						pSecond->GetVertexBuffers().FindChannel
 								(
 									result.m_channels[c].m_semantic,
@@ -860,7 +799,7 @@ namespace mu
 					// Reset the channel offsets if necessary
 					if (resetOffsets)
 					{
-						int offset = 0;
+						int32 offset = 0;
 						for ( int32 c=0; c<result.m_channels.Num(); ++c )
 						{
                             result.m_channels[c].m_offset = (uint8_t)offset;
@@ -890,8 +829,8 @@ namespace mu
 
 						someChannel = true;
 
-						int foundBuffer = -1;
-						int foundChannel = -1;
+						int32 foundBuffer = -1;
+						int32 foundChannel = -1;
 						pFirst->GetVertexBuffers().FindChannel(chan.m_semantic, chan.m_semanticIndex, &foundBuffer, &foundChannel);
 						if (foundBuffer >= 0)
 						{
@@ -964,12 +903,9 @@ namespace mu
 						}
 					}
 				}
-
-                // Allocate vertices
-                Result->GetVertexBuffers().SetElementCount( FirstCount + SecondCount );
 				
 				// Convert the source meshes to the new format
-                if (pFirst->GetVertexBuffers().HasSameFormat(Result->GetVertexBuffers()))
+                if (!bNeedsExplicitVertexIds && pFirst->GetVertexBuffers().HasSameFormat(Result->GetVertexBuffers()))
                 {
                     pVFirst = pFirst;
                 }
@@ -978,16 +914,26 @@ namespace mu
 					check(ScratchMeshes.FirstReformat);
 
 					bool bOutSuccess = false;
-					MeshFormat(ScratchMeshes.FirstReformat.get(), pFirst, Result, false, true, false, false, false, bOutSuccess);
-					
+					MeshFormat(ScratchMeshes.FirstReformat.get(), pFirst, Result, false, true, false, false, bOutSuccess);
+
+					if (bNeedsExplicitVertexIds)
+					{
+						ScratchMeshes.FirstReformat->MakeVertexIndicesExplicit();
+					}
+
 					if (bOutSuccess)
 					{
 						pVFirst = ScratchMeshes.FirstReformat;
 					}
+					else
+					{
+						check(false);
+					}
+
                 }
 
 
-                if (pSecond->GetVertexBuffers().HasSameFormat(Result->GetVertexBuffers()))
+                if (!bNeedsExplicitVertexIds && pSecond->GetVertexBuffers().HasSameFormat(Result->GetVertexBuffers()))
                 {
                     pVSecond = pSecond;
                 }
@@ -996,13 +942,22 @@ namespace mu
 					check(ScratchMeshes.SecondReformat);
 
 					bool bOutSuccess = false;
-                    MeshFormat(ScratchMeshes.SecondReformat.get(), pSecond, Result, false, true, false, false, false, bOutSuccess);
+                    MeshFormat(ScratchMeshes.SecondReformat.get(), pSecond, Result, false, true, false, false, bOutSuccess);
+
+					if (bNeedsExplicitVertexIds)
+					{
+						ScratchMeshes.SecondReformat->MakeVertexIndicesExplicit();
+					}
 
 					if (bOutSuccess)
 					{
 						pVSecond = ScratchMeshes.SecondReformat;
 					}
-                }
+					else
+					{
+						check(false);
+					}
+				}
 
 				check(pVFirst->GetVertexBuffers().HasSameFormat(pVSecond->GetVertexBuffers()));
 			}
@@ -1014,10 +969,12 @@ namespace mu
 				pVSecond = pSecond;
 			}
 
-			Result->m_VertexBuffers = pVFirst->m_VertexBuffers;
-			Result->GetVertexBuffers().SetElementCount( FirstCount + SecondCount );
+			Result->VertexBuffers = pVFirst->VertexBuffers;
+
+			// Allocate additional vertices
+			Result->VertexBuffers.SetElementCount( FirstCount + SecondCount );
 
-            // first copy all the vertex data
+            // Copy the second mesh vertex data
 			{
 				MUTABLE_CPUPROFILER_SCOPE(CopyVertexData);
 				for (int32 vb = 0; vb < Result->GetVertexBuffers().m_buffers.Num(); ++vb)
@@ -1027,11 +984,11 @@ namespace mu
 
 					if (SecondCount)
 					{
-						int elemSize = Result->GetVertexBuffers().GetElementSize((int)vb);
-						int firstSize = FirstCount * elemSize;
-						int secondSize = SecondCount * elemSize;
+						int32 elemSize = Result->GetVertexBuffers().GetElementSize(vb);
+						int32 firstSize = FirstCount * elemSize;
+						int32 secondSize = SecondCount * elemSize;
 
-						FMemory::Memcpy( &result.m_data[firstSize], &second.m_data[0], secondSize );
+						FMemory::Memcpy( &result.m_data[firstSize], second.m_data.GetData(), secondSize);
 					}
 				}
 			}
@@ -1139,11 +1096,11 @@ namespace mu
 		}
 
 		// Tags
-		Result->m_tags = pFirst->m_tags;
+		Result->Tags = pFirst->Tags;
 
-		for (const FString& SecondTag : pSecond->m_tags)
+		for (const FString& SecondTag : pSecond->Tags)
 		{
-			Result->m_tags.AddUnique(SecondTag);
+			Result->Tags.AddUnique(SecondTag);
 		}
 
 
diff --git a/Engine/Plugins/Experimental/Mutable/Source/MutableRuntime/Internal/MuR/Operations.h b/Engine/Plugins/Experimental/Mutable/Source/MutableRuntime/Internal/MuR/Operations.h
--- a/Engine/Plugins/Experimental/Mutable/Source/MutableRuntime/Internal/MuR/Operations.h
+++ b/Engine/Plugins/Experimental/Mutable/Source/MutableRuntime/Internal/MuR/Operations.h
@@ -881,7 +881,7 @@ namespace mu
             {
                 Vertex				= 1 << 0,
                 Index				= 1 << 1,
-                Face				= 1 << 2,
+                // deprecated Face  = 1 << 2,
 
                 /** This flag will not add blank channels for the channels in the format mesh but not in the source mesh. */
 				IgnoreMissing		= 1 << 4,
diff --git a/Engine/Plugins/Experimental/Mutable/Source/MutableRuntime/Internal/MuR/SystemPrivate.h b/Engine/Plugins/Experimental/Mutable/Source/MutableRuntime/Internal/MuR/SystemPrivate.h
--- a/Engine/Plugins/Experimental/Mutable/Source/MutableRuntime/Internal/MuR/SystemPrivate.h
+++ b/Engine/Plugins/Experimental/Mutable/Source/MutableRuntime/Internal/MuR/SystemPrivate.h
@@ -1516,6 +1516,8 @@ namespace mu
 		{
 			if (Resource)
 			{
+				//Resource->CheckIntegrity();
+
 				TempMeshes.RemoveSingle(Resource);
 
 				int32& Count = CacheResources.FindOrAdd(Resource, 0);
diff --git a/Engine/Plugins/Experimental/Mutable/Source/MutableRuntime/Private/MuR/CodeRunner.cpp b/Engine/Plugins/Experimental/Mutable/Source/MutableRuntime/Private/MuR/CodeRunner.cpp
--- a/Engine/Plugins/Experimental/Mutable/Source/MutableRuntime/Private/MuR/CodeRunner.cpp
+++ b/Engine/Plugins/Experimental/Mutable/Source/MutableRuntime/Private/MuR/CodeRunner.cpp
@@ -1298,7 +1298,7 @@ namespace mu
                     if (args.newSurfaceID)
                     {
 						check(pB->GetSurfaceCount() == 1);
-						Result->m_surfaces.Last().m_id = args.newSurfaceID;
+						Result->Surfaces.Last().Id = args.newSurfaceID;
                     }
 
 					Release(pA);
@@ -1318,7 +1318,7 @@ namespace mu
 
                     if (Result->GetSurfaceCount() > 0 && args.newSurfaceID)
                     {
-                        Result->m_surfaces.Last().m_id = args.newSurfaceID;
+                        Result->Surfaces.Last().Id = args.newSurfaceID;
                     }
 
 					Release(pA);
@@ -1820,7 +1820,6 @@ namespace mu
 							true,
 							(Flags & OP::MeshFormatArgs::Vertex) != 0,
 							(Flags & OP::MeshFormatArgs::Index) != 0,
-							(Flags & OP::MeshFormatArgs::Face) != 0,
 							(Flags & OP::MeshFormatArgs::IgnoreMissing) != 0,
 							bOutSuccess);
 
@@ -1896,7 +1895,8 @@ namespace mu
 
                 // Access with memcpy necessary for unaligned arm issues.
                 uint32 blocks[1024];
-				FMemory::Memcpy(blocks, data, sizeof(uint32)*FMath::Min(1024,int(blockCount)));
+				check(blockCount<1024);
+				FMemory::Memcpy(blocks, data, sizeof(uint32)*FMath::Min(1024,int32(blockCount)));
 
 				if (Source)
 				{
@@ -2879,8 +2879,8 @@ namespace mu
 					FMemory::Memcpy(&TagCount, Data, sizeof(uint16));
 					Data += sizeof(uint16);
 
-					int32 FirstMeshTagIndex = Result->m_tags.Num();
-					Result->m_tags.SetNum(FirstMeshTagIndex+TagCount);
+					int32 FirstMeshTagIndex = Result->Tags.Num();
+					Result->Tags.SetNum(FirstMeshTagIndex+TagCount);
 					for (uint16 TagIndex = 0; TagIndex < TagCount; ++TagIndex)
 					{
 						OP::ADDRESS TagConstant;
@@ -2889,7 +2889,7 @@ namespace mu
 
 						check(TagConstant < (uint32)pModel->GetPrivate()->m_program.m_constantStrings.Num());
 						const FString& Name = Program.m_constantStrings[TagConstant];
-						Result->m_tags[FirstMeshTagIndex+TagIndex] = Name;
+						Result->Tags[FirstMeshTagIndex+TagIndex] = Name;
 					}
 
 					StoreMesh(item, Result);
diff --git a/Engine/Plugins/Experimental/Mutable/Source/MutableRuntime/Private/MuR/Mesh.cpp b/Engine/Plugins/Experimental/Mutable/Source/MutableRuntime/Private/MuR/Mesh.cpp
--- a/Engine/Plugins/Experimental/Mutable/Source/MutableRuntime/Private/MuR/Mesh.cpp
+++ b/Engine/Plugins/Experimental/Mutable/Source/MutableRuntime/Private/MuR/Mesh.cpp
@@ -13,15 +13,19 @@
 namespace mu
 {
 
+MUTABLE_IMPLEMENT_ENUM_SERIALISABLE(EBoneUsageFlags);
+MUTABLE_IMPLEMENT_ENUM_SERIALISABLE(EMeshBufferType);
+MUTABLE_IMPLEMENT_ENUM_SERIALISABLE(EShapeBindingMethod);
+MUTABLE_IMPLEMENT_ENUM_SERIALISABLE(EVertexColorUsage);
+
 
-//---------------------------------------------------------------------------------------------
 void Mesh::Serialise( const Mesh* p, OutputArchive& arch )
 {
     //p->m_pD->CheckIntegrity();
     arch << *p;
 }
 
-//---------------------------------------------------------------------------------------------
+
 MeshPtr Mesh::StaticUnserialise( InputArchive& arch )
 {
     MUTABLE_CPUPROFILER_SCOPE(MeshUnserialise)
@@ -36,7 +40,6 @@ MeshPtr Mesh::StaticUnserialise( InputArchive& arch )
 }
 
 
-//---------------------------------------------------------------------------------------------
 MeshPtr Mesh::Clone() const
 {
     //MUTABLE_CPUPROFILER_SCOPE(MeshClone);
@@ -44,24 +47,24 @@ MeshPtr Mesh::Clone() const
 
     MeshPtr pResult = new Mesh();
 
-    pResult->m_internalId = m_internalId;
-	pResult->m_staticFormatFlags = m_staticFormatFlags;
-	pResult->m_surfaces = m_surfaces;
-	pResult->m_pSkeleton = m_pSkeleton;
-	pResult->m_pPhysicsBody = m_pPhysicsBody;
-	pResult->m_tags = m_tags;
+    pResult->InternalId = InternalId;
+	pResult->StaticFormatFlags = StaticFormatFlags;
+	pResult->Surfaces = Surfaces;
+	pResult->Skeleton = Skeleton;
+	pResult->PhysicsBody = PhysicsBody;
+	pResult->Tags = Tags;
 	pResult->StreamedResources = StreamedResources;
+	pResult->VertexIDPrefix = VertexIDPrefix;
 
     // Clone the main buffers
-    pResult->m_VertexBuffers = m_VertexBuffers;
-    pResult->m_IndexBuffers = m_IndexBuffers;
-    pResult->m_FaceBuffers = m_FaceBuffers;
+    pResult->VertexBuffers = VertexBuffers;
+    pResult->IndexBuffers = IndexBuffers;
 
 	// Clone additional buffers
 	pResult->AdditionalBuffers = AdditionalBuffers;
 
     // Clone the layouts
-	pResult->m_layouts = m_layouts;
+	pResult->Layouts = Layouts;
 
     // The skeleton is not cloned because it is not owned by this mesh and it is always assumed
     // to be shared.
@@ -79,6 +82,7 @@ MeshPtr Mesh::Clone() const
     return pResult;
 }
 
+
 MeshPtr Mesh::Clone(EMeshCopyFlags Flags) const
 {
     //MUTABLE_CPUPROFILER_SCOPE(MeshClone);
@@ -86,27 +90,27 @@ MeshPtr Mesh::Clone(EMeshCopyFlags Flags) const
 
     MeshPtr pResult = new Mesh();
 
-    pResult->m_internalId = m_internalId;
-	pResult->m_staticFormatFlags = m_staticFormatFlags;
+    pResult->InternalId = InternalId;
+	pResult->StaticFormatFlags = StaticFormatFlags;
 
 	if (EnumHasAnyFlags(Flags, EMeshCopyFlags::WithSurfaces))
 	{
-		pResult->m_surfaces = m_surfaces;
+		pResult->Surfaces = Surfaces;
 	}
 
 	if (EnumHasAnyFlags(Flags, EMeshCopyFlags::WithSkeleton))
 	{
-		pResult->m_pSkeleton = m_pSkeleton;
+		pResult->Skeleton = Skeleton;
 	}
 
 	if (EnumHasAnyFlags(Flags, EMeshCopyFlags::WithPhysicsBody))
 	{
-		pResult->m_pPhysicsBody = m_pPhysicsBody;
+		pResult->PhysicsBody = PhysicsBody;
 	}
 
 	if (EnumHasAnyFlags(Flags, EMeshCopyFlags::WithTags))
 	{
-		pResult->m_tags = m_tags;
+		pResult->Tags = Tags;
 	}
 
 	if (EnumHasAnyFlags(Flags, EMeshCopyFlags::WithStreamedResources))
@@ -117,17 +121,12 @@ MeshPtr Mesh::Clone(EMeshCopyFlags Flags) const
     // Clone the main buffers
 	if (EnumHasAnyFlags(Flags, EMeshCopyFlags::WithVertexBuffers))
     {
-		pResult->m_VertexBuffers = m_VertexBuffers;
+		pResult->VertexBuffers = VertexBuffers;
 	}
 
 	if (EnumHasAnyFlags(Flags, EMeshCopyFlags::WithIndexBuffers))
     {
-		pResult->m_IndexBuffers = m_IndexBuffers;
-	}
-
-	if (EnumHasAnyFlags(Flags, EMeshCopyFlags::WithFaceBuffers))
-    {
-		pResult->m_FaceBuffers = m_FaceBuffers;
+		pResult->IndexBuffers = IndexBuffers;
 	}
 
 	// Clone additional buffers
@@ -139,7 +138,7 @@ MeshPtr Mesh::Clone(EMeshCopyFlags Flags) const
     // Clone the layout	
 	if (EnumHasAnyFlags(Flags, EMeshCopyFlags::WithLayouts))
 	{
-		pResult->m_layouts = m_layouts;
+		pResult->Layouts = Layouts;
 	}
     // The skeleton is not cloned because it is not owned by this mesh and it is always assumed
     // to be shared.
@@ -173,32 +172,32 @@ MeshPtr Mesh::Clone(EMeshCopyFlags Flags) const
 }
 
 
-//---------------------------------------------------------------------------------------------
 void Mesh::CopyFrom(const Mesh& From, EMeshCopyFlags Flags)
 {
     //MUTABLE_CPUPROFILER_SCOPE(CopyFrom);
 
-    m_internalId = From.m_internalId;
-	m_staticFormatFlags = From.m_staticFormatFlags;
+    InternalId = From.InternalId;
+	StaticFormatFlags = From.StaticFormatFlags;
+	VertexIDPrefix = From.VertexIDPrefix;
 
 	if (EnumHasAnyFlags(Flags, EMeshCopyFlags::WithSurfaces))
 	{
-		m_surfaces = From.m_surfaces;
+		Surfaces = From.Surfaces;
 	}
 
 	if (EnumHasAnyFlags(Flags, EMeshCopyFlags::WithSkeleton))
 	{
-		m_pSkeleton = From.m_pSkeleton;
+		Skeleton = From.Skeleton;
 	}
 
 	if (EnumHasAnyFlags(Flags, EMeshCopyFlags::WithPhysicsBody))
 	{
-		m_pPhysicsBody = From.m_pPhysicsBody;
+		PhysicsBody = From.PhysicsBody;
 	}
 
 	if (EnumHasAnyFlags(Flags, EMeshCopyFlags::WithTags))
 	{
-		m_tags = From.m_tags;
+		Tags = From.Tags;
 	}
 
 	if (EnumHasAnyFlags(Flags, EMeshCopyFlags::WithStreamedResources))
@@ -209,17 +208,12 @@ void Mesh::CopyFrom(const Mesh& From, EMeshCopyFlags Flags)
     // Copy the main buffers
 	if (EnumHasAnyFlags(Flags, EMeshCopyFlags::WithVertexBuffers))
     {
-		m_VertexBuffers = From.m_VertexBuffers;
+		VertexBuffers = From.VertexBuffers;
 	}
 
 	if (EnumHasAnyFlags(Flags, EMeshCopyFlags::WithIndexBuffers))
     {
-		m_IndexBuffers = From.m_IndexBuffers;
-	}
-
-	if (EnumHasAnyFlags(Flags, EMeshCopyFlags::WithFaceBuffers))
-    {
-		m_FaceBuffers = From.m_FaceBuffers;
+		IndexBuffers = From.IndexBuffers;
 	}
 
 	// Copy additional buffers
@@ -231,7 +225,7 @@ void Mesh::CopyFrom(const Mesh& From, EMeshCopyFlags Flags)
     // Copy the layout	
 	if (EnumHasAnyFlags(Flags, EMeshCopyFlags::WithLayouts))
 	{
-		m_layouts = From.m_layouts;
+		Layouts = From.Layouts;
 	}
     // The skeleton is not copied because it is not owned by this mesh and it is always assumed
     // to be shared.
@@ -263,135 +257,249 @@ void Mesh::CopyFrom(const Mesh& From, EMeshCopyFlags Flags)
 
 }
 
-//---------------------------------------------------------------------------------------------
-uint32_t Mesh::GetId() const
+
+uint32 Mesh::GetId() const
 {
-    return m_internalId;
+    return InternalId;
 }
 
 
-//---------------------------------------------------------------------------------------------
 int Mesh::GetVertexCount() const
 {
     return GetVertexBuffers().GetElementCount();
 }
 
 
-//---------------------------------------------------------------------------------------------
 FMeshBufferSet& Mesh::GetVertexBuffers()
 {
-    return m_VertexBuffers;
+    return VertexBuffers;
 }
 
 
-//---------------------------------------------------------------------------------------------
 const FMeshBufferSet& Mesh::GetVertexBuffers() const
 {
-    return m_VertexBuffers;
+    return VertexBuffers;
 }
 
 
-//---------------------------------------------------------------------------------------------
-mu::Ptr<const Skeleton> Mesh::GetSkeleton() const
+bool Mesh::HasVertexIds() const
 {
-    return m_pSkeleton;
+	if (VertexIDPrefix != 0) return true;
+
+	// If there is no prefix, there must be a buffer with full indices
+	int32 BufferIndex = -1;
+	int32 ChannelIndex = -1;
+	VertexBuffers.FindChannel(MBS_VERTEXINDEX, 0, &BufferIndex, &ChannelIndex);
+	if (BufferIndex >= 0 && ChannelIndex >= 0)
+	{
+		if (VertexBuffers.m_buffers[BufferIndex].m_channels[ChannelIndex].m_format == MBF_UINT64)
+		{
+			return true;
+		}
+	}
+
+	return false;
 }
 
 
-//---------------------------------------------------------------------------------------------
-void Mesh::SetSkeleton( Ptr<const Skeleton> s )
+bool Mesh::AreVertexIdsImplicit() const
 {
-    m_pSkeleton = s;
+	// Is there a buffer for vertex ids?
+	int32 BufferIndex = -1;
+	int32 ChannelIndex = -1;
+	VertexBuffers.FindChannel(MBS_VERTEXINDEX, 0, &BufferIndex, &ChannelIndex);
+
+	return (VertexIDPrefix != 0) && (BufferIndex < 0) && (ChannelIndex < 0);
 }
 
-//---------------------------------------------------------------------------------------------
-mu::Ptr<const PhysicsBody> Mesh::GetPhysicsBody() const
+
+bool Mesh::AreVertexIdsExplicit() const
 {
-    return m_pPhysicsBody;
+	// Is there a buffer for vertex ids?
+	int32 BufferIndex = -1;
+	int32 ChannelIndex = -1;
+	VertexBuffers.FindChannel(MBS_VERTEXINDEX, 0, &BufferIndex, &ChannelIndex);
+
+	bool bExplicit = (BufferIndex >= 0) && (ChannelIndex >= 0) && (VertexBuffers.m_buffers[BufferIndex].m_channels[ChannelIndex].m_format == MBF_UINT64);
+	if (bExplicit)
+	{
+		check(VertexIDPrefix == 0);
+	}
+	return bExplicit;
 }
 
 
-//---------------------------------------------------------------------------------------------
-void Mesh::SetPhysicsBody( Ptr<const PhysicsBody> PhysicsBody )
+void Mesh::MakeVertexIndicesRelative()
 {
-    m_pPhysicsBody = PhysicsBody;
+	check(AreVertexIdsImplicit());
+
+	int32 NewBuffer = VertexBuffers.GetBufferCount();
+	VertexBuffers.SetBufferCount(NewBuffer + 1);
+	EMeshBufferSemantic Semantic = MBS_VERTEXINDEX;
+	int32 SemanticIndex = 0;
+	EMeshBufferFormat Format = MBF_UINT32;
+	int32 Components = 1;
+	int32 Offset = 0;
+	VertexBuffers.SetBuffer(NewBuffer, sizeof(uint32), 1, &Semantic, &SemanticIndex, &Format, &Components, &Offset );
+	uint32* pIdData = reinterpret_cast<uint32*>( VertexBuffers.GetBufferData(NewBuffer) );
+
+	int32 VertexCount = GetVertexCount();
+	for (int32 i = 0; i < VertexCount; ++i)
+	{
+		(*pIdData++) = i;
+	}
 }
 
-//---------------------------------------------------------------------------------------------
-int32 Mesh::AddAdditionalPhysicsBody(Ptr<const PhysicsBody> Body)
+
+void Mesh::MakeVertexIndicesExplicit()
 {
-	return AdditionalPhysicsBodies.Add(Body);
+	MUTABLE_CPUPROFILER_SCOPE(Mesh_MakeVertexIndicesExplicit);
+
+	if (!VertexIDPrefix)
+	{
+		// We already have explicit vertex IDs, or we don't have any.
+		return;
+	}
+
+	int32 VertexCount = GetVertexCount();
+
+	bool bHasRelativeVertexIndices = false;
+
+	int32 OldBuf = -1;
+	int32 OldChan = -1;
+	VertexBuffers.FindChannel(MBS_VERTEXINDEX, 0, &OldBuf, &OldChan);
+	bool bHasVertexIndices = (OldBuf >= 0 && OldChan >= 0);
+	if (bHasVertexIndices)
+	{
+		check(OldChan == 0 && VertexBuffers.m_buffers[OldBuf].m_channels.Num() == 1);
+
+		bool bHasExplicitVertexIndices = VertexBuffers.m_buffers[OldBuf].m_channels[0].m_format == MBF_UINT64;
+		if (bHasExplicitVertexIndices)
+		{
+			// nothing to do
+			return;
+		}
+				
+		// The mesh has relative vertex IDs.
+		bHasRelativeVertexIndices = true;
+
+		const uint32* OldIdData = reinterpret_cast<const uint32*>(VertexBuffers.GetBufferData(OldBuf));
+
+		// Create a new buffer with explicit ids and remove the old buffer with relative ids
+		int32 NewBuffer = VertexBuffers.GetBufferCount();
+		VertexBuffers.SetBufferCount(NewBuffer + 1);
+		EMeshBufferSemantic Semantic = MBS_VERTEXINDEX;
+		int32 SemanticIndex = 0;
+		EMeshBufferFormat Format = MBF_UINT64;
+		int32 Components = 1;
+		int32 Offset = 0;
+		VertexBuffers.SetBuffer(NewBuffer, sizeof(uint64), 1, &Semantic, &SemanticIndex, &Format, &Components, &Offset);
+		uint64* IdData = reinterpret_cast<uint64*>(VertexBuffers.GetBufferData(NewBuffer));
+
+		for (int32 i = 0; i < VertexCount; ++i)
+		{
+			uint32 OldId = *OldIdData++;
+			uint64 Id = (uint64(VertexIDPrefix) << 32) | uint64(OldId);
+			(*IdData++) = Id;
+		}
+
+		// TODO: Optimize by overwritting the existing buffer.
+		VertexBuffers.RemoveBuffer(OldBuf);
+	}
+
+	if (!bHasRelativeVertexIndices)
+	{
+		// The mesh has implicit Ids
+		// Create a new buffer with explicit ids
+		int32 NewBuffer = VertexBuffers.GetBufferCount();
+		VertexBuffers.SetBufferCount(NewBuffer + 1);
+		EMeshBufferSemantic Semantic = MBS_VERTEXINDEX;
+		int32 SemanticIndex = 0;
+		EMeshBufferFormat Format = MBF_UINT64;
+		int32 Components = 1;
+		int32 Offset = 0;
+		VertexBuffers.SetBuffer(NewBuffer, sizeof(uint64), 1, &Semantic, &SemanticIndex, &Format, &Components, &Offset);
+		uint64* IdData = reinterpret_cast<uint64*>(VertexBuffers.GetBufferData(NewBuffer));
+
+		for (int32 VertexIndex = 0; VertexIndex < VertexCount; ++VertexIndex)
+		{
+			uint64 Id = (uint64(VertexIDPrefix) << 32) | uint64(VertexIndex);
+			(*IdData++) = Id;
+		}
+	}
+
+	VertexIDPrefix = 0;
 }
 
-Ptr<const PhysicsBody> Mesh::GetAdditionalPhysicsBody(int32 I) const
+
+mu::Ptr<const Skeleton> Mesh::GetSkeleton() const
 {
-	check(I >= 0 && I < AdditionalPhysicsBodies.Num());
+    return Skeleton;
+}
 
-	return AdditionalPhysicsBodies[I];
+
+void Mesh::SetSkeleton( Ptr<const mu::Skeleton> s )
+{
+    Skeleton = s;
 }
 
 
-//---------------------------------------------------------------------------------------------
-int Mesh::GetFaceCount() const
+mu::Ptr<const PhysicsBody> Mesh::GetPhysicsBody() const
 {
-    int result = GetFaceBuffers().GetElementCount();
-    if (result)
-    {
-        check(GetIndexBuffers().GetElementCount() == result * 3);
-    }
-    else
-    {
-        result = GetIndexBuffers().GetElementCount() / 3;
-    }
-    return result;
+    return PhysicsBody;
 }
 
 
-//---------------------------------------------------------------------------------------------
-FMeshBufferSet& Mesh::GetFaceBuffers()
+void Mesh::SetPhysicsBody( Ptr<const mu::PhysicsBody> InPhysicsBody )
 {
-    return m_FaceBuffers;
+    PhysicsBody = InPhysicsBody;
 }
 
 
-//---------------------------------------------------------------------------------------------
-const FMeshBufferSet& Mesh::GetFaceBuffers() const
+int32 Mesh::AddAdditionalPhysicsBody(Ptr<const mu::PhysicsBody> Body)
 {
-    return m_FaceBuffers;
+	return AdditionalPhysicsBodies.Add(Body);
+}
+
+
+Ptr<const PhysicsBody> Mesh::GetAdditionalPhysicsBody(int32 I) const
+{
+	check(I >= 0 && I < AdditionalPhysicsBodies.Num());
+
+	return AdditionalPhysicsBodies[I];
+}
+
+
+int32 Mesh::GetFaceCount() const
+{
+    return GetIndexBuffers().GetElementCount() / 3;
 }
 
 
-//---------------------------------------------------------------------------------------------
 int Mesh::GetIndexCount() const
 {
-    // It is possible to ignore face buffers altogether
-    //check( GetIndexBuffers().GetElementCount()==GetFaceBuffers().GetElementCount()*3 );
     return GetIndexBuffers().GetElementCount();
 }
 
 
-//---------------------------------------------------------------------------------------------
 FMeshBufferSet& Mesh::GetIndexBuffers()
 {
-    return m_IndexBuffers;
+    return IndexBuffers;
 }
 
 
-//---------------------------------------------------------------------------------------------
 const FMeshBufferSet& Mesh::GetIndexBuffers() const
 {
-    return m_IndexBuffers;
+    return IndexBuffers;
 }
 
 
-//---------------------------------------------------------------------------------------------
 int Mesh::GetSurfaceCount() const
 {
-    return m_surfaces.Num();
+    return Surfaces.Num();
 }
 
 
-//---------------------------------------------------------------------------------------------
 void Mesh::GetSurface( int32 surfaceIndex,
                        int32* firstVertex, int32* vertexCount,
                        int32* firstIndex, int32* indexCount,
@@ -402,13 +510,13 @@ void Mesh::GetSurface( int32 surfaceIndex,
 
     if (surfaceIndex>=0 && surfaceIndex<count)
     {
-        if (surfaceIndex<m_surfaces.Num())
+        if (surfaceIndex<Surfaces.Num())
         {
-            const MESH_SURFACE& surf = m_surfaces[surfaceIndex];
-            if (firstVertex) *firstVertex = surf.m_firstVertex;
-            if (vertexCount) *vertexCount = surf.m_vertexCount;
-            if (firstIndex) *firstIndex = surf.m_firstIndex;
-            if (indexCount) *indexCount = surf.m_indexCount;
+            const FMeshSurface& surf = Surfaces[surfaceIndex];
+            if (firstVertex) *firstVertex = surf.FirstVertex;
+            if (vertexCount) *vertexCount = surf.VertexCount;
+            if (firstIndex) *firstIndex = surf.FirstIndex;
+            if (indexCount) *indexCount = surf.IndexCount;
             if (BoneIndex) *BoneIndex = surf.BoneMapIndex;
             if (BoneCount) *BoneCount = surf.BoneMapCount;
             if (bCastShadow) *bCastShadow = surf.bCastShadow;
@@ -439,75 +547,67 @@ void Mesh::GetSurface( int32 surfaceIndex,
 }
 
 
-//---------------------------------------------------------------------------------------------
 uint32_t Mesh::GetSurfaceId( int surfaceIndex ) const
 {
-    if (surfaceIndex>=0 && surfaceIndex<m_surfaces.Num())
+    if (surfaceIndex>=0 && surfaceIndex<Surfaces.Num())
     {
-        const MESH_SURFACE& surf = m_surfaces[surfaceIndex];
-        return surf.m_id;
+        const FMeshSurface& surf = Surfaces[surfaceIndex];
+        return surf.Id;
     }
 
     return 0;
 }
 
 
-//---------------------------------------------------------------------------------------------
 void Mesh::AddLayout(Ptr<const Layout> pLayout )
 {
 	LLM_SCOPE_BYNAME(TEXT("MutableRuntime"));
-    m_layouts.Add( pLayout );
+    Layouts.Add( pLayout );
 }
 
 
-//---------------------------------------------------------------------------------------------
 int Mesh::GetLayoutCount() const
 {
-    return m_layouts.Num();
+    return Layouts.Num();
 }
 
 
-//---------------------------------------------------------------------------------------------
 const Layout* Mesh::GetLayout( int i ) const
 {
-    check( i>=0 && i<m_layouts.Num() );
+    check( i>=0 && i<Layouts.Num() );
 
-    return m_layouts[i].get();
+    return Layouts[i].get();
 }
 
 
-//---------------------------------------------------------------------------------------------
 void Mesh::SetLayout( int i, Ptr<const Layout> pLayout )
 {
-    check( i>=0 && i<m_layouts.Num() );
+    check( i>=0 && i<Layouts.Num() );
 
-    m_layouts[i] = pLayout;
+    Layouts[i] = pLayout;
 }
 
 
-//---------------------------------------------------------------------------------------------
 int Mesh::GetTagCount() const
 {
-    return m_tags.Num();
+    return Tags.Num();
 }
 
 
-//---------------------------------------------------------------------------------------------
 void Mesh::SetTagCount( int count )
 {
 	LLM_SCOPE_BYNAME(TEXT("MutableRuntime"));
-    m_tags.SetNum( count );
+    Tags.SetNum( count );
 }
 
 
-//---------------------------------------------------------------------------------------------
 const FString& Mesh::GetTag( int tagIndex ) const
 {
     check( tagIndex>=0 && tagIndex<GetTagCount() );
 
     if (tagIndex >= 0 && tagIndex < GetTagCount())
     {
-        return m_tags[tagIndex];
+        return Tags[tagIndex];
     }
     else
     {
@@ -517,7 +617,6 @@ const FString& Mesh::GetTag( int tagIndex ) const
 }
 
 
-//---------------------------------------------------------------------------------------------
 void Mesh::SetTag( int tagIndex, const FString& Name )
 {
     check( tagIndex>=0 && tagIndex<GetTagCount() );
@@ -525,33 +624,29 @@ void Mesh::SetTag( int tagIndex, const FString& Name )
 
     if (tagIndex >= 0 && tagIndex < GetTagCount())
     {
-        m_tags[tagIndex] = Name;
+        Tags[tagIndex] = Name;
     }
 }
 
 
-//---------------------------------------------------------------------------------------------
 void Mesh::AddStreamedResource(uint64 ResourceId)
 {
 	StreamedResources.AddUnique(ResourceId);
 }
 
 
-//---------------------------------------------------------------------------------------------
 const TArray<uint64>& Mesh::GetStreamedResources() const
 {
 	return StreamedResources;
 }
 
 
-//---------------------------------------------------------------------------------------------
 int32 Mesh::FindBonePose(const FBoneName& BoneId) const
 {
 	return BonePoses.IndexOfByPredicate([BoneId](const FBonePose& Pose) { return Pose.BoneId == BoneId; });
 }
 
 
-//---------------------------------------------------------------------------------------------
 void mu::Mesh::SetBonePoseCount(int32 count)
 {
 	LLM_SCOPE_BYNAME(TEXT("MutableRuntime"));
@@ -559,14 +654,12 @@ void mu::Mesh::SetBonePoseCount(int32 count)
 }
 
 
-//---------------------------------------------------------------------------------------------
 int32 mu::Mesh::GetBonePoseCount() const
 {
 	return BonePoses.Num();
 }
 
 
-//---------------------------------------------------------------------------------------------
 void mu::Mesh::SetBonePose(int32 Index, const FBoneName& BoneId, FTransform3f Transform, EBoneUsageFlags BoneUsageFlags)
 {
 	check(BonePoses.IsValidIndex(Index));
@@ -577,7 +670,6 @@ void mu::Mesh::SetBonePose(int32 Index, const FBoneName& BoneId, FTransform3f Tr
 }
 
 
-//---------------------------------------------------------------------------------------------
 const FBoneName& Mesh::GetBonePoseId(int32 Index) const
 {
 	check(BonePoses.IsValidIndex(Index));
@@ -585,7 +677,6 @@ const FBoneName& Mesh::GetBonePoseId(int32 Index) const
 }
 
 
-//---------------------------------------------------------------------------------------------
 void mu::Mesh::GetBonePoseTransform(int32 BoneIndex, FTransform3f& Transform) const
 {
 	check(BoneIndex >= 0 && BoneIndex < BonePoses.Num());
@@ -593,7 +684,6 @@ void mu::Mesh::GetBonePoseTransform(int32 BoneIndex, FTransform3f& Transform) co
 }
 
 
-//---------------------------------------------------------------------------------------------
 EBoneUsageFlags Mesh::GetBoneUsageFlags(int32 BoneIndex) const
 {
 	check(BoneIndex >= 0 && BoneIndex < BonePoses.Num());
@@ -601,35 +691,30 @@ EBoneUsageFlags Mesh::GetBoneUsageFlags(int32 BoneIndex) const
 }
 
 
-//---------------------------------------------------------------------------------------------
 void Mesh::SetBoneMap(const TArray<FBoneName>& InBoneMap)
 {
 	BoneMap = InBoneMap;
 }
 
 
-//---------------------------------------------------------------------------------------------
 const TArray<FBoneName>& Mesh::GetBoneMap() const
 {
 	return BoneMap;
 }
 
 
-//---------------------------------------------------------------------------------------------
 int32 Mesh::GetSkeletonIDsCount() const
 {
     return SkeletonIDs.Num();
 }
 
 
-//---------------------------------------------------------------------------------------------
 int32 Mesh::GetSkeletonID(int32 SkeletonIndex) const
 {
 	return SkeletonIDs.IsValidIndex(SkeletonIndex) ? SkeletonIDs[SkeletonIndex] : INDEX_NONE;
 }
 
 
-//---------------------------------------------------------------------------------------------
 void Mesh::AddSkeletonID(int32 SkeletonID)
 {
 	check(SkeletonID != INDEX_NONE);
@@ -637,9 +722,6 @@ void Mesh::AddSkeletonID(int32 SkeletonID)
 }
 
 
-//---------------------------------------------------------------------------------------------
-//---------------------------------------------------------------------------------------------
-//---------------------------------------------------------------------------------------------
 int32 Mesh::GetDataSize() const
 {
 	// TODO: review if other mesh fields like additional physics assets
@@ -653,33 +735,32 @@ int32 Mesh::GetDataSize() const
 	}
 
 	return sizeof(Mesh)
-		+ m_IndexBuffers.GetDataSize()
-		+ m_VertexBuffers.GetDataSize()
-		+ m_FaceBuffers.GetDataSize()
+		+ IndexBuffers.GetDataSize()
+		+ VertexBuffers.GetDataSize()
 		+ BonePoses.Num() * sizeof(FBonePose)
 		+ AdditionalBuffersSize;
 }
 
-//---------------------------------------------------------------------------------------------
+
 bool Mesh::HasCompatibleFormat( const Mesh* pOther ) const
 {
     bool compatible = true;
 
-    compatible &= m_layouts.Num()==pOther->m_layouts.Num();
-    compatible &= m_VertexBuffers.GetBufferCount()
-            == pOther->m_VertexBuffers.GetBufferCount();
+    compatible &= Layouts.Num()==pOther->Layouts.Num();
+    compatible &= VertexBuffers.GetBufferCount()
+            == pOther->VertexBuffers.GetBufferCount();
 
 
     // Indices
     //-----------------
-    if ( m_IndexBuffers.GetElementCount()>0 && pOther->GetIndexCount()>0 )
+    if ( IndexBuffers.GetElementCount()>0 && pOther->GetIndexCount()>0 )
     {
-        check( m_IndexBuffers.m_buffers.Num() == 1 );
+        check( IndexBuffers.m_buffers.Num() == 1 );
         check( pOther->GetIndexBuffers().m_buffers.Num() == 1 );
-        check( m_IndexBuffers.GetBufferChannelCount(0) == 1 );
+        check( IndexBuffers.GetBufferChannelCount(0) == 1 );
         check( pOther->GetIndexBuffers().GetBufferChannelCount(0) == 1 );
 
-        const FMeshBuffer& dest = m_IndexBuffers.m_buffers[0];
+        const FMeshBuffer& dest = IndexBuffers.m_buffers[0];
         const FMeshBuffer& source = pOther->GetIndexBuffers().m_buffers[0];
 
         compatible &= dest.m_channels[0].m_format == source.m_channels[0].m_format;
@@ -693,9 +774,9 @@ bool Mesh::HasCompatibleFormat( const Mesh* pOther ) const
 
     // Vertices
     //-----------------
-    for ( int vb = 0; vb<m_VertexBuffers.GetBufferCount(); ++vb )
+    for ( int vb = 0; vb<VertexBuffers.GetBufferCount(); ++vb )
     {
-        const FMeshBuffer& dest = m_VertexBuffers.m_buffers[vb];
+        const FMeshBuffer& dest = VertexBuffers.m_buffers[vb];
         const FMeshBuffer& source = pOther->GetVertexBuffers().m_buffers[vb];
 
         // TODO: More checks about channels formats and semantics
@@ -707,12 +788,11 @@ bool Mesh::HasCompatibleFormat( const Mesh* pOther ) const
 }
 
 
-//---------------------------------------------------------------------------------------------
 UE::Math::TIntVector3<uint32_t>  Mesh::GetFaceVertexIndices( int f ) const
 {
 	UE::Math::TIntVector3<uint32> res;
 
-    MeshBufferIteratorConst<MBF_UINT32,uint32_t,1> it( m_IndexBuffers, MBS_VERTEXINDEX );
+    MeshBufferIteratorConst<MBF_UINT32,uint32_t,1> it( IndexBuffers, MBS_VERTEXINDEX );
     it += f*3;
 
     res[0] = (*it)[0];
@@ -728,154 +808,17 @@ UE::Math::TIntVector3<uint32_t>  Mesh::GetFaceVertexIndices( int f ) const
 }
 
 
-//---------------------------------------------------------------------------------------------
-bool Mesh::HasFace
-    (
-        const Mesh& other,
-        int otherFaceIndex,
-        const VERTEX_MATCH_MAP& vertexMap
-    ) const
+bool Mesh::FVertexMatchMap::DoMatch(int32 v, int32 ov) const
 {
-    bool found = false;
-
-	UE::Math::TIntVector3<uint32> ov = other.GetFaceVertexIndices( otherFaceIndex );
-
-    UntypedMeshBufferIteratorConst it( m_IndexBuffers, MBS_VERTEXINDEX );
-
-    for ( int f=0; !found && f<m_FaceBuffers.GetElementCount(); f++ )
-    {
-		UE::Math::TIntVector3<uint32> v;
-        v[0] = it.GetAsUINT32(); ++it;
-        v[1] = it.GetAsUINT32(); ++it;
-        v[2] = it.GetAsUINT32(); ++it;
-
-        found = true;
-        for ( int vi=0; found && vi<3; ++vi )
-        {
-            found = vertexMap.Matches(v[vi],ov[0])
-                 || vertexMap.Matches(v[vi],ov[1])
-                 || vertexMap.Matches(v[vi],ov[2]);
-        }
-    }
-
-    return found;
-}
-
-
-//---------------------------------------------------------------------------------------------
-bool Mesh::IsSameVertex
-    (
-        uint32_t vertexIndex,
-        const Mesh& other,
-        uint32_t otherVertexIndex,
-        float tolerance
-    ) const
-{
-    bool same = true;
-
-    // For all the attributes in this mesh
-    for ( int b=0; same && b<m_VertexBuffers.GetBufferCount(); ++b )
-    {
-        for ( int c=0; same && c<m_VertexBuffers.GetBufferChannelCount(b); ++c )
-        {
-            EMeshBufferSemantic semantic;
-            int semanticIndex = 0;
-            EMeshBufferFormat format;
-            int components;
-            int offset = 0;
-            m_VertexBuffers.GetChannel( b, c, &semantic, &semanticIndex, &format, &components, &offset );
-
-            // If it is not one of the relevant semantics
-            if ( semantic!=MBS_POSITION &&
-                    //semantic!=MBS_TEXCOORDS &&
-                    //semantic!=MBS_NORMAL &&
-                    //semantic!=MBS_TANGENT &&
-                    //semantic!=MBS_BINORMAL &&
-                    semantic!=MBS_BONEINDICES &&
-                    semantic!=MBS_BONEWEIGHTS )
-            {
-                break;
-            }
-
-            // Find the channel in the other mesh
-            int otherBuffer = -1;
-            int otherChannel = -1;
-            other.m_VertexBuffers.FindChannel( semantic, semanticIndex, &otherBuffer, &otherChannel );
-            check( otherBuffer>=0 && otherChannel>=0 );
-
-            EMeshBufferSemantic otherSemantic;
-            int otherSemanticIndex = 0;
-            EMeshBufferFormat otherFormat;
-            int otherComponents;
-            int otherOffset;
-            other.m_VertexBuffers.GetChannel
-                    (
-                        b, c,
-                        &otherSemantic, &otherSemanticIndex,
-                        &otherFormat, &otherComponents,
-                        &otherOffset
-                    );
-            check( otherSemantic == semantic );
-            check( otherFormat == format );
-            check( otherComponents == components );
-
-            int elemSize = m_VertexBuffers.GetElementSize( b );
-            const uint8_t* pData = m_VertexBuffers.GetBufferData( b )
-                    + elemSize*vertexIndex + offset;
-
-            int otherElemSize = other.m_VertexBuffers.GetElementSize( b );
-            const uint8_t* pOtherData = other.m_VertexBuffers.GetBufferData( b )
-                    + otherElemSize*otherVertexIndex + otherOffset;
-
-            switch (format)
-            {
-            case MBF_FLOAT32:
-            {
-                const float* pfData = (const float*)pData;
-                const float* pfOtherData = (const float*)pOtherData;
-                for ( int d=0; same && d<components; ++d )
-                {
-                    float diff = fabs( (*pfData) - (*pfOtherData) );
-                    same = diff <= tolerance;
-                    ++pfData;
-                    ++pfOtherData;
-                }
-                break;
-            }
-
-            case MBF_UINT8:
-            {
-                for ( int d=0; same && d<components; ++d )
-                {
-                    same = (*pData) == (*pOtherData);
-                    ++pData;
-                    ++pOtherData;
-                }
-                break;
-            }
-
-            default:
-                check( false );
-            }
-        }
-    }
-
-    return same;
-}
-
-
-//---------------------------------------------------------------------------------------------
-bool Mesh::VERTEX_MATCH_MAP::Matches(int v, int ov) const
-{
-	if (v >= 0 && v < (int)m_firstMatch.Num())
+	if (v >= 0 && v < FirstMatch.Num())
 	{
-		int start = m_firstMatch[v];
-		int end = v + 1 < m_firstMatch.Num() ? m_firstMatch[v + 1] : m_matches.Num();
+		int32 start = FirstMatch[v];
+		int32 end = v + 1 < FirstMatch.Num() ? FirstMatch[v + 1] : Matches.Num();
 		bool res = false;
 
 		while (!res && start < end)
 		{
-			if (m_matches[start] == ov)
+			if (Matches[start] == ov)
 			{
 				res = true;
 			}
@@ -889,19 +832,13 @@ bool Mesh::VERTEX_MATCH_MAP::Matches(int v, int ov) const
 }
 
 
-//---------------------------------------------------------------------------------------------
-void Mesh::GetVertexMap
-    (
-        const Mesh& other,
-        VERTEX_MATCH_MAP& vertexMap,
-        float tolerance
-    ) const
+void Mesh::GetVertexMap( const Mesh& other, FVertexMatchMap& vertexMap, float tolerance ) const
 {
-    int vertexCount = m_VertexBuffers.GetElementCount();
-    vertexMap.m_firstMatch.SetNum( vertexCount );
-    vertexMap.m_matches.SetNum( vertexCount+(vertexCount>>2) );
+    int32 vertexCount = VertexBuffers.GetElementCount();
+    vertexMap.FirstMatch.SetNum( vertexCount );
+    vertexMap.Matches.SetNum( vertexCount+(vertexCount>>2) );
 
-    int otherVertexCount = other.m_VertexBuffers.GetElementCount();
+    int otherVertexCount = other.VertexBuffers.GetElementCount();
 
     if ( !vertexCount || !otherVertexCount )
     {
@@ -909,8 +846,8 @@ void Mesh::GetVertexMap
     }
 
 
-    MeshBufferIteratorConst< MBF_FLOAT32, float, 3 > itp( m_VertexBuffers, MBS_POSITION);
-    MeshBufferIteratorConst< MBF_FLOAT32, float, 3 > itopBegin( other.m_VertexBuffers, MBS_POSITION);
+    MeshBufferIteratorConst< MBF_FLOAT32, float, 3 > itp( VertexBuffers, MBS_POSITION);
+    MeshBufferIteratorConst< MBF_FLOAT32, float, 3 > itopBegin( other.VertexBuffers, MBS_POSITION);
 
 
     // Bucket the other mesh
@@ -920,7 +857,7 @@ void Mesh::GetVertexMap
     float rangeMin = TNumericLimits<float>::Max();
     float rangeMax = -TNumericLimits<float>::Max();
     MeshBufferIteratorConst< MBF_FLOAT32, float, 3 >  itop = itopBegin;
-    for ( int ov=0; ov<otherVertexCount; ++ov )
+    for ( int32 ov=0; ov<otherVertexCount; ++ov )
     {
         float v = (*itop)[MUTABLE_BUCKET_CHANNEL];
         rangeMin = FMath::Min( rangeMin, v );
@@ -930,23 +867,23 @@ void Mesh::GetVertexMap
     rangeMin -= tolerance;
     rangeMax += tolerance;
 
-    TArray<int> buckets[MUTABLE_NUM_BUCKETS];
-    for ( int b=0; b<MUTABLE_NUM_BUCKETS; ++b )
+    TArray<int32> buckets[MUTABLE_NUM_BUCKETS];
+    for ( int32 b=0; b<MUTABLE_NUM_BUCKETS; ++b )
     {
         buckets[b].Reserve( otherVertexCount/MUTABLE_NUM_BUCKETS*2 );
     }
 
     float bucketSize = (rangeMax-rangeMin)/float(MUTABLE_NUM_BUCKETS);
     itop = itopBegin;
-    for ( int ov=0; ov<otherVertexCount; ++ov )
+    for ( int32 ov=0; ov<otherVertexCount; ++ov )
     {
         float v = (*itop)[MUTABLE_BUCKET_CHANNEL];
 
-        int bucket0 = int( floor( (v-tolerance-rangeMin)/bucketSize ) );
+        int32 bucket0 = int32( floor( (v-tolerance-rangeMin)/bucketSize ) );
         bucket0 = FMath::Min( MUTABLE_NUM_BUCKETS-1, FMath::Max( 0, bucket0 ) );
         buckets[bucket0].Add(ov);
 
-        int bucket1 = int( floor( (v+tolerance-rangeMin)/bucketSize ) );
+        int32 bucket1 = int32( floor( (v+tolerance-rangeMin)/bucketSize ) );
         bucket1 = FMath::Min( MUTABLE_NUM_BUCKETS-1, FMath::Max( 0, bucket1 ) );
 
         if (bucket1!=bucket0)
@@ -957,136 +894,84 @@ void Mesh::GetVertexMap
         ++itop;
     }
 
-
     // TODO Compare only positions?
-    if (true)
+
+    // Use buckets
+    for ( int32 v=0; v<vertexCount; ++v )
     {
+        vertexMap.FirstMatch[v] = vertexMap.Matches.Num();
 
-        /*
-        // Don't use buckets
-        for ( int v=0; v<vertexCount; ++v )
-        {
-            itop = itopBegin;
+        float vbucket = (*itp)[MUTABLE_BUCKET_CHANNEL];
+        int32 bucket = int32( floor( (vbucket-rangeMin)/bucketSize ) );
 
-            vertexMap.m_firstMatch[v] = (int)vertexMap.m_matches.size();
-            for ( int ov=0; ov<otherVertexCount; ++ov )
+        if (bucket>=0 && bucket<MUTABLE_NUM_BUCKETS)
+        {
+            int32 bucketVertexCount = buckets[bucket].Num();
+            for ( int32 ov=0; ov<bucketVertexCount; ++ov )
             {
+                int32 otherVertexIndex = buckets[bucket][ov];
+                FVector3f p = (itopBegin+otherVertexIndex).GetAsVec3f();
+
                 bool same = true;
-                for ( int d=0; same && d<3; ++d )
+                for ( int32 d=0; same && d<3; ++d )
                 {
-                    float diff = fabs( (*itp)[d] - (*itop)[d] );
+                    float diff = fabs( (*itp)[d] - p[d] );
                     same = diff <= tolerance;
                 }
 
                 if ( same )
                 {
-                    vertexMap.m_matches.push_back(ov);
+                    vertexMap.Matches.Add( otherVertexIndex );
                 }
-
-                ++itop;
             }
-
-            ++itp;
-        }
-/*/
-        // Use buckets
-        for ( int v=0; v<vertexCount; ++v )
-        {
-            vertexMap.m_firstMatch[v] = (int)vertexMap.m_matches.Num();
-
-            float vbucket = (*itp)[MUTABLE_BUCKET_CHANNEL];
-            int bucket = int( floor( (vbucket-rangeMin)/bucketSize ) );
-
-            if (bucket>=0 && bucket<MUTABLE_NUM_BUCKETS)
-            {
-                int bucketVertexCount = (int)buckets[bucket].Num();
-                for ( int ov=0; ov<bucketVertexCount; ++ov )
-                {
-                    int otherVertexIndex = buckets[bucket][ov];
-                    FVector3f p = (itopBegin+otherVertexIndex).GetAsVec3f();
-
-                    bool same = true;
-                    for ( int d=0; same && d<3; ++d )
-                    {
-                        float diff = fabs( (*itp)[d] - p[d] );
-                        same = diff <= tolerance;
-                    }
-
-                    if ( same )
-                    {
-                        vertexMap.m_matches.Add( otherVertexIndex );
-                    }
-                }
-            }
-
-            ++itp;
         }
 
-    }
-    else
-    {
-        // Slow generic way
-        for ( int v=0; v<vertexCount; ++v )
-        {
-            vertexMap.m_firstMatch[v] = (int)vertexMap.m_matches.Num();
-            for ( int ov=0; ov<otherVertexCount; ++ov )
-            {
-                // TODO: Optimize this by not using IsSameVertex
-                if ( IsSameVertex( v, other, ov, tolerance ) )
-                {
-                    vertexMap.m_matches.Add(ov);
-                }
-            }
-        }
+        ++itp;
     }
 
 }
 
 
-//---------------------------------------------------------------------------------------------
 void Mesh::EnsureSurfaceData()
 {
-	if (!m_surfaces.Num() && m_VertexBuffers.GetElementCount())
+	if (!Surfaces.Num() && VertexBuffers.GetElementCount())
 	{
-		MESH_SURFACE s;
-		s.m_vertexCount = m_VertexBuffers.GetElementCount();
-		s.m_indexCount = m_IndexBuffers.GetElementCount();
+		FMeshSurface s;
+		s.VertexCount = VertexBuffers.GetElementCount();
+		s.IndexCount = IndexBuffers.GetElementCount();
 		s.BoneMapCount = BoneMap.Num();
-		m_surfaces.Add(s);
+		Surfaces.Add(s);
 	}
 }
 
 
-//---------------------------------------------------------------------------------------------
 void Mesh::ResetBufferIndices()
 {
-	m_VertexBuffers.ResetBufferIndices();
-	m_IndexBuffers.ResetBufferIndices();
-	m_FaceBuffers.ResetBufferIndices();
+	VertexBuffers.ResetBufferIndices();
+	IndexBuffers.ResetBufferIndices();
 }
 
 
-//---------------------------------------------------------------------------------------------
-void UnserialiseLegacySurfaces(InputArchive& arch, TArray<MESH_SURFACE>& OutMeshSurfaces)
+void UnserialiseLegacySurfaces(InputArchive& arch, TArray<FMeshSurface>& OutMeshSurfaces)
 {
 	struct FMeshSurfaceLegacy
 	{
 		FMeshSurfaceLegacy()
 		{}
 
-		int32 m_firstVertex = 0;
-		int32 m_vertexCount = 0;
-		int32 m_firstIndex = 0;
-		int32 m_indexCount = 0;
-		uint32 m_id = 0;
+		int32 FirstVertex = 0;
+		int32 VertexCount = 0;
+		int32 FirstIndex = 0;
+		int32 IndexCount = 0;
+		uint32 Id = 0;
 
 		void Unserialise(InputArchive& arch)
 		{
-			arch >> m_firstVertex;
-			arch >> m_vertexCount;
-			arch >> m_firstIndex;
-			arch >> m_indexCount;
-			arch >> m_id;
+			arch >> FirstVertex;
+			arch >> VertexCount;
+			arch >> FirstIndex;
+			arch >> IndexCount;
+			arch >> Id;
 		}
 	}; 
 
@@ -1099,45 +984,43 @@ void UnserialiseLegacySurfaces(InputArchive& arch, TArray<MESH_SURFACE>& OutMesh
 	for (int32 SurfaceIndex = 0; SurfaceIndex < NumSurfaces; ++SurfaceIndex)
 	{
 		FMeshSurfaceLegacy& LegacySurface = LegacyMeshSurfaces[SurfaceIndex];
-		MESH_SURFACE& Surface = OutMeshSurfaces[SurfaceIndex];
-		Surface.m_firstVertex = LegacySurface.m_firstVertex;
-		Surface.m_vertexCount = LegacySurface.m_vertexCount;
-		Surface.m_firstIndex = LegacySurface.m_firstIndex;
-		Surface.m_indexCount = LegacySurface.m_indexCount;
-		Surface.m_id = LegacySurface.m_id;
+		FMeshSurface& Surface = OutMeshSurfaces[SurfaceIndex];
+		Surface.FirstVertex = LegacySurface.FirstVertex;
+		Surface.VertexCount = LegacySurface.VertexCount;
+		Surface.FirstIndex = LegacySurface.FirstIndex;
+		Surface.IndexCount = LegacySurface.IndexCount;
+		Surface.Id = LegacySurface.Id;
 	}
 }
 
 
-//-------------------------------------------------------------------------------------------------
-void MESH_SURFACE::Serialise(OutputArchive& arch) const
+void FMeshSurface::Serialise(OutputArchive& arch) const
 {
 	const int32 ver = 1;
 	arch << ver;
 
-	arch << m_firstVertex;
-	arch << m_vertexCount;
-	arch << m_firstIndex;
-	arch << m_indexCount;
+	arch << FirstVertex;
+	arch << VertexCount;
+	arch << FirstIndex;
+	arch << IndexCount;
 	arch << BoneMapIndex;
 	arch << BoneMapCount;
 	arch << bCastShadow;
 
-	arch << m_id;
+	arch << Id;
 }
 
 
-//-------------------------------------------------------------------------------------------------
-void MESH_SURFACE::Unserialise(InputArchive& arch)
+void FMeshSurface::Unserialise(InputArchive& arch)
 {
 	int32 ver = 0;
 	arch >> ver;
 	check(ver <= 1);
 
-	arch >> m_firstVertex;
-	arch >> m_vertexCount;
-	arch >> m_firstIndex;
-	arch >> m_indexCount;
+	arch >> FirstVertex;
+	arch >> VertexCount;
+	arch >> FirstIndex;
+	arch >> IndexCount;
 	arch >> BoneMapIndex;
 	arch >> BoneMapCount;
 
@@ -1146,11 +1029,10 @@ void MESH_SURFACE::Unserialise(InputArchive& arch)
 		arch >> bCastShadow;
 	}
 
-	arch >> m_id;
+	arch >> Id;
 }
 
 
-//-------------------------------------------------------------------------------------------------
 void Mesh::FBonePose::Serialise(OutputArchive& arch) const
 {
 	const int32 ver = 2;
@@ -1162,7 +1044,6 @@ void Mesh::FBonePose::Serialise(OutputArchive& arch) const
 }
 
 
-//-------------------------------------------------------------------------------------------------
 void Mesh::FBonePose::Unserialise(InputArchive& arch)
 {
 	int32 ver = 0;
@@ -1196,74 +1077,79 @@ void Mesh::FBonePose::Unserialise(InputArchive& arch)
 }
 
 
-//-------------------------------------------------------------------------------------------------
 void Mesh::Serialise(OutputArchive& arch) const
 {
-	uint32 ver = 18;
+	uint32 ver = 19;
 	arch << ver;
 
-	arch << m_IndexBuffers;
-	arch << m_VertexBuffers;
-	arch << m_FaceBuffers;
+	arch << IndexBuffers;
+	arch << VertexBuffers;
 	arch << AdditionalBuffers;
-	arch << m_layouts;
+	arch << Layouts;
 
 	arch << SkeletonIDs;
 
-	arch << m_pSkeleton;
-	arch << m_pPhysicsBody;
+	arch << Skeleton;
+	arch << PhysicsBody;
 
-	arch << m_staticFormatFlags;
-	arch << m_surfaces;
+	arch << StaticFormatFlags;
+	arch << Surfaces;
 
-	arch << m_tags;
+	arch << Tags;
 	arch << StreamedResources;
 
 	arch << BonePoses;
 	arch << BoneMap;
 
 	arch << AdditionalPhysicsBodies;
+
+	arch << VertexIDPrefix;
 }
 
 
-//-------------------------------------------------------------------------------------------------
 void Mesh::Unserialise(InputArchive& arch)
 {
 	uint32 ver;
 	arch >> ver;
-	check(ver <= 18);
+	check(ver <= 19);
+
+	arch >> IndexBuffers;
+	arch >> VertexBuffers;
+	
+	if (ver < 19)
+	{
+		FMeshBufferSet Dummy;
+		arch >> Dummy;
+	}
 
-	arch >> m_IndexBuffers;
-	arch >> m_VertexBuffers;
-	arch >> m_FaceBuffers;
 	arch >> AdditionalBuffers;
-	arch >> m_layouts;
+	arch >> Layouts;
 
 	if (ver >= 14)
 	{
 		arch >> SkeletonIDs;
 	}
 
-	arch >> m_pSkeleton;
+	arch >> Skeleton;
 	if (ver >= 12)
 	{ 
-		arch >> m_pPhysicsBody;
+		arch >> PhysicsBody;
 	}
 	else
 	{
-		m_pPhysicsBody = nullptr;
+		PhysicsBody = nullptr;
 	}
 
-	arch >> m_staticFormatFlags;
+	arch >> StaticFormatFlags;
 
 	if (ver >= 16)
 	{
-		arch >> m_surfaces;
+		arch >> Surfaces;
 	}
 	else
 	{
 		// Deserialize LegacySurfaces
-		UnserialiseLegacySurfaces(arch, m_surfaces);
+		UnserialiseLegacySurfaces(arch, Surfaces);
 	}
 
 	if (ver <= 16)
@@ -1289,15 +1175,15 @@ void Mesh::Unserialise(InputArchive& arch)
 	{
 		TArray < std::string > Temp;
 		arch >> Temp;
-		m_tags.SetNum(Temp.Num());
+		Tags.SetNum(Temp.Num());
 		for (int32 c = 0; c < Temp.Num(); ++c)
 		{
-			m_tags[c] = Temp[c].c_str();
+			Tags[c] = Temp[c].c_str();
 		}
 	}
 	else
 	{
-		arch >> m_tags;
+		arch >> Tags;
 	}
 
 	if (ver >= 18)
@@ -1309,17 +1195,17 @@ void Mesh::Unserialise(InputArchive& arch)
 	{
 		arch >> BonePoses;
 	}
-	else if (m_pSkeleton)
+	else if (Skeleton)
 	{
-		const int32 NumBones = m_pSkeleton->GetBoneCount();
+		const int32 NumBones = Skeleton->GetBoneCount();
 		BonePoses.SetNum(NumBones);
-		check(m_pSkeleton->m_boneTransforms_DEPRECATED.Num() == NumBones);
+		check(Skeleton->m_boneTransforms_DEPRECATED.Num() == NumBones);
 
 		for (int32 BoneIndex = 0; BoneIndex < NumBones; ++BoneIndex)
 		{
 			BonePoses[BoneIndex].BoneId = FBoneName(static_cast<uint32>(BoneIndex));
 			BonePoses[BoneIndex].BoneUsageFlags = EBoneUsageFlags::Skinning;
-			BonePoses[BoneIndex].BoneTransform = m_pSkeleton->m_boneTransforms_DEPRECATED[BoneIndex];
+			BonePoses[BoneIndex].BoneTransform = Skeleton->m_boneTransforms_DEPRECATED[BoneIndex];
 		}
 	}
 
@@ -1336,7 +1222,7 @@ void Mesh::Unserialise(InputArchive& arch)
 			BoneMap[BoneIndex] = FBoneName(static_cast<uint32>(BoneIndex));
 		}
 
-		for (MESH_SURFACE& Surface : m_surfaces)
+		for (FMeshSurface& Surface : Surfaces)
 		{
 			Surface.BoneMapCount = NumBonePoses;
 		}
@@ -1346,27 +1232,30 @@ void Mesh::Unserialise(InputArchive& arch)
 	{
 		arch >> AdditionalPhysicsBodies;
 	}
+
+	if (ver >= 19)
+	{
+		arch >> VertexIDPrefix;
+	}
 }
 
 
-//---------------------------------------------------------------------------------------------
 bool Mesh::IsSimilar(const Mesh& o, bool bCompareLayouts) const
 {
 	// Some meshes are just vertex indices (masks) we don't consider them for similarity,
 	// because the kind of vertex channel data they store is the kind that is ignored.
-	if (m_IndexBuffers.GetElementCount() == 0)
+	if (IndexBuffers.GetElementCount() == 0)
 	{
 		return false;
 	}
 
-	bool equal = m_IndexBuffers == o.m_IndexBuffers;
-	if (equal) equal = m_FaceBuffers == o.m_FaceBuffers;
-	if (equal && bCompareLayouts) equal = (m_layouts.Num() == o.m_layouts.Num());
-	if (equal && m_pSkeleton != o.m_pSkeleton)
+	bool equal = IndexBuffers == o.IndexBuffers;
+	if (equal && bCompareLayouts) equal = (Layouts.Num() == o.Layouts.Num());
+	if (equal && Skeleton != o.Skeleton)
 	{
-		if (m_pSkeleton && o.m_pSkeleton)
+		if (Skeleton && o.Skeleton)
 		{
-			equal = (*m_pSkeleton == *o.m_pSkeleton);
+			equal = (*Skeleton == *o.Skeleton);
 		}
 		else
 		{
@@ -1374,11 +1263,11 @@ bool Mesh::IsSimilar(const Mesh& o, bool bCompareLayouts) const
 		}
 	}
 	
-	if (equal && m_pPhysicsBody != o.m_pPhysicsBody)
+	if (equal && PhysicsBody != o.PhysicsBody)
 	{
-		if (m_pPhysicsBody && o.m_pPhysicsBody)
+		if (PhysicsBody && o.PhysicsBody)
 		{
-			equal = (*m_pPhysicsBody == *o.m_pPhysicsBody);
+			equal = (*PhysicsBody == *o.PhysicsBody);
 		}
 		else
 		{
@@ -1386,22 +1275,22 @@ bool Mesh::IsSimilar(const Mesh& o, bool bCompareLayouts) const
 		}
 	}
 
-	if (equal) equal = (m_surfaces == o.m_surfaces);
-	if (equal) equal = (m_tags == o.m_tags);
+	if (equal) equal = (Surfaces == o.Surfaces);
+	if (equal) equal = (Tags == o.Tags);
 
 	// Special comparison for layouts
 	if (bCompareLayouts)
 	{
-		for (int32 i = 0; equal && i < m_layouts.Num(); ++i)
+		for (int32 i = 0; equal && i < Layouts.Num(); ++i)
 		{
-			equal &= m_layouts[i]->IsSimilar(*o.m_layouts[i]);
+			equal &= Layouts[i]->IsSimilar(*o.Layouts[i]);
 		}
 	}
 
 	// Special comparison for vertex buffers
 	if (equal)
 	{
-		equal = m_VertexBuffers.IsSimilarRobust(o.m_VertexBuffers, bCompareLayouts);
+		equal = VertexBuffers.IsSimilarRobust(o.VertexBuffers, bCompareLayouts);
 	}
 
 	return equal;
@@ -1409,10 +1298,9 @@ bool Mesh::IsSimilar(const Mesh& o, bool bCompareLayouts) const
 }
 
 
-//---------------------------------------------------------------------------------------------
 void Mesh::ResetStaticFormatFlags() const
 {
-    m_staticFormatFlags = 0;
+    StaticFormatFlags = 0;
 
     for ( int f=0; f<SMF_COUNT; ++f )
     {
@@ -1420,48 +1308,67 @@ void Mesh::ResetStaticFormatFlags() const
              &&
              s_staticMeshFormatIdentify[f]( this ) )
         {
-            m_staticFormatFlags |= (1<<f);
+            StaticFormatFlags |= (1<<f);
         }
     }
 }
 
 
-//---------------------------------------------------------------------------------------------
 void Mesh::CheckIntegrity() const
 {
 #ifdef MUTABLE_DEBUG
 
+	{
+		int32 BufferIndex = -1;
+		int32 ChannelIndex = -1;
+		VertexBuffers.FindChannel(MBS_VERTEXINDEX, 0, &BufferIndex, &ChannelIndex);
+		if (BufferIndex >= 0 && ChannelIndex >= 0)
+		{
+			EMeshBufferFormat IdFormat = VertexBuffers.m_buffers[BufferIndex].m_channels[ChannelIndex].m_format;
+			if (IdFormat == MBF_UINT64)
+			{
+				check(VertexIDPrefix==0);
+			}
+			else if (IdFormat == MBF_UINT32)
+			{
+				check(VertexIDPrefix != 0);
+			}
+			else
+			{
+				check(false);
+			}
+		}
+	}
+
     // Check vertex indices
     {
-        for ( int b=0; b<m_IndexBuffers.GetBufferCount(); ++b )
+        for ( int32 b=0; b<IndexBuffers.GetBufferCount(); ++b )
         {
-            int elemSize = m_IndexBuffers.GetElementSize( b );
+            int32 elemSize = IndexBuffers.GetElementSize( b );
 
-            for ( int c=0; c<m_IndexBuffers.GetBufferChannelCount(b); ++c )
+            for ( int32 c=0; c<IndexBuffers.GetBufferChannelCount(b); ++c )
             {
                 EMeshBufferSemantic semantic;
-                int semanticIndex = 0;
+                int32 semanticIndex = 0;
                 EMeshBufferFormat format;
-                int components;
-                int offset = 0;
-                m_IndexBuffers.GetChannel
-                        ( b, c, &semantic, &semanticIndex, &format, &components, &offset );
+                int32 components;
+                int32 offset = 0;
+                IndexBuffers.GetChannel( b, c, &semantic, &semanticIndex, &format, &components, &offset );
 
                 if ( semantic==MBS_VERTEXINDEX )
                 {
-                    size_t icount = (size_t)m_IndexBuffers.GetElementCount();
-                    int elemCount = m_VertexBuffers.GetElementCount();
-                    for ( size_t indexIndex = 0; indexIndex < icount; ++indexIndex )
+                    int32 icount = IndexBuffers.GetElementCount();
+                    int32 elemCount = VertexBuffers.GetElementCount();
+                    for ( int32 indexIndex = 0; indexIndex < icount; ++indexIndex )
                     {
-                        const uint8_t* pData = m_IndexBuffers.GetBufferData( b )
-                                + elemSize*indexIndex + offset;
+                        const uint8* pData = IndexBuffers.GetBufferData( b ) + elemSize*indexIndex + offset;
 
                         switch (format)
                         {
                         case MBF_UINT32:
                         {
-                            uint32_t index = *(const uint32_t*)pData;
-                            check( index < uint32_t( elemCount ) );
+                            uint32 index = *(const uint32*)pData;
+                            check( index < uint32( elemCount ) );
                             break;
                         }
                         case MBF_UINT16:
@@ -1472,8 +1379,8 @@ void Mesh::CheckIntegrity() const
                         }
                         case MBF_UINT8:
                         {
-                            uint8_t index = *(const uint8_t*)pData;
-                            check( index < uint8_t( elemCount ) );
+                            uint8 index = *(const uint8*)pData;
+                            check( index < uint8( elemCount ) );
                             break;
                         }
                         default:
@@ -1489,21 +1396,20 @@ void Mesh::CheckIntegrity() const
 
     // Check bone indices, if there are bones. Bones could have been removed for later addition as an optimisation.
     // For all the attributes in this mesh
-    int boneCount = m_pSkeleton ? m_pSkeleton->GetBoneCount() : 0;
-    if ( m_pSkeleton && boneCount )
+    int32 boneCount = Skeleton ? Skeleton->GetBoneCount() : 0;
+    if ( Skeleton && boneCount )
     {
-        for ( int b = 0; b < m_VertexBuffers.GetBufferCount(); ++b )
+        for ( int32 b = 0; b < VertexBuffers.GetBufferCount(); ++b )
         {
-            int channelCount = m_VertexBuffers.GetBufferChannelCount( b );
-            for ( int c = 0; c < channelCount; ++c )
+            int32 channelCount = VertexBuffers.GetBufferChannelCount( b );
+            for ( int32 c = 0; c < channelCount; ++c )
             {
                 EMeshBufferSemantic semantic;
-                int semanticIndex = 0;
+                int32 semanticIndex = 0;
                 EMeshBufferFormat format;
-                int components;
-                int offset = 0;
-                m_VertexBuffers.GetChannel
-                        ( b, c, &semantic, &semanticIndex, &format, &components, &offset );
+                int32 components;
+                int32 offset = 0;
+                VertexBuffers.GetChannel( b, c, &semantic, &semanticIndex, &format, &components, &offset );
 
                 // If it is not one of the relevant semantics
                 if (
@@ -1519,23 +1425,21 @@ void Mesh::CheckIntegrity() const
                     continue;
                 }
 
-                size_t elemCount = (size_t)m_VertexBuffers.GetElementCount();
-                int elemSize = m_VertexBuffers.GetElementSize( b );
+                int32 elemCount = VertexBuffers.GetElementCount();
+                int32 elemSize = VertexBuffers.GetElementSize( b );
 
-                for ( size_t vertexIndex = 0;
-                      vertexIndex < elemCount; ++vertexIndex )
+                for (int32 vertexIndex = 0; vertexIndex < elemCount; ++vertexIndex )
                 {
-                    const uint8_t* pData = m_VertexBuffers.GetBufferData( b )
-                            + elemSize*vertexIndex + offset;
+                    const uint8* pData = VertexBuffers.GetBufferData( b ) + elemSize*vertexIndex + offset;
 
                     switch (format)
                     {
                     case MBF_UINT8:
                     {
-                        for ( int d = 0; d < components; ++d )
+                        for ( int32 d = 0; d < components; ++d )
                         {
-                            uint8_t index = *pData;
-                            check( index < uint64_t(boneCount) );
+                            uint8 index = *pData;
+                            check( index < uint64(boneCount) );
                             ++pData;
                         }
                         break;
@@ -1543,10 +1447,10 @@ void Mesh::CheckIntegrity() const
 
                     case MBF_UINT16:
                     {
-                        for ( int d = 0; d < components; ++d )
+                        for ( int32 d = 0; d < components; ++d )
                         {
                             uint16 index = *(uint16*)pData;
-                            check( index < uint64_t( boneCount ) );
+                            check( index < uint64( boneCount ) );
                             pData+=2;
                         }
                         break;
@@ -1554,10 +1458,10 @@ void Mesh::CheckIntegrity() const
 
                     case MBF_UINT32:
                     {
-                        for ( int d = 0; d < components; ++d )
+                        for ( int32 d = 0; d < components; ++d )
                         {
-                            uint32_t index = *(uint32_t*)pData;
-                            check( index < uint64_t( boneCount ) );
+                            uint32 index = *(uint32*)pData;
+                            check( index < uint64( boneCount ) );
                             pData+=4;
                         }
                         break;
@@ -1574,15 +1478,12 @@ void Mesh::CheckIntegrity() const
 }
 
 
-//---------------------------------------------------------------------------------------------
-//---------------------------------------------------------------------------------------------
-//---------------------------------------------------------------------------------------------
 static bool StaticMeshFormatIdentify_None( const Mesh* )
 {
     return false;
 }
 
-//---------------------------------------------------------------------------------------------
+
 static bool StaticMeshFormatIdentify_Project( const Mesh* pM )
 {
     // This format is used internally for the mesh project
@@ -1590,17 +1491,16 @@ static bool StaticMeshFormatIdentify_Project( const Mesh* pM )
 
     // The first vertex buffer must be texcoords(2f), position(3f), normal(3f)
     // all tightly packed
-    res &= pM->m_VertexBuffers.GetBufferCount()>=1;
+    res &= pM->VertexBuffers.GetBufferCount()>=1;
 
     if ( res )
     {
-        res &= pM->m_VertexBuffers.m_buffers[0].m_channels.Num()==3;
+        res &= pM->VertexBuffers.m_buffers[0].m_channels.Num()==3;
     }
 
     if ( res )
     {
-        const FMeshBufferChannel& chan =
-                pM->m_VertexBuffers.m_buffers[0].m_channels[0];
+        const FMeshBufferChannel& chan = pM->VertexBuffers.m_buffers[0].m_channels[0];
 
         res &= chan.m_semantic == MBS_TEXCOORDS;
         res &= chan.m_format == MBF_FLOAT32;
@@ -1612,8 +1512,7 @@ static bool StaticMeshFormatIdentify_Project( const Mesh* pM )
 
     if ( res )
     {
-        const FMeshBufferChannel& chan =
-                pM->m_VertexBuffers.m_buffers[0].m_channels[1];
+        const FMeshBufferChannel& chan = pM->VertexBuffers.m_buffers[0].m_channels[1];
 
         res &= chan.m_semantic == MBS_POSITION;
         res &= chan.m_format == MBF_FLOAT32;
@@ -1624,8 +1523,7 @@ static bool StaticMeshFormatIdentify_Project( const Mesh* pM )
 
     if ( res )
     {
-        const FMeshBufferChannel& chan =
-                pM->m_VertexBuffers.m_buffers[0].m_channels[2];
+        const FMeshBufferChannel& chan = pM->VertexBuffers.m_buffers[0].m_channels[2];
 
         res &= chan.m_semantic == MBS_NORMAL;
         res &= chan.m_format == MBF_FLOAT32;
@@ -1637,13 +1535,12 @@ static bool StaticMeshFormatIdentify_Project( const Mesh* pM )
     // The first index buffer must be just index buffers u32
     if ( res )
     {
-        res &= pM->m_IndexBuffers.m_buffers[0].m_channels.Num()>=1;
+        res &= pM->IndexBuffers.m_buffers[0].m_channels.Num()>=1;
     }
 
     if ( res )
     {
-        const FMeshBufferChannel& chan =
-                pM->m_IndexBuffers.m_buffers[0].m_channels[0];
+        const FMeshBufferChannel& chan = pM->IndexBuffers.m_buffers[0].m_channels[0];
 
         res &= chan.m_semantic == MBS_VERTEXINDEX;
         res &= chan.m_format == MBF_UINT32;
@@ -1656,7 +1553,6 @@ static bool StaticMeshFormatIdentify_Project( const Mesh* pM )
 }
 
 
-//---------------------------------------------------------------------------------------------
 static bool StaticMeshFormatIdentify_ProjectWrapping( const Mesh* pM )
 {
     // This format is used internally for the mesh project
@@ -1664,16 +1560,16 @@ static bool StaticMeshFormatIdentify_ProjectWrapping( const Mesh* pM )
 
     // The first vertex buffer must be texcoords(2f), position(3f), normal(3f), layoutBlock(uint32_t)
     // all tightly packed
-    res &= pM->m_VertexBuffers.GetBufferCount()>=1;
+    res &= pM->VertexBuffers.GetBufferCount()>=1;
 
     if ( res )
     {
-        res &= pM->m_VertexBuffers.m_buffers[0].m_channels.Num()==4;
+        res &= pM->VertexBuffers.m_buffers[0].m_channels.Num()==4;
     }
 
     if ( res )
     {
-        const FMeshBufferChannel& chan = pM->m_VertexBuffers.m_buffers[0].m_channels[0];
+        const FMeshBufferChannel& chan = pM->VertexBuffers.m_buffers[0].m_channels[0];
 
         res &= chan.m_semantic == MBS_TEXCOORDS;
         res &= chan.m_format == MBF_FLOAT32;
@@ -1685,7 +1581,7 @@ static bool StaticMeshFormatIdentify_ProjectWrapping( const Mesh* pM )
 
     if ( res )
     {
-        const FMeshBufferChannel& chan = pM->m_VertexBuffers.m_buffers[0].m_channels[1];
+        const FMeshBufferChannel& chan = pM->VertexBuffers.m_buffers[0].m_channels[1];
 
         res &= chan.m_semantic == MBS_POSITION;
         res &= chan.m_format == MBF_FLOAT32;
@@ -1696,7 +1592,7 @@ static bool StaticMeshFormatIdentify_ProjectWrapping( const Mesh* pM )
 
     if ( res )
     {
-        const FMeshBufferChannel& chan = pM->m_VertexBuffers.m_buffers[0].m_channels[2];
+        const FMeshBufferChannel& chan = pM->VertexBuffers.m_buffers[0].m_channels[2];
 
         res &= chan.m_semantic == MBS_NORMAL;
         res &= chan.m_format == MBF_FLOAT32;
@@ -1707,7 +1603,7 @@ static bool StaticMeshFormatIdentify_ProjectWrapping( const Mesh* pM )
 
     if ( res )
     {
-        const FMeshBufferChannel& chan = pM->m_VertexBuffers.m_buffers[0].m_channels[3];
+        const FMeshBufferChannel& chan = pM->VertexBuffers.m_buffers[0].m_channels[3];
 
         res &= chan.m_semantic == MBS_LAYOUTBLOCK;
         res &= chan.m_format == MBF_UINT32;
@@ -1720,12 +1616,12 @@ static bool StaticMeshFormatIdentify_ProjectWrapping( const Mesh* pM )
     // The first index buffer must be just index buffers u32
     if ( res )
     {
-        res &= pM->m_IndexBuffers.m_buffers[0].m_channels.Num()>=1;
+        res &= pM->IndexBuffers.m_buffers[0].m_channels.Num()>=1;
     }
 
     if ( res )
     {
-        const FMeshBufferChannel& chan = pM->m_IndexBuffers.m_buffers[0].m_channels[0];
+        const FMeshBufferChannel& chan = pM->IndexBuffers.m_buffers[0].m_channels[0];
 
         res &= chan.m_semantic == MBS_VERTEXINDEX;
         res &= chan.m_format == MBF_UINT32;
@@ -1738,7 +1634,6 @@ static bool StaticMeshFormatIdentify_ProjectWrapping( const Mesh* pM )
 }
 
 
-//---------------------------------------------------------------------------------------------
 STATIC_MESH_FORMAT_ID_FUNC s_staticMeshFormatIdentify[] =
 {
     StaticMeshFormatIdentify_None,
@@ -1747,7 +1642,6 @@ STATIC_MESH_FORMAT_ID_FUNC s_staticMeshFormatIdentify[] =
 };
 
 
-//---------------------------------------------------------------------------------------------
 namespace
 {
     void LogBuffer( FString& out, const FMeshBufferSet& bufset, int32 BufferElementLimit)
@@ -1810,19 +1704,11 @@ void Mesh::Log( FString& out, int32 BufferElementLimit)
     out += "Mesh:\n";
 
     out += "Indices:\n";
-    LogBuffer( out, m_IndexBuffers, BufferElementLimit);
+    LogBuffer( out, IndexBuffers, BufferElementLimit);
 
     out += "Vertices:\n";
-    LogBuffer( out, m_VertexBuffers, BufferElementLimit);
-
-    out += "Faces:\n";
-    LogBuffer( out, m_FaceBuffers, BufferElementLimit);
+    LogBuffer( out, VertexBuffers, BufferElementLimit);
 }
 
-
-MUTABLE_IMPLEMENT_ENUM_SERIALISABLE(EBoneUsageFlags)
-MUTABLE_IMPLEMENT_ENUM_SERIALISABLE(EMeshBufferType)
-MUTABLE_IMPLEMENT_ENUM_SERIALISABLE(EShapeBindingMethod)
-MUTABLE_IMPLEMENT_ENUM_SERIALISABLE(EVertexColorUsage)
 	
 }
diff --git a/Engine/Plugins/Experimental/Mutable/Source/MutableRuntime/Private/MuR/MeshBufferSet.cpp b/Engine/Plugins/Experimental/Mutable/Source/MutableRuntime/Private/MuR/MeshBufferSet.cpp
--- a/Engine/Plugins/Experimental/Mutable/Source/MutableRuntime/Private/MuR/MeshBufferSet.cpp
+++ b/Engine/Plugins/Experimental/Mutable/Source/MutableRuntime/Private/MuR/MeshBufferSet.cpp
@@ -20,7 +20,7 @@ namespace mu
 
 	
 	//---------------------------------------------------------------------------------------------
-	static FMeshBufferFormatData s_meshBufferFormatData[MBF_COUNT] =
+	static FMeshBufferFormatData s_meshBufferFormatData[] = // MBF_COUNT entries
 	{
 		{ 0, 0 },
 		{ 2, 0 },
@@ -45,9 +45,15 @@ namespace mu
 		{ 1, 0 },
 		{ 1, 0 },
 
+		{ 8, 0 },
+
+		{ 8, 64 },
+		{ 8, 63 },
+		{ 8, 0 },
 		{ 8, 0 }
 	};
 
+	static_assert(sizeof(s_meshBufferFormatData) / sizeof(FMeshBufferFormatData) == int32(MBF_COUNT));
 
 	//---------------------------------------------------------------------------------------------
 	const FMeshBufferFormatData& GetMeshFormatData(EMeshBufferFormat format)
diff --git a/Engine/Plugins/Experimental/Mutable/Source/MutableRuntime/Private/MuR/OpImageProject.cpp b/Engine/Plugins/Experimental/Mutable/Source/MutableRuntime/Private/MuR/OpImageProject.cpp
--- a/Engine/Plugins/Experimental/Mutable/Source/MutableRuntime/Private/MuR/OpImageProject.cpp
+++ b/Engine/Plugins/Experimental/Mutable/Source/MutableRuntime/Private/MuR/OpImageProject.cpp
@@ -949,7 +949,7 @@ void ImageRasterProjectedPlanar(const Mesh* pMesh, Image* pImage,
 		ProcessorFeatures |= EPixelProcessorFeatures::WithMask;
 	}
 
-	if (pMesh->m_staticFormatFlags & (1 << SMF_PROJECT))
+	if (pMesh->StaticFormatFlags & (1 << SMF_PROJECT))
 	{
 		float UnusedProjectionAngle = 0;
 		FImageRasterInvokeArgs RasterArgs = 
@@ -1163,7 +1163,7 @@ void ImageRasterProjectedWrapping( const Mesh* pMesh, Image* pImage,
 		ProcessorFeatures |= EPixelProcessorFeatures::WithMask;
 	}
 
-    if ( ( pMesh->m_staticFormatFlags & (1<<SMF_PROJECTWRAPPING) ) )
+    if ( ( pMesh->StaticFormatFlags & (1<<SMF_PROJECTWRAPPING) ) )
     {
 		float UnusedProjectionAngle = 0;
 
@@ -1377,7 +1377,7 @@ void ImageRasterProjectedCylindrical( const Mesh* pMesh, Image* pImage,
 		ProcessorFeatures |= EPixelProcessorFeatures::WithMask;
 	}
 
-    if ((pMesh->m_staticFormatFlags & (1 << SMF_PROJECT)))
+    if ((pMesh->StaticFormatFlags & (1 << SMF_PROJECT)))
     {
 		int32 UnusedLayout = 0;
 		int32 UnusedBlock = 0;
@@ -3010,7 +3010,6 @@ void MeshProject_Optimised(Mesh* Result, const Mesh* pMesh, const FProjector& pr
     // Shrink result mesh
     Result->GetVertexBuffers().SetElementCount(currentVertex);
     Result->GetIndexBuffers().SetElementCount(currentIndex);
-    Result->GetFaceBuffers().SetElementCount(currentIndex/3);
 }
 #ifdef DEBUG_PROJECTION
 UE_ENABLE_OPTIMIZATION
@@ -3023,12 +3022,12 @@ void MeshProject(Mesh* Result, const Mesh* pMesh, const FProjector& projector, b
 {
 	MUTABLE_CPUPROFILER_SCOPE(MeshProject);
 
-    if (pMesh->m_staticFormatFlags & (1<<SMF_PROJECT))
+    if (pMesh->StaticFormatFlags & (1<<SMF_PROJECT))
     {
         // Mesh-optimised version
         MeshProject_Optimised(Result, pMesh, projector, bOutSuccess);
     }
-    else if (pMesh->m_staticFormatFlags & (1<<SMF_PROJECTWRAPPING))
+    else if (pMesh->StaticFormatFlags & (1<<SMF_PROJECTWRAPPING))
     {
         // Mesh-optimised version for wrapping projectors
         // \todo: make sure the projector is a wrapping projector
@@ -3041,7 +3040,7 @@ void MeshProject(Mesh* Result, const Mesh* pMesh, const FProjector& projector, b
 
 	if (bOutSuccess)
 	{
-		Result->m_surfaces.SetNum(0);
+		Result->Surfaces.SetNum(0);
 		Result->EnsureSurfaceData();
 		Result->ResetStaticFormatFlags();
 	}
diff --git a/Engine/Plugins/Experimental/Mutable/Source/MutableRuntime/Private/MuR/OpMeshApplyPose.h b/Engine/Plugins/Experimental/Mutable/Source/MutableRuntime/Private/MuR/OpMeshApplyPose.h
--- a/Engine/Plugins/Experimental/Mutable/Source/MutableRuntime/Private/MuR/OpMeshApplyPose.h
+++ b/Engine/Plugins/Experimental/Mutable/Source/MutableRuntime/Private/MuR/OpMeshApplyPose.h
@@ -86,7 +86,7 @@ namespace mu
 		}
 
         // Get pointers to vertex position data
-		MeshBufferIteratorConst<MBF_FLOAT32, float, 3> itSource(pBase->m_VertexBuffers, MBS_POSITION, 0);
+		MeshBufferIteratorConst<MBF_FLOAT32, float, 3> itSource(pBase->VertexBuffers, MBS_POSITION, 0);
         if (!itSource.ptr())
         {
             // Formats not implemented
@@ -96,8 +96,8 @@ namespace mu
         }
 
         // Get pointers to skinning data
-		UntypedMeshBufferIteratorConst itBoneIndices(pBase->m_VertexBuffers, MBS_BONEINDICES, 0);
-		UntypedMeshBufferIteratorConst itBoneWeights(pBase->m_VertexBuffers, MBS_BONEWEIGHTS, 0);
+		UntypedMeshBufferIteratorConst itBoneIndices(pBase->VertexBuffers, MBS_BONEINDICES, 0);
+		UntypedMeshBufferIteratorConst itBoneWeights(pBase->VertexBuffers, MBS_BONEWEIGHTS, 0);
         if (!itBoneIndices.ptr() || !itBoneWeights.ptr())
         {
             // No skinning data
@@ -108,7 +108,7 @@ namespace mu
 
 		Result->CopyFrom(*pBase);
 
-		MeshBufferIterator<MBF_FLOAT32, float, 3> itTarget(Result->m_VertexBuffers, MBS_POSITION, 0);
+		MeshBufferIterator<MBF_FLOAT32, float, 3> itTarget(Result->VertexBuffers, MBS_POSITION, 0);
 		check(itTarget.ptr());
 
         // Proceed
diff --git a/Engine/Plugins/Experimental/Mutable/Source/MutableRuntime/Private/MuR/OpMeshBind.h b/Engine/Plugins/Experimental/Mutable/Source/MutableRuntime/Private/MuR/OpMeshBind.h
--- a/Engine/Plugins/Experimental/Mutable/Source/MutableRuntime/Private/MuR/OpMeshBind.h
+++ b/Engine/Plugins/Experimental/Mutable/Source/MutableRuntime/Private/MuR/OpMeshBind.h
@@ -1154,7 +1154,7 @@ namespace mu
 			Result->AdditionalBuffers.Emplace(EMeshBufferType::SkeletonDeformBinding, MoveTemp(SkeletonBuffer));
 		}
 
-		const PhysicsBody* ResultPhysicsBody = Result->m_pPhysicsBody.get();
+		const PhysicsBody* ResultPhysicsBody = Result->PhysicsBody.get();
 		if (bReshapePhysics && (ResultPhysicsBody || Result->AdditionalPhysicsBodies.Num()) && PhysicsToDeform.Num())
 		{
 			MUTABLE_CPUPROFILER_SCOPE(BindPhysicsBody);
diff --git a/Engine/Plugins/Experimental/Mutable/Source/MutableRuntime/Private/MuR/OpMeshClipMorphPlane.h b/Engine/Plugins/Experimental/Mutable/Source/MutableRuntime/Private/MuR/OpMeshClipMorphPlane.h
--- a/Engine/Plugins/Experimental/Mutable/Source/MutableRuntime/Private/MuR/OpMeshClipMorphPlane.h
+++ b/Engine/Plugins/Experimental/Mutable/Source/MutableRuntime/Private/MuR/OpMeshClipMorphPlane.h
@@ -417,18 +417,17 @@ namespace mu
 		SIZE_T removedIndices = itBase - itDest;
 		check(removedIndices % 3 == 0);
 
-		Result->GetFaceBuffers().SetElementCount(aFaceCount - (int32)removedIndices / 3);
 		Result->GetIndexBuffers().SetElementCount(aFaceCount * 3 - (int32)removedIndices);
 
 		// TODO: Should redo/reorder the face buffer before SetElementCount since some deleted faces could be left and some remaining faces deleted.
 
         // Fix the surface data if present.
-        if (Result->m_surfaces.Num())
+        if (Result->Surfaces.Num())
         {
             // We assume there will be only one.
-            check(Result->m_surfaces.Num()==1);
+            check(Result->Surfaces.Num()==1);
 
-            Result->m_surfaces[0].m_indexCount -= (int32)removedIndices;
+            Result->Surfaces[0].IndexCount -= (int32)removedIndices;
         }
 	}
 }
diff --git a/Engine/Plugins/Experimental/Mutable/Source/MutableRuntime/Private/MuR/OpMeshClipWithMesh.cpp b/Engine/Plugins/Experimental/Mutable/Source/MutableRuntime/Private/MuR/OpMeshClipWithMesh.cpp
--- a/Engine/Plugins/Experimental/Mutable/Source/MutableRuntime/Private/MuR/OpMeshClipWithMesh.cpp
+++ b/Engine/Plugins/Experimental/Mutable/Source/MutableRuntime/Private/MuR/OpMeshClipWithMesh.cpp
@@ -558,72 +558,6 @@ namespace mu { namespace
 	}
 
 
-
-    //---------------------------------------------------------------------------------------------
-    //! Reference version
-    //---------------------------------------------------------------------------------------------
-  //  inline void MeshClipWithMesh_Reference(Mesh* Result, const Mesh* pBase, const Mesh* pClipMesh, bool& bOutSuccess)
-  //  {
-  //      MUTABLE_CPUPROFILER_SCOPE(MeshClipWithMesh_Reference);
-  //      
-		//bOutSuccess = true;
-
-  //      uint32 vcount = pClipMesh->GetVertexBuffers().GetElementCount();
-  //      if (!vcount)
-  //      {
-		//	bOutSuccess = false;
-		//	return; // Since there's nothing to clip against return a copy of the original base mesh
-  //      }
-
-  //      TArray<uint8> vertex_in_clip_mesh;  // Stores whether each vertex in the original mesh in in the clip mesh volume
-  //      MeshClipMeshClassifyVertices(vertex_in_clip_mesh, pBase, pClipMesh);
-
-		//Result->CopyFrom(*pBase);
-  //      // Now remove all the faces from the result mesh that have all the vertices outside the clip volume
-  //      UntypedMeshBufferIteratorConst itBase(Result->GetIndexBuffers(), MBS_VERTEXINDEX);
-  //      UntypedMeshBufferIterator itDest(Result->GetIndexBuffers(), MBS_VERTEXINDEX);
-  //      int aFaceCount = Result->GetFaceCount();
-
-  //      UntypedMeshBufferIteratorConst ito(Result->GetIndexBuffers(), MBS_VERTEXINDEX);
-  //      for (int f = 0; f < aFaceCount; ++f)
-  //      {
-  //          vec3<uint32> ov;
-  //          ov[0] = ito.GetAsUINT32(); ++ito;
-  //          ov[1] = ito.GetAsUINT32(); ++ito;
-  //          ov[2] = ito.GetAsUINT32(); ++ito;
-
-  //          bool all_verts_in = vertex_in_clip_mesh[ov[0]] && vertex_in_clip_mesh[ov[1]] && vertex_in_clip_mesh[ov[2]];
-
-  //          if (!all_verts_in)
-  //          {
-  //              if (itDest.ptr() != itBase.ptr())
-  //              {
-		//			FMemory::Memcpy(itDest.ptr(), itBase.ptr(), itBase.GetElementSize() * 3);
-  //              }
-
-  //              itDest += 3;
-  //          }
-
-  //          itBase += 3;
-  //      }
-
-  //      SIZE_T removedIndices = itBase - itDest;
-  //      check(removedIndices % 3 == 0);
-
-  //      Result->GetFaceBuffers().SetElementCount(aFaceCount - (int)removedIndices / 3);
-  //      Result->GetIndexBuffers().SetElementCount(aFaceCount * 3 - (int)removedIndices);
-
-  //      // TODO: Should redo/reorder the face buffer before SetElementCount since some deleted faces could be left and some remaining faces deleted.
-
-  //      //if (Result->GetFaceBuffers().GetElementCount() == 0)
-  //      //{
-  //      //	Result->CopyFrom(pBase); // If all faces have been discarded, return a copy of the unmodified mesh because unreal doesn't like empty meshes
-  //      //}
-
-    //      MeshRemoveUnusedVertices(Result);
-  //  }
-
-
 	/** Make a mask with the indices of the vertices with 0 in the IncludedVertices array. */
 	void CreateMask(Mesh* Result, const Mesh* Base, const TArray<uint8>& IncludedVertices)
 	{
@@ -636,43 +570,52 @@ namespace mu { namespace
 			}
 		}
 
-		// Create the vertex buffer
-		{
-			Result->GetVertexBuffers().SetElementCount(MaskVertexCount);
-			Result->GetVertexBuffers().SetBufferCount(1);
+		Result->GetVertexBuffers().SetElementCount(MaskVertexCount);
+		Result->GetVertexBuffers().SetBufferCount(1);
+		Result->VertexIDPrefix = Base->VertexIDPrefix;
 
-			// Vertex index channel
-			EMeshBufferSemantic Semantic = MBS_VERTEXINDEX;
-			int32 SemanticIndex = 0;
-			EMeshBufferFormat Format = MBF_UINT32;
-			int32 Components = 1;
-			int32 Offsets = 0;
-
-			Result->GetVertexBuffers().SetBuffer
-			(
-				0,
-				sizeof(uint32),
-				1,
-				&Semantic,
-				&SemanticIndex,
-				&Format,
-				&Components,
-				&Offsets
-			);
-		}
+		EMeshBufferSemantic Semantic = MBS_VERTEXINDEX;
+		int32 SemanticIndex = 0;
+		int32 Components = 1;
+		int32 Offsets = 0;
 
-		// Fill the buffer
-		MeshBufferIterator<MBF_UINT32, uint32, 1> itMask(Result->GetVertexBuffers(), MBS_VERTEXINDEX, 0);
-		UntypedMeshBufferIteratorConst itBase(Base->GetVertexBuffers(), MBS_VERTEXINDEX, 0);
-		for (int32 v = 0; v<IncludedVertices.Num(); ++v)
+		bool bMakeRelativeIds = !Base->AreVertexIdsExplicit();
+		if (bMakeRelativeIds)
 		{
-			if (!IncludedVertices[v])
+			EMeshBufferFormat Format = MBF_UINT32;
+			Result->GetVertexBuffers().SetBuffer( 0, sizeof(uint32), 1, &Semantic, &SemanticIndex, &Format, &Components, &Offsets );
+
+			// Fill the buffer
+			uint32* Data = reinterpret_cast<uint32*>(Result->GetVertexBuffers().GetBufferData(0));
+			MeshVertexIdIteratorConst ItBase(Base);
+			for (int32 v = 0; v < IncludedVertices.Num(); ++v)
 			{
-				(*itMask)[0] = itBase.GetAsUINT32();
-				++itMask;
+				if (!IncludedVertices[v])
+				{
+					uint64 ID = ItBase.Get();
+					*Data = uint32(ID & 0xffffffff);
+					++Data;
+				}
+				++ItBase;
 			}
+		}
+		else
+		{
+			EMeshBufferFormat Format = MBF_UINT64;
+			Result->GetVertexBuffers().SetBuffer( 0, sizeof(uint64), 1, &Semantic, &SemanticIndex, &Format, &Components, &Offsets );
 
-			++itBase;
+			// Fill the buffer
+			uint64* Data = reinterpret_cast<uint64*>( Result->GetVertexBuffers().GetBufferData(0) );
+			MeshVertexIdIteratorConst ItBase(Base);
+			for (int32 v = 0; v < IncludedVertices.Num(); ++v)
+			{
+				if (!IncludedVertices[v])
+				{
+					*Data = ItBase.Get();
+					++Data;
+				}
+				++ItBase;
+			}
 		}
 	}
 
@@ -731,17 +674,8 @@ namespace mu
         SIZE_T RemovedIndices = ItBase - ItDest;
         check(RemovedIndices % 3 == 0);
 
-        Result->GetFaceBuffers().SetElementCount(AFaceCount - (int32)RemovedIndices / 3);
         Result->GetIndexBuffers().SetElementCount(AFaceCount * 3 - (int32)RemovedIndices);
 
-        // TODO: Should redo/reorder the face buffer before SetElementCount since some deleted faces could be left and some remaining faces deleted.
-
-        //if (Result->GetFaceBuffers().GetElementCount() == 0)
-        //{
-        //	Result->CopyFrom(pBase); // If all faces have been discarded, return a copy of the unmodified mesh because unreal doesn't like empty meshes
-        //}
-
-
         MeshRemoveUnusedVertices(Result);
     }
 
@@ -837,49 +771,6 @@ namespace mu
 	}
 
 
-  //  void MeshMaskClipMesh_Reference(Mesh* Result, const Mesh* pBase, const Mesh* pClipMesh )
-  //  {
-  //      MUTABLE_CPUPROFILER_SCOPE(MeshMaskClipMesh_Reference);
-
-  //      uint32 vcount = pClipMesh->GetVertexBuffers().GetElementCount();
-  //      if (!vcount)
-  //      {
-  //          return; // Since there's nothing to clip against return a copy of the original base mesh
-  //      }
-
-  //      TArray<uint8> vertex_in_clip_mesh;  // Stores whether each vertex in the original mesh in in the clip mesh volume
-  //      MeshClipMeshClassifyVertices( vertex_in_clip_mesh, pBase, pClipMesh );
-
-  //      // We only remove vertices if all their faces are clipped
-		//TArray<uint8> vertex_with_face_not_clipped;
-		//vertex_with_face_not_clipped.SetNumZeroed(vertex_in_clip_mesh.Num());
-
-  //      UntypedMeshBufferIteratorConst ito(pBase->GetIndexBuffers(), MBS_VERTEXINDEX);
-  //      int aFaceCount = pBase->GetFaceCount();
-  //      for (int f = 0; f < aFaceCount; ++f)
-  //      {
-  //          vec3<uint32> ov;
-  //          ov[0] = ito.GetAsUINT32(); ++ito;
-  //          ov[1] = ito.GetAsUINT32(); ++ito;
-  //          ov[2] = ito.GetAsUINT32(); ++ito;
-
-  //          bool faceClipped =
-  //                  vertex_in_clip_mesh[ov[0]] &&
-  //                  vertex_in_clip_mesh[ov[1]] &&
-  //                  vertex_in_clip_mesh[ov[2]];
-
-  //          if (!faceClipped)
-  //          {
-  //              vertex_with_face_not_clipped[ov[0]] = true;
-  //              vertex_with_face_not_clipped[ov[1]] = true;
-  //              vertex_with_face_not_clipped[ov[2]] = true;
-  //          }
-  //      }
-
-		//CreateMask(Result, pBase, vertex_with_face_not_clipped );
-  //  }
-
-
 	void MeshMaskDiff(Mesh* Result, const Mesh* pBase, const Mesh* pFragment, bool& bOutSuccess)
     {
         MUTABLE_CPUPROFILER_SCOPE(MeshMaskDiff);
@@ -914,7 +805,7 @@ namespace mu
             }
         }
         float tolerance = 1e-5f * aabbox.size.Length();
-        Mesh::VERTEX_MATCH_MAP vertexMap;
+        Mesh::FVertexMatchMap vertexMap;
         pFragment->GetVertexMap( *pBase, vertexMap, tolerance );
 
 
@@ -988,9 +879,9 @@ namespace mu
                 hasFace = true;
                 for ( int vi=0; hasFace && vi<3; ++vi )
                 {
-                    hasFace = vertexMap.Matches(v[vi],ov[0])
-                         || vertexMap.Matches(v[vi],ov[1])
-                         || vertexMap.Matches(v[vi],ov[2]);
+                    hasFace = vertexMap.DoMatch(v[vi],ov[0])
+                         || vertexMap.DoMatch(v[vi],ov[1])
+                         || vertexMap.DoMatch(v[vi],ov[2]);
                 }
             }
 
diff --git a/Engine/Plugins/Experimental/Mutable/Source/MutableRuntime/Private/MuR/OpMeshDifference.h b/Engine/Plugins/Experimental/Mutable/Source/MutableRuntime/Private/MuR/OpMeshDifference.h
--- a/Engine/Plugins/Experimental/Mutable/Source/MutableRuntime/Private/MuR/OpMeshDifference.h
+++ b/Engine/Plugins/Experimental/Mutable/Source/MutableRuntime/Private/MuR/OpMeshDifference.h
@@ -13,9 +13,9 @@ namespace mu
 	//! If the channel list is empty all the channels will be compared.
 	//---------------------------------------------------------------------------------------------
     inline void MeshDifference( Mesh* Result, const Mesh* pBase, const Mesh* pTarget,
-                                   int numChannels,
+                                   int32 numChannels,
                                    const EMeshBufferSemantic* semantics,
-                                   const int* semanticIndices,
+                                   const int32* semanticIndices,
                                    bool ignoreTexCoords, bool& bOutSuccess)
 
 	{
@@ -30,8 +30,7 @@ namespace mu
 		bool bIsCorrect =
                 pBase
                 &&
-				( pBase->GetVertexBuffers().GetElementCount()
-						== pTarget->GetVertexBuffers().GetElementCount() )
+				( pBase->GetVertexBuffers().GetElementCount() == pTarget->GetVertexBuffers().GetElementCount() )
 				&&
                 ( pBase->GetIndexCount() == pTarget->GetIndexCount() )
                 &&
@@ -44,22 +43,20 @@ namespace mu
 			return ;
 		}
 
-		uint32 vcount = pBase->GetVertexBuffers().GetElementCount();		
+		int32 VertexCount = pBase->GetVertexBuffers().GetElementCount();		
 
 		// If no channels were specified, get them all
 		TArray<EMeshBufferSemantic> allSemantics;		
-		TArray<int> allSemanticIndices;
+		TArray<int32> allSemanticIndices;
 		if ( !numChannels )
 		{
-			for ( int vb=0; vb<pBase->GetVertexBuffers().GetBufferCount(); ++vb )
+			for ( int32 vb=0; vb<pBase->GetVertexBuffers().GetBufferCount(); ++vb )
 			{
-				for ( int c=0; c<pBase->GetVertexBuffers().GetBufferChannelCount(vb); ++c )
+				for ( int32 c=0; c<pBase->GetVertexBuffers().GetBufferChannelCount(vb); ++c )
 				{
 					EMeshBufferSemantic sem = MBS_NONE;
-					int semIndex = 0;
-					pBase->GetVertexBuffers().GetChannel( vb, c,
-															&sem, &semIndex,
-															nullptr, nullptr, nullptr );			
+					int32 semIndex = 0;
+					pBase->GetVertexBuffers().GetChannel( vb, c, &sem, &semIndex, nullptr, nullptr, nullptr );			
 					if ( sem != MBS_VERTEXINDEX &&
 					     sem != MBS_BONEINDICES &&
 					     sem != MBS_BONEWEIGHTS &&
@@ -82,13 +79,13 @@ namespace mu
 		// Make a delta of every vertex
 		// TODO: Not always 4 components
 		// TODO: Not always floats
-		int differentVertexCount = 0;
+		int32 differentVertexCount = 0;
 		TArray<bool> isVertexDifferent;
-		isVertexDifferent.SetNumZeroed(vcount);
+		isVertexDifferent.SetNumZeroed(VertexCount);
 		TArray< FVector4f > deltas;
-		deltas.SetNum(vcount * numChannels);
+		deltas.SetNum(VertexCount * numChannels);
 
-		for ( int c=0; c<numChannels; ++c )
+		for ( int32 c=0; c<numChannels; ++c )
 		{
 			UntypedMeshBufferIteratorConst baseIt		
 					( pBase->GetVertexBuffers(), semantics[c], semanticIndices[c] );
@@ -96,7 +93,7 @@ namespace mu
 			UntypedMeshBufferIteratorConst targetIt		
 					( pTarget->GetVertexBuffers(), semantics[c], semanticIndices[c] );
 
-			for ( size_t v=0; v<vcount; ++v )
+			for ( int32 v=0; v<VertexCount; ++v )
 			{
 				FVector4f base = baseIt.GetAsVec4f();
 				FVector4f target = targetIt.GetAsVec4f();
@@ -122,77 +119,101 @@ namespace mu
 		// Create the morph mesh
 		{
 			Result->GetVertexBuffers().SetElementCount( differentVertexCount );			
-			Result->GetVertexBuffers().SetBufferCount( 1 );								
+			Result->GetVertexBuffers().SetBufferCount( 2 );								
 
 			TArray<EMeshBufferSemantic> semantic;
-			TArray<int> semanticIndex;
+			TArray<int32> semanticIndex;
 			TArray<EMeshBufferFormat> format;
-			TArray<int> components;
-			TArray<int> offsets;
-			int offset = 0;										
-
-			// Vertex index channel
-			semantic.Add( MBS_VERTEXINDEX );					
-			semanticIndex.Add( 0 );
-			format.Add( MBF_UINT32 );
-			components.Add( 1 );
-			offsets.Add( offset );
-			offset += 4;						
-
-			for ( int c=0; c<numChannels; ++c )
+			TArray<int32> components;
+			TArray<int32> offsets;
+			int32 ElementSize = 0;
+
+			// The first buffer will contain the actual morph data.
+			for (int32 c = 0; c < numChannels; ++c)
+			{
+				semantic.Add(semantics[c]);
+				semanticIndex.Add(semanticIndices[c]);
+				format.Add(MBF_FLOAT32);
+				components.Add(4);
+				offsets.Add(ElementSize);
+				ElementSize += 4 * sizeof(float);
+			}
+
+			int32 BufferIndex = 0;
+			Result->GetVertexBuffers().SetBuffer(BufferIndex, ElementSize, semantic.Num(), semantic.GetData(), semanticIndex.GetData(), format.GetData(), components.GetData(), offsets.GetData());
+
+			// The second buffer will contain the ids of the vertices to morph.
+			semantic = { MBS_VERTEXINDEX };
+			semanticIndex = { 0 };
+			components = { 1 };
+			offsets = { 0 };
+			bool bGenerateRelativeIds = !pBase->AreVertexIdsExplicit();
+			if (bGenerateRelativeIds)
 			{
-				semantic.Add( semantics[c] );
-				semanticIndex.Add( semanticIndices[c] );
-				format.Add( MBF_FLOAT32 );
-				components.Add( 4 );
-				offsets.Add( offset );
-				offset += 16;				
+				format = { MBF_UINT32 };
+				ElementSize = sizeof(uint32);
 			}
+			else
+			{
+				format = { MBF_UINT64 };
+				ElementSize = sizeof(uint64);
+			}
+
+			BufferIndex = 1;
+			Result->GetVertexBuffers().SetBuffer(BufferIndex, ElementSize, semantic.Num(), semantic.GetData(), semanticIndex.GetData(), format.GetData(), components.GetData(), offsets.GetData());
+
 
-			Result->GetVertexBuffers().SetBuffer
-				(
-					0,
-					offset,
-					numChannels+1,
-					&semantic[0],
-					&semanticIndex[0],
-					&format[0],
-					&components[0],
-					&offsets[0]
-				);
+			if (bGenerateRelativeIds)
+			{
+				Result->VertexIDPrefix = pBase->VertexIDPrefix;
+			}
+			else
+			{
+				Result->VertexIDPrefix = 0;
+			}
 
 			// Source vertex index channel
-			UntypedMeshBufferIteratorConst itBaseId( pBase->GetVertexBuffers(), MBS_VERTEXINDEX, 0 );
-			check(itBaseId.GetElementSize());
+			MeshVertexIdIteratorConst IdIterator( pBase );
+			check(IdIterator.IsValid());
 
 			// Set the data
-			uint8_t* ResultData = Result->GetVertexBuffers().GetBufferData( 0 );
-			for (size_t v=0; v<vcount; ++v)
+			uint8* ResultData = Result->GetVertexBuffers().GetBufferData(0);
+			uint8* ResultIds = Result->GetVertexBuffers().GetBufferData(1);
+			for (int32 v=0; v<VertexCount; ++v)
 			{
 				if ( isVertexDifferent[v] )
 				{
-					// index
-					*((uint32*)ResultData) = itBaseId.GetAsUINT32();			
-					ResultData += 4;
+					// Vertex Id
+					uint64 FullId = IdIterator.Get();
+					if (bGenerateRelativeIds)
+					{
+						uint32 RelativeId = uint32_t(FullId & 0xffffffff);
+						*((uint32*)ResultIds) = RelativeId;
+						ResultIds += sizeof(uint32);
+					}
+					else
+					{
+						*((uint64*)ResultIds) = FullId;
+						ResultIds += sizeof(uint64);
+					}
 
 					// all channels
-					for ( int c=0; c<numChannels; ++c )
+					for ( int32 c=0; c<numChannels; ++c )
 					{
-						FMemory::Memcpy(ResultData, &deltas[v * numChannels + c], 16);
-						ResultData += 4*4;											
+						constexpr int32 ChannelSizeBytes = 16;
+						FMemory::Memcpy(ResultData, &deltas[v * numChannels + c], ChannelSizeBytes);
+						ResultData += ChannelSizeBytes;
 					}
 				}
 
-				++itBaseId;
+				++IdIterator;
 			}
 		}
 
 
-        // Rebuild surface data.
-        // \todo: For now make single surfaced
-        Result->m_surfaces.SetNum(0);
+        // Morphs are always single surfaced
+        Result->Surfaces.SetNum(0);
         Result->EnsureSurfaceData();
 	}
 
-
 }
diff --git a/Engine/Plugins/Experimental/Mutable/Source/MutableRuntime/Private/MuR/OpMeshExtractLayoutBlock.h b/Engine/Plugins/Experimental/Mutable/Source/MutableRuntime/Private/MuR/OpMeshExtractLayoutBlock.h
--- a/Engine/Plugins/Experimental/Mutable/Source/MutableRuntime/Private/MuR/OpMeshExtractLayoutBlock.h
+++ b/Engine/Plugins/Experimental/Mutable/Source/MutableRuntime/Private/MuR/OpMeshExtractLayoutBlock.h
@@ -11,39 +11,65 @@ namespace mu
 
 
     //---------------------------------------------------------------------------------------------
-    inline void MeshExtractFromVertices( const Mesh* pSource,
-                                         Mesh* pResult,
-                                         const TArray<int>& oldToNew,
-                                         const TArray<int>& newToOld )
+    inline void MeshExtractFromVertices( const Mesh* Source,
+                                         Mesh* Result,
+                                         const TArray<int32>& oldToNew,
+                                         const TArray<int32>& newToOld )
     {
-        int resultVertices = newToOld.Num();
+        int32 ResultVertices = newToOld.Num();
 
         // Assemble the new vertex buffer
-        pResult->GetVertexBuffers().SetElementCount( resultVertices );
-        for ( int b=0; b<pResult->GetVertexBuffers().GetBufferCount(); ++b )
+        Result->GetVertexBuffers().SetElementCount(ResultVertices);
+        for ( int b=0; b<Result->GetVertexBuffers().GetBufferCount(); ++b )
         {
-            const uint8_t* pSourceData = pSource->GetVertexBuffers().GetBufferData(b);
-            uint8_t* pDest = pResult->GetVertexBuffers().GetBufferData(b);
-            int size = pResult->GetVertexBuffers().GetElementSize(b);
-            for ( int v=0; v<resultVertices; ++v )
+            const uint8* pSourceData = Source->GetVertexBuffers().GetBufferData(b);
+            uint8* pDest = Result->GetVertexBuffers().GetBufferData(b);
+            int32 size = Result->GetVertexBuffers().GetElementSize(b);
+            for ( int32 NewIndex =0; NewIndex < ResultVertices; ++NewIndex)
             {
-                int i = newToOld[v];
-                memcpy( pDest, pSourceData + size*i, size );
+                int32 OldIndex = newToOld[NewIndex];
+                FMemory::Memcpy( pDest, pSourceData + size*OldIndex, size );
                 pDest+=size;
             }
         }
 
+		// If the vertices are explicit or relative, the above operation will already handle them correctly
+		// Otherwise, create a relative vertex ID buffer if necessary
+		Result->VertexIDPrefix = Source->VertexIDPrefix;
+		if (Source->AreVertexIdsImplicit() 
+			&& 
+			// If we extract everything, we can keep the ids implicit.
+			ResultVertices!=Source->GetVertexCount())
+		{
+			// Add a new buffer
+			int32 NewBuffer = Result->VertexBuffers.GetBufferCount();
+			Result->VertexBuffers.SetBufferCount(NewBuffer + 1);
+			EMeshBufferSemantic Semantic = MBS_VERTEXINDEX;
+			int32 SemanticIndex = 0;
+			EMeshBufferFormat Format = MBF_UINT32;
+			int32 Components = 1;
+			int32 Offset = 0;
+			Result->VertexBuffers.SetBuffer(NewBuffer, sizeof(uint32), 1, &Semantic, &SemanticIndex, &Format, &Components, &Offset);
+			uint32* pIdData = reinterpret_cast<uint32*>(Result->VertexBuffers.GetBufferData(NewBuffer));
+
+			for (int32 NewIndex = 0; NewIndex < ResultVertices; ++NewIndex)
+			{
+				uint32 OldIndex = newToOld[NewIndex];
+				(*pIdData++) = OldIndex;
+			}
+		}
+
         // Assemble the new index buffers
 		TArray<bool> usedSourceFaces;
-		usedSourceFaces.SetNumZeroed(pSource->GetFaceCount());
-        UntypedMeshBufferIteratorConst itIndex( pSource->GetIndexBuffers(), MBS_VERTEXINDEX );
-        UntypedMeshBufferIterator itResultIndex( pResult->GetIndexBuffers(), MBS_VERTEXINDEX );
+		usedSourceFaces.SetNumZeroed(Source->GetFaceCount());
+        UntypedMeshBufferIteratorConst itIndex( Source->GetIndexBuffers(), MBS_VERTEXINDEX );
+        UntypedMeshBufferIterator itResultIndex( Result->GetIndexBuffers(), MBS_VERTEXINDEX );
         int indexCount = 0;
         if ( itIndex.GetFormat()==MBF_UINT32 )
         {
-            const uint32_t* pIndices = reinterpret_cast<const uint32_t*>( itIndex.ptr() );
-            uint32_t* pDestIndices = reinterpret_cast<uint32_t*>( itResultIndex.ptr() );
-            for ( int i=0; i<pSource->GetIndexCount()/3; ++i )
+            const uint32* pIndices = reinterpret_cast<const uint32_t*>( itIndex.ptr() );
+            uint32* pDestIndices = reinterpret_cast<uint32_t*>( itResultIndex.ptr() );
+            for ( int32 i=0; i<Source->GetIndexCount()/3; ++i )
             {
                 if ( oldToNew[ pIndices[i*3+0] ]>=0
                      &&
@@ -64,7 +90,7 @@ namespace mu
         {
             const uint16* pIndices = reinterpret_cast<const uint16*>( itIndex.ptr() );
             uint16* pDestIndices = reinterpret_cast<uint16*>( itResultIndex.ptr() );
-            for ( int i=0; i<pSource->GetIndexCount()/3; ++i )
+            for ( int32 i=0; i<Source->GetIndexCount()/3; ++i )
             {
                 if ( oldToNew[ pIndices[i*3+0] ]>=0
                      &&
@@ -85,30 +111,7 @@ namespace mu
         {
             check( false );
         }
-        pResult->GetIndexBuffers().SetElementCount( indexCount );
-
-
-        // Face buffers
-        {
-            int faceCount = 0;
-            for ( int i=0; i<pSource->GetFaceCount(); ++i )
-            {
-                if ( usedSourceFaces[i] )
-                {
-                    for ( int b=0; b<pResult->GetFaceBuffers().GetBufferCount(); ++b )
-                    {
-                        const uint8_t* pSourceData = pSource->GetFaceBuffers().GetBufferData(b);
-                        uint8_t* pDest = pResult->GetFaceBuffers().GetBufferData(b);
-                        int size = pResult->GetFaceBuffers().GetElementSize(b);
-                        memcpy( pDest + size*faceCount, pSourceData + size*i, size );
-                    }
-
-                    faceCount++;
-                }
-            }
-            check( faceCount == indexCount/3 );
-            pResult->GetFaceBuffers().SetElementCount( faceCount );
-        }
+        Result->GetIndexBuffers().SetElementCount( indexCount );
     }
 
 
@@ -124,15 +127,14 @@ namespace mu
 		// TODO: Optimise
 		Result->CopyFrom(*pSource);
 
-		UntypedMeshBufferIteratorConst itBlocks( pSource->GetVertexBuffers(),
-												 MBS_LAYOUTBLOCK, layout );
+		UntypedMeshBufferIteratorConst itBlocks( pSource->GetVertexBuffers(), MBS_LAYOUTBLOCK, layout );
 
         if (itBlocks.GetFormat()!=MBF_NONE)
         {
             int resultVertices = 0;
-			TArray<int> oldToNew;
+			TArray<int32> oldToNew;
 			oldToNew.Init(-1,pSource->GetVertexCount());
-			TArray<int> newToOld;
+			TArray<int32> newToOld;
             newToOld.Reserve( pSource->GetVertexCount() );
 
             if ( itBlocks.GetFormat()==MBF_UINT16 )
@@ -167,7 +169,7 @@ namespace mu
             MeshExtractFromVertices(pSource, Result, oldToNew, newToOld);
         }
 
-        Result->m_surfaces.Empty();
+        Result->Surfaces.Empty();
         Result->EnsureSurfaceData();
 	}
 
diff --git a/Engine/Plugins/Experimental/Mutable/Source/MutableRuntime/Private/MuR/OpMeshFormat.cpp b/Engine/Plugins/Experimental/Mutable/Source/MutableRuntime/Private/MuR/OpMeshFormat.cpp
--- a/Engine/Plugins/Experimental/Mutable/Source/MutableRuntime/Private/MuR/OpMeshFormat.cpp
+++ b/Engine/Plugins/Experimental/Mutable/Source/MutableRuntime/Private/MuR/OpMeshFormat.cpp
@@ -1,6 +1,5 @@
 // Copyright Epic Games, Inc. All Rights Reserved.
 
-
 #include "MuR/OpMeshFormat.h"
 
 #include "Containers/Array.h"
@@ -25,12 +24,7 @@ namespace mu
 {
 
 	//-------------------------------------------------------------------------------------------------
-	void MeshFormatBuffer
-	(
-		const FMeshBufferSet& Source,
-		FMeshBufferSet& Result,
-		int bufferIndex
-	)
+	void MeshFormatBuffer( const FMeshBufferSet& Source, FMeshBufferSet& Result, int32 bufferIndex )
 	{
 		int vCount = Source.GetElementCount();
 
@@ -368,7 +362,7 @@ namespace mu
 		// For every vertex buffer in result
 		int32 vCount = Source.GetElementCount();
 		Result.SetElementCount(vCount);
-		for (int b = 0; b < Result.GetBufferCount(); ++b)
+		for (int32 b = 0; b < Result.GetBufferCount(); ++b)
 		{
 			MeshFormatBuffer(Source, Result, b);
 		}
@@ -377,7 +371,7 @@ namespace mu
 		// Detect internal system buffers and clone them unmodified.
 		if (bKeepSystemBuffers)
 		{
-			for (int b = 0; b < Source.GetBufferCount(); ++b)
+			for (int32 b = 0; b < Source.GetBufferCount(); ++b)
 			{
 				bool bIsSystemBuffer = false;
 
@@ -385,10 +379,10 @@ namespace mu
 				if (Source.GetBufferChannelCount(b) == 1)
 				{
 					EMeshBufferSemantic sourceSemantic;
-					int sourceSemanticIndex;
+					int32 sourceSemanticIndex;
 					EMeshBufferFormat sourceFormat;
-					int sourceComponents;
-					int sourceOffset;
+					int32 sourceComponents;
+					int32 sourceOffset;
 					Source.GetChannel
 					(
 						b, 0,
@@ -420,12 +414,11 @@ namespace mu
 	void MeshFormat
 	(
 		Mesh* Result, 
-		const Mesh* pPureSource,
-		const Mesh* pFormat,
+		const Mesh* PureSource,
+		const Mesh* Format,
 		bool keepSystemBuffers,
 		bool formatVertices,
 		bool formatIndices,
-		bool formatFaces,
 		bool ignoreMissingChannels,
 		bool& bOutSuccess
 	)
@@ -433,28 +426,29 @@ namespace mu
 		MUTABLE_CPUPROFILER_SCOPE(MeshFormat);
 		bOutSuccess = true;
 
-		if (!pPureSource) 
+		if (!PureSource) 
 		{
 			check(false);
 			bOutSuccess = false;	
 			return;
 		}
 
-		if (!pFormat)
+		if (!Format)
 		{
 			check(false);
 			bOutSuccess = false;
 			return;
 		}
 
-		Ptr<const Mesh> pSource = pPureSource;
+		Ptr<const Mesh> Source = PureSource;
 
-		Result->CopyFrom(*pFormat);
+		Result->CopyFrom(*Format);
+		Result->VertexIDPrefix = Source->VertexIDPrefix;
 
 		// Make sure that the bone indices will fit in this format, or extend it.
 		if (formatVertices)
 		{
-			const FMeshBufferSet& VertexBuffers = pSource->GetVertexBuffers();
+			const FMeshBufferSet& VertexBuffers = Source->GetVertexBuffers();
 
 			const int32 BufferCount = VertexBuffers.GetBufferCount();
 			for (int32 BufferIndex = 0; BufferIndex < BufferCount; ++BufferIndex)
@@ -466,8 +460,8 @@ namespace mu
 
 					if (Channel.m_semantic == MBS_BONEINDICES)
 					{
-						int resultBuf = 0;
-						int resultChan = 0;
+						int32 resultBuf = 0;
+						int32 resultChan = 0;
 						FMeshBufferSet& formBuffs = Result->GetVertexBuffers();
 						formBuffs.FindChannel(MBS_BONEINDICES, Channel.m_semanticIndex, &resultBuf, &resultChan);
 						if (resultBuf >= 0)
@@ -514,65 +508,51 @@ namespace mu
 			}
 		}
 
-		// \todo Make sure that the vertex indices will fit in this format, or extend it.
-
 		if (formatVertices)
 		{
-			FormatBufferSet(pSource->GetVertexBuffers(), Result->GetVertexBuffers(),
+			FormatBufferSet(Source->GetVertexBuffers(), Result->GetVertexBuffers(),
 				keepSystemBuffers, ignoreMissingChannels, true);
 		}
 		else
 		{
-			Result->m_VertexBuffers = pSource->GetVertexBuffers();
+			Result->VertexBuffers = Source->GetVertexBuffers();
 		}
 
 		if (formatIndices)
 		{
-			FormatBufferSet(pSource->GetIndexBuffers(), Result->GetIndexBuffers(), keepSystemBuffers,
+			// \todo Make sure that the vertex indices will fit in this format, or extend it.
+			FormatBufferSet(Source->GetIndexBuffers(), Result->GetIndexBuffers(), keepSystemBuffers,
 				ignoreMissingChannels, false);
 		}
 		else
 		{
-			Result->m_IndexBuffers = pSource->GetIndexBuffers();
-		}
-
-		if (formatFaces)
-		{
-			FormatBufferSet(pSource->GetFaceBuffers(),
-				Result->GetFaceBuffers(),
-				keepSystemBuffers,
-				ignoreMissingChannels,
-				false);
-		}
-		else
-		{
-			Result->m_FaceBuffers = pSource->GetFaceBuffers();
+			Result->IndexBuffers = Source->GetIndexBuffers();
 		}
 
 		// Copy the rest of the data
-		Result->SetSkeleton(pSource->GetSkeleton());
-		Result->SetPhysicsBody(pSource->GetPhysicsBody());
+		Result->SetSkeleton(Source->GetSkeleton());
+		Result->SetPhysicsBody(Source->GetPhysicsBody());
 
-		Result->m_layouts.Empty();
-		for (const Ptr<const Layout>& Layout : pSource->m_layouts)
+		Result->Layouts.Empty();
+		for (const Ptr<const Layout>& Layout : Source->Layouts)
 		{
-			Result->m_layouts.Add(Layout->Clone());
+			Result->Layouts.Add(Layout->Clone());
 		}
 
-		Result->m_tags = pSource->m_tags;
-		Result->StreamedResources = pSource->StreamedResources;
+		Result->Tags = Source->Tags;
+		Result->StreamedResources = Source->StreamedResources;
 
-		Result->AdditionalBuffers = pSource->AdditionalBuffers;
+		Result->AdditionalBuffers = Source->AdditionalBuffers;
 
-		Result->BonePoses = pSource->BonePoses;
-		Result->BoneMap = pSource->BoneMap;
+		Result->BonePoses = Source->BonePoses;
+		Result->BoneMap = Source->BoneMap;
 
-		Result->SkeletonIDs = pSource->SkeletonIDs;
+		Result->SkeletonIDs = Source->SkeletonIDs;
 
 		// A shallow copy is done here, it should not be a problem.
-		Result->AdditionalPhysicsBodies = pSource->AdditionalPhysicsBodies;
+		Result->AdditionalPhysicsBodies = Source->AdditionalPhysicsBodies;
 
-		Result->m_surfaces = pSource->m_surfaces;
+		Result->Surfaces = Source->Surfaces;
 
 		Result->ResetStaticFormatFlags();
 		Result->EnsureSurfaceData();
@@ -586,7 +566,7 @@ namespace mu
 			return;
 		}
 
-		FMeshBufferSet& VertexBuffers = InMesh->m_VertexBuffers;
+		FMeshBufferSet& VertexBuffers = InMesh->VertexBuffers;
 
 		// Reduce the number of influences if possible
 		constexpr int32 SemanticIndex = 0;
@@ -658,7 +638,7 @@ namespace mu
 
 				FormatBufferSet( VertexBuffers, NewVertexBuffers, true, false, true);
 
-				InMesh->m_VertexBuffers = NewVertexBuffers;
+				InMesh->VertexBuffers = NewVertexBuffers;
 			}
 		}
 
diff --git a/Engine/Plugins/Experimental/Mutable/Source/MutableRuntime/Private/MuR/OpMeshMorph.h b/Engine/Plugins/Experimental/Mutable/Source/MutableRuntime/Private/MuR/OpMeshMorph.h
--- a/Engine/Plugins/Experimental/Mutable/Source/MutableRuntime/Private/MuR/OpMeshMorph.h
+++ b/Engine/Plugins/Experimental/Mutable/Source/MutableRuntime/Private/MuR/OpMeshMorph.h
@@ -26,26 +26,44 @@ namespace mu
 		}
 
 		const auto MakeIndexMap = [](
-				UntypedMeshBufferIteratorConst BaseIdIter, int32 BaseNum,
-				MeshBufferIteratorConst<MBF_UINT32, uint32, 1> MorphIdIter, int32 MorphNum) 
-		-> SparseIndexMap
+			MeshVertexIdIteratorConst BaseIdIter, int32 BaseNum,
+			MeshVertexIdIteratorConst MorphIdIter, int32 MorphNum)
+		-> SparseIndexMapSet
 		{	
-            uint32 MinBaseId = TNumericLimits<uint32>::Max();
-            uint32 MaxBaseId = 0;
+			TArray<SparseIndexMapSet::FRangeDesc> RangeDescs;
+
+			// Detect all ranges and their limits
 			{
 				for (int32 Index = 0; Index < BaseNum; ++Index, ++BaseIdIter)
 				{
-					const uint32 BaseId = BaseIdIter.GetAsUINT32();
-					MinBaseId = FMath::Min(BaseId, MinBaseId);
-					MaxBaseId = FMath::Max(BaseId, MaxBaseId);
+					const uint64 BaseId = BaseIdIter.Get();
+
+					uint32 Prefix = BaseId >> 32;
+					uint32 Id = BaseId & 0xffffffff;
+					bool bFound = false;
+					for (SparseIndexMapSet::FRangeDesc& Range : RangeDescs)
+					{
+						if (Range.Prefix == Prefix)
+						{
+							Range.MinIndex = FMath::Min(Id, Range.MinIndex);
+							Range.MaxIndex = FMath::Max(Id, Range.MaxIndex);
+							bFound = true;
+							break;
+						}
+					}
+
+					if (!bFound)
+					{
+						RangeDescs.Add({Prefix, Id, Id});
+					}
 				}
 			}
 
-            SparseIndexMap IndexMap(MinBaseId, MaxBaseId);
+			SparseIndexMapSet IndexMap(RangeDescs);
            
             for (int32 Index = 0; Index < MorphNum; ++Index, ++MorphIdIter)
             {
-                const uint32 MorphId = (*MorphIdIter)[0];
+                const uint64 MorphId = MorphIdIter.Get();
                 
                 IndexMap.Insert(MorphId, Index);
             }
@@ -54,9 +72,9 @@ namespace mu
 		};
 
         const auto ApplyNormalMorph = [](
-				UntypedMeshBufferIteratorConst BaseIdIter, const TStaticArray<UntypedMeshBufferIterator, 3>& BaseTangentFrameIters, const int32 BaseNum,
-				MeshBufferIteratorConst<MBF_UINT32, uint32, 1> MorphIdIter, UntypedMeshBufferIteratorConst& MorphNormalIter, const int32 MorphNum,
-				const SparseIndexMap& IndexMap, const float Factor)
+			MeshVertexIdIteratorConst BaseIdIter, const TStaticArray<UntypedMeshBufferIterator, 3>& BaseTangentFrameIters, const int32 BaseNum,
+			MeshVertexIdIteratorConst MorphIdIter, UntypedMeshBufferIteratorConst& MorphNormalIter, const int32 MorphNum,
+				const SparseIndexMapSet& IndexMap, const float Factor)
         -> void
         {
 			const UntypedMeshBufferIterator& BaseNormalIter = BaseTangentFrameIters[2];
@@ -77,7 +95,7 @@ namespace mu
 
             for (int32 VertexIndex = 0; VertexIndex < BaseNum; ++VertexIndex)
             {
-                const uint32 BaseId = (BaseIdIter + VertexIndex).GetAsUINT32();
+                const uint64 BaseId = (BaseIdIter + VertexIndex).Get();
                 const int32 MorphIndex = static_cast<int32>(IndexMap.Find(BaseId));
 
                 if (MorphIndex == SparseIndexMap::NotFoundValue)
@@ -86,11 +104,11 @@ namespace mu
                 }
 
                 // Find consecutive run.
-                UntypedMeshBufferIteratorConst RunBaseIter = BaseIdIter + VertexIndex;
-                MeshBufferIteratorConst<MBF_UINT32, uint32, 1> RunMorphIter = MorphIdIter + MorphIndex;
+				MeshVertexIdIteratorConst RunBaseIter = BaseIdIter + VertexIndex;
+				MeshVertexIdIteratorConst RunMorphIter = MorphIdIter + MorphIndex;
 
                 int32 RunSize = 0;
-                for ( ; VertexIndex + RunSize < BaseNum && MorphIndex + RunSize < MorphNum && RunBaseIter.GetAsUINT32() == (*RunMorphIter)[0];
+                for ( ; VertexIndex + RunSize < BaseNum && MorphIndex + RunSize < MorphNum && RunBaseIter.Get() == RunMorphIter.Get();
                         ++RunSize, ++RunBaseIter, ++RunMorphIter);
 
 				for (int32 RunIndex = 0; RunIndex < RunSize; ++RunIndex)
@@ -153,33 +171,31 @@ namespace mu
         };
 
         const auto ApplyGenericMorph = []( 
-				UntypedMeshBufferIteratorConst BaseIdIter, const TArray<UntypedMeshBufferIterator>& BaseChannelsIters, const int32 BaseNum,
-				MeshBufferIteratorConst<MBF_UINT32, uint32, 1> MorphIdIter, const TArray<UntypedMeshBufferIteratorConst>& MorphChannelsIters, const int32 MorphNum,
-				const SparseIndexMap& IndexMap, const float Factor)
+			MeshVertexIdIteratorConst BaseIdIter, const TArray<UntypedMeshBufferIterator>& BaseChannelsIters, const int32 BaseNum,
+			MeshVertexIdIteratorConst MorphIdIter, const TArray<UntypedMeshBufferIteratorConst>& MorphChannelsIters, const int32 MorphNum,
+				const SparseIndexMapSet& IndexMap, const float Factor)
         -> void
         {
             for (int32 VertexIndex = 0; VertexIndex < BaseNum; ++VertexIndex)
             {
-                const uint32 BaseId = (BaseIdIter + VertexIndex).GetAsUINT32();
-                const uint32 MorphIdx = IndexMap.Find(BaseId);
+                const uint64 BaseId = (BaseIdIter + VertexIndex).Get();
+                const uint32 MorphIndex = IndexMap.Find(BaseId);
 
-                if (MorphIdx == SparseIndexMap::NotFoundValue)
+                if (MorphIndex == SparseIndexMap::NotFoundValue)
                 {
                     continue;
                 }
 
-                const int32 MorphIndex = static_cast<int32>(MorphIdx);
-
                 // Find consecutive run.
-                UntypedMeshBufferIteratorConst RunBaseIter = BaseIdIter + VertexIndex;
-                MeshBufferIteratorConst<MBF_UINT32, uint32, 1> RunMorphIter = MorphIdIter + MorphIndex;
+				MeshVertexIdIteratorConst RunBaseIter = BaseIdIter + VertexIndex;
+				MeshVertexIdIteratorConst RunMorphIter = MorphIdIter + MorphIndex;
 
                 int32 RunSize = 0;
-                for ( ; VertexIndex + RunSize < BaseNum && MorphIndex + RunSize < MorphNum && RunBaseIter.GetAsUINT32() == (*RunMorphIter)[0];
+                for ( ; VertexIndex + RunSize < BaseNum && int32(MorphIndex) + RunSize < MorphNum && RunBaseIter.Get() == RunMorphIter.Get();
                         ++RunSize, ++RunBaseIter, ++RunMorphIter);
 
                 const int32 ChannelNum = MorphChannelsIters.Num();
-                for (int32 ChannelIndex = 1; ChannelIndex < ChannelNum; ++ChannelIndex)
+                for (int32 ChannelIndex = 0; ChannelIndex < ChannelNum; ++ChannelIndex)
                 {
                 	if (!(BaseChannelsIters[ChannelIndex].ptr() && MorphChannelsIters[ChannelIndex].ptr()))
                 	{
@@ -227,7 +243,8 @@ namespace mu
 
 		if (RefTarget)
 		{
-			const int32 ChannelsNum = RefTarget->GetVertexBuffers().GetBufferChannelCount(0);
+			constexpr int32 MorphBufferDataChannel = 0;
+			const int32 ChannelsNum = RefTarget->GetVertexBuffers().GetBufferChannelCount(MorphBufferDataChannel);
 
 			TArray<UntypedMeshBufferIterator> BaseChannelsIters;
 			BaseChannelsIters.SetNum(ChannelsNum);
@@ -242,11 +259,12 @@ namespace mu
 			UntypedMeshBufferIteratorConst MaxNormalChannelIter;
 
 			const bool bBaseHasNormals = UntypedMeshBufferIteratorConst(pBase->GetVertexBuffers(), MBS_NORMAL, 0).ptr() != nullptr;
-			for (int32 ChannelIndex = 1; ChannelIndex < ChannelsNum; ++ChannelIndex)
+			for (int32 ChannelIndex = 0; ChannelIndex < ChannelsNum; ++ChannelIndex)
 			{
 				const FMeshBufferSet& MBSPriv = RefTarget->GetVertexBuffers();
-				EMeshBufferSemantic Sem = MBSPriv.m_buffers[0].m_channels[ChannelIndex].m_semantic;
-				int32 SemIndex = MBSPriv.m_buffers[0].m_channels[ChannelIndex].m_semanticIndex;
+				const FMeshBufferChannel& Channel = MBSPriv.m_buffers[MorphBufferDataChannel].m_channels[ChannelIndex];
+				EMeshBufferSemantic Sem = Channel.m_semantic;
+				int32 SemIndex = Channel.m_semanticIndex;
 			
 				if (Sem == MBS_NORMAL && bBaseHasNormals)
 				{
@@ -284,12 +302,12 @@ namespace mu
 				}
 			}
 			
-			UntypedMeshBufferIteratorConst BaseIdIter(pBase->GetVertexBuffers(), MBS_VERTEXINDEX);
+			MeshVertexIdIteratorConst BaseIdIter(pBase);
 
 			if (MinNum > 0)
 			{
-				MeshBufferIteratorConst<MBF_UINT32, uint32, 1> MinIdIter(pMin->GetVertexBuffers(), MBS_VERTEXINDEX);
-				SparseIndexMap IndexMap = MakeIndexMap(BaseIdIter, BaseNum, MinIdIter, MinNum);
+				MeshVertexIdIteratorConst MinIdIter(pMin);
+				SparseIndexMapSet IndexMap = MakeIndexMap(BaseIdIter, BaseNum, MinIdIter, MinNum);
 
 				ApplyGenericMorph(BaseIdIter, BaseChannelsIters, BaseNum, MinIdIter, MinChannelsIters, MinNum, IndexMap, 1.0f - Factor);
 				if (MinNormalChannelIter.ptr())
@@ -300,8 +318,8 @@ namespace mu
 
 			if (MaxNum > 0)
 			{
-				MeshBufferIteratorConst<MBF_UINT32, uint32, 1> MaxIdIter(pMax->GetVertexBuffers(), MBS_VERTEXINDEX);
-				SparseIndexMap IndexMap = MakeIndexMap(BaseIdIter, BaseNum, MaxIdIter, MaxNum);
+				MeshVertexIdIteratorConst MaxIdIter(pMax);
+				SparseIndexMapSet IndexMap = MakeIndexMap(BaseIdIter, BaseNum, MaxIdIter, MaxNum);
 
 				ApplyGenericMorph(BaseIdIter, BaseChannelsIters, BaseNum, MaxIdIter, MaxChannelsIters, MaxNum, IndexMap, Factor);
 				
diff --git a/Engine/Plugins/Experimental/Mutable/Source/MutableRuntime/Private/MuR/OpMeshOptimizeSkinning.h b/Engine/Plugins/Experimental/Mutable/Source/MutableRuntime/Private/MuR/OpMeshOptimizeSkinning.h
--- a/Engine/Plugins/Experimental/Mutable/Source/MutableRuntime/Private/MuR/OpMeshOptimizeSkinning.h
+++ b/Engine/Plugins/Experimental/Mutable/Source/MutableRuntime/Private/MuR/OpMeshOptimizeSkinning.h
@@ -27,7 +27,7 @@ namespace mu
 		}
 
 		uint32 MaxBoneMapIndex = 0;
-		for (const MESH_SURFACE& Surface : InMesh->m_surfaces)
+		for (const FMeshSurface& Surface : InMesh->Surfaces)
 		{
 			MaxBoneMapIndex = FMath::Max(MaxBoneMapIndex, Surface.BoneMapCount);
 		}
@@ -147,11 +147,11 @@ namespace mu
 				const int32 TailSize = AuxOffset - (BoneIndexChannelOffset + BoneIndexComponentCount * BoneIndexSize);
 
 
-				for (const MESH_SURFACE& Surface : InMesh->m_surfaces)
+				for (const FMeshSurface& Surface : InMesh->Surfaces)
 				{
 					const uint32 NumBonesInBoneMap = Surface.BoneMapCount;
 
-					for (int32 VertexIndex = 0; VertexIndex < Surface.m_vertexCount; ++VertexIndex)
+					for (int32 VertexIndex = 0; VertexIndex < Surface.VertexCount; ++VertexIndex)
 					{
 						FMemory::Memcpy(Data, SourceData, BoneIndexChannelOffset);
 						Data += BoneIndexChannelOffset;
diff --git a/Engine/Plugins/Experimental/Mutable/Source/MutableRuntime/Private/MuR/OpMeshRemove.cpp b/Engine/Plugins/Experimental/Mutable/Source/MutableRuntime/Private/MuR/OpMeshRemove.cpp
--- a/Engine/Plugins/Experimental/Mutable/Source/MutableRuntime/Private/MuR/OpMeshRemove.cpp
+++ b/Engine/Plugins/Experimental/Mutable/Source/MutableRuntime/Private/MuR/OpMeshRemove.cpp
@@ -11,27 +11,27 @@ namespace mu
 {
 
     //---------------------------------------------------------------------------------------------
-    struct ID_INTERVAL
+    struct FIdInterval
     {
-        int32 idStart;
+        uint64 idStart;
         int32 idPosition;
         int32 size;
     };
 
 
-    void ExtractVertexIndexIntervals( TArray<ID_INTERVAL>& intervals, const Mesh* pSource )
+    void ExtractVertexIndexIntervals( TArray<FIdInterval>& intervals, const Mesh* pSource )
     {
-        MeshBufferIteratorConst<MBF_UINT32,uint32,1> itVI( pSource->GetVertexBuffers(), MBS_VERTEXINDEX );
-        ID_INTERVAL current;
-		current.idStart = -1;
+		MeshVertexIdIteratorConst itVI( pSource );
+		FIdInterval current;
+		current.idStart = TNumericLimits<uint64>::Max();
 		current.idPosition = 0;
 		current.size = 0;
-		for ( int sv=0; sv<pSource->GetVertexBuffers().GetElementCount(); ++sv )
+		for ( int32 sv=0; sv<pSource->GetVertexBuffers().GetElementCount(); ++sv )
         {
-            int32 id = (*itVI)[0];
+            uint64 id = itVI.Get();
             ++itVI;
 
-            if (current.idStart<0)
+            if (current.idStart==TNumericLimits<uint64>::Max())
             {
                 current.idStart = id;
                 current.idPosition = sv;
@@ -60,11 +60,11 @@ namespace mu
     }
 
 
-    int FindPositionInIntervals( const TArray<ID_INTERVAL>& intervals, int32 id )
+    int64 FindPositionInIntervals( const TArray<FIdInterval>& intervals, int64 id )
     {
-        for( const ID_INTERVAL& interval: intervals )
+        for( const FIdInterval& interval: intervals )
         {
-            int32 deltaId = id - interval.idStart;
+            int64 deltaId = id - interval.idStart;
             if (deltaId>=0 && deltaId<interval.size)
             {
                 return interval.idPosition+deltaId;
@@ -76,7 +76,8 @@ namespace mu
 
 	void MeshRemoveVerticesWithMap( Mesh* Result, const uint8* RemovedVertices, uint32 RemovedVertexCount)
 	{
-		int32 firstFreeVertex = 0;
+		int32 FirstFreeVertex = 0;
+		int32 RemovedIndices = 0;
 
         // Rebuild index buffers
 
@@ -86,8 +87,6 @@ namespace mu
 		TArray<int32> UsedVertexMap;
 		UsedVertexMap.Init( -1, Result->GetVertexCount() );
         {
-            int32 removedIndices = 0;
-
             if ( Result->GetIndexBuffers().GetElementSize(0)==4 )
             {
                 MeshBufferIteratorConst<MBF_UINT32,uint32,1> itSource( Result->GetIndexBuffers(), MBS_VERTEXINDEX );
@@ -120,8 +119,8 @@ namespace mu
 
                             if (UsedVertexMap[ sourceIndex ] < 0 )
                             {
-								UsedVertexMap[ sourceIndex ] = firstFreeVertex;
-                                firstFreeVertex++;
+								UsedVertexMap[ sourceIndex ] = FirstFreeVertex;
+                                FirstFreeVertex++;
                             }
 
                             uint32 destIndex = UsedVertexMap[ sourceIndex ];
@@ -134,7 +133,7 @@ namespace mu
                     itSource+=3;
                 }
 
-                removedIndices = itSource - itDest;
+                RemovedIndices = itSource - itDest;
             }
 
             else if ( Result->GetIndexBuffers().GetElementSize(0)==2 )
@@ -168,8 +167,8 @@ namespace mu
 
                             if (UsedVertexMap[ sourceIndex ] < 0 )
                             {
-								UsedVertexMap[ sourceIndex ] = firstFreeVertex;
-                                firstFreeVertex++;
+								UsedVertexMap[ sourceIndex ] = FirstFreeVertex;
+                                FirstFreeVertex++;
                             }
 
                             uint16 destIndex = (uint16)UsedVertexMap[ sourceIndex ];
@@ -182,7 +181,7 @@ namespace mu
                     itSource+=3;
                 }
 
-                removedIndices = itSource - itDest;
+                RemovedIndices = itSource - itDest;
             }
 
             else
@@ -191,16 +190,21 @@ namespace mu
                 check( false );
             }
 
-            check( removedIndices%3==0 );
+            check( RemovedIndices%3==0 );
 
-            int32 faceCount = Result->GetFaceCount();
-            Result->GetFaceBuffers().SetElementCount( faceCount-(int)removedIndices/3 );
-            Result->GetIndexBuffers().SetElementCount( faceCount*3-(int)removedIndices );
+            int32 FaceCount = Result->GetFaceCount();
+            Result->GetIndexBuffers().SetElementCount( FaceCount*3-RemovedIndices );
         }
 
 
         // Rebuild the vertex buffers
 
+		// If we had implicit indices, make them explicit or relative to keep them valid
+		if (RemovedIndices && Result->AreVertexIdsImplicit())
+		{
+			Result->MakeVertexIndicesRelative();
+		}
+
 		// The temp array is necessary because if the vertex buffer is not sorted according to the index buffer we cannot do it in-place
 		// This happens with some mesh import options.
 		TArray<uint8> Temp;
@@ -209,7 +213,7 @@ namespace mu
             int32 elemSize = Result->GetVertexBuffers().GetElementSize( b );
             const uint8* SourceData = Result->GetVertexBuffers().GetBufferData( b );
 
-			Temp.SetNumUninitialized(firstFreeVertex*elemSize,EAllowShrinking::No);
+			Temp.SetNumUninitialized(FirstFreeVertex*elemSize,EAllowShrinking::No);
             uint8* DestData = Temp.GetData();
 
             for ( int32 v=0; v<Result->GetVertexCount(); ++v )
@@ -251,13 +255,13 @@ namespace mu
             }
 
 			// Copy from temp buffer to final vertex buffer
-			FMemory::Memcpy( Result->GetVertexBuffers().GetBufferData(b), DestData, firstFreeVertex*elemSize);
+			FMemory::Memcpy( Result->GetVertexBuffers().GetBufferData(b), DestData, FirstFreeVertex*elemSize);
         }
-        Result->GetVertexBuffers().SetElementCount( firstFreeVertex );
+        Result->GetVertexBuffers().SetElementCount( FirstFreeVertex );
 
         // Rebuild surface data.
         // \todo: For now make single surfaced
-        Result->m_surfaces.Empty();
+        Result->Surfaces.Empty();
         Result->EnsureSurfaceData();
     }
 
@@ -282,13 +286,13 @@ namespace mu
 		TArray<uint8> RemovedVertices;
 		RemovedVertices.SetNumZeroed(ResultVertexCount);
         {
-			TArray<ID_INTERVAL> Intervals;
+			TArray<FIdInterval> Intervals;
             ExtractVertexIndexIntervals(Intervals, Result);
 
-			MeshBufferIteratorConst<MBF_UINT32, uint32, 1> itMaskVI(Mask->GetVertexBuffers(), MBS_VERTEXINDEX);
+			MeshVertexIdIteratorConst itMaskVI(Mask);
 			for ( int32 mv=0; mv<MaskElementCount; ++mv )
             {
-                uint32 MaskVertexId = (*itMaskVI)[0];
+                uint64 MaskVertexId = itMaskVI.Get();
                 ++itMaskVI;
 
                 int32 IndexInSource = FindPositionInIntervals(Intervals, MaskVertexId);
diff --git a/Engine/Plugins/Experimental/Mutable/Source/MutableRuntime/Private/MuR/OpMeshReshape.h b/Engine/Plugins/Experimental/Mutable/Source/MutableRuntime/Private/MuR/OpMeshReshape.h
--- a/Engine/Plugins/Experimental/Mutable/Source/MutableRuntime/Private/MuR/OpMeshReshape.h
+++ b/Engine/Plugins/Experimental/Mutable/Source/MutableRuntime/Private/MuR/OpMeshReshape.h
@@ -693,11 +693,11 @@ namespace mu
 		const int32 PhysicsBodiesNum = BodyOffsets.Num() - 1;
 		
 		// Apply main physics body
-		if (BaseMesh.m_pPhysicsBody)
+		if (BaseMesh.PhysicsBody)
 		{
-			Ptr<PhysicsBody> NewBody = BaseMesh.m_pPhysicsBody->Clone();
+			Ptr<PhysicsBody> NewBody = BaseMesh.PhysicsBody->Clone();
 			ApplyPhysicsBody(*NewBody, 0, BodyOffsets[1]);
-			OutNewMesh.m_pPhysicsBody = NewBody;
+			OutNewMesh.PhysicsBody = NewBody;
 		}
 
 
@@ -924,7 +924,7 @@ namespace mu
 		// order of operation is important.
 
 		// Transform physics volumes based on the deformed sampling points.
-		const PhysicsBody* OldPhysicsBody = Result->m_pPhysicsBody.get();
+		const PhysicsBody* OldPhysicsBody = Result->PhysicsBody.get();
 		const int32 AdditionalPhysicsBodiesNum = Result->AdditionalPhysicsBodies.Num();
 
 		if (bReshapePhysicsVolumes && (OldPhysicsBody || AdditionalPhysicsBodiesNum))
diff --git a/Engine/Plugins/Experimental/Mutable/Source/MutableRuntime/Private/MuR/SparseIndexMap.h b/Engine/Plugins/Experimental/Mutable/Source/MutableRuntime/Private/MuR/SparseIndexMap.h
--- a/Engine/Plugins/Experimental/Mutable/Source/MutableRuntime/Private/MuR/SparseIndexMap.h
+++ b/Engine/Plugins/Experimental/Mutable/Source/MutableRuntime/Private/MuR/SparseIndexMap.h
@@ -68,7 +68,7 @@ namespace mu
             if ( !BlockPtr )
             {
                 BlockPtr = MakeUnique<BlockType>(  
-                            MakeUniformStaticArray<uint32_t, 1 << Log2BlockElems>( NotFoundValue ) );
+                            MakeUniformStaticArray<uint32, 1 << Log2BlockElems>( NotFoundValue ) );
             }
 
             (*BlockPtr)[ MappedIndex & ( ( 1 << Log2BlockElems ) - 1 ) ] = ValueIndex;
@@ -101,4 +101,67 @@ namespace mu
     };
 
 
+	class SparseIndexMapSet
+	{
+	private:
+
+		struct FRange
+		{
+			uint32 Prefix;
+			SparseIndexMap Map;
+		};
+
+		TArray<FRange> Ranges;
+
+	public:
+
+		struct FRangeDesc
+		{
+			uint32 Prefix;
+			uint32 MinIndex;
+			uint32 MaxIndex;
+		};
+
+		SparseIndexMapSet(const TArray<FRangeDesc>& InRanges)
+		{
+			Ranges.Reserve(InRanges.Num());
+			for ( const FRangeDesc& Desc: InRanges )
+			{
+				Ranges.Add({ Desc.Prefix, SparseIndexMap(Desc.MinIndex,Desc.MaxIndex)});
+			}
+		}
+
+		bool Insert(const uint64 KeyIndex, const uint32 ValueIndex)
+		{
+			uint32 Prefix = KeyIndex >> 32;
+			uint32 KeyValue = KeyIndex & 0xffffffff;
+			for (FRange& Range : Ranges)
+			{
+				if (Range.Prefix == Prefix)
+				{
+					return Range.Map.Insert(KeyValue, ValueIndex);
+				}
+			}
+
+			check(false);
+			return false;
+		}
+
+		uint32 Find(const uint64 KeyIndex) const
+		{
+			uint32 Prefix = KeyIndex >> 32;
+			uint32 KeyValue = KeyIndex & 0xffffffff;
+			for (const FRange& Range : Ranges)
+			{
+				if (Range.Prefix == Prefix)
+				{
+					return Range.Map.Find(KeyValue);
+				}
+			}
+
+			check(false);
+			return SparseIndexMap::NotFoundValue;
+		}
+
+	};
 }
diff --git a/Engine/Plugins/Experimental/Mutable/Source/MutableRuntime/Public/MuR/Mesh.h b/Engine/Plugins/Experimental/Mutable/Source/MutableRuntime/Public/MuR/Mesh.h
--- a/Engine/Plugins/Experimental/Mutable/Source/MutableRuntime/Public/MuR/Mesh.h
+++ b/Engine/Plugins/Experimental/Mutable/Source/MutableRuntime/Public/MuR/Mesh.h
@@ -25,63 +25,39 @@ namespace mu
 
 	// Forward references
 	class Layout;
-	typedef Ptr<Layout> LayoutPtr;
-	typedef Ptr<const Layout> LayoutPtrConst;
-
     class Skeleton;
-    typedef Ptr<Skeleton> SkeletonPtr;
-    typedef Ptr<const Skeleton> SkeletonPtrConst;
-
     class PhysicsBody;
-    typedef Ptr<PhysicsBody> PhysicsBodyPtr;
-    typedef Ptr<const PhysicsBody> PhysicsBodyPtrConst;
 
-    class Mesh;
+	class Mesh;
     typedef Ptr<Mesh> MeshPtr;
     typedef Ptr<const Mesh> MeshPtrConst;
 
-	struct MESH_SURFACE
+	struct FMeshSurface
 	{
-		MESH_SURFACE()
-		{
-			m_firstVertex = 0;
-			m_vertexCount = 0;
-			m_firstIndex = 0;
-			m_indexCount = 0;
-			m_id = 0;
-
-			BoneMapIndex = 0;
-			BoneMapCount = 0;
-
-			bCastShadow = false;
-		}
-
-		int32 m_firstVertex;
-		int32 m_vertexCount;
-		int32 m_firstIndex;
-		int32 m_indexCount;
-		uint32 m_id;
+		int32 FirstVertex=0;
+		int32 VertexCount=0;
+		int32 FirstIndex=0;
+		int32 IndexCount=0;
+		uint32 Id=0;
 		
-		uint32 BoneMapIndex;
-		uint32 BoneMapCount;
+		uint32 BoneMapIndex=0;
+		uint32 BoneMapCount=0;
 
-		bool bCastShadow;
+		bool bCastShadow=false;
 
-		//!
-		inline bool operator==(const MESH_SURFACE& o) const
+		inline bool operator==(const FMeshSurface& o) const
 		{
-			return m_firstVertex == o.m_firstVertex
-				&& m_vertexCount == o.m_vertexCount
-				&& m_firstIndex == o.m_firstIndex
-				&& m_indexCount == o.m_indexCount
-				&& m_id == o.m_id
+			return FirstVertex == o.FirstVertex
+				&& VertexCount == o.VertexCount
+				&& FirstIndex == o.FirstIndex
+				&& IndexCount == o.IndexCount
+				&& Id == o.Id
 				&& BoneMapIndex == o.BoneMapIndex
 				&& BoneMapCount == o.BoneMapCount
 				&& bCastShadow == o.bCastShadow;
 		}
 
 		inline void Serialise(OutputArchive& arch) const;
-
 		inline void Unserialise(InputArchive& arch);
 
 	};
@@ -142,7 +118,7 @@ namespace mu
 		WithTags = 1 << 6,
 		WithVertexBuffers = 1 << 7,
 		WithIndexBuffers = 1 << 8,
-		WithFaceBuffers = 1 << 9,
+		// deprecated WithFaceBuffers = 1 << 9,
 		WithAdditionalBuffers = 1 << 10,
 		WithLayouts = 1 << 11,
 		WithPoses = 1 << 12,
@@ -171,17 +147,17 @@ namespace mu
         //-----------------------------------------------------------------------------------------
 
         //! Deep clone this mesh.
-        MeshPtr Clone() const;
+        Ptr<Mesh> Clone() const;
 		
 		// Clone with flags allowing to not include some parts in the cloned mesh
-		MeshPtr Clone(EMeshCopyFlags Flags) const;
+		Ptr<Mesh> Clone(EMeshCopyFlags Flags) const;
 
 		// Copy form another mesh.
 		void CopyFrom(const Mesh& From, EMeshCopyFlags Flags = EMeshCopyFlags::AllFlags);
 
         //! Serialisation
         static void Serialise( const Mesh* p, OutputArchive& arch );
-        static MeshPtr StaticUnserialise( InputArchive& arch );
+        static Ptr<Mesh> StaticUnserialise( InputArchive& arch );
 
 		// Resource interface
 		int32 GetDataSize() const override;
@@ -195,25 +171,21 @@ namespace mu
         //! \{
 
         //!
-        int GetIndexCount() const;
+        int32 GetIndexCount() const;
 
         //! Index buffers. They are owned by this mesh.
         FMeshBufferSet& GetIndexBuffers();
         const FMeshBufferSet& GetIndexBuffers() const;
 
         //
-        int GetVertexCount() const;
+        int32 GetVertexCount() const;
 
         //! Vertex buffers. They are owned by this mesh.
         FMeshBufferSet& GetVertexBuffers();
         const FMeshBufferSet& GetVertexBuffers() const;
 
         //
-        int GetFaceCount() const;
-
-        //! Face buffers. They are owned by this mesh.
-        FMeshBufferSet& GetFaceBuffers();
-        const FMeshBufferSet& GetFaceBuffers() const;
+        int32 GetFaceCount() const;
 
         //! Get the number of surfaces defined in this mesh. Surfaces are buffer-contiguous mesh
         //! fragments that share common properties (usually material)
@@ -230,6 +202,18 @@ namespace mu
 
         //! \}
 
+		/** Return true if the mesh has unique IDs for its vertices. */
+		bool HasVertexIds() const;
+
+		/** Return true if the mesh has unique vertex IDs and they stored in an implicit way. 
+		* This is relevant for some mesh operations that will need to make them explicit so that the result is still correct.
+		*/
+		bool AreVertexIdsImplicit() const;
+		bool AreVertexIdsExplicit() const;
+
+		/** Create an explicit vertex buffer for vertex IDs if they are implicit. */
+		void MakeVertexIndicesRelative();
+		void MakeVertexIndicesExplicit();
 
         //! \name Texture layouts
         //! \{
@@ -343,46 +327,47 @@ namespace mu
 		template<typename Type>
 		using TMemoryTrackedArray = TArray<Type, FDefaultMemoryTrackingAllocator<MemoryCounters::FMeshMemoryCounter>>;
 
-		//! Non-persistent internal id unique for a mesh generated for a specific state and
-		//! parameter values.
-		mutable uint32 m_internalId = 0;
+		/** Non-persistent internal id unique for a mesh generated for a specific state and parameter values. */
+		mutable uint32 InternalId = 0;
 
-		//!
-		FMeshBufferSet m_VertexBuffers;
+		//! This is bit-mask on the STATIC_MESH_FORMATS enumeration, marking what static formats
+		//! are compatible with this one. Usually precalculated at model compilation time.
+		//! It should be reset after any operation that modifies the format.
+		mutable uint32 StaticFormatFlags = 0;
+
+		/** Prefix for the unique vertex IDs, if the mesh has them in an implicit, or relative way. 
+		* See MeshVertexIdIterator for details.
+		*/
+		uint32 VertexIDPrefix = 0;
 
 		//!
-		FMeshBufferSet m_IndexBuffers;
+		FMeshBufferSet VertexBuffers;
 
 		//!
-		FMeshBufferSet m_FaceBuffers;
+		FMeshBufferSet IndexBuffers;
 
 		//! Additional buffers used for temporary or custom data in different algorithms.
 		TArray<TPair<EMeshBufferType, FMeshBufferSet>> AdditionalBuffers;
 
-		//! This is bit-mask on the STATIC_MESH_FORMATS enumeration, marking what static formats
-		//! are compatible with this one. Usually precalculated at model compilation time.
-		//! It should be reset after any operation that modifies the format.
-		mutable uint32 m_staticFormatFlags = 0;
-
-		TArray<MESH_SURFACE> m_surfaces;
+		TArray<FMeshSurface> Surfaces;
 
 		// Externally provided SkeletonIDs of the skeletons required by this mesh.
 		TArray<uint32> SkeletonIDs;
 
 		//! This skeleton and physics body are not owned and may be used by other meshes, so it cannot be modified
 		//! once the mesh has been fully created.
-		Ptr<const Skeleton> m_pSkeleton;
-		Ptr<const PhysicsBody> m_pPhysicsBody;
+		Ptr<const Skeleton> Skeleton;
+		Ptr<const PhysicsBody> PhysicsBody;
 
 		//! Additional physics bodies referenced by the mesh that don't merge.
-		TArray<Ptr<const PhysicsBody>> AdditionalPhysicsBodies;
+		TArray<Ptr<const mu::PhysicsBody>> AdditionalPhysicsBodies;
 
 		//! Texture Layout blocks attached to this mesh. They are const because they could be shared with
 		//! other meshes, so they need to be cloned and replaced if a modification is needed.
-		TArray<Ptr<const Layout>> m_layouts;		
+		TArray<Ptr<const Layout>> Layouts;		
 
 		//!
-		TArray<FString> m_tags;
+		TArray<FString> Tags;
 
 		// Opaque handle to external resources.
 		TArray<uint64> StreamedResources;
@@ -421,17 +406,16 @@ namespace mu
         //!
 		inline bool operator==(const Mesh& o) const
 		{
-			bool equal = (m_IndexBuffers == o.m_IndexBuffers);
-			if (equal) equal = (m_VertexBuffers == o.m_VertexBuffers);
-			if (equal) equal = (m_FaceBuffers == o.m_FaceBuffers);
-			if (equal) equal = (m_layouts.Num() == o.m_layouts.Num());
+			bool equal = (IndexBuffers == o.IndexBuffers);
+			if (equal) equal = (VertexBuffers == o.VertexBuffers);
+			if (equal) equal = (Layouts.Num() == o.Layouts.Num());
 			if (equal) equal = (BonePoses.Num() == o.BonePoses.Num());
 			if (equal) equal = (BoneMap.Num() == o.BoneMap.Num());
-			if (equal && m_pSkeleton != o.m_pSkeleton)
+			if (equal && Skeleton != o.Skeleton)
 			{
-				if (m_pSkeleton && o.m_pSkeleton)
+				if (Skeleton && o.Skeleton)
 				{
-					equal = (*m_pSkeleton == *o.m_pSkeleton);
+					equal = (*Skeleton == *o.Skeleton);
 				}
 				else
 				{
@@ -439,13 +423,14 @@ namespace mu
 				}
 			}
 			if (equal) equal = (StreamedResources == o.StreamedResources);
-			if (equal) equal = (m_surfaces == o.m_surfaces);
-			if (equal) equal = (m_tags == o.m_tags);
+			if (equal) equal = (Surfaces == o.Surfaces);
+			if (equal) equal = (Tags == o.Tags);
 			if (equal) equal = (SkeletonIDs == o.SkeletonIDs);
+			if (equal) equal = (VertexIDPrefix == o.VertexIDPrefix);
 
-			for (int32 i = 0; equal && i < m_layouts.Num(); ++i)
+			for (int32 i = 0; equal && i < Layouts.Num(); ++i)
 			{
-				equal &= (*m_layouts[i]) == (*o.m_layouts[i]);
+				equal &= (*Layouts[i]) == (*o.Layouts[i]);
 			}
 
 			equal &= AdditionalBuffers.Num() == o.AdditionalBuffers.Num();
@@ -475,34 +460,22 @@ namespace mu
 		//! indices.
 		bool IsSimilar(const Mesh& o, bool bCompareLayouts) const;
 
-		//! Compare the vertex attributes to check if they match.
-		bool IsSameVertex( uint32 v, const Mesh& other, uint32 otherVertexIndex, float tolerance = UE_SMALL_NUMBER) const;
-
 
 		//! Make a map from the vertices in this mesh to thefirst matching vertex of the given
 		//! mesh. If non is found, the index is set to -1.
-		struct VERTEX_MATCH_MAP
+		struct FVertexMatchMap
 		{
 			//! One for every vertex
-			TArray<int> m_firstMatch;
+			TArray<int32> FirstMatch;
 
 			//! The matches of every vertex in a sequence
-			TArray<int> m_matches;
+			TArray<int32> Matches;
 
 			//!
-			bool Matches(int v, int ov) const;
+			bool DoMatch(int32 v, int32 ov) const;
 		};
 
-		void GetVertexMap
-		(
-			const Mesh& other,
-			VERTEX_MATCH_MAP& vertexMap,
-			float tolerance = 1e-3f//1e-5f //std::numeric_limits<float>::min()
-		) const;
-
-		//! Return true if the mesh has a face with the same vertices than the face in the given
-		//! mesh. The actual vertex data is used in the comparison.
-		bool HasFace( const Mesh& other, int otherFaceIndex, const VERTEX_MATCH_MAP& vertexMap ) const;
+		void GetVertexMap( const Mesh& other, FVertexMatchMap& vertexMap, float tolerance = 1e-3f) const;
 
 		//! Compare the vertex attributes to check if they match.
 		UE::Math::TIntVector3<uint32> GetFaceVertexIndices(int f) const;
diff --git a/Engine/Plugins/Experimental/Mutable/Source/MutableRuntime/Public/MuR/MeshBufferSet.h b/Engine/Plugins/Experimental/Mutable/Source/MutableRuntime/Public/MuR/MeshBufferSet.h
--- a/Engine/Plugins/Experimental/Mutable/Source/MutableRuntime/Public/MuR/MeshBufferSet.h
+++ b/Engine/Plugins/Experimental/Mutable/Source/MutableRuntime/Public/MuR/MeshBufferSet.h
@@ -60,6 +60,10 @@ namespace mu
         MBF_PACKEDDIRS8_W_TANGENTSIGN,
 
 		MBF_FLOAT64,
+		MBF_UINT64,
+		MBF_INT64,
+		MBF_NUINT64,
+		MBF_NINT64,
 
 		MBF_COUNT,
 
diff --git a/Engine/Plugins/Experimental/Mutable/Source/MutableRuntime/Public/MuR/Model.h b/Engine/Plugins/Experimental/Mutable/Source/MutableRuntime/Public/MuR/Model.h
--- a/Engine/Plugins/Experimental/Mutable/Source/MutableRuntime/Public/MuR/Model.h
+++ b/Engine/Plugins/Experimental/Mutable/Source/MutableRuntime/Public/MuR/Model.h
@@ -13,7 +13,7 @@
 //! data structures. Compiled models are not necessarily compatible when the runtime is updated,
 //! so this version number can be used externally to verify this. It is not used internally, and
 //! serializing models from different versions than this runtime will probably result in a crash.
-#define MUTABLE_COMPILED_MODEL_CODE_VERSION		uint32( 83 )
+#define MUTABLE_COMPILED_MODEL_CODE_VERSION		uint32( 85 )
 #define MUTABLE_PARAMETERS_VERSION              uint32( 3 )
 
 
diff --git a/Engine/Plugins/Experimental/Mutable/Source/MutableTools/Internal/MuT/AST.h b/Engine/Plugins/Experimental/Mutable/Source/MutableTools/Internal/MuT/AST.h
--- a/Engine/Plugins/Experimental/Mutable/Source/MutableTools/Internal/MuT/AST.h
+++ b/Engine/Plugins/Experimental/Mutable/Source/MutableTools/Internal/MuT/AST.h
@@ -253,8 +253,8 @@ namespace mu
 			{
 				const mu::Mesh* Data = Key.get();
 				return HashCombineFast(
-					::GetTypeHash(Data->m_VertexBuffers.GetElementCount()),
-					::GetTypeHash(Data->m_IndexBuffers.GetElementCount())
+					::GetTypeHash(Data->VertexBuffers.GetElementCount()),
+					::GetTypeHash(Data->IndexBuffers.GetElementCount())
 				);
 			}
 		};
diff --git a/Engine/Plugins/Experimental/Mutable/Source/MutableTools/Internal/MuT/NodeMeshConstantPrivate.h b/Engine/Plugins/Experimental/Mutable/Source/MutableTools/Internal/MuT/NodeMeshConstantPrivate.h
--- a/Engine/Plugins/Experimental/Mutable/Source/MutableTools/Internal/MuT/NodeMeshConstantPrivate.h
+++ b/Engine/Plugins/Experimental/Mutable/Source/MutableTools/Internal/MuT/NodeMeshConstantPrivate.h
@@ -5,11 +5,11 @@
 #include "MuT/NodeMeshPrivate.h"
 
 #include "MuT/NodeMeshConstant.h"
-#include "MuT/NodeLayout.h"
 
 
 namespace mu
 {
+	class NodeLayout;
 
 	class NodeMeshConstant::Private : public NodeMesh::Private
 	{
@@ -17,12 +17,12 @@ namespace mu
 
 		static FNodeType s_type;
 
-		MeshPtr m_pValue;
+		Ptr<Mesh> Value;
 
-		TArray<NodeLayoutPtr> m_layouts;
+		TArray<Ptr<NodeLayout>> Layouts;
 
 		// NodeMesh::Private interface
-        NodeLayoutPtr GetLayout( int index ) const override;
+		Ptr<NodeLayout> GetLayout(int32 index) const override;
 
 	};
 
diff --git a/Engine/Plugins/Experimental/Mutable/Source/MutableTools/Internal/MuT/NodeMeshFormatPrivate.h b/Engine/Plugins/Experimental/Mutable/Source/MutableTools/Internal/MuT/NodeMeshFormatPrivate.h
--- a/Engine/Plugins/Experimental/Mutable/Source/MutableTools/Internal/MuT/NodeMeshFormatPrivate.h
+++ b/Engine/Plugins/Experimental/Mutable/Source/MutableTools/Internal/MuT/NodeMeshFormatPrivate.h
@@ -21,7 +21,6 @@ namespace mu
 		/** New mesh format.The buffers in the sets have no elements, but they define the formats. */
 		FMeshBufferSet VertexBuffers;
 		FMeshBufferSet IndexBuffers;
-		FMeshBufferSet FaceBuffers;
 		
 		/** */
 		bool bOptimizeBuffers = false;
diff --git a/Engine/Plugins/Experimental/Mutable/Source/MutableTools/Private/MuT/ASTOpMeshFormat.cpp b/Engine/Plugins/Experimental/Mutable/Source/MutableTools/Private/MuT/ASTOpMeshFormat.cpp
--- a/Engine/Plugins/Experimental/Mutable/Source/MutableTools/Private/MuT/ASTOpMeshFormat.cpp
+++ b/Engine/Plugins/Experimental/Mutable/Source/MutableTools/Private/MuT/ASTOpMeshFormat.cpp
@@ -157,12 +157,12 @@ namespace
 
 
 	// Make a mesh format suitable to morph a particular other format.
-	MeshPtr MakeMorphTargetFormat(MeshPtrConst pTargetFormat)
+	Ptr<Mesh> MakeMorphTargetFormat(Ptr<const Mesh> pTargetFormat)
 	{
 		MUTABLE_CPUPROFILER_SCOPE(MakeMorphTargetFormat);
 
 		// Make a morph format by adding all the vertex channels from the base into a single
-		// vertex buffer, adding the vertex index channel
+		// vertex buffer
 
 		int32 offset = 0;
 		int32 numChannels = 0;
@@ -172,15 +172,6 @@ namespace
 		TArray<int32> components;
 		TArray<int32> offsets;
 
-		// Add the vertex index channel
-		semantics.Add(MBS_VERTEXINDEX);
-		semanticIndices.Add(0);
-		formats.Add(MBF_UINT32);
-		components.Add(1);
-		offsets.Add(offset);
-		offset += components[numChannels] * GetMeshFormatData(formats[numChannels]).SizeInBytes;
-		numChannels++;
-
 		// Add the vertex channels from the new format
 		for (int32 vb = 0; vb < pTargetFormat->GetVertexBuffers().GetBufferCount(); ++vb)
 		{
@@ -208,16 +199,14 @@ namespace
 
 		MeshPtr pTargetMorphFormat = new Mesh;
 		pTargetMorphFormat->GetVertexBuffers().SetBufferCount(1);
-		if (numChannels > 0)
-		{
-			pTargetMorphFormat->GetVertexBuffers().SetBuffer(0, offset,
-				numChannels,
-				&semantics[0],
-				&semanticIndices[0],
-				&formats[0],
-				&components[0],
-				&offsets[0]);
-		}
+
+		pTargetMorphFormat->GetVertexBuffers().SetBuffer(0, offset,
+			numChannels,
+			semantics.GetData(),
+			semanticIndices.GetData(),
+			formats.GetData(),
+			components.GetData(),
+			offsets.GetData());
 
 		return pTargetMorphFormat;
 	}
@@ -279,8 +268,8 @@ mu::Ptr<ASTOp> Sink_MeshFormatAST::Visit(const mu::Ptr<ASTOp>& at, const ASTOpMe
 		NewOp->Base = Visit(NewOp->Base.child(), currentFormatOp);
 
 		// Reformat the morph targets to match the new format.
-		MeshPtrConst pTargetFormat = FindBaseMeshConstant(currentFormatOp->Format.child());
-		MeshPtrConst pTargetMorphFormat = MakeMorphTargetFormat(pTargetFormat);
+		Ptr<const Mesh> pTargetFormat = FindBaseMeshConstant(currentFormatOp->Format.child());
+		Ptr<const Mesh> pTargetMorphFormat = MakeMorphTargetFormat(pTargetFormat);
 
 		mu::Ptr<ASTOpConstantResource> NewFormatConstant = new ASTOpConstantResource();
 		NewFormatConstant->Type = OP_TYPE::ME_CONSTANT;
@@ -301,7 +290,7 @@ mu::Ptr<ASTOp> Sink_MeshFormatAST::Visit(const mu::Ptr<ASTOp>& at, const ASTOpMe
 
 	case OP_TYPE::ME_MERGE:
 	{
-		auto newOp = mu::Clone<ASTOpFixed>(at);
+		Ptr<ASTOpFixed> newOp = mu::Clone<ASTOpFixed>(at);
 		newOp->SetChild(newOp->op.args.MeshMerge.base, Visit(newOp->children[newOp->op.args.MeshMerge.base].child(), currentFormatOp));
 		newOp->SetChild(newOp->op.args.MeshMerge.added, Visit(newOp->children[newOp->op.args.MeshMerge.added].child(), currentFormatOp));
 		newAt = newOp;
@@ -311,26 +300,24 @@ mu::Ptr<ASTOp> Sink_MeshFormatAST::Visit(const mu::Ptr<ASTOp>& at, const ASTOpMe
 	case OP_TYPE::ME_INTERPOLATE:
 	{
 		// Move the format down the base of the morph
-		auto newOp = mu::Clone<ASTOpFixed>(at);
+		Ptr<ASTOpFixed> newOp = mu::Clone<ASTOpFixed>(at);
 		newOp->SetChild(newOp->op.args.MeshInterpolate.base, Visit(newOp->children[newOp->op.args.MeshInterpolate.base].child(), currentFormatOp));
 
 		// Reformat the morph targets to match the new format.
-		// \TODO: Cache pTargetMorphFormat? motaop?
-		MeshPtrConst pTargetFormat = FindBaseMeshConstant(currentFormatOp->Format.child());
-		MeshPtrConst pTargetMorphFormat = MakeMorphTargetFormat(pTargetFormat);
+		Ptr<const Mesh> pTargetFormat = FindBaseMeshConstant(currentFormatOp->Format.child());
+		Ptr<const Mesh> pTargetMorphFormat = MakeMorphTargetFormat(pTargetFormat);
 
-		mu::Ptr<ASTOpConstantResource> motaop = new ASTOpConstantResource();
-		motaop->Type = OP_TYPE::ME_CONSTANT;
-		motaop->SetValue(pTargetMorphFormat, nullptr );
-		auto targetMorphFormatAt = motaop;
+		mu::Ptr<ASTOpConstantResource> TargetMorphFormatOp = new ASTOpConstantResource();
+		TargetMorphFormatOp->Type = OP_TYPE::ME_CONSTANT;
+		TargetMorphFormatOp->SetValue(pTargetMorphFormat, nullptr );
 
-		for (int t = 0; t < MUTABLE_OP_MAX_INTERPOLATE_COUNT - 1; ++t)
+		for (int32 t = 0; t < MUTABLE_OP_MAX_INTERPOLATE_COUNT - 1; ++t)
 		{
 			if (newOp->children[newOp->op.args.MeshInterpolate.targets[t]])
 			{
 				mu::Ptr<ASTOpMeshFormat> newFormat = mu::Clone<ASTOpMeshFormat>(currentFormatOp);
 				newFormat->Flags = OP::MeshFormatArgs::Vertex | OP::MeshFormatArgs::IgnoreMissing;
-				newFormat->Format = targetMorphFormatAt;
+				newFormat->Format = TargetMorphFormatOp;
 
 				newOp->SetChild(newOp->op.args.MeshInterpolate.targets[t], Visit(newOp->children[newOp->op.args.MeshInterpolate.targets[t]].child(), newFormat.get()));
 			}
diff --git a/Engine/Plugins/Experimental/Mutable/Source/MutableTools/Private/MuT/CodeGenerator.cpp b/Engine/Plugins/Experimental/Mutable/Source/MutableTools/Private/MuT/CodeGenerator.cpp
--- a/Engine/Plugins/Experimental/Mutable/Source/MutableTools/Private/MuT/CodeGenerator.cpp
+++ b/Engine/Plugins/Experimental/Mutable/Source/MutableTools/Private/MuT/CodeGenerator.cpp
@@ -719,7 +719,6 @@ namespace mu
 
             // Generate the mesh
 			FMeshGenerationOptions MeshOptions;
-			MeshOptions.bUniqueVertexIDs = true;
 			MeshOptions.bLayouts = true;
 			MeshOptions.State = Options.State;
 			MeshOptions.ActiveTags = node.m_tags;
@@ -768,7 +767,6 @@ namespace mu
 						m_activeTags.Add(e.node->m_tags);
 						
 						FMeshGenerationOptions MergedMeshOptions;
-						MergedMeshOptions.bUniqueVertexIDs = true;
 						MergedMeshOptions.bLayouts = true;
 						MergedMeshOptions.bClampUVIslands = bShareSurface && bNormalizeUVs;
 						MergedMeshOptions.bNormalizeUVs = bNormalizeUVs;
@@ -839,7 +837,6 @@ namespace mu
                     {
                         FMeshGenerationResult removeResults;
 						FMeshGenerationOptions RemoveMeshOptions;
-						RemoveMeshOptions.bUniqueVertexIDs = false;
 						RemoveMeshOptions.bLayouts = false;
 						RemoveMeshOptions.State = Options.State;
 						RemoveMeshOptions.ActiveTags = e.node->m_tags;
@@ -892,7 +889,6 @@ namespace mu
 
                     // Target mesh
 					FMeshGenerationOptions MorphTargetMeshOptions;
-					MorphTargetMeshOptions.bUniqueVertexIDs = false;
 					MorphTargetMeshOptions.bLayouts = false;
 					MorphTargetMeshOptions.State = Options.State;
 
@@ -2136,7 +2132,6 @@ namespace mu
 
 				// Parameters
 				FMeshGenerationOptions ClipOptions;
-				ClipOptions.bUniqueVertexIDs = false;
 				ClipOptions.bLayouts = false;
 				ClipOptions.State = Options.State;
 
@@ -2309,7 +2304,6 @@ namespace mu
 				Ptr<ASTOpMeshClipDeform> ClipOp = new ASTOpMeshClipDeform();
 
 				FMeshGenerationOptions ClipOptions;
-				ClipOptions.bUniqueVertexIDs = false;
 				ClipOptions.bLayouts = false;
 				ClipOptions.State = Options.State;
 
diff --git a/Engine/Plugins/Experimental/Mutable/Source/MutableTools/Private/MuT/CodeGenerator.h b/Engine/Plugins/Experimental/Mutable/Source/MutableTools/Private/MuT/CodeGenerator.h
--- a/Engine/Plugins/Experimental/Mutable/Source/MutableTools/Private/MuT/CodeGenerator.h
+++ b/Engine/Plugins/Experimental/Mutable/Source/MutableTools/Private/MuT/CodeGenerator.h
@@ -194,42 +194,35 @@ namespace mu
 
     private:
 
-        //! List of meshes generated to be able to reuse them
-		TArray<Ptr<Mesh>> m_constantMeshes;
+        /** List of meshes generated to be able to reuse them. */
+		struct FGeneratedConstantMesh
+		{
+			Ptr<Mesh> Mesh;
+			Ptr<ASTOp> LastMeshOp;
+		};
+		TArray<FGeneratedConstantMesh> GeneratedConstantMeshes;
 
-        //! List of image resources for every image formata that have been generated so far as
-        //! palceholders for missing images.
+        /** List of image resources for every image formats that have been generated so far as palceholders for missing images. */
         Ptr<Image> m_missingImage[size_t(EImageFormat::IF_COUNT)];
 
         //! First free index for a layout block
         int32 m_absoluteLayoutIndex = 0;
 
-        //! First free index to be used to identify mesh vertices.
-        uint32 m_freeVertexIndex = 0;
+        //! List of already used vertex ID groups that must be unique.
+        TSet<uint32> UniqueVertexIDGroups;
 
 		// (top-down) Tags that are active when generating nodes.
 		TArray< TArray<FString> > m_activeTags;
 
         struct FParentKey
         {
-            FParentKey()
-            {
-                m_pObject = nullptr;
-                m_state = -1;
-                m_lod = -1;
-                m_component = -1;
-                m_surface = -1;
-                m_texture = -1;
-                m_block = -1;
-            }
-
-            const NodeObjectNew::Private* m_pObject;
-            int m_state;
-            int m_lod;
-            int m_component;
-            int m_surface;
-            int m_texture;
-            int m_block;
+            const NodeObjectNew::Private* m_pObject = nullptr;
+            int32 m_state = -1;
+            int32 m_lod = -1;
+            int32 m_component = -1;
+            int32 m_surface = -1;
+            int32 m_texture = -1;
+            int32 m_block = -1;
         };
 
 		TArray< FParentKey > m_currentParents;
@@ -456,11 +449,6 @@ namespace mu
 		*/
 		struct FMeshGenerationOptions : public FGenericGenerationOptions
 		{
-			/** Whatever mesh we reach at the leaves of the graph will need to have unique ids for its vertices.
-			* This is used to track mesh removal indices, morph data in other nodes, clothing data, etc.
-			*/
-			bool bUniqueVertexIDs = false;
-
 			/** The meshes at the leaves will need their own layout block data. */
 			bool bLayouts = false;
 
@@ -485,7 +473,6 @@ namespace mu
 			friend FORCEINLINE uint32 GetTypeHash(const FMeshGenerationOptions& InKey)
 			{
 				uint32 KeyHash = 0;
-				KeyHash = HashCombineFast(KeyHash, ::GetTypeHash(InKey.bUniqueVertexIDs));
 				KeyHash = HashCombineFast(KeyHash, ::GetTypeHash(InKey.bLayouts));
 				KeyHash = HashCombineFast(KeyHash, ::GetTypeHash(InKey.OverrideLayouts.Num()));
 				KeyHash = HashCombineFast(KeyHash, ::GetTypeHash(InKey.State));
@@ -496,7 +483,7 @@ namespace mu
 			FORCEINLINE bool operator==(const FMeshGenerationOptions& Other) const
 			{
 				return State==Other.State 
-					&& bUniqueVertexIDs==Other.bUniqueVertexIDs && bLayouts==Other.bLayouts 
+					&& bLayouts==Other.bLayouts 
 					&& bClampUVIslands == Other.bClampUVIslands && bNormalizeUVs == Other.bNormalizeUVs
 					&& ActiveTags==Other.ActiveTags
 					&& OverrideLayouts ==Other.OverrideLayouts;
diff --git a/Engine/Plugins/Experimental/Mutable/Source/MutableTools/Private/MuT/CodeGenerator_Image.cpp b/Engine/Plugins/Experimental/Mutable/Source/MutableTools/Private/MuT/CodeGenerator_Image.cpp
--- a/Engine/Plugins/Experimental/Mutable/Source/MutableTools/Private/MuT/CodeGenerator_Image.cpp
+++ b/Engine/Plugins/Experimental/Mutable/Source/MutableTools/Private/MuT/CodeGenerator_Image.cpp
@@ -1403,10 +1403,6 @@ namespace mu
 			MeshOptions.ActiveTags = Options.ActiveTags;
 			MeshOptions.bLayouts = true;			// We need the layout that we will use to render
 			MeshOptions.bNormalizeUVs = true;		// We need normalized UVs for the projection
-
-			// We may need IDs at this point if there are modifiers for the mesh.
-			// \TODO: Detect this at mesh constant generation and enable there instead of here?
-			MeshOptions.bUniqueVertexIDs = true;
 			
 			FMeshGenerationResult MeshResult;
 			GenerateMesh( MeshOptions, MeshResult, node.m_pMesh );
@@ -1446,7 +1442,6 @@ namespace mu
                 Ptr<ASTOpMeshFormat> FormatOp = new ASTOpMeshFormat();
 				FormatOp->Flags = OP::MeshFormatArgs::Vertex
                         | OP::MeshFormatArgs::Index
-                        | OP::MeshFormatArgs::Face
                         | OP::MeshFormatArgs::ResetBufferIndices;
 				FormatOp->Format = cop;
 				FormatOp->Source = CurrentMeshToProjectOp;
@@ -1478,7 +1473,6 @@ namespace mu
                 Ptr<ASTOpMeshFormat> FormatOp = new ASTOpMeshFormat();
 				FormatOp->Flags = OP::MeshFormatArgs::Vertex
 					| OP::MeshFormatArgs::Index
-					| OP::MeshFormatArgs::Face
 					| OP::MeshFormatArgs::ResetBufferIndices;
 				FormatOp->Format = cop;
 				FormatOp->Source = CurrentMeshToProjectOp;
diff --git a/Engine/Plugins/Experimental/Mutable/Source/MutableTools/Private/MuT/CodeGenerator_Mesh.cpp b/Engine/Plugins/Experimental/Mutable/Source/MutableTools/Private/MuT/CodeGenerator_Mesh.cpp
--- a/Engine/Plugins/Experimental/Mutable/Source/MutableTools/Private/MuT/CodeGenerator_Mesh.cpp
+++ b/Engine/Plugins/Experimental/Mutable/Source/MutableTools/Private/MuT/CodeGenerator_Mesh.cpp
@@ -761,9 +761,7 @@ namespace mu
         FMeshGenerationResult BaseResult;
         if ( node.Base )
         {
-			FMeshGenerationOptions BaseOptions = InOptions;
-			BaseOptions.bUniqueVertexIDs = true;
-            GenerateMesh(BaseOptions,BaseResult, node.Base );
+            GenerateMesh(InOptions,BaseResult, node.Base );
             OpMorph->Base = BaseResult.meshOp;
         }
         else
@@ -777,7 +775,6 @@ namespace mu
         {
             FMeshGenerationResult TargetResult;
 			FMeshGenerationOptions TargetOptions = InOptions;
-			TargetOptions.bUniqueVertexIDs = false;
 			TargetOptions.bLayouts = false;
 			// We need to override the layouts with the layouts that were generated for the base to make
 			// sure that we get the correct mesh when generating the target
@@ -859,7 +856,6 @@ namespace mu
         if ( node.m_pBase )
         {
 			FMeshGenerationOptions BaseOptions = InOptions;
-			BaseOptions.bUniqueVertexIDs = true;
 			BaseOptions.bLayouts = false;
 			GenerateMesh(BaseOptions, BaseResult, node.m_pBase );
 
@@ -876,7 +872,6 @@ namespace mu
 		if ( node.m_pTarget )
         {
 			FMeshGenerationOptions TargetOptions = InOptions;
-			TargetOptions.bUniqueVertexIDs = false;
 			TargetOptions.bLayouts = false;
 			TargetOptions.OverrideLayouts.Empty();
 			TargetOptions.ActiveTags.Empty();
@@ -998,7 +993,6 @@ namespace mu
             if ( NodeMesh* pA = node.m_targets[t].get() )
             {
 				FMeshGenerationOptions TargetOptions = InOptions;
-				TargetOptions.bUniqueVertexIDs = true;
 				TargetOptions.OverrideLayouts.Empty();
 
                 FMeshGenerationResult TargetResult;
@@ -1250,77 +1244,69 @@ namespace mu
 
 
     //---------------------------------------------------------------------------------------------
-    void CodeGenerator::GenerateMesh_Constant(const FMeshGenerationOptions& InOptions, FMeshGenerationResult& OutResult, const NodeMeshConstant* constant )
+    void CodeGenerator::GenerateMesh_Constant(const FMeshGenerationOptions& InOptions, FMeshGenerationResult& OutResult, const NodeMeshConstant* InNode )
     {
-        NodeMeshConstant::Private& node = *constant->GetPrivate();
+        NodeMeshConstant::Private& Node = *InNode->GetPrivate();
 
-        Ptr<ASTOpConstantResource> op = new ASTOpConstantResource();
-        op->Type = OP_TYPE::ME_CONSTANT;
-		OutResult.baseMeshOp = op;
-		OutResult.meshOp = op;
+        Ptr<ASTOpConstantResource> ConstantOp = new ASTOpConstantResource();
+		ConstantOp->Type = OP_TYPE::ME_CONSTANT;
+		OutResult.baseMeshOp = ConstantOp;
+		OutResult.meshOp = ConstantOp;
 		OutResult.GeneratedLayouts.Empty();
 
 		bool bIsOverridingLayouts = !InOptions.OverrideLayouts.IsEmpty();
 
-        MeshPtr pMesh = node.m_pValue.get();
+        Ptr<Mesh> pMesh = Node.Value.get();
 		if (!pMesh)
 		{
 			// This data is required
-			MeshPtr pTempMesh = new Mesh();
-			op->SetValue(pTempMesh, m_compilerOptions->OptimisationOptions.DiskCacheContext);
-			m_constantMeshes.Add(pTempMesh);
+			MeshPtr EmptyMesh = new Mesh();
+			ConstantOp->SetValue(EmptyMesh, m_compilerOptions->OptimisationOptions.DiskCacheContext);
+			EmptyMesh->VertexIDPrefix = ConstantOp->GetValueHash();
+
+			FGeneratedConstantMesh MeshEntry;
+			MeshEntry.Mesh = EmptyMesh;
+			MeshEntry.LastMeshOp = ConstantOp;
+			GeneratedConstantMeshes.Add(MeshEntry);
 
 			// Log an error message
-			m_pErrorLog->GetPrivate()->Add("Constant mesh not set.", ELMT_WARNING, node.m_errorContext);
+			m_pErrorLog->GetPrivate()->Add("Constant mesh not set.", ELMT_WARNING, Node.m_errorContext);
 
 			return;
 		}
 
 		// Separate the tags from the mesh
-		TArray<FString> Tags = pMesh->m_tags;
+		TArray<FString> Tags = pMesh->Tags;
 		if (Tags.Num())
 		{
 			Ptr<Mesh> TaglessMesh = CloneOrTakeOver(pMesh.get());
-			TaglessMesh->m_tags.SetNum(0, EAllowShrinking::No);
+			TaglessMesh->Tags.SetNum(0, EAllowShrinking::No);
 			pMesh = TaglessMesh;
 		}
 
 		// Find out if we can (or have to) reuse a mesh that we have already generated.
-		MeshPtrConst DuplicateOf;
-		for (int32 i = 0; i < m_constantMeshes.Num(); ++i)
+		FGeneratedConstantMesh DuplicateOf;
+		for (int32 i = 0; i < GeneratedConstantMeshes.Num(); ++i)
 		{
-			MeshPtrConst Candidate = m_constantMeshes[i];
+			FGeneratedConstantMesh Candidate = GeneratedConstantMeshes[i];
 			
 			bool bCompareLayouts = InOptions.bLayouts && !bIsOverridingLayouts;
 
-			if (Candidate->IsSimilar(*pMesh, bCompareLayouts))
+			if (Candidate.Mesh->IsSimilar(*pMesh, bCompareLayouts))
 			{
-				// If it is similar, and we need unique vertex IDs, check that it also has them. This was skipped in the IsSimilar.
-				if (InOptions.bUniqueVertexIDs)
-				{
-					int32 FoundBuffer = -1;
-					int32 FoundChannel = -1;
-					Candidate->GetVertexBuffers().FindChannel(MBS_VERTEXINDEX, 0, &FoundBuffer, &FoundChannel);
-					bool bHasUniqueVertexIDs = FoundBuffer >= 0 && FoundChannel >= 0;
-					if (!bHasUniqueVertexIDs)
-					{
-						continue;
-					}
-				}
-
 				// If it is similar and we are overriding the layouts, we must compare the layouts of the candidate with the ones
 				// we are using to override.
 				if (bIsOverridingLayouts)
 				{
-					if (Candidate->GetLayoutCount() != InOptions.OverrideLayouts.Num())
+					if (Candidate.Mesh->GetLayoutCount() != InOptions.OverrideLayouts.Num())
 					{
 						continue;
 					}
 
 					bool bLayoutsAreEqual = true;
-					for (int32 l = 0; l < Candidate->GetLayoutCount(); ++l)
+					for (int32 l = 0; l < Candidate.Mesh->GetLayoutCount(); ++l)
 					{
-						bLayoutsAreEqual = (*Candidate->GetLayout(l) == *InOptions.OverrideLayouts[l]);
+						bLayoutsAreEqual = (*Candidate.Mesh->GetLayout(l) == *InOptions.OverrideLayouts[l]);
 						if ( !bLayoutsAreEqual )
 						{
 							break;
@@ -1338,15 +1324,16 @@ namespace mu
 			}
 		}
 
-		Ptr<const Mesh> FinalMesh;
-		if (DuplicateOf)
+		Ptr<ASTOp> LastMeshOp = ConstantOp;
+
+		if (DuplicateOf.Mesh)
 		{
 			// Make sure the source layouts of the mesh are mapped to the layouts of the duplicated mesh.
 			if (InOptions.bLayouts)
 			{
 				if (bIsOverridingLayouts)
 				{
-					for (int32 l = 0; l < DuplicateOf->GetLayoutCount(); ++l)
+					for (int32 l = 0; l < DuplicateOf.Mesh->GetLayoutCount(); ++l)
 					{
 						const Layout* OverridingLayout = InOptions.OverrideLayouts[l].get();
 						OutResult.GeneratedLayouts.Add(OverridingLayout);
@@ -1354,30 +1341,58 @@ namespace mu
 				}
 				else
 				{
-					for (int32 l = 0; l < DuplicateOf->GetLayoutCount(); ++l)
+					for (int32 l = 0; l < DuplicateOf.Mesh->GetLayoutCount(); ++l)
 					{
-						const Layout* DuplicatedLayout = DuplicateOf->GetLayout(l);
+						const Layout* DuplicatedLayout = DuplicateOf.Mesh->GetLayout(l);
 						OutResult.GeneratedLayouts.Add(DuplicatedLayout);
 					}
 				}
 			}
 
-			FinalMesh = DuplicateOf;
+			LastMeshOp = DuplicateOf.LastMeshOp;
+			ConstantOp = nullptr;
 		}
 		else
 		{
 			// We need to clone the mesh in the node because we will modify it.
-			Ptr<Mesh> pCloned = pMesh->Clone();
-			pCloned->EnsureSurfaceData();
+			Ptr<Mesh> Cloned = pMesh->Clone();
+			Cloned->EnsureSurfaceData();
+
+			ConstantOp->SetValue(Cloned, m_compilerOptions->OptimisationOptions.DiskCacheContext);
+
+			// Add the unique vertex ID prefix in all cases, since it is free memory-wise
+			uint32 VertexIDPrefix = uint32(ConstantOp->GetValueHash());
+			{
+				// Ensure the ID group is unique
+				bool bValid = false;
+				do
+				{
+					bool bAlreadyPresent = false;
+					UniqueVertexIDGroups.FindOrAdd(VertexIDPrefix, &bAlreadyPresent);
+					bValid = !bAlreadyPresent && VertexIDPrefix != 0;
+					if (!bValid)
+					{
+						++VertexIDPrefix;
+					}
+				} while (bValid);
+
+				Cloned->VertexIDPrefix = VertexIDPrefix;
+			}
+
+			// Add the constant data
+			FGeneratedConstantMesh MeshEntry;
+			MeshEntry.Mesh = Cloned;
+			MeshEntry.LastMeshOp = LastMeshOp;
+			GeneratedConstantMeshes.Add(MeshEntry);
 
 			if (InOptions.bLayouts)
 			{
 				if (!bIsOverridingLayouts)
 				{
 					// Apply whatever transform is necessary for every layout
-					for (int32 LayoutIndex = 0; LayoutIndex < node.m_layouts.Num(); ++LayoutIndex)
+					for (int32 LayoutIndex = 0; LayoutIndex < Node.Layouts.Num(); ++LayoutIndex)
 					{
-						Ptr<NodeLayout> pLayoutNode = node.m_layouts[LayoutIndex];
+						Ptr<NodeLayout> pLayoutNode = Node.Layouts[LayoutIndex];
 						if (!pLayoutNode)
 						{
 							continue;
@@ -1389,8 +1404,8 @@ namespace mu
 
 						Ptr<const Layout> SourceLayout = TypedNode->GetPrivate()->m_pLayout;
 						Ptr<const Layout> GeneratedLayout = AddLayout( SourceLayout );
-						const void* Context = InOptions.OverrideContext.Get(node.m_errorContext);
-						PrepareForLayout(GeneratedLayout, pCloned, LayoutIndex, Context, InOptions);
+						const void* Context = InOptions.OverrideContext.Get(Node.m_errorContext);
+						PrepareForLayout(GeneratedLayout, Cloned, LayoutIndex, Context, InOptions);
 
 						OutResult.GeneratedLayouts.Add(GeneratedLayout);
 					}
@@ -1401,58 +1416,17 @@ namespace mu
 					for (int32 LayoutIndex = 0; LayoutIndex < InOptions.OverrideLayouts.Num(); ++LayoutIndex)
 					{
 						Ptr<const Layout> GeneratedLayout = InOptions.OverrideLayouts[LayoutIndex];
-						const void* Context = InOptions.OverrideContext.Get(node.m_errorContext);
-						PrepareForLayout(GeneratedLayout, pCloned, LayoutIndex, Context, InOptions);
+						const void* Context = InOptions.OverrideContext.Get(Node.m_errorContext);
+						PrepareForLayout(GeneratedLayout, Cloned, LayoutIndex, Context, InOptions);
 
 						OutResult.GeneratedLayouts.Add(GeneratedLayout);
 					}
 				}
 			}
-
-			if (InOptions.bUniqueVertexIDs)
-			{
-				// Enumerate the vertices uniquely unless they already have indices
-				int buf = -1;
-				int chan = -1;
-				pCloned->GetVertexBuffers().FindChannel(MBS_VERTEXINDEX, 0, &buf, &chan);
-				bool hasVertexIndices = (buf >= 0 && chan >= 0);
-				if (!hasVertexIndices)
-				{
-					int newBuffer = pCloned->GetVertexBuffers().GetBufferCount();
-					pCloned->GetVertexBuffers().SetBufferCount(newBuffer + 1);
-					EMeshBufferSemantic semantic = MBS_VERTEXINDEX;
-					int semanticIndex = 0;
-					EMeshBufferFormat format = MBF_UINT32;
-					int components = 1;
-					int offset = 0;
-					pCloned->GetVertexBuffers().SetBuffer
-					(
-						newBuffer,
-						sizeof(uint32),
-						1,
-						&semantic, &semanticIndex,
-						&format, &components,
-						&offset
-					);
-					uint32* pIdData = (uint32*)pCloned->GetVertexBuffers().GetBufferData(newBuffer);
-					for (int i = 0; i < pMesh->GetVertexCount(); ++i)
-					{
-						check(m_freeVertexIndex < TNumericLimits<uint32>::Max());
-
-						(*pIdData++) = m_freeVertexIndex++;
-						check(m_freeVertexIndex < TNumericLimits<uint32>::Max());
-					}
-				}
-			}
-
-			// Add the constant data
-			m_constantMeshes.Add(pCloned);
-			FinalMesh = pCloned;
 		}
 
-		op->SetValue(FinalMesh, m_compilerOptions->OptimisationOptions.DiskCacheContext);
-
-		Ptr<ASTOp> LastMeshOp = op;
+		OutResult.baseMeshOp = LastMeshOp;
+		OutResult.meshOp = LastMeshOp;
 
 		// Add the tags operation
 		if (Tags.Num())
@@ -1465,7 +1439,7 @@ namespace mu
 
 		// Apply the modifier for the pre-normal operations stage.
 		bool bModifiersForBeforeOperations = true;
-		OutResult.meshOp = ApplyMeshModifiers(InOptions, LastMeshOp, bModifiersForBeforeOperations, node.m_errorContext);
+		OutResult.meshOp = ApplyMeshModifiers(InOptions, LastMeshOp, bModifiersForBeforeOperations, Node.m_errorContext);
     }
 
 
@@ -1490,19 +1464,13 @@ namespace mu
             if (node.VertexBuffers.GetBufferCount())
             {
                 op->Flags |= OP::MeshFormatArgs::Vertex;
-                FormatMesh->m_VertexBuffers = node.VertexBuffers;
+                FormatMesh->VertexBuffers = node.VertexBuffers;
             }
 
             if (node.IndexBuffers.GetBufferCount())
             {
 				op->Flags |= OP::MeshFormatArgs::Index;
-                FormatMesh->m_IndexBuffers = node.IndexBuffers;
-            }
-
-            if (node.FaceBuffers.GetBufferCount())
-            {
-                op->Flags |= OP::MeshFormatArgs::Face;
-                FormatMesh->m_FaceBuffers = node.FaceBuffers;
+                FormatMesh->IndexBuffers = node.IndexBuffers;
             }
 
 			if (node.bOptimizeBuffers)
@@ -1515,8 +1483,6 @@ namespace mu
             cop->SetValue( FormatMesh, m_compilerOptions->OptimisationOptions.DiskCacheContext );
             op->Format = cop;
 
-            m_constantMeshes.Add(FormatMesh);
-
             OutResult.meshOp = op;
             OutResult.baseMeshOp = baseResult.baseMeshOp;
 			OutResult.GeneratedLayouts = baseResult.GeneratedLayouts;
@@ -1568,7 +1534,6 @@ namespace mu
         if (node.m_pSource)
         {
 			FMeshGenerationOptions BaseOptions = InOptions;
-			BaseOptions.bUniqueVertexIDs = true;
             GenerateMesh(BaseOptions, OutResult, node.m_pSource);
             op->source = OutResult.meshOp;
         }
@@ -1641,10 +1606,7 @@ namespace mu
         // Base
         if (node.m_pSource)
         {
-			FMeshGenerationOptions BaseOptions = InOptions;
-			BaseOptions.bUniqueVertexIDs = true;
-
-            GenerateMesh(BaseOptions, OutResult, node.m_pSource );
+            GenerateMesh(InOptions, OutResult, node.m_pSource );
             op->SetChild( op->op.args.MeshClipWithMesh.source, OutResult.meshOp );
         }
         else
@@ -1658,7 +1620,6 @@ namespace mu
         if (node.m_pClipMesh)
         {
 			FMeshGenerationOptions ClipOptions = InOptions;
-			ClipOptions.bUniqueVertexIDs = false;
 			ClipOptions.bLayouts = false;
 			ClipOptions.OverrideLayouts.Empty();
 			ClipOptions.ActiveTags.Empty();
@@ -1688,10 +1649,7 @@ namespace mu
 		// Base Mesh
 		if (Node.m_pBaseMesh)
 		{
-			FMeshGenerationOptions BaseOptions = InOptions;
-			BaseOptions.bUniqueVertexIDs = true;
-
-			GenerateMesh(BaseOptions, Result, Node.m_pBaseMesh);
+			GenerateMesh(InOptions, Result, Node.m_pBaseMesh);
 			OpBind->Mesh = Result.meshOp;
 		}
 		else
@@ -1705,7 +1663,6 @@ namespace mu
 		if (Node.m_pClipShape)
 		{
 			FMeshGenerationOptions ClipOptions = InOptions;
-			ClipOptions.bUniqueVertexIDs = false;
 			ClipOptions.bLayouts = false;
 			ClipOptions.OverrideLayouts.Empty();
 			ClipOptions.ActiveTags.Empty();
@@ -1746,7 +1703,6 @@ namespace mu
         if (node.m_pPose)
         {
 			FMeshGenerationOptions PoseOptions = InOptions;
-			PoseOptions.bUniqueVertexIDs = false;
 			PoseOptions.bLayouts = false;
 			PoseOptions.OverrideLayouts.Empty();
 			PoseOptions.ActiveTags.Empty();
@@ -1790,7 +1746,6 @@ namespace mu
 		if (node.m_pMeshB)
 		{
 			FMeshGenerationOptions OtherOptions = InOptions;
-			OtherOptions.bUniqueVertexIDs = false;
 			OtherOptions.bLayouts = false;
 			OtherOptions.OverrideLayouts.Empty();
 			OtherOptions.ActiveTags.Empty();
@@ -1836,10 +1791,7 @@ namespace mu
 		// Base Mesh
 		if (Node.BaseMesh)
 		{
-			FMeshGenerationOptions BaseOptions = InOptions;
-			BaseOptions.bUniqueVertexIDs = true;
-
-			GenerateMesh(BaseOptions, OutResult, Node.BaseMesh);
+			GenerateMesh(InOptions, OutResult, Node.BaseMesh);
 			OpBind->Mesh = OutResult.meshOp;
 		}
 		else
@@ -1850,7 +1802,6 @@ namespace mu
 
 		// Base and target shapes shouldn't have layouts or modifiers.
 		FMeshGenerationOptions ShapeOptions = InOptions;
-		ShapeOptions.bUniqueVertexIDs = false;
 		ShapeOptions.bLayouts = false;
 		ShapeOptions.OverrideLayouts.Empty();
 		ShapeOptions.ActiveTags.Empty();
diff --git a/Engine/Plugins/Experimental/Mutable/Source/MutableTools/Private/MuT/DataPacker.cpp b/Engine/Plugins/Experimental/Mutable/Source/MutableTools/Private/MuT/DataPacker.cpp
--- a/Engine/Plugins/Experimental/Mutable/Source/MutableTools/Private/MuT/DataPacker.cpp
+++ b/Engine/Plugins/Experimental/Mutable/Source/MutableTools/Private/MuT/DataPacker.cpp
@@ -701,7 +701,7 @@ namespace mu
 
             if ( (usedSemantics&layoutSemantics) == 0)
             {
-                pMesh->m_layouts.Empty();
+                pMesh->Layouts.Empty();
             }
         }
     }
diff --git a/Engine/Plugins/Experimental/Mutable/Source/MutableTools/Private/MuT/NodeMeshConstant.cpp b/Engine/Plugins/Experimental/Mutable/Source/MutableTools/Private/MuT/NodeMeshConstant.cpp
--- a/Engine/Plugins/Experimental/Mutable/Source/MutableTools/Private/MuT/NodeMeshConstant.cpp
+++ b/Engine/Plugins/Experimental/Mutable/Source/MutableTools/Private/MuT/NodeMeshConstant.cpp
@@ -11,87 +11,69 @@
 
 namespace mu
 {
+	FNodeType NodeMeshConstant::Private::s_type = FNodeType( "MeshConstant", NodeMesh::GetStaticType() );
 
-	//---------------------------------------------------------------------------------------------
-	// Static initialisation
-	//---------------------------------------------------------------------------------------------
-	FNodeType NodeMeshConstant::Private::s_type =
-			FNodeType( "MeshConstant", NodeMesh::GetStaticType() );
-
-
-	//---------------------------------------------------------------------------------------------
-	//!
-	//---------------------------------------------------------------------------------------------
 
 	MUTABLE_IMPLEMENT_NODE( NodeMeshConstant, EType::Constant, Node, Node::EType::Mesh);
 
 
-	//---------------------------------------------------------------------------------------------
-	// Own Interface
-	//---------------------------------------------------------------------------------------------
-	MeshPtr NodeMeshConstant::GetValue() const
+	Ptr<Mesh> NodeMeshConstant::GetValue() const
 	{
-		return m_pD->m_pValue;
+		return m_pD->Value;
 	}
 
 
-	//---------------------------------------------------------------------------------------------
-	void NodeMeshConstant::SetValue( MeshPtr pValue )
+	void NodeMeshConstant::SetValue( Ptr<Mesh> pValue )
 	{
-		m_pD->m_pValue = pValue;
+		m_pD->Value = pValue;
 
-        if (m_pD->m_pValue)
+        if (m_pD->Value)
         {
             // Ensure mesh is well formed
-            m_pD->m_pValue->EnsureSurfaceData();
+            m_pD->Value->EnsureSurfaceData();
         }
     }
 
 
-	//---------------------------------------------------------------------------------------------
-	int NodeMeshConstant::GetLayoutCount() const
+	int32 NodeMeshConstant::GetLayoutCount() const
 	{
-		return m_pD->m_layouts.Num();
+		return m_pD->Layouts.Num();
 	}
 
 
-	//---------------------------------------------------------------------------------------------
-	void NodeMeshConstant::SetLayoutCount( int num )
+	void NodeMeshConstant::SetLayoutCount( int32 num )
 	{
 		check( num >=0 );
-		m_pD->m_layouts.SetNum( num );
+		m_pD->Layouts.SetNum( num );
 	}
 
 
-	//---------------------------------------------------------------------------------------------
-	NodeLayoutPtr NodeMeshConstant::GetLayout( int index ) const
+	Ptr<NodeLayout> NodeMeshConstant::GetLayout( int32 index ) const
 	{
-		check( index >=0 && index < m_pD->m_layouts.Num() );
+		check( index >=0 && index < m_pD->Layouts.Num() );
 
 		return m_pD->GetLayout( index );
 	}
 
 
-	//---------------------------------------------------------------------------------------------
-	NodeLayoutPtr NodeMeshConstant::Private::GetLayout( int index ) const
+	Ptr<NodeLayout> NodeMeshConstant::Private::GetLayout( int32 index ) const
 	{
 		NodeLayoutPtr pResult;
 
-		if ( index >=0 && index < m_layouts.Num() )
+		if ( index >=0 && index < Layouts.Num() )
 		{
-			pResult = m_layouts[ index ].get();
+			pResult = Layouts[ index ];
 		}
 
 		return pResult;
 	}
 
 
-	//---------------------------------------------------------------------------------------------
-	void NodeMeshConstant::SetLayout( int index, NodeLayoutPtr pLayout )
+	void NodeMeshConstant::SetLayout( int32 index, Ptr<NodeLayout> pLayout )
 	{
-		check( index >=0 && index < m_pD->m_layouts.Num() );
+		check( index >=0 && index < m_pD->Layouts.Num() );
 
-		m_pD->m_layouts[ index ] = pLayout;
+		m_pD->Layouts[ index ] = pLayout;
 	}
 
 
diff --git a/Engine/Plugins/Experimental/Mutable/Source/MutableTools/Private/MuT/NodeMeshFormat.cpp b/Engine/Plugins/Experimental/Mutable/Source/MutableTools/Private/MuT/NodeMeshFormat.cpp
--- a/Engine/Plugins/Experimental/Mutable/Source/MutableTools/Private/MuT/NodeMeshFormat.cpp
+++ b/Engine/Plugins/Experimental/Mutable/Source/MutableTools/Private/MuT/NodeMeshFormat.cpp
@@ -55,11 +55,6 @@ class FMeshBufferSet;
 	}
 
 
-	FMeshBufferSet& NodeMeshFormat::GetFaceBuffers()
-	{
-		return m_pD->FaceBuffers;
-	}
-
 	void NodeMeshFormat::SetOptimizeBuffers(bool bEnable)
 	{
 		m_pD->bOptimizeBuffers = bEnable;
diff --git a/Engine/Plugins/Experimental/Mutable/Source/MutableTools/Private/MuT/TypeInfo.cpp b/Engine/Plugins/Experimental/Mutable/Source/MutableTools/Private/MuT/TypeInfo.cpp
--- a/Engine/Plugins/Experimental/Mutable/Source/MutableTools/Private/MuT/TypeInfo.cpp
+++ b/Engine/Plugins/Experimental/Mutable/Source/MutableTools/Private/MuT/TypeInfo.cpp
@@ -116,6 +116,11 @@ namespace mu
 		"PackedDirS8_WTangentSign",
 
 		"Float64",
+
+		"UInt64",
+		"Int64",
+		"NUint64",
+		"NInt64",
 	};
 
 	static_assert(sizeof(TypeInfo::s_meshBufferFormatName) / sizeof(void*) == int32(MBF_COUNT));
diff --git a/Engine/Plugins/Experimental/Mutable/Source/MutableTools/Public/MuT/NodeMeshConstant.h b/Engine/Plugins/Experimental/Mutable/Source/MutableTools/Public/MuT/NodeMeshConstant.h
--- a/Engine/Plugins/Experimental/Mutable/Source/MutableTools/Public/MuT/NodeMeshConstant.h
+++ b/Engine/Plugins/Experimental/Mutable/Source/MutableTools/Public/MuT/NodeMeshConstant.h
@@ -2,7 +2,6 @@
 
 #pragma once
 
-#include "MuR/Mesh.h"
 #include "MuR/Ptr.h"
 #include "MuR/RefCounted.h"
 #include "MuT/Node.h"
@@ -17,10 +16,8 @@ namespace mu
 	typedef Ptr<NodeMeshConstant> NodeMeshConstantPtr;
 	typedef Ptr<const NodeMeshConstant> NodeMeshConstantPtrConst;
 
+	class Mesh;
 	class NodeLayout;
-	typedef Ptr<NodeLayout> NodeLayoutPtr;
-	typedef Ptr<const NodeLayout> NodeLayoutPtrConst;
-
 
 	//! Node that outputs a constant mesh.
 	//! It allows to define the layouts for the texture channels of the constant mesh
@@ -42,18 +39,18 @@ namespace mu
 		//-----------------------------------------------------------------------------------------
 
 		//! Get the constant mesh that will be returned.
-		MeshPtr GetValue() const;
+		Ptr<Mesh> GetValue() const;
 
 		//! Set the constant mesh that will be returned.
-		void SetValue( MeshPtr );
+		void SetValue( Ptr<Mesh> );
 
 		//! Get the number of layouts defined in this mesh.
-		int GetLayoutCount() const;
-		void SetLayoutCount( int );
+		int32 GetLayoutCount() const;
+		void SetLayoutCount( int32 );
 
 		//! Get the node defining a layout of the returned mesh.
-		NodeLayoutPtr GetLayout( int index ) const;
-		void SetLayout( int index, NodeLayoutPtr );
+		Ptr<NodeLayout> GetLayout( int32 index ) const;
+		void SetLayout( int32 index, Ptr<NodeLayout>);
 
 		//-----------------------------------------------------------------------------------------
 		// Interface pattern
diff --git a/Engine/Plugins/Experimental/Mutable/Source/MutableTools/Public/MuT/NodeMeshFormat.h b/Engine/Plugins/Experimental/Mutable/Source/MutableTools/Public/MuT/NodeMeshFormat.h
--- a/Engine/Plugins/Experimental/Mutable/Source/MutableTools/Public/MuT/NodeMeshFormat.h
+++ b/Engine/Plugins/Experimental/Mutable/Source/MutableTools/Public/MuT/NodeMeshFormat.h
@@ -47,12 +47,6 @@ namespace mu
 		*/
 		FMeshBufferSet& GetIndexBuffers();
 
-		/** Access the MeshBufferSet that defines the new format for the mesh faces. These
-		* buffers don't really contain any data (they have 0 elements) but they define the
-		* structure. If this is null, the vertex buffers will not be changed.
-		*/
-		FMeshBufferSet& GetFaceBuffers();
-
 		/** Optimize the buffers is possible. This may change the target format to reduce the number
 		* of channels or even the type if possible to minimize size.
 		* By default it is disabled.
