[Mutable] Added gathered AnimBPs to Baked Instances.

#rnx

diff --git a/Engine/Plugins/Experimental/Mutable/Source/CustomizableObject/Private/MuCO/CustomizableObjectInstance.cpp b/Engine/Plugins/Experimental/Mutable/Source/CustomizableObject/Private/MuCO/CustomizableObjectInstance.cpp
--- a/Engine/Plugins/Experimental/Mutable/Source/CustomizableObject/Private/MuCO/CustomizableObjectInstance.cpp
+++ b/Engine/Plugins/Experimental/Mutable/Source/CustomizableObject/Private/MuCO/CustomizableObjectInstance.cpp
@@ -3468,10 +3468,21 @@ void UCustomizableInstancePrivate::InitSkeletalMeshData(const TSharedRef<FUpdate
 		}
 
 		//Custom Asset User Data
-		if (OperationData->Instance->GetAnimationGameplayTags().Num())
+		if (OperationData->Instance->GetAnimationGameplayTags().Num() ||
+			ComponentData.AnimSlotToBP.Num())
 		{
 			UCustomizableObjectInstanceUserData* InstanceData = NewObject<UCustomizableObjectInstanceUserData>(SkeletalMesh, NAME_None, RF_Public | RF_Transactional);
-			InstanceData->SetAnimationGameplayTags(OperationData->Instance->GetAnimationGameplayTags());
+			InstanceData->AnimationGameplayTag = OperationData->Instance->GetAnimationGameplayTags();
+			
+			for (const TTuple<FName, TSoftClassPtr<UAnimInstance>>& AnimSlot : ComponentData.AnimSlotToBP)
+			{
+				FCustomizableObjectAnimationSlot AnimationSlot;
+				AnimationSlot.Name = AnimSlot.Key;
+				AnimationSlot.AnimInstance = AnimSlot.Value;
+				
+				InstanceData->AnimationSlots.Add(AnimationSlot);
+			}
+			
 			SkeletalMesh->AddAssetUserData(InstanceData);
 		}
 	}
diff --git a/Engine/Plugins/Experimental/Mutable/Source/CustomizableObject/Public/MuCO/CustomizableObjectInstanceAssetUserData.h b/Engine/Plugins/Experimental/Mutable/Source/CustomizableObject/Public/MuCO/CustomizableObjectInstanceAssetUserData.h
--- a/Engine/Plugins/Experimental/Mutable/Source/CustomizableObject/Public/MuCO/CustomizableObjectInstanceAssetUserData.h
+++ b/Engine/Plugins/Experimental/Mutable/Source/CustomizableObject/Public/MuCO/CustomizableObjectInstanceAssetUserData.h
@@ -7,6 +7,21 @@
 
 #include "CustomizableObjectInstanceAssetUserData.generated.h"
 
+class UAnimInstance;
+
+
+USTRUCT(BlueprintType)
+struct FCustomizableObjectAnimationSlot
+{
+	GENERATED_BODY()
+
+	UPROPERTY(BlueprintReadWrite, Category = CustomizableObjectInstance)
+	FName Name;
+
+	UPROPERTY(BlueprintReadWrite, Category = CustomizableObjectInstance)
+	TSoftClassPtr<UAnimInstance> AnimInstance;
+};
+
 
 /** Additional data attached to Skeletal Meshes. */
 UCLASS(BlueprintType)
@@ -15,23 +30,9 @@ class CUSTOMIZABLEOBJECT_API UCustomizableObjectInstanceUserData : public UAsset
 	GENERATED_BODY()
 
 public:
+	UPROPERTY(BlueprintReadWrite, Category = CustomizableObjectInstance)
+	FGameplayTagContainer AnimationGameplayTag;
 
-	/** Return the list of tags for this instance. */
-	UFUNCTION(BlueprintCallable, Category = CustomizableObjectInstance)
-	const FGameplayTagContainer& GetAnimationGameplayTags() const
-	{
-		return AnimBPGameplayTags;
-	};
-
-	/** Sets the list of tags for this instance. */
-	UFUNCTION(BlueprintCallable, Category = CustomizableObjectInstance)
-	void SetAnimationGameplayTags(const FGameplayTagContainer& InstanceTags)
-	{
-		AnimBPGameplayTags = InstanceTags;
-	};
-
-private:
-
-	UPROPERTY()
-	FGameplayTagContainer AnimBPGameplayTags;
+	UPROPERTY(BlueprintReadWrite, Category = CustomizableObjectInstance)
+	TArray<FCustomizableObjectAnimationSlot> AnimationSlots;
 };
diff --git a/Engine/Plugins/Experimental/Mutable/Source/CustomizableObjectEditor/Private/MuCOE/CustomizableObjectInstanceBakingUtils.cpp b/Engine/Plugins/Experimental/Mutable/Source/CustomizableObjectEditor/Private/MuCOE/CustomizableObjectInstanceBakingUtils.cpp
--- a/Engine/Plugins/Experimental/Mutable/Source/CustomizableObjectEditor/Private/MuCOE/CustomizableObjectInstanceBakingUtils.cpp
+++ b/Engine/Plugins/Experimental/Mutable/Source/CustomizableObjectEditor/Private/MuCOE/CustomizableObjectInstanceBakingUtils.cpp
@@ -18,6 +18,7 @@
 
 #include "MuCO/CustomizableObject.h"
 #include "MuCO/CustomizableObjectInstanceAssetUserData.h"
+#include "MuCO/CustomizableObjectInstancePrivate.h"
 #include "MuCO/CustomizableObjectMipDataProvider.h"
 #include "MuCO/ICustomizableObjectEditorModule.h"
 #include "MuT/UnrealPixelFormatOverride.h"
@@ -842,10 +843,24 @@ bool BakeCustomizableObjectInstance(
 
 			// Add Instance Info in a custom AssetUserData
 			{
-				if (InInstance.GetAnimationGameplayTags().Num())
+				const FCustomizableInstanceComponentData* ComponentData = InInstance.GetPrivate()->GetComponentData(ComponentIndex);
+				check(ComponentData);
+				
+				if (InInstance.GetAnimationGameplayTags().Num() ||
+					ComponentData->AnimSlotToBP.Num())
 				{
 					UCustomizableObjectInstanceUserData* InstanceData = NewObject<UCustomizableObjectInstanceUserData>(SkeletalMesh, NAME_None, RF_Public | RF_Transactional);
-					InstanceData->SetAnimationGameplayTags(InInstance.GetAnimationGameplayTags());
+					InstanceData->AnimationGameplayTag = InInstance.GetAnimationGameplayTags();
+
+					for (const TTuple<FName, TSoftClassPtr<UAnimInstance>>& AnimSlot : ComponentData->AnimSlotToBP)
+					{
+						FCustomizableObjectAnimationSlot AnimationSlot;
+						AnimationSlot.Name = AnimSlot.Key;
+						AnimationSlot.AnimInstance = AnimSlot.Value;
+				
+						InstanceData->AnimationSlots.Add(AnimationSlot);
+					}
+
 					SkeletalMesh->AddAssetUserData(InstanceData);
 				}
 			}
