[Mutable] Load textures concurrently during an update.

RuntimeReferencedTextures will be kept alive after a none compilation (or loading the Model Resources). Those will live until the mu::Model dies.

- Split RuntimeReferencedTextures from PassThroughTextures.
- Added constructor to FMutableSourceTextureData.
- GetImageAsync (MutableThread) now does not access the TextureToLoad UTexture2D.

#rnx

diff --git a/Engine/Plugins/Experimental/Mutable/Source/CustomizableObject/Internal/MuCO/CustomizableObjectPrivate.h b/Engine/Plugins/Experimental/Mutable/Source/CustomizableObject/Internal/MuCO/CustomizableObjectPrivate.h
--- a/Engine/Plugins/Experimental/Mutable/Source/CustomizableObject/Internal/MuCO/CustomizableObjectPrivate.h
+++ b/Engine/Plugins/Experimental/Mutable/Source/CustomizableObject/Internal/MuCO/CustomizableObjectPrivate.h
@@ -651,10 +651,16 @@ struct FModelResources
 	UPROPERTY()
 	TArray<TSoftObjectPtr<UMaterialInterface>> Materials;
 
-	/** PassThrough Textures used by the mu::Model. */
+	/** PassThrough textures used by the mu::Model. */
 	UPROPERTY()
 	TArray<TSoftObjectPtr<UTexture>> PassThroughTextures;
 
+#if WITH_EDITORONLY_DATA
+	/** Runtime referenced textures used by the mu::Model. */
+	UPROPERTY()
+	TArray<TSoftObjectPtr<UTexture2D>> RuntimeReferencedTextures;
+#endif
+	
 	/** Physics assets gathered from the SkeletalMeshes, to be used in mesh generation in-game */
 	UPROPERTY()
 	TArray<TSoftObjectPtr<UPhysicsAsset>> PhysicsAssets;
diff --git a/Engine/Plugins/Experimental/Mutable/Source/CustomizableObject/Internal/MuCO/CustomizableObjectSystemPrivate.h b/Engine/Plugins/Experimental/Mutable/Source/CustomizableObject/Internal/MuCO/CustomizableObjectSystemPrivate.h
--- a/Engine/Plugins/Experimental/Mutable/Source/CustomizableObject/Internal/MuCO/CustomizableObjectSystemPrivate.h
+++ b/Engine/Plugins/Experimental/Mutable/Source/CustomizableObject/Internal/MuCO/CustomizableObjectSystemPrivate.h
@@ -545,6 +545,8 @@ public:
 	mu::Instance::ID InstanceID = 0; // Redundant
 	const mu::Instance* MutableInstance = nullptr;
 
+	TSharedPtr<mu::Model> Model;
+
 	uint8 NumComponents = 0;
 	uint8 NumLODsAvailable = 0;
 	uint8 FirstLODAvailable = 0;
@@ -755,10 +757,7 @@ public:
 	void CacheTextureParameters(const TArray<FCustomizableObjectTextureParameterValue>& TextureParameters) const;
 
 	void UnCacheTextureParameters(const TArray<FCustomizableObjectTextureParameterValue>& TextureParameters) const;
-
-	/** Tick, in the Game Thread, anything that could potentially block the Mutable Thread. */
-	void TickMutableThreadDependencies();
-
+	
 	/** Mutable TaskGraph system (Mutable Thread). */
 	FMutableTaskGraph MutableTaskGraph;
 	
diff --git a/Engine/Plugins/Experimental/Mutable/Source/CustomizableObject/Internal/MuCO/UnrealToMutableTextureConversionUtils.h b/Engine/Plugins/Experimental/Mutable/Source/CustomizableObject/Internal/MuCO/UnrealToMutableTextureConversionUtils.h
--- a/Engine/Plugins/Experimental/Mutable/Source/CustomizableObject/Internal/MuCO/UnrealToMutableTextureConversionUtils.h
+++ b/Engine/Plugins/Experimental/Mutable/Source/CustomizableObject/Internal/MuCO/UnrealToMutableTextureConversionUtils.h
@@ -5,8 +5,24 @@
 #include "HAL/Platform.h"
 #include "Engine/Texture.h"
 
+#if WITH_EDITOR
 struct CUSTOMIZABLEOBJECT_API FMutableSourceTextureData
 {
+	FMutableSourceTextureData() = default;
+	
+	FMutableSourceTextureData(const UTexture& Texture);
+	
+	FTextureSource& GetSource();
+	
+	bool GetFlipGreenChannel() const;
+	
+	bool HasAlphaChannel() const;
+
+	bool GetCompressionForceAlpha() const;
+
+	bool IsNormalComposite() const;
+
+private:
 	FTextureSource Source;
 	bool bFlipGreenChannel = false;
 	bool bHasAlphaChannel = false;
@@ -15,8 +31,6 @@ struct CUSTOMIZABLEOBJECT_API FMutableSourceTextureData
 };
 
 
-#if WITH_EDITOR
-
 namespace mu
 {
 	class Image;
diff --git a/Engine/Plugins/Experimental/Mutable/Source/CustomizableObject/Private/MuCO/CustomizableObject.cpp b/Engine/Plugins/Experimental/Mutable/Source/CustomizableObject/Private/MuCO/CustomizableObject.cpp
--- a/Engine/Plugins/Experimental/Mutable/Source/CustomizableObject/Private/MuCO/CustomizableObject.cpp
+++ b/Engine/Plugins/Experimental/Mutable/Source/CustomizableObject/Private/MuCO/CustomizableObject.cpp
@@ -541,6 +541,17 @@ void UCustomizableObjectPrivate::SaveCompiledData(FArchive& MemoryWriter, bool b
 		MemoryWriter << StringRef;
 	}
 
+#if WITH_EDITORONLY_DATA
+	int32 NumRuntimeReferencedTextures = LocalModelResources.RuntimeReferencedTextures.Num();
+	MemoryWriter << NumRuntimeReferencedTextures;
+	
+	for (const TSoftObjectPtr<UTexture2D>& RuntimeReferencedTexture : LocalModelResources.RuntimeReferencedTextures)
+	{
+		FString StringRef = RuntimeReferencedTexture.ToString();
+		MemoryWriter << StringRef;
+	}
+#endif
+
 	int32 NumPhysicsAssets = LocalModelResources.PhysicsAssets.Num();
 	MemoryWriter << NumPhysicsAssets;
 
@@ -658,6 +669,19 @@ void UCustomizableObjectPrivate::LoadCompiledData(FArchive& MemoryReader, const
 			LocalModelResource.PassThroughTextures.Add(TSoftObjectPtr<UTexture>(FSoftObjectPath(StringRef)));
 		}
 
+#if WITH_EDITORONLY_DATA
+		int32 NumRuntimeReferencedTextures = 0;
+		MemoryReader << NumRuntimeReferencedTextures;
+
+		for (int32 Index = 0; Index < NumRuntimeReferencedTextures; ++Index)
+		{
+			FString StringRef;
+			MemoryReader << StringRef;
+
+			LocalModelResource.RuntimeReferencedTextures.Add(TSoftObjectPtr<UTexture2D>(FSoftObjectPath(StringRef)));
+		}
+#endif
+		
 		int32 NumPhysicsAssets = 0;
 		MemoryReader << NumPhysicsAssets;
 
diff --git a/Engine/Plugins/Experimental/Mutable/Source/CustomizableObject/Private/MuCO/CustomizableObjectInstance.cpp b/Engine/Plugins/Experimental/Mutable/Source/CustomizableObject/Private/MuCO/CustomizableObjectInstance.cpp
--- a/Engine/Plugins/Experimental/Mutable/Source/CustomizableObject/Private/MuCO/CustomizableObjectInstance.cpp
+++ b/Engine/Plugins/Experimental/Mutable/Source/CustomizableObject/Private/MuCO/CustomizableObjectInstance.cpp
@@ -5840,7 +5840,7 @@ void UCustomizableInstancePrivate::BuildMaterials(const TSharedRef<FUpdateContex
 		CustomizableObject->GetPathName(),
 		Public->GetPathName(),
 		UCustomizableObjectSystem::GetInstance()->GetPrivate()->MutableSystem,
-		CustomizableObject->GetPrivate()->GetModel(),
+		OperationData->Model,
 		OperationData->Parameters,
 	    OperationData->GetCapturedDescriptor().GetState());
 
diff --git a/Engine/Plugins/Experimental/Mutable/Source/CustomizableObject/Private/MuCO/CustomizableObjectSystem.cpp b/Engine/Plugins/Experimental/Mutable/Source/CustomizableObject/Private/MuCO/CustomizableObjectSystem.cpp
--- a/Engine/Plugins/Experimental/Mutable/Source/CustomizableObject/Private/MuCO/CustomizableObjectSystem.cpp
+++ b/Engine/Plugins/Experimental/Mutable/Source/CustomizableObject/Private/MuCO/CustomizableObjectSystem.cpp
@@ -1400,17 +1400,7 @@ bool UCustomizableObjectSystem::LockObject(const class UCustomizableObject* InOb
 				});
 
 			
-			if (GetMutableDefault<UCustomizableObjectSettings>()->bEnableStreamingManager)
-			{
-				while (!Task.IsCompleted())				
-				{
-					GetPrivate()->TickMutableThreadDependencies();
-				}
-			}
-			else
-			{
-				Task.Wait();
-			}
+			Task.Wait();
 		}
 
 		Private->MutablePendingInstanceWork.RemoveUpdatesForObject(InObject);
@@ -1648,15 +1638,13 @@ namespace impl
 		SystemPrivate->Streamer->PrepareStreamingForObject(Operation->Instance->GetCustomizableObject());			
 
 		const mu::Ptr<mu::System> MutableSystem = SystemPrivate->MutableSystem;
-
-		const TSharedPtr<mu::Model> Model = Operation->Instance->GetCustomizableObject()->GetPrivate()->GetModel();
-
+		
 		if (Operation->bLiveUpdateMode)
 		{
 			if (Operation->InstanceID == 0)
 			{
 				// It's the first update since the instance was put in LiveUpdate Mode, this ID will be reused from now on
-				Operation->InstanceID = MutableSystem->NewInstance(Model);
+				Operation->InstanceID = MutableSystem->NewInstance(Operation->Model);
 				UE_LOG(LogMutable, Verbose, TEXT("Creating Mutable instance with id [%d] for reuse "), Operation->InstanceID);
 			}
 			else
@@ -1670,7 +1658,7 @@ namespace impl
 		{
 			// In non-LiveUpdate mode, we are forcing the recreation of mutable-side instances with every update.
 			check(Operation->InstanceID == 0);
-			Operation->InstanceID = MutableSystem->NewInstance(Model);
+			Operation->InstanceID = MutableSystem->NewInstance(Operation->Model);
 			UE_LOG(LogMutable, Verbose, TEXT("Creating Mutable instance with id [%d] "), Operation->InstanceID);
 		}
 
@@ -3296,6 +3284,11 @@ namespace impl
 			CandidateInstancePrivateData->LiveUpdateModeInstanceID = 0;
 		}
 
+		Operation->Model = CustomizableObject->GetPrivate()->GetModel().ToSharedRef();
+
+#if WITH_EDITOR
+		SystemPrivateData->GetImageProviderChecked()->CacheRuntimeReferencedImages(Operation->Model.ToSharedRef(), CustomizableObject->GetPrivate()->GetModelResources().RuntimeReferencedTextures);
+#endif
 		
 		// Task: Mutable Update and GetMesh
 		//-------------------------------------------------------------
@@ -3681,7 +3674,6 @@ int32 UCustomizableObjectSystem::TickInternal()
 	}
 
 #if WITH_EDITOR
-	GetPrivate()->TickMutableThreadDependencies();
 #endif
 	
 	const int32 RemainingTasks = Private->MutableTaskGraph.Tick();
@@ -4025,13 +4017,6 @@ int32 UCustomizableObjectSystem::GetWorkingMemory() const
 }
 
 
-void UCustomizableObjectSystemPrivate::TickMutableThreadDependencies()
-{
-	check(IsInGameThread());
-	ImageProvider->Tick();
-}
-
-
 uint64 UCustomizableObjectSystem::GetMaxChunkSizeForPlatform(const ITargetPlatform* TargetPlatform)
 {
 	const FString& PlatformName = TargetPlatform ? TargetPlatform->IniPlatformName() : FPlatformProperties::IniPlatformName();
diff --git a/Engine/Plugins/Experimental/Mutable/Source/CustomizableObject/Private/MuCO/UnrealMutableImageProvider.cpp b/Engine/Plugins/Experimental/Mutable/Source/CustomizableObject/Private/MuCO/UnrealMutableImageProvider.cpp
--- a/Engine/Plugins/Experimental/Mutable/Source/CustomizableObject/Private/MuCO/UnrealMutableImageProvider.cpp
+++ b/Engine/Plugins/Experimental/Mutable/Source/CustomizableObject/Private/MuCO/UnrealMutableImageProvider.cpp
@@ -2,6 +2,7 @@
 
 #include "MuCO/UnrealMutableImageProvider.h"
 
+#include "ImageCoreUtils.h"
 #include "MuCO/CustomizableObjectSystem.h"
 #include "MuCO/CustomizableObject.h"
 #include "MuCO/CustomizableObjectPrivate.h"
@@ -9,6 +10,7 @@
 #include "TextureResource.h"
 #include "MuR/Parameters.h"
 #include "MuR/ImageTypes.h"
+#include "MuR/Model.h"
 
 
 //-------------------------------------------------------------------------------------------------
@@ -77,92 +79,11 @@ mu::EImageFormat GetMutablePixelFormat(EPixelFormat InTextureFormat)
 }
 
 
-#if WITH_EDITOR
-
-bool FUnrealMutableImageProvider::Tick()
-{
-	FReferencedImageRequest* Request=nullptr;
-
-	// Process only one request per tick, or zero.
-	if (!QueuedReferencedImageRequests.Dequeue(Request))
-	{
-		return true;
-	}
-
-	// Find the CO for this model.
-	UCustomizableObject* CO = nullptr;
-	for (TObjectIterator<UCustomizableObject> It; It; ++It)
-	{
-		if (IsValid(*It) && It->GetPrivate()->GetModel().Get() == Request->ModelPtr)
-		{
-			CO = *It;
-			break;
-		}
-	}
-
-	if (!CO)
-	{
-		// The CO for this request has been unloaded!
-		check(false);
-		return true;
-	}
-
-	const FModelResources& ModelResources = CO->GetPrivate()->GetModelResources();
-	if (!ModelResources.PassThroughTextures.IsValidIndex(Request->Id))
-	{
-		// The id is not valid for this CO
-		check(false);
-		return true;
-	}
-
-	// Find the texture id
-	TSoftObjectPtr<UTexture> TexturePtr = ModelResources.PassThroughTextures[Request->Id];
-
-	// This can cause a stall because of loading the asset.
-	UTexture2D* Texture = Cast<UTexture2D>( TexturePtr.LoadSynchronous() );
-	if (!Texture)
-	{
-		// Failed to load the texture
-		check(false);
-		return true;
-	}
-
-	// In the editor the src data can be directly accessed
-	int32 MipIndex = (Request->MipmapsToSkip < Texture->GetPlatformData()->Mips.Num()) ? Request->MipmapsToSkip : Texture->GetPlatformData()->Mips.Num() - 1;
-	check(MipIndex >= 0);
-
-	FMutableSourceTextureData Tex;
-	Tex.Source = Texture->Source.CopyTornOff();
-	Tex.bFlipGreenChannel = Texture->bFlipGreenChannel;
-	Tex.bHasAlphaChannel = 
-		Texture->AdjustMinAlpha != Texture->AdjustMaxAlpha
-		&& Texture->CompressionSettings != TextureCompressionSettings::TC_Normalmap
-		&& !Texture->CompressionNoAlpha;
-	Tex.bCompressionForceAlpha = Texture->CompressionForceAlpha;
-	Tex.bIsNormalComposite = false; // TODO?
-
-	EUnrealToMutableConversionError Error = ConvertTextureUnrealSourceToMutable(Request->ResultImage.get(), Tex, MipIndex);
-	
-	if (Error != EUnrealToMutableConversionError::Success)
-	{
-		// This could happen in the editor, because some source textures may have changed while there was a background compilation.
-		// We just show a warning and move on. This cannot happen during cooks, so it is fine.
-		UE_LOG(LogMutable, Warning, TEXT("Failed to load some source texture data for [%s]. Some materials may look corrupted."), *Texture->GetName());
-	}
-
-	Request->CompletionEvent.Trigger();
-
-	return true;
-}
-
-#endif
-
-
 //-------------------------------------------------------------------------------------------------
 TTuple<UE::Tasks::FTask, TFunction<void()>> FUnrealMutableImageProvider::GetImageAsync(FName Id, uint8 MipmapsToSkip, TFunction<void(mu::Ptr<mu::Image>)>& ResultCallback)
 {
 	// Thread: worker
-	MUTABLE_CPUPROFILER_SCOPE(FUnrealMutableImageProvider::GetImage);
+	MUTABLE_CPUPROFILER_SCOPE(FUnrealMutableImageProvider::GetImageAsync);
 
 	// Some data that may have to be copied from the GlobalExternalImages while it's locked
 	IBulkDataIORequest* IORequest = nullptr;
@@ -201,23 +122,33 @@ TTuple<UE::Tasks::FTask, TFunction<void()>> FUnrealMutableImageProvider::GetImag
 			ResultCallback(ImageInfo.Image);
 			return Invoke(TrivialReturn);
 		}
-		else if (UTexture2D* TextureToLoad = ImageInfo.TextureToLoad)
-		{
-			// It's safe to access TextureToLoad because ExternalImagesLock guarantees that the data in GlobalExternalImages is valid,
-			// not being modified by the game thread at the moment and the texture cannot be GCed because of the AddReferencedObjects
-			// in the FUnrealMutableImageProvider
 
 #if WITH_EDITOR
-			if (!TextureToLoad->IsAsyncCacheComplete())
+		if (const TSharedPtr<FMutableSourceTextureData>& SourceTextureData = ImageInfo.SourceTextureData)
+		{
+			const int32 MipIndex = FMath::Min(static_cast<int32>(MipmapsToSkip), SourceTextureData->GetSource().GetNumMips() - 1);
+			check(MipIndex >= 0);
+			
+			// In the editor the src data can be directly accessed
+			mu::Ptr<mu::Image> Image = new mu::Image();
+
+			EUnrealToMutableConversionError Error = ConvertTextureUnrealSourceToMutable(Image.get(), *SourceTextureData, MipIndex);
+			if (Error != EUnrealToMutableConversionError::Success)
 			{
-				// If the UE texture is being compiled/processed in a background process, it's not safe to access
-				FString TextureName;
-				TextureToLoad->GetName(TextureName);
-				UE_LOG(LogMutable, Warning, TEXT("Failed to get the external texture [%s] because it's still being async cached. Wait until all background processed finish and retry again."), *TextureName);
-				ResultCallback(CreateDummy());
-				return Invoke(TrivialReturn);
+				// This could happen in the editor, because some source textures may have changed while there was a background compilation.
+				// We just show a warning and move on. This cannot happen during cooks, so it is fine.
+				UE_LOG(LogMutable, Warning, TEXT("Failed to load some source texture data for image [%s]. Some materials may look corrupted."), *Id.ToString());
 			}
-#endif
+
+			ResultCallback(Image);
+			return Invoke(TrivialReturn);
+		}
+#else
+		if (UTexture2D* TextureToLoad = ImageInfo.TextureToLoad)
+		{
+			// It's safe to access TextureToLoad because ExternalImagesLock guarantees that the data in GlobalExternalImages is valid,
+			// not being modified by the game thread at the moment and the texture cannot be GCed because of the AddReferencedObjects
+			// in the FUnrealMutableImageProvider
 
 			int32 MipIndex = MipmapsToSkip < TextureToLoad->GetPlatformData()->Mips.Num() ? MipmapsToSkip : TextureToLoad->GetPlatformData()->Mips.Num() - 1;
 			check (MipIndex >= 0);
@@ -230,33 +161,7 @@ TTuple<UE::Tasks::FTask, TFunction<void()>> FUnrealMutableImageProvider::GetImag
 					break;
 				}
 			}
-
-#if WITH_EDITOR
-			// In the editor the src data can be directly accessed
-			mu::Ptr<mu::Image> Image = new mu::Image();
-
-			FMutableSourceTextureData Tex;
-			Tex.Source = TextureToLoad->Source.CopyTornOff();
-			Tex.bFlipGreenChannel = TextureToLoad->bFlipGreenChannel;
-			Tex.bHasAlphaChannel =
-				TextureToLoad->AdjustMinAlpha != TextureToLoad->AdjustMaxAlpha
-				&& TextureToLoad->CompressionSettings != TextureCompressionSettings::TC_Normalmap
-				&& !TextureToLoad->CompressionNoAlpha;
-			Tex.bCompressionForceAlpha = TextureToLoad->CompressionForceAlpha;
-			Tex.bIsNormalComposite = false; // TODO?
-
-			EUnrealToMutableConversionError Error = ConvertTextureUnrealSourceToMutable(Image.get(), Tex, MipIndex);
-
-			if (Error != EUnrealToMutableConversionError::Success)
-			{
-				// This could happen in the editor, because some source textures may have changed while there was a background compilation.
-				// We just show a warning and move on. This cannot happen during cooks, so it is fine.
-				UE_LOG(LogMutable, Warning, TEXT("Failed to load some source texture data for [%s]. Some materials may look corrupted."), *TextureToLoad->GetName());
-			}
-
-			ResultCallback(Image);
-			return Invoke(TrivialReturn);
-#else
+			
 			// Texture format and the equivalent mutable format
 			Format = TextureToLoad->GetPlatformData()->PixelFormat;
 			MutImageFormat = GetMutablePixelFormat(Format);
@@ -410,15 +315,13 @@ TTuple<UE::Tasks::FTask, TFunction<void()>> FUnrealMutableImageProvider::GetImag
 					return Invoke(TrivialReturn);
 				}
 			}
-#endif
-		}
-		else
-		{
-			// No UTexture2D was provided, cannot do anything, just provide a dummy texture
-			UE_LOG(LogMutable, Warning, TEXT("No UTexture2D was provided for an application-specific image parameter."));
-			ResultCallback(CreateDummy());
-			return Invoke(TrivialReturn);
 		}
+#endif
+		
+		// No UTexture2D was provided, cannot do anything, just provide a dummy texture
+		UE_LOG(LogMutable, Warning, TEXT("No UTexture2D was provided for an application-specific image parameter."));
+		ResultCallback(CreateDummy());
+		return Invoke(TrivialReturn);
 	}
 
 	// Make sure the returned event is dispatched at some point for all code paths, 
@@ -437,33 +340,39 @@ TTuple<UE::Tasks::FTask, TFunction<void()>> FUnrealMutableImageProvider::GetRefe
 		return MakeTuple(UE::Tasks::MakeCompletedTask<void>(), []() -> void {});
 	};
 
-
 #if WITH_EDITOR
 
-	TUniquePtr<FReferencedImageRequest> Request( new FReferencedImageRequest(UE::Tasks::FTaskEvent(TEXT("GetReferencedImageCompletion"))) );
-	Request->ModelPtr = ModelPtr;
-	Request->Id = Id;
-	Request->MipmapsToSkip = MipmapsToSkip;
-	Request->ResultImage = new mu::Image();
-
-	QueuedReferencedImageRequests.Enqueue(Request.Get());
+	mu::Ptr<mu::Image> Image = new mu::Image();
 
-	if (IsInGameThread())
+	FScopeLock Lock(&RuntimeReferencedLock);
+	
+	TArray<FMutableSourceTextureData>& SourceTextures = RuntimeReferencedImages[ModelPtr].SourceTextures;
+	if (!SourceTextures.IsValidIndex(Id))
 	{
-		// This may happen in the mutable debugger.
-		while (!Request->CompletionEvent.IsCompleted())
-		{
-			Tick();
-		}
+		// This could happen in the editor, because some source textures may have changed while there was a background compilation.
+		// We just show a warning and move on. This cannot happen during cooks, so it is fine.
+		UE_LOG(LogMutable, Warning, TEXT("Failed to load image [%i]."), Id);
+
+		ResultCallback(Image);
+		return Invoke(TrivialReturn);
 	}
-	else
+
+	FMutableSourceTextureData& SourceTextureData = SourceTextures[Id];
+	
+	const int32 MipIndex = FMath::Min(static_cast<int32>(MipmapsToSkip), SourceTextureData.GetSource().GetNumMips() - 1);
+	check(MipIndex >= 0);
+	
+	EUnrealToMutableConversionError Error = ConvertTextureUnrealSourceToMutable(Image.get(), SourceTextureData, MipIndex);
+	if (Error != EUnrealToMutableConversionError::Success)
 	{
-		Request->CompletionEvent.BusyWait();
+		// This could happen in the editor, because some source textures may have changed while updating.
+		// We just show a warning and move on. This cannot happen during cooks, so it is fine.
+		UE_LOG(LogMutable, Warning, TEXT("Failed to load some source texture data for image [%i]. Some textures may be corrupted."), Id);
 	}
 
-	ResultCallback(Request->ResultImage);
+	ResultCallback(Image);
 	return Invoke(TrivialReturn);
-
+	
 #else // WITH_EDITOR
 
 	// Not supported outside editor yet.
@@ -481,56 +390,82 @@ mu::FImageDesc FUnrealMutableImageProvider::GetImageDesc(FName Id, uint8 Mipmaps
 {
 	MUTABLE_CPUPROFILER_SCOPE(FUnrealMutableImageProvider::GetImageDesc);
 
-	mu::FImageDesc Result;	
+	FScopeLock Lock(&ExternalImagesLock);
+	// Inside this scope it's safe to access GlobalExternalImages
+
+	if (!GlobalExternalImages.Contains(Id))
 	{
-		FScopeLock Lock(&ExternalImagesLock);
-		// Inside this scope it's safe to access GlobalExternalImages
+		// Null case, no image was provided
+		return CreateDummyDesc();
+	}
 
-		if (!GlobalExternalImages.Contains(Id))
-		{
-			// Null case, no image was provided
-			return CreateDummyDesc();
-		}
+	const FUnrealMutableImageInfo& ImageInfo = GlobalExternalImages[Id];
 
-		const FUnrealMutableImageInfo& ImageInfo = GlobalExternalImages[Id];
+	if (ImageInfo.Image)
+	{
+		// Easy case where the image was directly provided
+		return mu::FImageDesc(ImageInfo.Image->GetSize(), ImageInfo.Image->GetFormat(), ImageInfo.Image->GetLODCount());
+	}
 
-		if (ImageInfo.Image)
-		{
-			// Easy case where the image was directly provided
-			Result = mu::FImageDesc(ImageInfo.Image->GetSize(), ImageInfo.Image->GetFormat(), ImageInfo.Image->GetLODCount());
-		}
-		else if (ImageInfo.TextureToLoad)
-		{
-			// It's safe to access TextureToLoad because ExternalImagesLock guarantees that the data in GlobalExternalImages is valid,
-			// not being modified by the game thread at the moment and the texture cannot be GCed because of the AddReferencedObjects
-			// in the FUnrealMutableImageProvider
-			
-			// Check if it's a format we support
-			if (ImageInfo.MutableFormat == mu::EImageFormat::IF_NONE)
-			{
-				UE_LOG(LogMutable, Warning, TEXT("Failed to get external image descriptor. Unexpected image format. EPixelFormat [%s]."), GetPixelFormatString(ImageInfo.Format));
-				return CreateDummyDesc();
-			}
+#if WITH_EDITOR
+	if (const TSharedPtr<FMutableSourceTextureData>& SourceTextureData = ImageInfo.SourceTextureData)
+	{
+		const FTextureSource& Source = SourceTextureData->GetSource();
 
-			const int32 MipIndex = FMath::Min3(static_cast<int32>(MipmapsToSkip), ImageInfo.FirstAvailableMip, ImageInfo.NumMips - 1);
+		const int32 MipIndex = FMath::Min(static_cast<int32>(MipmapsToSkip),  Source.GetNumMips() - 1);
 
-			const mu::FImageSize ImageSize = mu::FImageSize(
-					ImageInfo.SizeX >> MipIndex,
-					ImageInfo.SizeY >> MipIndex);
+		const mu::FImageSize ImageSize = mu::FImageSize(
+				Source.GetSizeX() >> MipIndex,
+				Source.GetSizeY() >> MipIndex);
 
-			const int32 Lods = 1;
+		const int32 Lods = 1;
+		
+		return mu::FImageDesc(ImageSize, mu::EImageFormat::IF_NONE, Lods);
+	}
+#else
+	if (UTexture2D* TextureToLoad = ImageInfo.TextureToLoad)
+	{
+		mu::FImageDesc Result;	
+
+		// It's safe to access TextureToLoad because ExternalImagesLock guarantees that the data in GlobalExternalImages is valid,
+		// not being modified by the game thread at the moment and the texture cannot be GCed because of the AddReferencedObjects
+		// in the FUnrealMutableImageProvider
+
+		int32 MipIndex = MipmapsToSkip < TextureToLoad->GetPlatformData()->Mips.Num() ? MipmapsToSkip : TextureToLoad->GetPlatformData()->Mips.Num() - 1;
 
-			Result = mu::FImageDesc(ImageSize, ImageInfo.MutableFormat, Lods);
+		// Mips in the mip tail are inlined and can't be streamed, find the smallest mip available.
+		for (; MipIndex > 0; --MipIndex)
+		{
+			if (TextureToLoad->GetPlatformData()->Mips[MipIndex].BulkData.CanLoadFromDisk())
+			{
+				break;
+			}
 		}
-		else
+
+		// Texture format and the equivalent mutable format
+		const EPixelFormat Format = TextureToLoad->GetPlatformData()->PixelFormat;
+		const mu::EImageFormat MutableFormat = GetMutablePixelFormat(Format);
+
+		// Check if it's a format we support
+		if (MutableFormat == mu::EImageFormat::IF_NONE)
 		{
-			// No UTexture2D was provided, cannot do anything, just provide a dummy texture
-			UE_LOG(LogMutable, Warning, TEXT("No UTexture2D was provided for an application-specific image parameter descriptor."));
+			UE_LOG(LogMutable, Warning, TEXT("Failed to get external image descriptor. Unexpected image format. EImageFormat [%s]."), GetPixelFormatString(Format));
 			return CreateDummyDesc();
 		}
 
-		return Result;
+		const mu::FImageSize ImageSize = mu::FImageSize(
+				TextureToLoad->GetSizeX() >> MipIndex,
+				TextureToLoad->GetSizeY() >> MipIndex);
+
+		const int32 Lods = 1;
+
+		return mu::FImageDesc(ImageSize, MutableFormat, Lods);
 	}
+#endif
+	
+	// No UTexture2D was provided, cannot do anything, just provide a dummy texture
+	UE_LOG(LogMutable, Warning, TEXT("No UTexture2D was provided for an application-specific image parameter descriptor."));
+	return CreateDummyDesc();
 }
 
 
@@ -557,44 +492,45 @@ void FUnrealMutableImageProvider::CacheImage(FName Id, bool bUser)
 		}
 		else
 		{
-			mu::ImagePtr pResult;
-			UTexture2D* UnrealDeferredTexture = nullptr;
+			bool bFound = false;
+
+			FUnrealMutableImageInfo ImageInfo;
 
 			// See if any provider provides this id.
-			for (int32 p = 0; !pResult && p < ImageProviders.Num(); ++p)
+			for (const TWeakObjectPtr<UCustomizableSystemImageProvider>& Provider : ImageProviders)
 			{
-				TWeakObjectPtr<UCustomizableSystemImageProvider> Provider = ImageProviders[p];
-
-				if (Provider.IsValid())
+				if (bFound)
 				{
-					// \TODO: all these queries could probably be optimized into a single call.
-					switch (Provider->HasTextureParameterValue(Id))
-					{
+					break;
+				}
+				
+				if (!Provider.IsValid())
+				{
+					continue;
+				}
 
-					case UCustomizableSystemImageProvider::ValueType::Raw:
+				// \TODO: all these queries could probably be optimized into a single call.
+				switch (Provider->HasTextureParameterValue(Id))
+				{
+				case UCustomizableSystemImageProvider::ValueType::Raw:
 					{
 						FIntVector desc = Provider->GetTextureParameterValueSize(Id);
-						pResult = new mu::Image(desc[0], desc[1], 1, mu::EImageFormat::IF_RGBA_UBYTE, mu::EInitializationType::Black);
+						mu::Ptr<mu::Image> pResult = new mu::Image(desc[0], desc[1], 1, mu::EImageFormat::IF_RGBA_UBYTE, mu::EInitializationType::Black);
 						Provider->GetTextureParameterValueData(Id, pResult->GetLODData(0));
+
+						ImageInfo = FUnrealMutableImageInfo(pResult);
+						
+						bFound = true;
 						break;
 					}
 
-					case UCustomizableSystemImageProvider::ValueType::Unreal:
+				case UCustomizableSystemImageProvider::ValueType::Unreal:
 					{
 						UTexture2D* UnrealTexture = Provider->GetTextureParameterValue(Id);
-						pResult = new mu::Image();
+						mu::Ptr<mu::Image> pResult = new mu::Image();
 
 #if WITH_EDITOR
-						FMutableSourceTextureData Tex;
-						Tex.Source = UnrealTexture->Source.CopyTornOff();
-						Tex.bFlipGreenChannel = UnrealTexture->bFlipGreenChannel;
-						Tex.bHasAlphaChannel =
-							UnrealTexture->AdjustMinAlpha != UnrealTexture->AdjustMaxAlpha
-							&& UnrealTexture->CompressionSettings != TextureCompressionSettings::TC_Normalmap
-							&& !UnrealTexture->CompressionNoAlpha;
-						Tex.bCompressionForceAlpha = UnrealTexture->CompressionForceAlpha;
-						Tex.bIsNormalComposite = false; // TODO?
-
+						FMutableSourceTextureData Tex(*UnrealTexture);
 						EUnrealToMutableConversionError Error = ConvertTextureUnrealSourceToMutable(pResult.get(), Tex, 0);
 						if (Error != EUnrealToMutableConversionError::Success)
 						{
@@ -606,39 +542,40 @@ void FUnrealMutableImageProvider::CacheImage(FName Id, bool bUser)
 						ConvertTextureUnrealPlatformToMutable(pResult.get(), UnrealTexture, 0);
 #endif
 
+						ImageInfo = FUnrealMutableImageInfo(pResult);
+
+						bFound = true;
 						break;
 					}
 
-					case UCustomizableSystemImageProvider::ValueType::Unreal_Deferred:
+				case UCustomizableSystemImageProvider::ValueType::Unreal_Deferred:
 					{
-						UnrealDeferredTexture = Provider->GetTextureParameterValue(Id);
-						break;
-					}
+						if (UTexture2D* UnrealDeferredTexture = Provider->GetTextureParameterValue(Id))
+						{
+							ImageInfo = FUnrealMutableImageInfo(*UnrealDeferredTexture);
+
+							bFound = true;
+						}
 
-					default:
 						break;
 					}
+
+				default:
+					break;
 				}
 			}
 
-			if (!pResult && !UnrealDeferredTexture)
+			if (!bFound)
 			{
-				UE_LOG(LogMutable, Warning, TEXT("Failed to cache external image %s. Missing result or source texture. Result [%d]. Unreal texture [%d]. Num providers [%d]"),
-					*Id.ToString(),
-					pResult ? 1 : 0,
-					UnrealDeferredTexture ? 1 : 0,
-					ImageProviders.Num());
-
+				UE_LOG(LogMutable, Warning, TEXT("Failed to cache external image %s. Missing result and source texture."), *Id.ToString());
 				return;
 			}
 			
-			FUnrealMutableImageInfo ImageInfo(pResult, UnrealDeferredTexture);
-
 			if (bUser)
 			{
 				ImageInfo.ReferencesUser = true;			
 			}
-			else
+				else
 			{
 				ImageInfo.ReferencesSystem++;							
 			}
@@ -745,6 +682,34 @@ void FUnrealMutableImageProvider::UnCacheImages(const mu::Parameters& Parameters
 }
 
 
+#if WITH_EDITOR
+void FUnrealMutableImageProvider::CacheRuntimeReferencedImages(const TSharedRef<const mu::Model>& Model, const TArray<TSoftObjectPtr<UTexture2D>>& RuntimeReferencedTextures)
+{
+	check(IsInGameThread());
+	
+	MUTABLE_CPUPROFILER_SCOPE(FUnrealMutableImageProvider::CacheRuntimeReferencedImages);
+	
+	FScopeLock Lock(&RuntimeReferencedLock);
+
+	FRuntimeReferencedImages& ModelImages = RuntimeReferencedImages.Add(&Model.Get());
+	ModelImages.Model = Model.ToWeakPtr();
+
+	ModelImages.SourceTextures.Reset();
+	for (const TSoftObjectPtr<UTexture2D>& RuntimeReferencedTexture : RuntimeReferencedTextures)
+	{
+		UTexture2D* Texture = RuntimeReferencedTexture.LoadSynchronous();
+		if (!Texture)
+		{
+			UE_LOG(LogMutable, Warning, TEXT("Failed to load texture [%s]."), *Texture->GetName());
+			continue;
+		}
+
+		ModelImages.SourceTextures.Emplace(*Texture);
+	}
+}
+#endif
+
+
 mu::ImagePtr FUnrealMutableImageProvider::CreateDummy()
 {
 	// Create a dummy image
@@ -795,6 +760,19 @@ TAutoConsoleVariable<bool> CVarMutableLockExternalImagesDuringGC(
 
 void FUnrealMutableImageProvider::AddReferencedObjects(FReferenceCollector& Collector)
 {
+#if WITH_EDITOR
+	{
+		FScopeLock Lock(&RuntimeReferencedLock);
+
+		for (TMap<const void*, FRuntimeReferencedImages>::TIterator It = RuntimeReferencedImages.CreateIterator(); It; ++It)
+		{
+			if (!It.Value().Model.IsValid())
+			{
+				It.RemoveCurrent();
+			}
+		}
+	}
+#else
 	bool bDoLock = CVarMutableLockExternalImagesDuringGC.GetValueOnAnyThread();
 
 	if (bDoLock)
@@ -814,31 +792,25 @@ void FUnrealMutableImageProvider::AddReferencedObjects(FReferenceCollector& Coll
 	{
 		ExternalImagesLock.Unlock();
 	}
+#endif
 }
 
 
-FUnrealMutableImageProvider::FUnrealMutableImageInfo::FUnrealMutableImageInfo(const mu::ImagePtr& InImage, UTexture2D* InTextureToLoad) :
-	Image(InImage),
-	TextureToLoad(InTextureToLoad)
+FUnrealMutableImageProvider::FUnrealMutableImageInfo::FUnrealMutableImageInfo(const mu::ImagePtr& InImage)
 {
 	check(IsInGameThread())
 	
-	FTexturePlatformData* PlatformData = TextureToLoad->GetPlatformData();
-	
-	Format = PlatformData->PixelFormat;
-	MutableFormat = GetMutablePixelFormat(Format);
+	Image = InImage;
+}
 
-	NumMips = PlatformData->Mips.Num();
-	
-	// Mips in the mip tail are inlined and can't be streamed, find the smallest mip available.
-	for (FirstAvailableMip = NumMips - 1; FirstAvailableMip > 0; --FirstAvailableMip)
-	{
-		if (PlatformData->Mips[FirstAvailableMip].BulkData.CanLoadFromDisk())
-		{
-			break;
-		}
-	}
 
-	SizeX = TextureToLoad->GetSizeX();
-	SizeY = TextureToLoad->GetSizeY();
+FUnrealMutableImageProvider::FUnrealMutableImageInfo::FUnrealMutableImageInfo(UTexture2D& Texture)
+{
+	check(IsInGameThread())
+
+#if WITH_EDITOR
+	SourceTextureData = MakeShared<FMutableSourceTextureData>(Texture);
+#else
+	TextureToLoad = &Texture;
+#endif
 }
diff --git a/Engine/Plugins/Experimental/Mutable/Source/CustomizableObject/Private/MuCO/UnrealMutableImageProvider.h b/Engine/Plugins/Experimental/Mutable/Source/CustomizableObject/Private/MuCO/UnrealMutableImageProvider.h
--- a/Engine/Plugins/Experimental/Mutable/Source/CustomizableObject/Private/MuCO/UnrealMutableImageProvider.h
+++ b/Engine/Plugins/Experimental/Mutable/Source/CustomizableObject/Private/MuCO/UnrealMutableImageProvider.h
@@ -9,9 +9,11 @@
 #include "Containers/Map.h"
 #include "Containers/Queue.h"
 #include "PixelFormat.h"
+#include "MuCO/UnrealToMutableTextureConversionUtils.h"
 
 #include "Tasks/Task.h"
 
+class UCustomizableObject;
 class UTexture2D;
 
 
@@ -44,6 +46,10 @@ public:
 
 	void CacheImages(const mu::Parameters& Parameters);
 	void UnCacheImages(const mu::Parameters& Parameters);
+
+#if WITH_EDITOR
+	void CacheRuntimeReferencedImages(const TSharedRef<const mu::Model>& Model, const TArray<TSoftObjectPtr<UTexture2D>>& RuntimeReferencedTextures);
+#endif
 	
 	/** List of actual image providers that have been registered to the CustomizableObjectSystem. */
 	TArray< TWeakObjectPtr<class UCustomizableSystemImageProvider> > ImageProviders;
@@ -55,23 +61,23 @@ public:
 	}
 
 private:
-
 	struct FUnrealMutableImageInfo
 	{
-		FUnrealMutableImageInfo() {}
-		FUnrealMutableImageInfo(const mu::ImagePtr& InImage, UTexture2D* InTextureToLoad);
+		FUnrealMutableImageInfo() = default;
+
+		FUnrealMutableImageInfo(const mu::ImagePtr& InImage);
+
+		FUnrealMutableImageInfo(UTexture2D& Texture);
 
 		mu::ImagePtr Image;
 
+#if WITH_EDITOR
+		TSharedPtr<FMutableSourceTextureData> SourceTextureData;
+#else
 		/** If the above Image has not been loaded in the game thread, the TextureToLoad bulk data will be loaded
 		* from the Mutable thread when it's needed */
 		TObjectPtr<UTexture2D> TextureToLoad = nullptr;
-		EPixelFormat Format = PF_Unknown;
-		mu::EImageFormat MutableFormat = mu::EImageFormat::IF_NONE;
-		int32 NumMips = 0;
-		int32 FirstAvailableMip = -1;
-		int32 SizeX = 0;
-		int32 SizeY = 0;
+#endif
 		
 		/** true of the reference maintained by the user. */
 		bool ReferencesUser = false;
@@ -89,36 +95,24 @@ private:
 
 	/** Map of Ids to external textures that may be required for any instance or Mutable texture mip under construction.
 	* This is only safely written from the game thread protected by the following critical section, and it
-	* is safely read from the mutable thread during the update of the instance or texture mips
-	*/
+	* is safely read from the mutable thread during the update of the instance or texture mips. */
 	TMap<FName, FUnrealMutableImageInfo> GlobalExternalImages;
+
+#if WITH_EDITOR
+	struct FRuntimeReferencedImages
+	{
+		TArray<FMutableSourceTextureData> SourceTextures;
+		TWeakPtr<const mu::Model> Model;
+	};
+	
+	TMap<const void*, FRuntimeReferencedImages> RuntimeReferencedImages;
+#endif
 	
 	/** Access to GlobalExternalImages must be protected with this because it may be accessed concurrently from the 
 	Game thread to modify it and from the Mutable thread to read it. */
 	FCriticalSection ExternalImagesLock;
 
 #if WITH_EDITOR
-public:
-	/** The provider may tick to support referenced images. */
-	bool Tick();
-
-private:
-	struct FReferencedImageRequest
-	{
-		FReferencedImageRequest(const UE::Tasks::FTaskEvent&& Event) : CompletionEvent(Event) {}
-
-		// Sync
-		UE::Tasks::FTaskEvent CompletionEvent;
-
-		// Input
-		int32 Id = 0;
-		uint8 MipmapsToSkip = 0;
-		const void* ModelPtr = nullptr;
-
-		// Result
-		mu::Ptr<mu::Image> ResultImage;
-	};
-	TQueue<FReferencedImageRequest*, EQueueMode::Spsc> QueuedReferencedImageRequests;
+	FCriticalSection RuntimeReferencedLock;
 #endif
-
 };
diff --git a/Engine/Plugins/Experimental/Mutable/Source/CustomizableObject/Private/MuCO/UnrealToMutableTextureConversionUtils.cpp b/Engine/Plugins/Experimental/Mutable/Source/CustomizableObject/Private/MuCO/UnrealToMutableTextureConversionUtils.cpp
--- a/Engine/Plugins/Experimental/Mutable/Source/CustomizableObject/Private/MuCO/UnrealToMutableTextureConversionUtils.cpp
+++ b/Engine/Plugins/Experimental/Mutable/Source/CustomizableObject/Private/MuCO/UnrealToMutableTextureConversionUtils.cpp
@@ -173,21 +173,67 @@ void BlurNormalForComposite(FImage& Image)
 } //namespace UnrealToMutableImageConversion_Internal
 
 
+FMutableSourceTextureData::FMutableSourceTextureData(const UTexture& Texture)
+{
+	check(IsInGameThread());
+		
+	Source = Texture.Source.CopyTornOff();
+	bFlipGreenChannel = Texture.bFlipGreenChannel;
+	bHasAlphaChannel = Texture.AdjustMinAlpha != Texture.AdjustMaxAlpha &&
+		Texture.CompressionSettings != TC_Normalmap &&
+		!Texture.CompressionNoAlpha;
+	bCompressionForceAlpha = Texture.CompressionForceAlpha;
+	bIsNormalComposite = false; // TODO?
+}
+
+
+FTextureSource& FMutableSourceTextureData::GetSource()
+{
+	return Source;	
+}
+
+
+bool FMutableSourceTextureData::GetFlipGreenChannel() const
+{
+	return bFlipGreenChannel;
+}
+
+
+bool FMutableSourceTextureData::HasAlphaChannel() const
+{
+	return bHasAlphaChannel;
+}
+
+
+bool FMutableSourceTextureData::GetCompressionForceAlpha() const
+{
+	return bCompressionForceAlpha;
+}
+
+
+bool FMutableSourceTextureData::IsNormalComposite() const
+{
+	return bIsNormalComposite;
+}
+
+
 EUnrealToMutableConversionError ConvertTextureUnrealSourceToMutable(mu::Image* OutResult, FMutableSourceTextureData& Tex, uint8 MipmapsToSkip)
 {
 	MUTABLE_CPUPROFILER_SCOPE(ConvertTextureUnrealSourceToMutable);
 
     using namespace UnrealToMutableImageConversion_Internal;
 
+	FTextureSource& Source = Tex.GetSource();
+	
 	// Correct mips to skip to fit source data
-	MipmapsToSkip = FMath::Clamp(MipmapsToSkip, 0, Tex.Source.GetNumMips()-1);
+	MipmapsToSkip = FMath::Clamp(MipmapsToSkip, 0, Source.GetNumMips()-1);
 
     const int32 LODs = 1;
-    const int32 SizeX = Tex.Source.GetSizeX() >> MipmapsToSkip;
-    const int32 SizeY = Tex.Source.GetSizeY() >> MipmapsToSkip;
+    const int32 SizeX = Source.GetSizeX() >> MipmapsToSkip;
+    const int32 SizeY = Source.GetSizeY() >> MipmapsToSkip;
 	check(SizeX > 0 && SizeY > 0);
 
-    ETextureSourceFormat Format = Tex.Source.GetFormat();
+    ETextureSourceFormat Format = Source.GetFormat();
  
     ERawImageFormat::Type RawFormat = ConvertFormatSourceToRaw(Format);
 
@@ -195,15 +241,15 @@ EUnrealToMutableConversionError ConvertTextureUnrealSourceToMutable(mu::Image* O
     FImage TempImage(SizeX, SizeY, 1, RawFormat, EGammaSpace::Linear);
 	FImage TempImage2;
 
-    if (!Tex.Source.GetMipData(TempImage.RawData, MipmapsToSkip))
+    if (!Source.GetMipData(TempImage.RawData, MipmapsToSkip))
     {
         return EUnrealToMutableConversionError::Unknown;
     }
 
-    bool bFlipGreenChannel = Tex.bFlipGreenChannel;
+    bool bFlipGreenChannel = Tex.GetFlipGreenChannel();
 
     // If any post processes of the image is needed, convert to RGBA32F
-    if (Tex.bIsNormalComposite)
+    if (Tex.IsNormalComposite())
     { 
 		MUTABLE_CPUPROFILER_SCOPE(FlipOrComposite);
 
@@ -257,8 +303,8 @@ EUnrealToMutableConversionError ConvertTextureUnrealSourceToMutable(mu::Image* O
 	case ERawImageFormat::BGRA8:
 	{
 		// Try to find out if the texture has and actually makes use of the alpha channel
-		bool bHasAlphaChannel = Tex.bHasAlphaChannel
-			&& (Tex.bCompressionForceAlpha
+		bool bHasAlphaChannel = Tex.HasAlphaChannel()
+			&& (Tex.GetCompressionForceAlpha()
 				||
 				FImageCore::DetectAlphaChannel(TempImage));
 
diff --git a/Engine/Plugins/Experimental/Mutable/Source/CustomizableObjectEditor/Private/MuCOE/CustomizableObjectCompiler.cpp b/Engine/Plugins/Experimental/Mutable/Source/CustomizableObjectEditor/Private/MuCOE/CustomizableObjectCompiler.cpp
--- a/Engine/Plugins/Experimental/Mutable/Source/CustomizableObjectEditor/Private/MuCOE/CustomizableObjectCompiler.cpp
+++ b/Engine/Plugins/Experimental/Mutable/Source/CustomizableObjectEditor/Private/MuCOE/CustomizableObjectCompiler.cpp
@@ -923,12 +923,18 @@ void FCustomizableObjectCompiler::CompileInternal(bool bAsync)
 			ModelResources.Materials.Emplace(Material);
 		}
 
-		for (const TPair<TSoftObjectPtr<UTexture>, FMutableGraphGenerationContext::FGeneratedReferencedTexture>& Pair : GenerationContext.RuntimeReferencedTextureMap)
+		for (const TPair<TSoftObjectPtr<UTexture>, FMutableGraphGenerationContext::FGeneratedReferencedTexture>& Pair : GenerationContext.PassthroughTextureMap)
 		{
 			check(Pair.Value.ID == ModelResources.PassThroughTextures.Num());
 			ModelResources.PassThroughTextures.Add(Pair.Key);
 		}
 
+		for (const TPair<TSoftObjectPtr<UTexture2D>, FMutableGraphGenerationContext::FGeneratedReferencedTexture>& Pair : GenerationContext.RuntimeReferencedTextureMap)
+		{
+			check(Pair.Value.ID == ModelResources.RuntimeReferencedTextures.Num());
+			ModelResources.RuntimeReferencedTextures.Add(Pair.Key);
+		}
+		
 		ModelResources.PhysicsAssets = MoveTemp(GenerationContext.PhysicsAssets);
 
 		ModelResources.AnimBPs = MoveTemp(GenerationContext.AnimBPAssets);
@@ -1112,17 +1118,7 @@ void FCustomizableObjectCompiler::CompileInternal(bool bAsync)
 		{
 			check(Pair.Value.ID == NewCompileTimeReferencedTextures.Num());
 
-			UTexture* Texture = Pair.Key.LoadSynchronous();
-			FMutableSourceTextureData Tex;
-			Tex.Source = Texture->Source.CopyTornOff();
-			Tex.bFlipGreenChannel = Texture->bFlipGreenChannel;
-			Tex.bHasAlphaChannel =
-				Texture->AdjustMinAlpha != Texture->AdjustMaxAlpha
-				&& Texture->CompressionSettings != TextureCompressionSettings::TC_Normalmap
-				&& !Texture->CompressionNoAlpha;
-			Tex.bCompressionForceAlpha = Texture->CompressionForceAlpha;
-			Tex.bIsNormalComposite = false; // TODO?
-
+			FMutableSourceTextureData Tex(*Pair.Key.LoadSynchronous());
 			NewCompileTimeReferencedTextures.Add(Tex);
 		}
 
@@ -1326,7 +1322,7 @@ mu::NodePtr FCustomizableObjectCompiler::Export(UCustomizableObject* Object, con
 
 	// Pass out the references textures
 	OutRuntimeReferencedTextures.Empty();
-	for (const TPair<TSoftObjectPtr<UTexture>, FMutableGraphGenerationContext::FGeneratedReferencedTexture>& Pair : GenerationContext.RuntimeReferencedTextureMap)
+	for (const TPair<TSoftObjectPtr<UTexture2D>, FMutableGraphGenerationContext::FGeneratedReferencedTexture>& Pair : GenerationContext.RuntimeReferencedTextureMap)
 	{
 		check(Pair.Value.ID == OutRuntimeReferencedTextures.Num());
 		OutRuntimeReferencedTextures.Add(Pair.Key);
@@ -1337,17 +1333,7 @@ mu::NodePtr FCustomizableObjectCompiler::Export(UCustomizableObject* Object, con
 	{
 		check(Pair.Value.ID == OutCompilerReferencedTextures.Num());
 
-		UTexture* Texture = Pair.Key.LoadSynchronous();
-		FMutableSourceTextureData Tex;
-		Tex.Source = Texture->Source.CopyTornOff();
-		Tex.bFlipGreenChannel = Texture->bFlipGreenChannel;
-		Tex.bHasAlphaChannel =
-			Texture->AdjustMinAlpha != Texture->AdjustMaxAlpha
-			&& Texture->CompressionSettings != TextureCompressionSettings::TC_Normalmap
-			&& !Texture->CompressionNoAlpha;
-		Tex.bCompressionForceAlpha = Texture->CompressionForceAlpha;
-		Tex.bIsNormalComposite = false; // TODO?
-
+		FMutableSourceTextureData Tex(*Pair.Key.LoadSynchronous());
 		OutCompilerReferencedTextures.Add(Tex);
 	}
 
diff --git a/Engine/Plugins/Experimental/Mutable/Source/CustomizableObjectEditor/Private/MuCOE/GenerateMutableSource/GenerateMutableSource.cpp b/Engine/Plugins/Experimental/Mutable/Source/CustomizableObjectEditor/Private/MuCOE/GenerateMutableSource/GenerateMutableSource.cpp
--- a/Engine/Plugins/Experimental/Mutable/Source/CustomizableObjectEditor/Private/MuCOE/GenerateMutableSource/GenerateMutableSource.cpp
+++ b/Engine/Plugins/Experimental/Mutable/Source/CustomizableObjectEditor/Private/MuCOE/GenerateMutableSource/GenerateMutableSource.cpp
@@ -1916,26 +1916,40 @@ mu::Ptr<mu::Image> GenerateImageConstant(UTexture* Texture, FMutableGraphGenerat
 		}
 	}
 
-	TMap<TSoftObjectPtr<UTexture>, FMutableGraphGenerationContext::FGeneratedReferencedTexture>& TextureMap = bIsCompileTime
-		? GenerationContext.CompileTimeTextureMap
-		: GenerationContext.RuntimeReferencedTextureMap;
+	// Create a descriptor for the image.
+	// \TODO: If passthrough (bIsReference) we should apply lod bias, and max texture size to this desc.
+	// For now it is not a problem because passthrough textures shouldn't mix with any other operations.
+	mu::FImageDesc ImageDesc = GenerateImageDescriptor(Texture);
 
 	FMutableGraphGenerationContext::FGeneratedReferencedTexture InvalidEntry;
 	InvalidEntry.ID = TNumericLimits<uint32>::Max();
-	FMutableGraphGenerationContext::FGeneratedReferencedTexture& Entry = TextureMap.FindOrAdd(Texture,InvalidEntry);
 
-	if (Entry.ID == TNumericLimits<uint32>::Max())
+	FMutableGraphGenerationContext::FGeneratedReferencedTexture* Entry;
+	int32 Num = 0;
+	
+	if (bIsReference)
 	{
-		Entry.ID = TextureMap.Num()-1;
+		Entry = &GenerationContext.PassthroughTextureMap.FindOrAdd(Texture, InvalidEntry);
+		Num = GenerationContext.PassthroughTextureMap.Num();
+	}
+	else if (bIsCompileTime)
+	{
+        Entry = &GenerationContext.CompileTimeTextureMap.FindOrAdd(Texture, InvalidEntry);
+		Num = GenerationContext.CompileTimeTextureMap.Num();
+	}
+	else
+	{
+		Entry = &GenerationContext.RuntimeReferencedTextureMap.FindOrAdd(Texture, InvalidEntry);
+		Num = GenerationContext.RuntimeReferencedTextureMap.Num();
 	}
 
-	// Create a descriptor for the image.
-	// \TODO: If passthrough (bIsReference) we should apply lod bias, and max texture size to this desc.
-	// For now it is not a problem because passthrough textures shouldn't mix with any other operations.
-	mu::FImageDesc ImageDesc = GenerateImageDescriptor(Texture);
-
+	if (Entry->ID == TNumericLimits<uint32>::Max())
+	{
+		Entry->ID = Num - 1;
+	}
+	
 	// Compile-time references that are left should be resolved immediately (should only happen in editor).
-	mu::Ptr<mu::Image> Result = mu::Image::CreateAsReference(Entry.ID, ImageDesc, bForceLoad);
+	mu::Ptr<mu::Image> Result = mu::Image::CreateAsReference(Entry->ID, ImageDesc, bForceLoad);
 	return Result;
 }
 
diff --git a/Engine/Plugins/Experimental/Mutable/Source/CustomizableObjectEditor/Private/MuCOE/GenerateMutableSource/GenerateMutableSource.h b/Engine/Plugins/Experimental/Mutable/Source/CustomizableObjectEditor/Private/MuCOE/GenerateMutableSource/GenerateMutableSource.h
--- a/Engine/Plugins/Experimental/Mutable/Source/CustomizableObjectEditor/Private/MuCOE/GenerateMutableSource/GenerateMutableSource.h
+++ b/Engine/Plugins/Experimental/Mutable/Source/CustomizableObjectEditor/Private/MuCOE/GenerateMutableSource/GenerateMutableSource.h
@@ -584,8 +584,12 @@ struct FMutableGraphGenerationContext
 
 	// Cache of runtime pass-through images and their IDs used in the core to indentify them.
 	// These textures will remain as external references even in optimized models.
-	TMap<TSoftObjectPtr<UTexture>, FGeneratedReferencedTexture> RuntimeReferencedTextureMap;
-
+	TMap<TSoftObjectPtr<UTexture>, FGeneratedReferencedTexture> PassthroughTextureMap;
+	
+	// Cache of runtime images and their IDs used in the core to indentify them.
+	// These textures will remain as external references even in optimized models.
+	TMap<TSoftObjectPtr<UTexture2D>, FGeneratedReferencedTexture> RuntimeReferencedTextureMap;
+	
 	// Cache of runtime pass-through images and their IDs used in the core to indentify them
 	// These textures will become mutable images in the compiled model.
 	TMap<TSoftObjectPtr<UTexture>, FGeneratedReferencedTexture> CompileTimeTextureMap;
diff --git a/Engine/Plugins/Experimental/Mutable/Source/CustomizableObjectEditor/Private/MuCOE/SMutableCodeViewer.cpp b/Engine/Plugins/Experimental/Mutable/Source/CustomizableObjectEditor/Private/MuCOE/SMutableCodeViewer.cpp
--- a/Engine/Plugins/Experimental/Mutable/Source/CustomizableObjectEditor/Private/MuCOE/SMutableCodeViewer.cpp
+++ b/Engine/Plugins/Experimental/Mutable/Source/CustomizableObjectEditor/Private/MuCOE/SMutableCodeViewer.cpp
@@ -2506,16 +2506,7 @@ namespace
 
 			mu::Ptr<mu::Image> ResultImage = new mu::Image();
 
-			FMutableSourceTextureData Tex;
-			Tex.Source = Texture->Source.CopyTornOff();
-			Tex.bFlipGreenChannel = Texture->bFlipGreenChannel;
-			Tex.bHasAlphaChannel =
-				Texture->AdjustMinAlpha != Texture->AdjustMaxAlpha
-				&& Texture->CompressionSettings != TextureCompressionSettings::TC_Normalmap
-				&& !Texture->CompressionNoAlpha;
-			Tex.bCompressionForceAlpha = Texture->CompressionForceAlpha;
-			Tex.bIsNormalComposite = false; // TODO?
-
+			FMutableSourceTextureData Tex(*Texture);
 			EUnrealToMutableConversionError Error = ConvertTextureUnrealSourceToMutable(ResultImage.get(), Tex, MipmapsToSkip);
 
 			if (Error != EUnrealToMutableConversionError::Success)
