[Mutable] Fix LOD update candidates not being discarded when the requested LODs are a subset of the generated ones.

#rnx

diff --git a/Engine/Plugins/Experimental/Mutable/Source/CustomizableObject/Private/MuCO/CustomizableObjectSystem.cpp b/Engine/Plugins/Experimental/Mutable/Source/CustomizableObject/Private/MuCO/CustomizableObjectSystem.cpp
--- a/Engine/Plugins/Experimental/Mutable/Source/CustomizableObject/Private/MuCO/CustomizableObjectSystem.cpp
+++ b/Engine/Plugins/Experimental/Mutable/Source/CustomizableObject/Private/MuCO/CustomizableObjectSystem.cpp
@@ -3590,10 +3590,17 @@ int32 UCustomizableObjectSystem::TickInternal()
 
 					if (LODUpdateCandidate.Priority <= MaxPriorityFound)
 					{
-						if (LODUpdateCandidate.CustomizableObjectInstance->GetPrivate()->MinSquareDistFromComponentToPlayer < MaxSquareDistanceFound)
+						UCustomizableInstancePrivate* CustomizableInstancePrivate = LODUpdateCandidate.CustomizableObjectInstance->GetPrivate();
+
+						FDescriptorHash LODUpdateDescriptorHash = CustomizableInstancePrivate->CommittedDescriptorHash;
+						LODUpdateDescriptorHash.MinLOD = LODUpdateCandidate.MinLOD;
+						LODUpdateDescriptorHash.RequestedLODsPerComponent = LODUpdateCandidate.RequestedLODLevels;
+
+						if (CustomizableInstancePrivate->MinSquareDistFromComponentToPlayer < MaxSquareDistanceFound &&
+							!LODUpdateDescriptorHash.IsSubset(CustomizableInstancePrivate->CommittedDescriptorHash))
 						{
 							MaxPriorityFound = LODUpdateCandidate.Priority;
-							MaxSquareDistanceFound = LODUpdateCandidate.CustomizableObjectInstance->GetPrivate()->MinSquareDistFromComponentToPlayer;
+							MaxSquareDistanceFound = CustomizableInstancePrivate->MinSquareDistFromComponentToPlayer;
 							PendingInstanceUpdateFound = nullptr;
 							LODUpdateCandidateFound = &LODUpdateCandidate;
 						}
