[Mutable] Restored GetAnimationGameplayTags due to being used in some blueprints.

#rnx

diff --git a/Engine/Plugins/Experimental/Mutable/Source/CustomizableObject/Public/MuCO/CustomizableObjectInstanceAssetUserData.h b/Engine/Plugins/Experimental/Mutable/Source/CustomizableObject/Public/MuCO/CustomizableObjectInstanceAssetUserData.h
--- a/Engine/Plugins/Experimental/Mutable/Source/CustomizableObject/Public/MuCO/CustomizableObjectInstanceAssetUserData.h
+++ b/Engine/Plugins/Experimental/Mutable/Source/CustomizableObject/Public/MuCO/CustomizableObjectInstanceAssetUserData.h
@@ -30,6 +30,20 @@ class CUSTOMIZABLEOBJECT_API UCustomizableObjectInstanceUserData : public UAsset
 	GENERATED_BODY()
 
 public:
+	/** Return the list of tags for this instance. */
+	UFUNCTION(BlueprintCallable, Category = CustomizableObjectInstance)
+	const FGameplayTagContainer& GetAnimationGameplayTags() const
+	{
+		return AnimationGameplayTag;
+	};
+
+	/** Sets the list of tags for this instance. */
+	UFUNCTION(BlueprintCallable, Category = CustomizableObjectInstance)
+	void SetAnimationGameplayTags(const FGameplayTagContainer& InstanceTags)
+	{
+		AnimationGameplayTag = InstanceTags;
+	};
+	
 	UPROPERTY(BlueprintReadWrite, Category = CustomizableObjectInstance)
 	FGameplayTagContainer AnimationGameplayTag;
 
