[Mutable] Fix mesh reshape normals and add the option in the editor node to regenerate them.

- The old algorithm rotated normals using the shape mesh faces, this can imprint the underlying shape to the normal direction. By reshaping a point representing the normal (vertex positon + normal) and recomputing the final normal using the reshaped position very nice results are achieved. This adds 8 bytes to the binding data structures.

- Changed the projection vertex binding method so it gives more controllable results. Before, the distance along the interpolated normal was used to reconstruct the point but that tuns out to not be correct, we need to store the distance along the 3 vertex normals. This uses 8 more bytes per vertex but the quality of the deformation is improved.

- Added mesh normal generation algorithm using triangle area weighted normals.

#rnx

diff --git a/Engine/Plugins/Experimental/Mutable/Source/CustomizableObjectEditor/Classes/MuCOE/Nodes/CustomizableObjectNodeMeshReshape.h b/Engine/Plugins/Experimental/Mutable/Source/CustomizableObjectEditor/Classes/MuCOE/Nodes/CustomizableObjectNodeMeshReshape.h
--- a/Engine/Plugins/Experimental/Mutable/Source/CustomizableObjectEditor/Classes/MuCOE/Nodes/CustomizableObjectNodeMeshReshape.h
+++ b/Engine/Plugins/Experimental/Mutable/Source/CustomizableObjectEditor/Classes/MuCOE/Nodes/CustomizableObjectNodeMeshReshape.h
@@ -75,7 +75,10 @@ public:
 
 	/** Enable the deformation of the vertices of the base mesh. */
 	UPROPERTY(EditAnywhere, Category = CustomizableObject)
-	bool bReshapeVertices = true;	
+	bool bReshapeVertices = true;
+
+	UPROPERTY(EditAnywhere, Category = CustomizableObject)
+	bool bRecomputeNormals = false;
 
 	/** Enable laplacian smoothing to the result of the base mesh reshape. */
 	UPROPERTY(EditAnywhere, Category = CustomizableObject, meta = (EditCondition = "bReshapeVertices"))
@@ -83,7 +86,7 @@ public:
 
 	/** Enable the deformation of the skeleton of the base mesh. */
 	UPROPERTY(EditAnywhere, Category = CustomizableObject)
-	bool bReshapeSkeleton = false;	
+	bool bReshapeSkeleton = false;
 
 	/** Enable the deformation of physics volumes of the base mesh */
     UPROPERTY(EditAnywhere, Category = CustomizableObject)
diff --git a/Engine/Plugins/Experimental/Mutable/Source/CustomizableObjectEditor/Private/MuCOE/GenerateMutableSource/GenerateMutableSourceMesh.cpp b/Engine/Plugins/Experimental/Mutable/Source/CustomizableObjectEditor/Private/MuCOE/GenerateMutableSource/GenerateMutableSourceMesh.cpp
--- a/Engine/Plugins/Experimental/Mutable/Source/CustomizableObjectEditor/Private/MuCOE/GenerateMutableSource/GenerateMutableSourceMesh.cpp
+++ b/Engine/Plugins/Experimental/Mutable/Source/CustomizableObjectEditor/Private/MuCOE/GenerateMutableSource/GenerateMutableSourceMesh.cpp
@@ -3682,6 +3682,7 @@ mu::NodeMeshPtr GenerateMutableSourceMesh(const UEdGraphPin* Pin,
 		{
 
 			MeshNode->SetReshapeVertices(TypedNodeReshape->bReshapeVertices);
+			MeshNode->SetRecomputeNormals(TypedNodeReshape->bRecomputeNormals);
 			MeshNode->SetApplyLaplacian(TypedNodeReshape->bApplyLaplacianSmoothing);
 			MeshNode->SetReshapeSkeleton(TypedNodeReshape->bReshapeSkeleton);
 			MeshNode->SetReshapePhysicsVolumes(TypedNodeReshape->bReshapePhysicsVolumes);
diff --git a/Engine/Plugins/Experimental/Mutable/Source/MutableRuntime/Internal/MuR/Operations.h b/Engine/Plugins/Experimental/Mutable/Source/MutableRuntime/Internal/MuR/Operations.h
--- a/Engine/Plugins/Experimental/Mutable/Source/MutableRuntime/Internal/MuR/Operations.h
+++ b/Engine/Plugins/Experimental/Mutable/Source/MutableRuntime/Internal/MuR/Operations.h
@@ -357,6 +357,7 @@ namespace mu
 		ReshapePhysicsVolumes  = 1 << 4,
 		ReshapeVertices		   = 1 << 5,
 		ApplyLaplacian		   = 1 << 6,
+		RecomputeNormals	   = 1 << 7,
 	};
 	ENUM_CLASS_FLAGS(EMeshBindShapeFlags);
 
diff --git a/Engine/Plugins/Experimental/Mutable/Source/MutableRuntime/Private/MuR/OpMeshBind.h b/Engine/Plugins/Experimental/Mutable/Source/MutableRuntime/Private/MuR/OpMeshBind.h
--- a/Engine/Plugins/Experimental/Mutable/Source/MutableRuntime/Private/MuR/OpMeshBind.h
+++ b/Engine/Plugins/Experimental/Mutable/Source/MutableRuntime/Private/MuR/OpMeshBind.h
@@ -16,7 +16,6 @@
 
 #include "MuR/OpMeshSmoothing.h"
 
-
 // TODO: Make the handling of rotations an option. It is more expensive on CPU and memory, and for some
 // cases it is not required at all.
 
@@ -113,15 +112,15 @@ namespace mu
 		// Barycentric coordinates on the shape triangle
 		float S, T;
 
-		// Distance along the interpolated normal of the shape triangle.
-		float D;
+		// Distance along the normals of the shape triangle.
+		FVector3f D;
 
 		// Index of the shape triangle.
 		int32 Triangle;
 
-
 		// Used to calculate the rotation to apply to the reshaped vertex tangent space.
-		FVector3f ShapeNormal;
+		float NS, NT;
+		FVector3f NormalD;
 		
 		// Bind point, if the point belongs to a rigid cluster, the attachment point, otherwise
 		// the original point.
@@ -131,18 +130,18 @@ namespace mu
 		// full effect. This weight should be proportional to the confidence we have that the binding data is valid. 
 		float Weight;
 	};
-	static_assert(sizeof(FReshapeVertexBindingData) == 44);
+	static_assert(sizeof(FReshapeVertexBindingData) == 4*15);
 
 	struct FReshapeVertexBindingDataBufferDescriptor
 	{
-		constexpr static int ElementSize = sizeof(FReshapeVertexBindingData);
-		constexpr static int Channels = 4;
+		constexpr static int32 ElementSize = sizeof(FReshapeVertexBindingData);
+		constexpr static int32 Channels = 4;
 		constexpr static EMeshBufferSemantic Semantics[Channels] = { MBS_BARYCENTRICCOORDS, MBS_DISTANCE, MBS_TRIANGLEINDEX, MBS_OTHER };
 		constexpr static EMeshBufferFormat Formats[Channels] = { MBF_FLOAT32, MBF_FLOAT32, MBF_INT32, MBF_FLOAT32 };
-		constexpr static int Components[Channels] = { 2, 1, 1, 3+3+1 };
-		constexpr static int Offsets[Channels] = { 0, 8, 12, 16 };
+		constexpr static int32 Components[Channels] = { 2, 3, 1, 2+3+3+1 };
+		constexpr static int32 Offsets[Channels] = { 0, 4*2, (2+3)*4, (2+3+1)*4};
 
-		const int SemanticIndices[Channels] = { 0, 0, 0, 0 };
+		const int32 SemanticIndices[Channels] = { 0, 0, 0, 0 };
 
 		FReshapeVertexBindingDataBufferDescriptor(int32 DataSetIndex)
 			: SemanticIndices{ DataSetIndex, DataSetIndex, DataSetIndex, DataSetIndex }
@@ -153,23 +152,23 @@ namespace mu
 	struct FReshapePointBindingData
 	{
 		float S, T;
-		float D;
+		FVector3f D;
 		int32 Triangle;
 		float Weight;
 	};
 
-	static_assert(sizeof(FReshapePointBindingData) == 4*5);
+	static_assert(sizeof(FReshapePointBindingData) == 24+4);
 
 	struct FReshapePointBindingDataBufferDescriptor
 	{
-		constexpr static int ElementSize = sizeof(FReshapeVertexBindingData);
-		constexpr static int Channels = 4;
+		constexpr static int32 ElementSize = sizeof(FReshapeVertexBindingData);
+		constexpr static int32 Channels = 4;
 		constexpr static EMeshBufferSemantic Semantics[Channels] = { MBS_BARYCENTRICCOORDS, MBS_DISTANCE, MBS_TRIANGLEINDEX, MBS_OTHER };
 		constexpr static EMeshBufferFormat Formats[Channels] = { MBF_FLOAT32, MBF_FLOAT32, MBF_INT32, MBF_FLOAT32 };
-		constexpr static int Components[Channels] = { 2, 1, 1, 1 };
-		constexpr static int Offsets[Channels] = { 0, 8, 12, 16 };
+		constexpr static int32 Components[Channels] = { 2, 3, 1, 1 };
+		constexpr static int32 Offsets[Channels] = { 0, 8, 20, 24 };
 
-		const int SemanticIndices[Channels] = { 0, 0, 0, 0 };
+		const int32 SemanticIndices[Channels] = { 0, 0, 0, 0 };
 
 		FReshapePointBindingDataBufferDescriptor(int32 DataSetIndex)
 			: SemanticIndices{ DataSetIndex, DataSetIndex, DataSetIndex, DataSetIndex }
@@ -179,13 +178,13 @@ namespace mu
 	
 	struct FIntBufferDescriptor
 	{
-		constexpr static int ElementSize = sizeof(int32);
-		constexpr static int Channels = 1;
+		constexpr static int32 ElementSize = sizeof(int32);
+		constexpr static int32 Channels = 1;
 		constexpr static EMeshBufferSemantic Semantics[Channels] = { MBS_OTHER };
 		constexpr static EMeshBufferFormat Formats[Channels] = { MBF_INT32 };
-		constexpr static int Components[Channels] = { 1 };
-		constexpr static int Offsets[Channels] = { 0 };
-		constexpr static int SemanticIndices[Channels] = { 0 };
+		constexpr static int32 Components[Channels] = { 1 };
+		constexpr static int32 Offsets[Channels] = { 0 };
+		constexpr static int32 SemanticIndices[Channels] = { 0 };
 
 		FIntBufferDescriptor()
 		{
@@ -210,13 +209,13 @@ namespace mu
 
 	struct FClipDeformVertexBindingDataBufferDescriptor
 	{
-		constexpr static int Channels = 3;
+		constexpr static int32 Channels = 3;
 		constexpr static EMeshBufferSemantic Semantics[Channels] = { MBS_BARYCENTRICCOORDS, MBS_TRIANGLEINDEX, MBS_OTHER };
 		constexpr static EMeshBufferFormat Formats[Channels] = { MBF_FLOAT32, MBF_INT32, MBF_FLOAT32 };
-		constexpr static int Components[Channels] = { 2, 1, 1 };
-		constexpr static int Offsets[Channels] = { 0, 8, 12 };
+		constexpr static int32 Components[Channels] = { 2, 1, 1 };
+		constexpr static int32 Offsets[Channels] = { 0, 8, 12 };
 
-		const int SemanticIndices[Channels] = { 0, 0, 0 };
+		const int32 SemanticIndices[Channels] = { 0, 0, 0 };
 
 		FClipDeformVertexBindingDataBufferDescriptor(int32 DataSetIndex)
 			: SemanticIndices{ DataSetIndex, DataSetIndex, DataSetIndex }
@@ -255,9 +254,31 @@ namespace mu
 		return MaskedValue | ~ChannelUsages.ClusterId;
 	}
 
+
+	FORCEINLINE FVector2f ComputeBarycentric(const FVector3f& Point, const FVector3f& A, const FVector3f& B, const FVector3f& C)
+	{
+		const FVector3f TriNorm = FVector3f::CrossProduct(B - A, C - A);
+
+		const float TriNormSizeSquared = TriNorm.SizeSquared();
+		// Return the center of the triangle if the area is very small.
+		if (TriNormSizeSquared <= UE_SMALL_NUMBER)
+		{
+			return FVector2f(1.0f/3.0f);
+		}
+
+		const float AreaABCInv = FMath::InvSqrt(TriNormSizeSquared);
+
+		const FVector3f N = TriNorm * AreaABCInv;
+
+		const float AreaPBC = FVector3f::DotProduct(N, FVector3f::CrossProduct(B - Point, C - Point));
+		const float AreaPCA = FVector3f::DotProduct(N, FVector3f::CrossProduct(C - Point, A - Point));
+		
+		return FVector2f(AreaPBC, AreaPCA) * AreaABCInv;
+	}
+
 	inline void BindReshapePoint(
 			FShapeMeshTree& ShapeMeshTree,
-			const FVector3f& Point, float MaskWeight, FReshapeVertexBindingData& OutBindData, 
+			const FVector3f& Point, const FVector3f& Normal, float MaskWeight, FReshapeVertexBindingData& OutBindData, 
 			const float ValidityTolerance = UE_KINDA_SMALL_NUMBER)
 	{
 
@@ -265,9 +286,7 @@ namespace mu
 		
 		OutBindData.S = 0.0f;
 		OutBindData.T = 0.0f;
-		OutBindData.D = 0.0f;
 		OutBindData.Triangle = -1;
-		OutBindData.ShapeNormal = FVector3f::Zero();
 		OutBindData.AttachmentPoint = Point;
 		float Weight = MaskWeight;
 
@@ -295,52 +314,47 @@ namespace mu
 		FVector3f TriangleC = (FVector3f)ShapeMesh.Positions[Triangle.C];
 
 		FPlane4f TrianglePlane(TriangleA, TriangleB, TriangleC);
-		FPlane4f VertexPlane(Point, TrianglePlane.GetNormal());
+		FVector3f PlaneNormal = TrianglePlane.GetNormal();
+		FPlane4f VertexPlane(Point, PlaneNormal);
+		FPlane4f NormalPlane(Point + Normal, PlaneNormal);
 	
 		// T1 = Triangle projected on the vertex plane along the triangle vertex normals
-		FVector3f TriangleA_VertexPlane = FMath::RayPlaneIntersection(TriangleA, ShapeMesh.Normals[Triangle.A], VertexPlane);
-		FVector3f TriangleB_VertexPlane = FMath::RayPlaneIntersection(TriangleB, ShapeMesh.Normals[Triangle.B], VertexPlane);
-		FVector3f TriangleC_VertexPlane = FMath::RayPlaneIntersection(TriangleC, ShapeMesh.Normals[Triangle.C], VertexPlane);
-	
-		// Barycentric coordinates of the vertex on in T1
-		FVector3f Barycentric = (FVector3f)FMath::ComputeBaryCentric2D((FVector)Point, (FVector)TriangleA_VertexPlane, (FVector)TriangleB_VertexPlane, (FVector)TriangleC_VertexPlane);
-
-		FVector3f ProjectedPoint = TriangleA * Barycentric.X + TriangleB * Barycentric.Y + TriangleC * Barycentric.Z;
-
-		FVector3f InterpolatedVertexNormal = ShapeMesh.Normals[Triangle.A] * Barycentric.X + ShapeMesh.Normals[Triangle.B] * Barycentric.Y + ShapeMesh.Normals[Triangle.C] * Barycentric.Z;
-		FVector3f ProjectedToVertex = (Point - ProjectedPoint);
-
-		// If the interpolated normal is not normalized the dot product gives the distance to the point times the length of the normal.
-		// We don't want that since when deforming a point, using the interpolated normal, we'd multiply by the length twice. See GetDeform().
-		// One option would be to normalize here (computing the distance) and when deforming the mesh.
-		// We can also compensate the signed distance so that the same interpolated normal (not normalized) gives us the correct point
-		// ( InterpolatedVertexNormal * d = ProjectedToVertex). That way modifications in the interpolated bound normal will affect the
-		// resulting point. We are doing the later.
-		const float InterpolatedNormalSizeSquared = InterpolatedVertexNormal.SizeSquared();
-		const float InvInterpolatedNormalSizeSquared = InterpolatedNormalSizeSquared > UE_SMALL_NUMBER ? 1.0f / InterpolatedNormalSizeSquared : 0.0f;	
-		float d = FVector3f::DotProduct(ProjectedToVertex, InterpolatedVertexNormal) * InvInterpolatedNormalSizeSquared;
+		OutBindData.D = 
+			FVector3f 
+			{
+				FMath::RayPlaneIntersectionParam(TriangleA, ShapeMesh.Normals[Triangle.A], VertexPlane),
+				FMath::RayPlaneIntersectionParam(TriangleB, ShapeMesh.Normals[Triangle.B], VertexPlane),
+				FMath::RayPlaneIntersectionParam(TriangleC, ShapeMesh.Normals[Triangle.C], VertexPlane)
+			};
+
+		OutBindData.NormalD = 
+			FVector3f
+			{
+				FMath::RayPlaneIntersectionParam(TriangleA, ShapeMesh.Normals[Triangle.A], NormalPlane),
+				FMath::RayPlaneIntersectionParam(TriangleB, ShapeMesh.Normals[Triangle.B], NormalPlane),
+				FMath::RayPlaneIntersectionParam(TriangleC, ShapeMesh.Normals[Triangle.C], NormalPlane)
+			};
+
+
+		FVector2f PositionBarycentric = ComputeBarycentric(Point, 
+					TriangleA + ShapeMesh.Normals[Triangle.A] * OutBindData.D.X, 
+					TriangleB + ShapeMesh.Normals[Triangle.B] * OutBindData.D.Y, 
+					TriangleC + ShapeMesh.Normals[Triangle.C] * OutBindData.D.Z); 
+
+		OutBindData.S = PositionBarycentric.X;
+		OutBindData.T = PositionBarycentric.Y;
 	
-		OutBindData.S = Barycentric.Y;
-		OutBindData.T = Barycentric.Z;
-		OutBindData.D = d;
-		
+		FVector2f NormalBarycentric = ComputeBarycentric((Point + Normal), 
+				TriangleA + ShapeMesh.Normals[Triangle.A] * OutBindData.NormalD.X, 
+				TriangleB + ShapeMesh.Normals[Triangle.B] * OutBindData.NormalD.Y, 
+				TriangleC + ShapeMesh.Normals[Triangle.C] * OutBindData.NormalD.Z); 
+
+		OutBindData.NS = NormalBarycentric.X;
+		OutBindData.NT = NormalBarycentric.Y;
+
 		OutBindData.Triangle = FoundIndex;
-		
-		OutBindData.ShapeNormal = ((TriangleB - TriangleA) ^ (TriangleC - TriangleA)).GetSafeNormal();
-		
-		const FVector3f ReprojectedPoint = ProjectedPoint + InterpolatedVertexNormal * d;
-		
-		const FVector3f ReprojectedVector = ReprojectedPoint - Point;
-		const float ErrorEstimate = (ReprojectedPoint - Point).GetAbsMax();
-		
-		// If within the tolerance, 1.0, otherwise linear falloff based on the tolerance
 
-		// Arbitrary factor, a binding will be considered valid (with its corresponding weight) to ErrorFalloffFactor times the validity tolerance.
-		constexpr float ErrorFalloffFactor = 4.0f;
-		OutBindData.Weight = FMath::Min(
-			MaskWeight, 
-			1.0f - FMath::Clamp( (ErrorEstimate - ValidityTolerance) / (ValidityTolerance * ErrorFalloffFactor), 0.0f, 1.0f));
-		
+		OutBindData.Weight = FMath::Clamp(MaskWeight, 0.0f, 1.0f);
 		OutBindData.Triangle = FMath::IsNearlyZero(OutBindData.Weight) ? -1 : OutBindData.Triangle;
 	}
 	
@@ -423,7 +437,7 @@ namespace mu
 		// Mask weight is set on a vertex by vertex basis, ignore weight for the shared data.
 		// This will be filled in afterwards.
 		constexpr float MaskWeight = 1.0f;
-		BindReshapePoint(ShapeMeshTree, BoundPoint, MaskWeight, OutBindingData, BindTolerance);
+		BindReshapePoint(ShapeMeshTree, BoundPoint, FVector3f::ZAxisVector, MaskWeight, OutBindingData, BindTolerance);
 	}
 
 	inline TTuple<TArray<FReshapePointBindingData>, TArray<int32>, TArray<int32>> BindPhysicsBodies( 
@@ -680,7 +694,7 @@ namespace mu
 		for (int32 PointIndex = 0; PointIndex < TotalNumPoints; ++PointIndex)
 		{
 			constexpr float MaskWeight = 1.0f;
-			BindReshapePoint(ShapeMeshTree, Points[PointIndex], MaskWeight, VertexBindData, 0.1f);
+			BindReshapePoint(ShapeMeshTree, Points[PointIndex], FVector3f::ZAxisVector, MaskWeight, VertexBindData, 0.1f);
 			PhysicsBodyBindData[PointIndex] = FReshapePointBindingData
 				{ VertexBindData.S, VertexBindData.T, VertexBindData.D, VertexBindData.Triangle, VertexBindData.Weight };
 		}
@@ -715,6 +729,7 @@ namespace mu
 			BindData.SetNum(MeshVertexCount);
 
 			const UntypedMeshBufferIteratorConst ItPositionBase(BaseMesh->GetVertexBuffers(), MBS_POSITION);
+			const UntypedMeshBufferIteratorConst ItNormalBase(BaseMesh->GetVertexBuffers(), MBS_NORMAL);
 
 			// Disable vertex color reads if the color is not used for mask weights.
 			const UntypedMeshBufferIteratorConst ItColorBase = ColorUsageMasks.MaskWeight == 0 
@@ -731,9 +746,10 @@ namespace mu
 				for (int32 VertexIndex = 0; VertexIndex < VertexCount; ++VertexIndex)
 				{
 					const FVector3f VertexPosition = (ItPositionBase + VertexIndex).GetAsVec3f();
+					const FVector3f VertexNormal = ItNormalBase.ptr() ? (ItNormalBase + VertexIndex).GetAsVec3f() : FVector3f::ZAxisVector;
 					const float MaskWeight = ItColorBase.ptr() ? GetVertexMaskWeight(ItColorBase + VertexIndex, ColorUsageMasks) : 1.0f;
 
-					BindReshapePoint(ShapeMeshTree, VertexPosition, MaskWeight, BindData[VertexIndex], BindValidityTolerance);
+					BindReshapePoint(ShapeMeshTree, VertexPosition, VertexNormal, MaskWeight, BindData[VertexIndex], BindValidityTolerance);
 				}
 			}
 			else
@@ -745,9 +761,10 @@ namespace mu
 				{
 					const int32 VertexIndex = NonRigidCluster[I];
 					const FVector3f VertexPosition = (ItPositionBase + VertexIndex).GetAsVec3f();
+					const FVector3f VertexNormal = ItNormalBase.ptr() ? (ItNormalBase + VertexIndex).GetAsVec3f() : FVector3f::ZAxisVector;
 					const float MaskWeight = ItColorBase.ptr() ? GetVertexMaskWeight(ItColorBase + VertexIndex, ColorUsageMasks) : 1.0f;
 					
-					BindReshapePoint(ShapeMeshTree, VertexPosition, MaskWeight, BindData[VertexIndex], BindValidityTolerance);
+					BindReshapePoint(ShapeMeshTree, VertexPosition, VertexNormal, MaskWeight, BindData[VertexIndex], BindValidityTolerance);
 				}
 
 				// Remove data form the non rigid cluster so it ios not processed in the rigid parts binding step.
@@ -933,7 +950,7 @@ namespace mu
 
 			FReshapeVertexBindingData BindData;
 			constexpr float MaskWeight = 1.0f;
-			BindReshapePoint(ShapeMeshTree, Mesh->BonePoses[BoneIndex].BoneTransform.GetLocation(), MaskWeight, BindData, BindValidityTolerance);
+			BindReshapePoint(ShapeMeshTree, Mesh->BonePoses[BoneIndex].BoneTransform.GetLocation(), FVector3f::ZAxisVector, MaskWeight, BindData, BindValidityTolerance);
 
 			// Only add binding  if there is a chance of the bone moving.
 			if (BindData.Weight > UE_SMALL_NUMBER && BindData.Triangle >= 0)
@@ -994,10 +1011,12 @@ namespace mu
 			return;
 		}
 
-		const bool bReshapeVertices = EnumHasAnyFlags(BindFlags, EMeshBindShapeFlags::ReshapeVertices);
-		const bool bApplyLaplacian  = EnumHasAnyFlags(BindFlags, EMeshBindShapeFlags::ApplyLaplacian);
-		const bool bReshapeSkeleton = EnumHasAnyFlags(BindFlags, EMeshBindShapeFlags::ReshapeSkeleton);
-		const bool bReshapePhysics  = EnumHasAnyFlags(BindFlags, EMeshBindShapeFlags::ReshapePhysicsVolumes);
+		const bool bReshapeVertices   = EnumHasAnyFlags(BindFlags, EMeshBindShapeFlags::ReshapeVertices);
+		// TODO: For now bRecomputeNormals is not used when binding, we could skip normal morph data generation. 
+		const bool bRecomputeNormals  = EnumHasAnyFlags(BindFlags, EMeshBindShapeFlags::RecomputeNormals);
+		const bool bApplyLaplacian    = EnumHasAnyFlags(BindFlags, EMeshBindShapeFlags::ApplyLaplacian);
+		const bool bReshapeSkeleton   = EnumHasAnyFlags(BindFlags, EMeshBindShapeFlags::ReshapeSkeleton);
+		const bool bReshapePhysics    = EnumHasAnyFlags(BindFlags, EMeshBindShapeFlags::ReshapePhysicsVolumes);
 
 		const FMeshBindColorChannelUsageMasks ColorUsagesMasks = MakeColorChannelUsageMasks(ColorChannelUsages);
 
@@ -1019,7 +1038,7 @@ namespace mu
 		}
 	
 		int32 ShapeVertexCount = ShapeMesh->GetVertexCount();
-		int ShapeTriangleCount = ShapeMesh->GetFaceCount();
+		int32 ShapeTriangleCount = ShapeMesh->GetFaceCount();
 		if (!ShapeVertexCount || !ShapeTriangleCount)
 		{
 			bOutSuccess = false;
@@ -1034,17 +1053,15 @@ namespace mu
 			ShapeMeshDescriptor.Normals.SetNum(ShapeVertexCount);
 
 			// \TODO: Simple but inefficient
-			UntypedMeshBufferIteratorConst ItPosition(ShapeMesh->GetVertexBuffers(), MBS_POSITION);
-			UntypedMeshBufferIteratorConst ItNormal(ShapeMesh->GetVertexBuffers(), MBS_NORMAL);
+			const UntypedMeshBufferIteratorConst ItPosition(ShapeMesh->GetVertexBuffers(), MBS_POSITION);
+			const UntypedMeshBufferIteratorConst ItNormal(ShapeMesh->GetVertexBuffers(), MBS_NORMAL);
 			for (int32 ShapeVertexIndex = 0; ShapeVertexIndex < ShapeVertexCount; ++ShapeVertexIndex)
 			{
-				FVector3f Position = ItPosition.GetAsVec3f();
+				FVector3f Position = (ItPosition + ShapeVertexIndex).GetAsVec3f();
 				ShapeMeshDescriptor.Positions[ShapeVertexIndex] = static_cast<FVector3d>(Position);
-				++ItPosition;
 				
-				FVector3f Normal = ItNormal.GetAsVec3f();
+				FVector3f Normal = (ItNormal + ShapeVertexIndex).GetAsVec3f();
 				ShapeMeshDescriptor.Normals[ShapeVertexIndex] = Normal;
-				++ItNormal;
 			}
 		}
 		// Generate the temp face query data for the shape
@@ -1053,16 +1070,13 @@ namespace mu
 			MUTABLE_CPUPROFILER_SCOPE(GenerateTrianglesQueryData);	
 			ShapeMeshDescriptor.Triangles.SetNum(ShapeTriangleCount);
 			// \TODO: Simple but inefficient
-			UntypedMeshBufferIteratorConst ItIndices(ShapeMesh->GetIndexBuffers(), MBS_VERTEXINDEX);
+			const UntypedMeshBufferIteratorConst ItIndices(ShapeMesh->GetIndexBuffers(), MBS_VERTEXINDEX);
 			for (int32 TriangleIndex = 0; TriangleIndex < ShapeTriangleCount; ++TriangleIndex)
 			{
 				UE::Geometry::FIndex3i Triangle;
-				Triangle.A = int(ItIndices.GetAsUINT32());
-				++ItIndices;
-				Triangle.B = int(ItIndices.GetAsUINT32());
-				++ItIndices;
-				Triangle.C = int(ItIndices.GetAsUINT32());
-				++ItIndices;
+				Triangle.A = int32((ItIndices + TriangleIndex*3 + 0).GetAsUINT32());
+				Triangle.B = int32((ItIndices + TriangleIndex*3 + 1).GetAsUINT32());
+				Triangle.C = int32((ItIndices + TriangleIndex*3 + 2).GetAsUINT32());
 
 				ShapeMeshDescriptor.Triangles[TriangleIndex] = Triangle;
 			}	
@@ -1142,9 +1156,9 @@ namespace mu
 			// Bone indices buffer
 			EMeshBufferSemantic BoneSemantics[1] = { MBS_OTHER };
 			EMeshBufferFormat BoneFormats[1] = { MBF_INT32 };
-			int BoneSemanticIndices[1] = { 0 };
-			int BoneComponents[1] = { 1 };
-			int BoneOffsets[1] = { 0 };
+			int32 BoneSemanticIndices[1] = { 0 };
+			int32 BoneComponents[1] = { 1 };
+			int32 BoneOffsets[1] = { 0 };
 
 			SkeletonBuffer.SetBuffer(1, sizeof(int32), 1, BoneSemantics, BoneSemanticIndices, BoneFormats, BoneComponents, BoneOffsets);
 
@@ -1467,13 +1481,11 @@ namespace mu
 			UntypedMeshBufferIteratorConst ItNormal(ShapeMesh->GetVertexBuffers(), MBS_NORMAL);
 			for (int32 ShapeVertexIndex = 0; ShapeVertexIndex < ShapeVertexCount; ++ShapeVertexIndex)
 			{
-				FVector3f Position = ItPosition.GetAsVec3f();
+				FVector3f Position = (ItPosition + ShapeVertexIndex).GetAsVec3f();
 				ShapeMeshDescriptor.Positions[ShapeVertexIndex] = static_cast<FVector3d>(Position);
-				++ItPosition;
 				
-				FVector3f Normal = ItNormal.GetAsVec3f();
+				FVector3f Normal = (ItNormal + ShapeVertexIndex).GetAsVec3f();
 				ShapeMeshDescriptor.Normals[ShapeVertexIndex] = Normal;
-				++ItNormal;
 			}
 		}
 		// Generate the temp face query data for the shape
@@ -1482,16 +1494,13 @@ namespace mu
 			MUTABLE_CPUPROFILER_SCOPE(GenerateTrianglesQueryData);	
 			ShapeMeshDescriptor.Triangles.SetNum(ShapeTriangleCount);
 			// \TODO: Simple but inefficient
-			UntypedMeshBufferIteratorConst ItIndices(ShapeMesh->GetIndexBuffers(), MBS_VERTEXINDEX);
+			const UntypedMeshBufferIteratorConst ItIndices(ShapeMesh->GetIndexBuffers(), MBS_VERTEXINDEX);
 			for (int32 TriangleIndex = 0; TriangleIndex < ShapeTriangleCount; ++TriangleIndex)
 			{
 				UE::Geometry::FIndex3i Triangle;
-				Triangle.A = int(ItIndices.GetAsUINT32());
-				++ItIndices;
-				Triangle.B = int(ItIndices.GetAsUINT32());
-				++ItIndices;
-				Triangle.C = int(ItIndices.GetAsUINT32());
-				++ItIndices;
+				Triangle.A = int32((ItIndices + TriangleIndex*3 + 0).GetAsUINT32());
+				Triangle.B = int32((ItIndices + TriangleIndex*3 + 1).GetAsUINT32());
+				Triangle.C = int32((ItIndices + TriangleIndex*3 + 2).GetAsUINT32());
 
 				ShapeMeshDescriptor.Triangles[TriangleIndex] = Triangle;
 			}	
diff --git a/Engine/Plugins/Experimental/Mutable/Source/MutableRuntime/Private/MuR/OpMeshComputeNormals.h b/Engine/Plugins/Experimental/Mutable/Source/MutableRuntime/Private/MuR/OpMeshComputeNormals.h
new file mode 100644
--- /dev/null
+++ b/Engine/Plugins/Experimental/Mutable/Source/MutableRuntime/Private/MuR/OpMeshComputeNormals.h
@@ -0,0 +1,155 @@
+// Copyright Epic Games, Inc. All Rights Reserved.
+
+#pragma once
+
+
+#include "Math/Vector.h"
+#include "Math/UnrealMathUtility.h"
+#include "Containers/Array.h"
+#include "Containers/ArrayView.h"
+#include "Containers/Map.h"
+
+#include "MuR/Mesh.h"
+#include "MuR/MeshPrivate.h"
+#include "MuR/MutableTrace.h"
+
+
+namespace mu
+{
+
+	float ComputeTangentBasisDeterminantSign(const FVector3f& N, const FVector3f& T, const FVector3f& B)
+	{
+		const float BaseTangentBasisDeterminant =  
+				B.X*T.Y*N.Z + B.Z*T.X*N.Y + B.Y*T.Z*N.Y -
+				B.Z*T.Y*N.X - B.Y*T.X*N.Z - B.X*T.Z*N.Y;
+
+		return BaseTangentBasisDeterminant < 0.0f ? -1.0f : 1.0f;
+	}
+
+	void OrthogonalizeTangentSpace(
+			const FVector3f* NormalPtr, FVector3f* TangentPtr, FVector3f* BiNormalPtr, 
+			float TangentBasisDeterminantSign)	
+	{
+		if (!!NormalPtr & !!TangentPtr)
+		{
+			// Orthogonalize Tangent based on new Normal. This assumes Normal and Tangent are normalized and different.
+			*TangentPtr = (*TangentPtr - FVector3f::DotProduct(*NormalPtr, *TangentPtr) * *NormalPtr).GetSafeNormal();
+
+			// BiNormal
+			if (BiNormalPtr)
+			{
+				*BiNormalPtr = FVector3f::CrossProduct(*TangentPtr, *NormalPtr) * TangentBasisDeterminantSign;
+			}
+		}
+	}
+
+	void ComputeMeshNormals(Mesh& DestMesh)
+	{
+		MUTABLE_CPUPROFILER_SCOPE(ComputeMeshNormals);
+
+		const UntypedMeshBufferIteratorConst BaseIndicesIter = UntypedMeshBufferIteratorConst(DestMesh.GetIndexBuffers(), MBS_VERTEXINDEX, 0);
+		const int32 NumIndices = DestMesh.GetIndexBuffers().GetElementCount();
+
+		const UntypedMeshBufferIteratorConst BasePositionIter = UntypedMeshBufferIteratorConst(DestMesh.GetVertexBuffers(), MBS_POSITION, 0);
+		const int32 NumVertices = DestMesh.GetVertexBuffers().GetElementCount();
+
+		check(BaseIndicesIter.ptr());
+		check(BasePositionIter.ptr());
+
+		TArray<FVector3f> NormalsAccumulation;
+		NormalsAccumulation.SetNumZeroed(NumVertices);	
+
+		check(NumIndices % 3 == 0);
+		for (int32 Face = 0; Face < NumIndices; Face += 3)
+		{
+			const UE::Geometry::FIndex3i FaceIndices = UE::Geometry::FIndex3i 
+			{ 
+				static_cast<int32>((BaseIndicesIter + Face + 0).GetAsUINT32()), 
+				static_cast<int32>((BaseIndicesIter + Face + 1).GetAsUINT32()),
+				static_cast<int32>((BaseIndicesIter + Face + 2).GetAsUINT32())
+			};
+
+			const FVector3f V0 = (BasePositionIter + FaceIndices[0]).GetAsVec3f();
+			const FVector3f V1 = (BasePositionIter + FaceIndices[1]).GetAsVec3f();
+			const FVector3f V2 = (BasePositionIter + FaceIndices[2]).GetAsVec3f();
+
+			const FVector3f AreaWeightedTriangleNormal = FVector3f::CrossProduct(V2 - V0, V1 - V0);
+
+			NormalsAccumulation[FaceIndices[0]] += AreaWeightedTriangleNormal;
+			NormalsAccumulation[FaceIndices[1]] += AreaWeightedTriangleNormal;
+			NormalsAccumulation[FaceIndices[2]] += AreaWeightedTriangleNormal;
+		}
+
+		for (int32 VertexIndex = 0; VertexIndex < NumVertices; ++VertexIndex)
+		{
+			NormalsAccumulation[VertexIndex] = NormalsAccumulation[VertexIndex].GetSafeNormal();
+		}
+
+		const UntypedMeshBufferIterator NormalIter = UntypedMeshBufferIterator(DestMesh.GetVertexBuffers(), MBS_NORMAL, 0);
+		check(NormalIter.ptr());
+
+		const UntypedMeshBufferIterator TangentIter  = UntypedMeshBufferIterator(DestMesh.GetVertexBuffers(), MBS_TANGENT, 0);
+		const UntypedMeshBufferIterator BiNormalIter = UntypedMeshBufferIterator(DestMesh.GetVertexBuffers(), MBS_BINORMAL, 0);
+
+        const EMeshBufferFormat NormalFormat   = NormalIter.GetFormat();
+        const EMeshBufferFormat TangentFormat  = TangentIter.GetFormat();
+        const EMeshBufferFormat BiNormalFormat = BiNormalIter.GetFormat();
+	
+        const int32 NormalComps   = NormalIter.GetComponents();
+        const int32 TangentComps  = TangentIter.GetComponents();
+        const int32 BiNormalComps = BiNormalIter.GetComponents();
+
+		// When normal is packed, binormal channel is not expected. It is not a big deal if it's there but we would be doing extra unused work in that case. 
+		ensureAlways(!(NormalFormat == MBF_PACKEDDIR8_W_TANGENTSIGN || NormalFormat == MBF_PACKEDDIRS8_W_TANGENTSIGN) || !BiNormalIter.ptr());
+		for (int32 VertexIndex = 0; VertexIndex < NumVertices; ++VertexIndex)
+		{
+			const FVector3f Normal = NormalsAccumulation[VertexIndex];
+			const FVector3f OriginalNormal = NormalIter.ptr()
+					? FVector3f::ZeroVector
+					: (NormalIter + VertexIndex).GetAsVec3f();
+
+			FVector3f Tangent = !TangentIter.ptr() 
+					? FVector3f::ZeroVector
+					: (TangentIter + VertexIndex).GetAsVec3f();
+
+			FVector3f BiNormal = !BiNormalIter.ptr() 
+					? FVector3f::ZeroVector
+					: (BiNormalIter + VertexIndex).GetAsVec3f();
+			
+			OrthogonalizeTangentSpace(
+					&Normal, 
+					TangentIter.ptr() ? &Tangent : nullptr, 
+					BiNormalIter.ptr() ? &BiNormal : nullptr,
+					BiNormalIter.ptr() ? ComputeTangentBasisDeterminantSign(OriginalNormal, Tangent, BiNormal) : 0.0f);
+			
+			if (NormalIter.ptr())
+			{
+				// Leave the tangent basis sign untouched for packed normals formats.
+				uint8 * const NormalElemPtr = (NormalIter + VertexIndex).ptr();
+				for (int32 C = 0; C < NormalComps && C < 3; ++C)
+				{
+					ConvertData(C, NormalElemPtr, NormalFormat, &Normal, MBF_FLOAT32);
+				}
+			}
+
+			if (TangentIter.ptr())
+			{
+				uint8 * const TangentElemPtr = (TangentIter + VertexIndex).ptr();
+				for (int32 C = 0; C < TangentComps && C < 3; ++C)
+				{
+					ConvertData(C, TangentElemPtr, TangentFormat, &Tangent, MBF_FLOAT32);
+				}
+			}
+			
+			if (BiNormalIter.ptr())
+			{
+				uint8 * const BiNormalElemPtr = (BiNormalIter + VertexIndex).ptr();
+				for (int32 C = 0; C < BiNormalComps && C < 3; ++C)
+				{
+					ConvertData(C, BiNormalElemPtr, BiNormalFormat, &BiNormal, MBF_FLOAT32);
+				}
+			}
+		}
+	}
+}
+
diff --git a/Engine/Plugins/Experimental/Mutable/Source/MutableRuntime/Private/MuR/OpMeshReshape.h b/Engine/Plugins/Experimental/Mutable/Source/MutableRuntime/Private/MuR/OpMeshReshape.h
--- a/Engine/Plugins/Experimental/Mutable/Source/MutableRuntime/Private/MuR/OpMeshReshape.h
+++ b/Engine/Plugins/Experimental/Mutable/Source/MutableRuntime/Private/MuR/OpMeshReshape.h
@@ -15,7 +15,8 @@
 #include "Spatial/MeshAABBTree3.h"
 #include "Distance/DistPoint3Triangle3.h"
 #include "LineTypes.h"
-#include "OpMeshSmoothing.h"
+#include "MuR/OpMeshSmoothing.h"
+#include "MuR/OpMeshComputeNormals.h"
 
 #include "Algo/AnyOf.h" 
 
@@ -44,45 +45,32 @@ namespace mu
 	};
 
 	// Methods to actually deform a point
-	inline void GetDeform( const FShapeMeshDescriptorApply& Shape, const FReshapeVertexBindingData& Binding, FVector3f& NewPosition, FQuat4f& Rotation)
+	inline void GetDeform( const FShapeMeshDescriptorApply& Shape, const FReshapeVertexBindingData& Binding, FVector3f& NewPosition, FVector3f& NewNormalPosition)
 	{
 		const UE::Geometry::FIndex3i& Triangle = Shape.Triangles[Binding.Triangle];
+	
+		FVector3f VA = Shape.Positions[Triangle.A] + Shape.Normals[Triangle.A]*Binding.D.X;
+		FVector3f VB = Shape.Positions[Triangle.B] + Shape.Normals[Triangle.B]*Binding.D.Y;
+		FVector3f VC = Shape.Positions[Triangle.C] + Shape.Normals[Triangle.C]*Binding.D.Z;
 
-		FVector3f ProjectedVertexPosition = 
-			Shape.Positions[Triangle.A] * (1.0f - Binding.S - Binding.T)
-			+ Shape.Positions[Triangle.B] * Binding.S
-			+ Shape.Positions[Triangle.C] * Binding.T;
-
-		// This method approximates the shape face rotation
-		FVector3f InterpolatedNormal = 
-			Shape.Normals[Triangle.A] * (1.0f - Binding.S - Binding.T)
-			+ Shape.Normals[Triangle.B] * Binding.S
-			+ Shape.Normals[Triangle.C] * Binding.T;
-		
-		NewPosition = ProjectedVertexPosition + InterpolatedNormal*Binding.D;
-
-		FVector3f CurrentShapeNormal = ((Shape.Positions[Triangle.B] - Shape.Positions[Triangle.A]) ^ 
-										(Shape.Positions[Triangle.C] - Shape.Positions[Triangle.A])).GetSafeNormal();
+		NewPosition = VA*Binding.S + VB*Binding.T + VC*(1.0f - Binding.S - Binding.T);
+	
+		FVector3f NVA = Shape.Positions[Triangle.A] + Shape.Normals[Triangle.A]*Binding.NormalD.X;
+		FVector3f NVB = Shape.Positions[Triangle.B] + Shape.Normals[Triangle.B]*Binding.NormalD.Y;
+		FVector3f NVC = Shape.Positions[Triangle.C] + Shape.Normals[Triangle.C]*Binding.NormalD.Z;
 
-		Rotation = FQuat4f::FindBetween(Binding.ShapeNormal.GetSafeNormal(), CurrentShapeNormal);
+		NewNormalPosition = (NVA*Binding.NS + NVB*Binding.NT + NVC*(1.0f - Binding.NS - Binding.NT));
 	}
 
 	inline void GetDeform( const FShapeMeshDescriptorApply& Shape, const FReshapePointBindingData& Binding, FVector3f& NewPosition)
 	{
 		const UE::Geometry::FIndex3i& Triangle = Shape.Triangles[Binding.Triangle];
+	
+		FVector3f VA = Shape.Positions[Triangle.A] + Shape.Normals[Triangle.A]*Binding.D.X;
+		FVector3f VB = Shape.Positions[Triangle.B] + Shape.Normals[Triangle.B]*Binding.D.Y;
+		FVector3f VC = Shape.Positions[Triangle.C] + Shape.Normals[Triangle.C]*Binding.D.Z;
 
-		FVector3f ProjectedVertexPosition = 
-			Shape.Positions[Triangle.A] * (1.0f - Binding.S - Binding.T)
-			+ Shape.Positions[Triangle.B] * Binding.S
-			+ Shape.Positions[Triangle.C] * Binding.T;
-
-		// This method approximates the shape face rotation
-		FVector3f InterpolatedNormal = 
-			Shape.Normals[Triangle.A] * (1.0f - Binding.S - Binding.T)
-			+ Shape.Normals[Triangle.B] * Binding.S
-			+ Shape.Normals[Triangle.C] * Binding.T;
-		
-		NewPosition = ProjectedVertexPosition + InterpolatedNormal*Binding.D;
+		NewPosition = VA*Binding.S + VB*Binding.T + VC*(1.0f - Binding.S - Binding.T);
 	}
 
 
@@ -381,15 +369,23 @@ namespace mu
 		OutL = FMath::Max(0.0f, (Centroids[0] - Centroids[1]).Length() - (OutR0 + OutR1));
 	}
 
-	inline void ApplyToVertices(Mesh* Mesh, TArrayView<const FReshapeVertexBindingData> BindingData, const FShapeMeshDescriptorApply& Shape)
+	inline void ApplyToVertices(Mesh* Mesh, TArrayView<const FReshapeVertexBindingData> BindingData, const FShapeMeshDescriptorApply& Shape, bool bSkipNormalReshape)
 	{
 		check(Mesh);
 		check(Mesh->GetVertexCount() == BindingData.Num());
 		
-		UntypedMeshBufferIterator ItPosition(Mesh->GetVertexBuffers(), MBS_POSITION);
-		UntypedMeshBufferIterator ItNormal(Mesh->GetVertexBuffers(), MBS_NORMAL);
-		UntypedMeshBufferIterator ItTangent(Mesh->GetVertexBuffers(), MBS_TANGENT);
-		UntypedMeshBufferIterator ItBinormal(Mesh->GetVertexBuffers(), MBS_BINORMAL);
+		const UntypedMeshBufferIterator PositionIter(Mesh->GetVertexBuffers(), MBS_POSITION);
+		const UntypedMeshBufferIterator NormalIter(Mesh->GetVertexBuffers(), MBS_NORMAL);
+		const UntypedMeshBufferIterator TangentIter(Mesh->GetVertexBuffers(), MBS_TANGENT);
+		const UntypedMeshBufferIterator BiNormalIter(Mesh->GetVertexBuffers(), MBS_BINORMAL);
+
+        const EMeshBufferFormat NormalFormat   = NormalIter.GetFormat();
+        const EMeshBufferFormat TangentFormat  = TangentIter.GetFormat();
+        const EMeshBufferFormat BiNormalFormat = BiNormalIter.GetFormat();
+	
+        const int32 NormalComps   = NormalIter.GetComponents();
+        const int32 TangentComps  = TangentIter.GetComponents();
+        const int32 BiNormalComps = BiNormalIter.GetComponents();
 
 #if DO_CHECK
 		// checking if the Base shape has more triangles than the target shape
@@ -407,70 +403,69 @@ namespace mu
 
 		const int32 MeshVertexCount = BindingData.Num();
 
-		for (int32 MeshVertexIndex = 0; MeshVertexIndex < MeshVertexCount; ++MeshVertexIndex)
+		const bool bComputeNormal = bSkipNormalReshape && NormalIter.ptr();
+		for (int32 VertexIndex = 0; VertexIndex < MeshVertexCount; ++VertexIndex)
 		{
-			const FReshapeVertexBindingData& Binding = BindingData[MeshVertexIndex];
-
-			FVector3f NewPosition;
-			FQuat4f TangentSpaceCorrection;
+			const FReshapeVertexBindingData& Binding = BindingData[VertexIndex];
 
 			const bool bModified = (Binding.Triangle >= 0) & (Binding.Triangle < ShapeTriangleCount);
 			if (bModified)
 			{
-				GetDeform(Shape, Binding, NewPosition, TangentSpaceCorrection);
+				FVector3f NewPosition;
+				FVector3f NewNormalPosition;
+				
+				GetDeform(Shape, Binding, NewPosition, NewNormalPosition);
 
-				const FVector3f OldPosition = ItPosition.GetAsVec3f();
-				FVector3f Displacement = NewPosition - OldPosition;
+				const FVector3f OldPosition = (PositionIter + VertexIndex).GetAsVec3f();
 
-				if (!FMath::IsNearlyEqual(Binding.Weight, 1.0f))
-				{
-					Displacement *= Binding.Weight;
-					TangentSpaceCorrection = FQuat4f::Slerp(FQuat4f::Identity, TangentSpaceCorrection, Binding.Weight);
-				}
+				FVector3f PositionDisplacement = (NewPosition - OldPosition) * Binding.Weight;
 				
-				// Non rigid vertices will not rotate since are attached to themselves 
-				NewPosition = TangentSpaceCorrection.RotateVector(OldPosition - Binding.AttachmentPoint) + (OldPosition + Displacement);
+				// TODO: Separeat rigid from not rigid processing.
+				(PositionIter + VertexIndex).SetFromVec3f(OldPosition + PositionDisplacement);
 				
-				ItPosition.SetFromVec3f(NewPosition);
-					
-				if (ItNormal.ptr())
+				if (bComputeNormal)
 				{
-					FVector3f OldNormal = ItNormal.GetAsVec3f();
-					FVector3f NewNormal = TangentSpaceCorrection.RotateVector(OldNormal);
-					ItNormal.SetFromVec3f(NewNormal);
-				}
-
-				if (ItTangent.ptr())
-				{
-					FVector3f OldTangent = ItTangent.GetAsVec3f();
-					FVector3f NewTangent = TangentSpaceCorrection.RotateVector(OldTangent);
-					ItTangent.SetFromVec3f(NewTangent);
-				}
+					const FVector3f OldNormal = (NormalIter + VertexIndex).GetAsVec3f();
+					const FVector3f OldNormalPosition = OldPosition + OldNormal;
+					
+					FVector3f NormalDisplacement = (NewNormalPosition - OldNormalPosition) * Binding.Weight;
 
-				if (ItBinormal.ptr())
-				{
-					FVector3f OldBinormal = ItBinormal.GetAsVec3f();
-					FVector3f NewBinormal = TangentSpaceCorrection.RotateVector(OldBinormal);
-					ItBinormal.SetFromVec3f(NewBinormal);
+					FVector3f Normal = ((OldNormalPosition + NormalDisplacement) - NewPosition).GetSafeNormal();
+					FVector3f Tangent = TangentIter.ptr() ? (TangentIter + VertexIndex).GetAsVec3f() : FVector3f::XAxisVector;
+					FVector3f BiNormal = BiNormalIter.ptr() ? (BiNormalIter + VertexIndex).GetAsVec3f() : FVector3f::YAxisVector;
+				
+					OrthogonalizeTangentSpace(
+							&Normal, 
+							TangentIter.ptr() ? &Tangent : nullptr, 
+							BiNormalIter.ptr() ? &BiNormal : nullptr,
+							BiNormalIter.ptr() ? ComputeTangentBasisDeterminantSign(OldNormal, Tangent, BiNormal) : 0.0f);
+					
+					// Leave the tangent basis sign untouched for packed normals formats.
+					uint8 * const NormalElemPtr = (NormalIter + VertexIndex).ptr();
+					for (int32 C = 0; C < NormalComps && C < 3; ++C)
+					{
+						ConvertData(C, NormalElemPtr, NormalFormat, &Normal, MBF_FLOAT32);
+					}
+
+					if (TangentIter.ptr())
+					{
+						uint8 * const TangentElemPtr = (TangentIter + VertexIndex).ptr();
+						for (int32 C = 0; C < TangentComps && C < 3; ++C)
+						{
+							ConvertData(C, TangentElemPtr, TangentFormat, &Tangent, MBF_FLOAT32);
+						}
+					}
+					
+					if (BiNormalIter.ptr())
+					{
+						uint8 * const BiNormalElemPtr = (BiNormalIter + VertexIndex).ptr();
+						for (int32 C = 0; C < BiNormalComps && C < 3; ++C)
+						{
+							ConvertData(C, BiNormalElemPtr, BiNormalFormat, &BiNormal, MBF_FLOAT32);
+						}
+					}
 				}
 			}
-			
-			++ItPosition;
-
-			if (ItNormal.ptr())
-			{
-				++ItNormal;
-			}
-
-			if (ItTangent.ptr())
-			{
-				++ItTangent;
-			}
-
-			if (ItBinormal.ptr())
-			{
-				++ItBinormal;
-			}
 		}
 	}
 
@@ -729,6 +724,7 @@ namespace mu
 		}
 		
 		const bool bReshapeVertices = EnumHasAnyFlags(BindFlags, EMeshBindShapeFlags::ReshapeVertices);
+		const bool bRecomputeNormals = EnumHasAnyFlags(BindFlags, EMeshBindShapeFlags::RecomputeNormals);
 		const bool bReshapeSkeleton = EnumHasAnyFlags(BindFlags, EMeshBindShapeFlags::ReshapeSkeleton);
 		const bool bReshapePhysicsVolumes = EnumHasAnyFlags(BindFlags, EMeshBindShapeFlags::ReshapePhysicsVolumes);
 		const bool bApplyLaplacian = EnumHasAnyFlags(BindFlags, EMeshBindShapeFlags::ApplyLaplacian);
@@ -820,17 +816,15 @@ namespace mu
 			ShapeDescriptor.Normals.SetNum(ShapeVertexCount);
 
 			// \TODO: Simple but inefficient
-			UntypedMeshBufferIteratorConst ItPosition(ShapeMesh->GetVertexBuffers(), MBS_POSITION);
-			UntypedMeshBufferIteratorConst ItNormal(ShapeMesh->GetVertexBuffers(), MBS_NORMAL);
+			const UntypedMeshBufferIteratorConst ItPosition(ShapeMesh->GetVertexBuffers(), MBS_POSITION);
+			const UntypedMeshBufferIteratorConst ItNormal(ShapeMesh->GetVertexBuffers(), MBS_NORMAL);
 			for (int32 ShapeVertexIndex = 0; ShapeVertexIndex < ShapeVertexCount; ++ShapeVertexIndex)
 			{
-				FVector3f Position = ItPosition.GetAsVec3f();
+				FVector3f Position = (ItPosition + ShapeVertexIndex).GetAsVec3f();
 				ShapeDescriptor.Positions[ShapeVertexIndex] = Position;
-				++ItPosition;
 
-				FVector3f Normal = ItNormal.GetAsVec3f();
+				FVector3f Normal = (ItNormal + ShapeVertexIndex).GetAsVec3f();
 				ShapeDescriptor.Normals[ShapeVertexIndex] = Normal;
-				++ItNormal;
 			}
 		}
 
@@ -842,16 +836,13 @@ namespace mu
 			ShapeDescriptor.Triangles.SetNum(ShapeTriangleCount);
 
 			// \TODO: Simple but inefficient
-			UntypedMeshBufferIteratorConst ItIndices(ShapeMesh->GetIndexBuffers(), MBS_VERTEXINDEX);
+			const UntypedMeshBufferIteratorConst ItIndices(ShapeMesh->GetIndexBuffers(), MBS_VERTEXINDEX);
 			for (int32 TriangleIndex = 0; TriangleIndex < ShapeTriangleCount; ++TriangleIndex)
 			{
 				UE::Geometry::FIndex3i Triangle;
-				Triangle.A = int(ItIndices.GetAsUINT32());
-				++ItIndices;
-				Triangle.B = int(ItIndices.GetAsUINT32());
-				++ItIndices;
-				Triangle.C = int(ItIndices.GetAsUINT32());
-				++ItIndices;
+				Triangle.A = int32((ItIndices + TriangleIndex*3 + 0).GetAsUINT32());
+				Triangle.B = int32((ItIndices + TriangleIndex*3 + 1).GetAsUINT32());
+				Triangle.C = int32((ItIndices + TriangleIndex*3 + 2).GetAsUINT32());
 
 				ShapeDescriptor.Triangles[TriangleIndex] = Triangle;
 			}
@@ -869,7 +860,7 @@ namespace mu
 						(const FReshapeVertexBindingData*)VB.GetBufferData(BarycentricDataBuffer),
 						VB.GetElementCount());
 
-				ApplyToVertices(VerticesReshapeMesh, VerticesBindingData, ShapeDescriptor);
+				ApplyToVertices(VerticesReshapeMesh, VerticesBindingData, ShapeDescriptor, bRecomputeNormals);
 			}
 
 			if (bApplyLaplacian)
@@ -879,6 +870,11 @@ namespace mu
 				check(Result->GetVertexCount() == 0 && Result->GetIndexCount() == 0); 
 				SmoothMeshLaplacian(*Result, *VerticesReshapeMesh);
 			}
+
+			if (bRecomputeNormals)
+			{
+				ComputeMeshNormals(*Result);
+			}
 		}
 	
 		if (bReshapeSkeleton)
diff --git a/Engine/Plugins/Experimental/Mutable/Source/MutableTools/Private/MuT/ASTOpMeshApplyShape.cpp b/Engine/Plugins/Experimental/Mutable/Source/MutableTools/Private/MuT/ASTOpMeshApplyShape.cpp
--- a/Engine/Plugins/Experimental/Mutable/Source/MutableTools/Private/MuT/ASTOpMeshApplyShape.cpp
+++ b/Engine/Plugins/Experimental/Mutable/Source/MutableTools/Private/MuT/ASTOpMeshApplyShape.cpp
@@ -15,6 +15,7 @@ namespace mu
 	ASTOpMeshApplyShape::ASTOpMeshApplyShape()
 		: Mesh(this)
 		, Shape(this)
+		, bRecomputeNormals(false)
 		, bReshapeSkeleton(false)
 		, bReshapePhysicsVolumes(false)
 		, bReshapeVertices(true)
@@ -37,6 +38,7 @@ namespace mu
 			const ASTOpMeshApplyShape* Other = static_cast<const ASTOpMeshApplyShape*>(&OtherUntyped);
 
 			const bool bSameFlags =
+				bRecomputeNormals == Other->bRecomputeNormals &&
 				bReshapePhysicsVolumes == Other->bReshapePhysicsVolumes &&
 				bReshapeSkeleton == Other->bReshapeSkeleton &&
 				bReshapeVertices == Other->bReshapeVertices &&
@@ -52,6 +54,7 @@ namespace mu
 	{
 		uint64 Result = std::hash<void*>()(Mesh.child().get());
 		hash_combine(Result, Shape.child().get());
+		hash_combine(Result, bool(bRecomputeNormals));
 		hash_combine(Result, bool(bReshapeSkeleton));
 		hash_combine(Result, bool(bReshapePhysicsVolumes));
 		hash_combine(Result, bool(bReshapeVertices));
@@ -66,6 +69,7 @@ namespace mu
 		Ptr<ASTOpMeshApplyShape> NewOp = new ASTOpMeshApplyShape();
 		NewOp->Mesh = MapChild(Mesh.child());
 		NewOp->Shape = MapChild(Shape.child());
+		NewOp->bRecomputeNormals = bRecomputeNormals;
 		NewOp->bReshapeSkeleton = bReshapeSkeleton;
 		NewOp->bReshapePhysicsVolumes = bReshapePhysicsVolumes;
 		NewOp->bReshapeVertices = bReshapeVertices;
@@ -91,6 +95,7 @@ namespace mu
 
 			constexpr EMeshBindShapeFlags NoFlags = EMeshBindShapeFlags::None;
 			EMeshBindShapeFlags BindFlags = NoFlags;
+			EnumAddFlags(BindFlags, bRecomputeNormals ? EMeshBindShapeFlags::RecomputeNormals : NoFlags);
 			EnumAddFlags(BindFlags, bReshapeSkeleton ? EMeshBindShapeFlags::ReshapeSkeleton : NoFlags);
 			EnumAddFlags(BindFlags, bReshapePhysicsVolumes ? EMeshBindShapeFlags::ReshapePhysicsVolumes : NoFlags);
 			EnumAddFlags(BindFlags, bReshapeVertices ? EMeshBindShapeFlags::ReshapeVertices : NoFlags);
diff --git a/Engine/Plugins/Experimental/Mutable/Source/MutableTools/Private/MuT/ASTOpMeshApplyShape.h b/Engine/Plugins/Experimental/Mutable/Source/MutableTools/Private/MuT/ASTOpMeshApplyShape.h
--- a/Engine/Plugins/Experimental/Mutable/Source/MutableTools/Private/MuT/ASTOpMeshApplyShape.h
+++ b/Engine/Plugins/Experimental/Mutable/Source/MutableTools/Private/MuT/ASTOpMeshApplyShape.h
@@ -22,6 +22,7 @@ struct FProgram;
 		ASTChild Mesh;
 		ASTChild Shape;
 
+		uint32 bRecomputeNormals	  : 1;
 		uint32 bReshapeSkeleton       : 1;
 		uint32 bReshapePhysicsVolumes : 1;
 		uint32 bReshapeVertices       : 1;
diff --git a/Engine/Plugins/Experimental/Mutable/Source/MutableTools/Private/MuT/ASTOpMeshBindShape.cpp b/Engine/Plugins/Experimental/Mutable/Source/MutableTools/Private/MuT/ASTOpMeshBindShape.cpp
--- a/Engine/Plugins/Experimental/Mutable/Source/MutableTools/Private/MuT/ASTOpMeshBindShape.cpp
+++ b/Engine/Plugins/Experimental/Mutable/Source/MutableTools/Private/MuT/ASTOpMeshBindShape.cpp
@@ -20,6 +20,7 @@ namespace mu
 	ASTOpMeshBindShape::ASTOpMeshBindShape()
 		: Mesh(this)
 		, Shape(this)
+		, bRecomputeNormals(false)
 		, bReshapeSkeleton(false)
 		, bReshapePhysicsVolumes(false)
 		, bReshapeVertices(false)
@@ -42,6 +43,7 @@ namespace mu
 			const ASTOpMeshBindShape* Other = static_cast<const ASTOpMeshBindShape*>(&OtherUntyped);
 
 			const bool bSameFlags =
+				bRecomputeNormals == Other->bRecomputeNormals &&
 				bReshapeSkeleton == Other->bReshapeSkeleton	&&
 				bReshapePhysicsVolumes == Other->bReshapePhysicsVolumes &&
 				bReshapeVertices == Other->bReshapeVertices &&
@@ -67,6 +69,7 @@ namespace mu
 	{
 		uint64 Result = std::hash<void*>()(Mesh.child().get());
 		hash_combine(Result, Shape.child().get());
+		hash_combine(Result, bool(bRecomputeNormals));
 		hash_combine(Result, bool(bReshapeSkeleton));
 		hash_combine(Result, bool(bReshapePhysicsVolumes));
 		hash_combine(Result, bool(bReshapeVertices));
@@ -97,6 +100,7 @@ namespace mu
 		Ptr<ASTOpMeshBindShape> NewOp = new ASTOpMeshBindShape();
 		NewOp->Mesh = mapChild(Mesh.child());
 		NewOp->Shape = mapChild(Shape.child());
+		NewOp->bRecomputeNormals = bRecomputeNormals;
 		NewOp->bReshapeSkeleton	= bReshapeSkeleton;
 		NewOp->bReshapePhysicsVolumes = bReshapePhysicsVolumes;
 		NewOp->bReshapeVertices = bReshapeVertices;
@@ -131,6 +135,7 @@ namespace mu
 
 			constexpr EMeshBindShapeFlags NoFlags = EMeshBindShapeFlags::None;
 			EMeshBindShapeFlags BindFlags = NoFlags;
+			EnumAddFlags(BindFlags, bRecomputeNormals ? EMeshBindShapeFlags::RecomputeNormals : NoFlags);
 			EnumAddFlags(BindFlags, bReshapeSkeleton ? EMeshBindShapeFlags::ReshapeSkeleton : NoFlags);
 			EnumAddFlags(BindFlags, bReshapePhysicsVolumes ? EMeshBindShapeFlags::ReshapePhysicsVolumes : NoFlags);
 			EnumAddFlags(BindFlags, bReshapeVertices ? EMeshBindShapeFlags::ReshapeVertices : NoFlags);
diff --git a/Engine/Plugins/Experimental/Mutable/Source/MutableTools/Private/MuT/ASTOpMeshBindShape.h b/Engine/Plugins/Experimental/Mutable/Source/MutableTools/Private/MuT/ASTOpMeshBindShape.h
--- a/Engine/Plugins/Experimental/Mutable/Source/MutableTools/Private/MuT/ASTOpMeshBindShape.h
+++ b/Engine/Plugins/Experimental/Mutable/Source/MutableTools/Private/MuT/ASTOpMeshBindShape.h
@@ -29,6 +29,7 @@ namespace mu
 
 		uint32 BindingMethod = 0;
 		
+		uint32 bRecomputeNormals	  : 1;
 		uint32 bReshapeSkeleton	      : 1;
 		uint32 bReshapePhysicsVolumes : 1;
 		uint32 bReshapeVertices       : 1;
diff --git a/Engine/Plugins/Experimental/Mutable/Source/MutableTools/Private/MuT/CodeGenerator_Mesh.cpp b/Engine/Plugins/Experimental/Mutable/Source/MutableTools/Private/MuT/CodeGenerator_Mesh.cpp
--- a/Engine/Plugins/Experimental/Mutable/Source/MutableTools/Private/MuT/CodeGenerator_Mesh.cpp
+++ b/Engine/Plugins/Experimental/Mutable/Source/MutableTools/Private/MuT/CodeGenerator_Mesh.cpp
@@ -800,6 +800,7 @@ namespace mu
 			// PhysicsBodies and the Skeleton.
             OpBind->bReshapeVertices = false;
             OpBind->bApplyLaplacian = false;
+            OpBind->bRecomputeNormals = false;
 		    OpBind->bReshapeSkeleton = node.bReshapeSkeleton;
 		    OpBind->BonesToDeform = node.BonesToDeform;
     	    OpBind->bReshapePhysicsVolumes = node.bReshapePhysicsVolumes; 
@@ -810,6 +811,7 @@ namespace mu
             OpBind->Shape = BaseResult.meshOp;
            
 			OpApply->bReshapeVertices = OpBind->bReshapeVertices;
+			OpApply->bRecomputeNormals = OpBind->bRecomputeNormals;
 		    OpApply->bReshapeSkeleton = OpBind->bReshapeSkeleton;
 		    OpApply->bReshapePhysicsVolumes = OpBind->bReshapePhysicsVolumes;
 
@@ -1775,6 +1777,7 @@ namespace mu
     	OpBind->bReshapePhysicsVolumes = Node.bReshapePhysicsVolumes;
 		OpBind->PhysicsToDeform = Node.PhysicsToDeform;
 		OpBind->bReshapeVertices = Node.bReshapeVertices;
+		OpBind->bRecomputeNormals = Node.bRecomputeNormals;	
 		OpBind->bApplyLaplacian = Node.bApplyLaplacian;
 		OpBind->BindingMethod = static_cast<uint32>(EShapeBindingMethod::ReshapeClosestProject);
 
@@ -1784,6 +1787,7 @@ namespace mu
 		OpBind->AChannelUsage = Node.ColorAChannelUsage;
 
 		OpApply->bReshapeVertices = OpBind->bReshapeVertices;
+		OpApply->bRecomputeNormals = OpBind->bRecomputeNormals;
 		OpApply->bReshapeSkeleton = OpBind->bReshapeSkeleton;
 		OpApply->bApplyLaplacian = OpBind->bApplyLaplacian;
 		OpApply->bReshapePhysicsVolumes = OpBind->bReshapePhysicsVolumes;
diff --git a/Engine/Plugins/Experimental/Mutable/Source/MutableTools/Private/MuT/NodeMeshReshape.cpp b/Engine/Plugins/Experimental/Mutable/Source/MutableTools/Private/MuT/NodeMeshReshape.cpp
--- a/Engine/Plugins/Experimental/Mutable/Source/MutableTools/Private/MuT/NodeMeshReshape.cpp
+++ b/Engine/Plugins/Experimental/Mutable/Source/MutableTools/Private/MuT/NodeMeshReshape.cpp
@@ -77,6 +77,13 @@ namespace mu
 		m_pD->bReshapeVertices = bEnable;
 	}
 
+	//---------------------------------------------------------------------------------------------
+	void NodeMeshReshape::SetRecomputeNormals(bool bEnable)
+	{
+		m_pD->bRecomputeNormals = bEnable;
+	}
+
+
 	//---------------------------------------------------------------------------------------------
 	void NodeMeshReshape::SetApplyLaplacian(bool bEnable)
 	{
diff --git a/Engine/Plugins/Experimental/Mutable/Source/MutableTools/Private/MuT/NodeMeshReshapePrivate.h b/Engine/Plugins/Experimental/Mutable/Source/MutableTools/Private/MuT/NodeMeshReshapePrivate.h
--- a/Engine/Plugins/Experimental/Mutable/Source/MutableTools/Private/MuT/NodeMeshReshapePrivate.h
+++ b/Engine/Plugins/Experimental/Mutable/Source/MutableTools/Private/MuT/NodeMeshReshapePrivate.h
@@ -21,6 +21,7 @@ namespace mu
 		Ptr<NodeMesh> BaseShape;
 		Ptr<NodeMesh> TargetShape;
 		bool bReshapeVertices = true;
+		bool bRecomputeNormals = false;
 		bool bApplyLaplacian = false;
 		bool bReshapeSkeleton = false;
 		bool bReshapePhysicsVolumes = false;
diff --git a/Engine/Plugins/Experimental/Mutable/Source/MutableTools/Public/MuT/NodeMeshReshape.h b/Engine/Plugins/Experimental/Mutable/Source/MutableTools/Public/MuT/NodeMeshReshape.h
--- a/Engine/Plugins/Experimental/Mutable/Source/MutableTools/Public/MuT/NodeMeshReshape.h
+++ b/Engine/Plugins/Experimental/Mutable/Source/MutableTools/Public/MuT/NodeMeshReshape.h
@@ -41,6 +41,10 @@ namespace mu
 
 		void SetReshapeVertices(bool);
 
+		/** Recompute normals after reshaping. This will ignore the reshaped normals. Disabled by default. */
+		void SetRecomputeNormals(bool);
+
+
 		/** Also deform the mesh skeleton. Disabled by default. */
 		void SetReshapeSkeleton(bool);
 	
