[Mutable] Fix cloth crash when clothing assets have more lods than the generated skeletal mesh.

#rnx

diff --git a/Engine/Plugins/Experimental/Mutable/Source/CustomizableObject/Private/MuCO/CustomizableObjectInstance.cpp b/Engine/Plugins/Experimental/Mutable/Source/CustomizableObject/Private/MuCO/CustomizableObjectInstance.cpp
--- a/Engine/Plugins/Experimental/Mutable/Source/CustomizableObject/Private/MuCO/CustomizableObjectInstance.cpp
+++ b/Engine/Plugins/Experimental/Mutable/Source/CustomizableObject/Private/MuCO/CustomizableObjectInstance.cpp
@@ -4567,6 +4567,59 @@ void UCustomizableInstancePrivate::BuildOrCopyClothingData(const TSharedRef<FUpd
 		}
 	}
 
+	// Remove empty LODs and redo LodMaps.
+	{
+		const int32 NumClothingAssets = NewClothingAssetsData.Num();
+		for (int32 I = 0; I < NumClothingAssets; ++I)
+		{
+			// Skip assets not set.
+			if (!NewClothingAssetsData[I].LodData.Num())
+			{
+				continue;
+			}
+
+			TArray<int32>& LodMap = NewClothingAssetsData[I].LodMap;
+			TArray<FClothLODDataCommon>& LodData = NewClothingAssetsData[I].LodData;
+
+			const int32 LodMapNum = LodMap.Num();
+			for (int32 LODIndex = 0; LODIndex < LodMapNum; ++LODIndex)
+			{
+				int32 MappedLODIndex = LodMap[LODIndex];
+				
+				if (MappedLODIndex != INDEX_NONE && !NewClothingAssetsData[I].LodData[MappedLODIndex].PhysicalMeshData.Vertices.Num())
+				{
+					LodMap[LODIndex] = INDEX_NONE;
+				}
+			}
+
+			TArray<int32> RemappedLODDataIndices;
+			RemappedLODDataIndices.Init(INDEX_NONE, LodData.Num());
+
+			TArray<FClothLODDataCommon> TrimmedLodData;
+			TrimmedLodData.Reserve(NewClothingAssetsData[I].LodData.Num());
+
+			for (int32 LODIndex = 0; LODIndex < LodMapNum; ++LODIndex)
+			{
+				const int32 MappedLODIndex = LodMap[LODIndex];
+				if (MappedLODIndex != INDEX_NONE)
+				{
+					if (RemappedLODDataIndices[MappedLODIndex] == INDEX_NONE)
+					{
+						const int32 RemappedIndex = TrimmedLodData.Emplace(MoveTemp(LodData[MappedLODIndex]));
+						LodMap[LODIndex] = RemappedIndex; 
+						RemappedLODDataIndices[MappedLODIndex] = RemappedIndex;
+					}
+					else
+					{
+						LodMap[LODIndex] = RemappedLODDataIndices[MappedLODIndex];
+					}
+				}
+			}
+
+			LodData = MoveTemp(TrimmedLodData);
+		}
+	}
+	
 	// From here up, could be moved to an async task similar to what is done with the other prepare tasks.  
  
 	// Create Clothing Assets.
@@ -4717,7 +4770,7 @@ void UCustomizableInstancePrivate::BuildOrCopyClothingData(const TSharedRef<FUpd
 					: FName(FString::Printf(TEXT("%s_%d"), *NewClothingAssetsData[I].Name.ToString(), I));;
 
 			NewClothingAssets[I] = NewObject<UCustomizableObjectClothingAsset>(SkeletalMesh, ClothingAssetObjectName);
-			
+
 			// The data can be moved to the actual asset since it will not be used anymore.
 			NewClothingAssets[I]->LodMap = MoveTemp(NewClothingAssetsData[I].LodMap);
 			NewClothingAssets[I]->LodData = MoveTemp(NewClothingAssetsData[I].LodData);
