[Mutable] Fix static assert with CacheRuntimeReferencedImages.

#rnx

diff --git a/Engine/Plugins/Experimental/Mutable/Source/CustomizableObject/Private/MuCO/UnrealMutableImageProvider.cpp b/Engine/Plugins/Experimental/Mutable/Source/CustomizableObject/Private/MuCO/UnrealMutableImageProvider.cpp
--- a/Engine/Plugins/Experimental/Mutable/Source/CustomizableObject/Private/MuCO/UnrealMutableImageProvider.cpp
+++ b/Engine/Plugins/Experimental/Mutable/Source/CustomizableObject/Private/MuCO/UnrealMutableImageProvider.cpp
@@ -700,7 +700,7 @@ void FUnrealMutableImageProvider::CacheRuntimeReferencedImages(const TSharedRef<
 		UTexture2D* Texture = RuntimeReferencedTexture.LoadSynchronous();
 		if (!Texture)
 		{
-			UE_LOG(LogMutable, Warning, TEXT("Failed to load texture [%s]."), *Texture->GetName());
+			UE_LOG(LogMutable, Warning, TEXT("Failed to load texture [%s]."), *RuntimeReferencedTexture->GetPathName());
 			continue;
 		}
 
