[mutable] Fix bug with empty index intervals on mesh removal.
#rnx

diff --git a/Engine/Plugins/Experimental/Mutable/Source/MutableRuntime/Private/MuR/OpMeshRemove.cpp b/Engine/Plugins/Experimental/Mutable/Source/MutableRuntime/Private/MuR/OpMeshRemove.cpp
--- a/Engine/Plugins/Experimental/Mutable/Source/MutableRuntime/Private/MuR/OpMeshRemove.cpp
+++ b/Engine/Plugins/Experimental/Mutable/Source/MutableRuntime/Private/MuR/OpMeshRemove.cpp
@@ -53,7 +53,7 @@ namespace mu
             }
         }
 
-        if (current.idStart>=0)
+        if (current.idStart!=TNumericLimits<uint64>::Max())
         {
             intervals.Add(current);
         }
