[mutable] Fix issue with vertex IDs when a mesh merges with itself, which is a corner-case that was crashing. Now, the second mesh IDs will be changed, to mimic old behaviour.
Fix a crash in the debugger related to missing data for the case of compilation for debugging.
#rnx

diff --git a/Engine/Plugins/Experimental/Mutable/Source/CustomizableObjectEditor/Private/MuCOE/CustomizableObjectCompiler.cpp b/Engine/Plugins/Experimental/Mutable/Source/CustomizableObjectEditor/Private/MuCOE/CustomizableObjectCompiler.cpp
--- a/Engine/Plugins/Experimental/Mutable/Source/CustomizableObjectEditor/Private/MuCOE/CustomizableObjectCompiler.cpp
+++ b/Engine/Plugins/Experimental/Mutable/Source/CustomizableObjectEditor/Private/MuCOE/CustomizableObjectCompiler.cpp
@@ -1300,7 +1300,10 @@ mu::NodePtr FCustomizableObjectCompiler::Export(UCustomizableObject* Object, con
 	CompilerOptions.bAnimBpPhysicsManipulationEnabled = Object->bEnableAnimBpPhysicsAssetsManipualtion;
 
 	FMutableGraphGenerationContext GenerationContext(Object, this, CompilerOptions);
-	GenerationContext.ParamNamesToSelectedOptions = CurrentRequest->GetParameterNamesToSelectedOptions();
+	if (CurrentRequest)
+	{
+		GenerationContext.ParamNamesToSelectedOptions = CurrentRequest->GetParameterNamesToSelectedOptions();
+	}
 
 	// Generate the mutable node expression
 	FText ErrorMsg;
diff --git a/Engine/Plugins/Experimental/Mutable/Source/MutableRuntime/Internal/MuR/OpMeshMerge.h b/Engine/Plugins/Experimental/Mutable/Source/MutableRuntime/Internal/MuR/OpMeshMerge.h
--- a/Engine/Plugins/Experimental/Mutable/Source/MutableRuntime/Internal/MuR/OpMeshMerge.h
+++ b/Engine/Plugins/Experimental/Mutable/Source/MutableRuntime/Internal/MuR/OpMeshMerge.h
@@ -692,6 +692,11 @@ namespace mu
 			bool bFastPath = pFirst->GetVertexBuffers().HasSameFormat( pSecond->GetVertexBuffers() );
 
 			bool bNeedsExplicitVertexIds = pFirst->VertexIDPrefix != pSecond->VertexIDPrefix;
+			if (!bNeedsExplicitVertexIds)
+			{
+				Result->VertexIDPrefix = pFirst->VertexIDPrefix;
+			}
+
 			bFastPath = bFastPath && !bNeedsExplicitVertexIds;
 
 			// Check if the format of the BoneIndex buffer has to change
