Adding missing includes.

diff --git a/Engine/Plugins/Experimental/Mutable/Source/MutableRuntime/Private/MuR/OpMeshComputeNormals.h b/Engine/Plugins/Experimental/Mutable/Source/MutableRuntime/Private/MuR/OpMeshComputeNormals.h
--- a/Engine/Plugins/Experimental/Mutable/Source/MutableRuntime/Private/MuR/OpMeshComputeNormals.h
+++ b/Engine/Plugins/Experimental/Mutable/Source/MutableRuntime/Private/MuR/OpMeshComputeNormals.h
@@ -8,6 +8,7 @@
 #include "Containers/Array.h"
 #include "Containers/ArrayView.h"
 #include "Containers/Map.h"
+#include "IndexTypes.h"
 
 #include "MuR/Mesh.h"
 #include "MuR/MeshPrivate.h"
diff --git a/Engine/Source/Runtime/Core/Public/Math/LFSR.h b/Engine/Source/Runtime/Core/Public/Math/LFSR.h
--- a/Engine/Source/Runtime/Core/Public/Math/LFSR.h
+++ b/Engine/Source/Runtime/Core/Public/Math/LFSR.h
@@ -3,6 +3,7 @@
 #pragma once
 
 #include "CoreTypes.h"
+#include "Misc/AssertionMacros.h"
 
 namespace UE
 {
