[Mutable] Add missing CO compile options.

#rnx

diff --git a/Engine/Plugins/Experimental/Mutable/Source/CustomizableObject/Private/MuCO/CustomizableObject.cpp b/Engine/Plugins/Experimental/Mutable/Source/CustomizableObject/Private/MuCO/CustomizableObject.cpp
--- a/Engine/Plugins/Experimental/Mutable/Source/CustomizableObject/Private/MuCO/CustomizableObject.cpp
+++ b/Engine/Plugins/Experimental/Mutable/Source/CustomizableObject/Private/MuCO/CustomizableObject.cpp
@@ -2737,6 +2737,12 @@ FCompilationRequest::FCompilationRequest(UCustomizableObject& InCustomizableObje
 {
 	CustomizableObject = &InCustomizableObject;
 	Options = InCustomizableObject.CompileOptions;
+	Options.bRealTimeMorphTargetsEnabled = InCustomizableObject.bEnableRealTimeMorphTargets;
+	Options.bClothingEnabled = InCustomizableObject.bEnableClothing;
+	Options.b16BitBoneWeightsEnabled = InCustomizableObject.bEnable16BitBoneWeights;
+	Options.bSkinWeightProfilesEnabled = InCustomizableObject.bEnableAltSkinWeightProfiles;
+	Options.bPhysicsAssetMergeEnabled = InCustomizableObject.bEnablePhysicsAssetMerge;
+	Options.bAnimBpPhysicsManipulationEnabled = InCustomizableObject.bEnableAnimBpPhysicsAssetsManipualtion;
 	bAsync = bAsyncCompile;
 }
 
