[Mutable] Fix hang when Mutable Streaming Manager is blocking.

#rnx

diff --git a/Engine/Plugins/Experimental/Mutable/Source/CustomizableObject/Internal/MuCO/CustomizableObjectSystemPrivate.h b/Engine/Plugins/Experimental/Mutable/Source/CustomizableObject/Internal/MuCO/CustomizableObjectSystemPrivate.h
--- a/Engine/Plugins/Experimental/Mutable/Source/CustomizableObject/Internal/MuCO/CustomizableObjectSystemPrivate.h
+++ b/Engine/Plugins/Experimental/Mutable/Source/CustomizableObject/Internal/MuCO/CustomizableObjectSystemPrivate.h
@@ -158,14 +158,10 @@ class FMutablePendingInstanceWork
 
 	TSet<mu::Instance::ID> PendingIDsToRelease;
 
-	int32 NumLODUpdatesLastTick = 0;
-
 public:
-	// Returns the number of pending instance updates, LOD Updates, discards and releases last tick.
+	// Returns the number of pending instance updates, LOD Updates and discards.
 	int32 Num() const;
-
-	void SetLODUpdatesLastTick(int32 NumLODUpdates);
-
+	
 	// Adds a new instance update
 	void AddUpdate(const FMutablePendingInstanceUpdate& UpdateToAdd);
 
@@ -803,6 +799,8 @@ public:
 	/** Map used to cache per platform MaxChunkSize. If MaxChunkSize > 0, streamed data will be split in multiple files */
 	TMap<FString, int64> PlatformMaxChunkSize;
 #endif
+
+	int32 NumLODUpdatesLastTick = 0;
 };
 
 
diff --git a/Engine/Plugins/Experimental/Mutable/Source/CustomizableObject/Private/MuCO/CustomizableObjectSystem.cpp b/Engine/Plugins/Experimental/Mutable/Source/CustomizableObject/Private/MuCO/CustomizableObjectSystem.cpp
--- a/Engine/Plugins/Experimental/Mutable/Source/CustomizableObject/Private/MuCO/CustomizableObjectSystem.cpp
+++ b/Engine/Plugins/Experimental/Mutable/Source/CustomizableObject/Private/MuCO/CustomizableObjectSystem.cpp
@@ -297,13 +297,7 @@ uint32 FPendingInstanceUpdateKeyFuncs::GetKeyHash(const TWeakObjectPtr<const UCu
 
 int32 FMutablePendingInstanceWork::Num() const
 {
-	return PendingInstanceUpdates.Num() + PendingInstanceDiscards.Num() + PendingIDsToRelease.Num() + NumLODUpdatesLastTick;
-}
-
-
-void FMutablePendingInstanceWork::SetLODUpdatesLastTick(int32 NumLODUpdates)
-{
-	NumLODUpdatesLastTick = NumLODUpdates;
+	return PendingInstanceUpdates.Num() + PendingInstanceDiscards.Num() + PendingIDsToRelease.Num();
 }
 
 
@@ -3608,7 +3602,7 @@ int32 UCustomizableObjectSystem::TickInternal()
 				}
 			}
 
-			Private->MutablePendingInstanceWork.SetLODUpdatesLastTick(RequestedLODUpdates.Num());
+			Private->NumLODUpdatesLastTick = RequestedLODUpdates.Num();
 
 			// If the chosen LODUpdate has the same instance as a PendingUpdate, choose the PendingUpdate to apply both the LOD update
 			// and customization change
@@ -3866,11 +3860,13 @@ int32 UCustomizableObjectSystem::GetNumInstances() const
 	return NumBuiltInstances;
 }
 
+
 int32 UCustomizableObjectSystem::GetNumPendingInstances() const
 {
-	return GetPrivate()->MutablePendingInstanceWork.Num();
+	return GetPrivate()->MutablePendingInstanceWork.Num() + GetPrivate()->NumLODUpdatesLastTick;
 }
 
+
 int32 UCustomizableObjectSystem::GetTotalInstances() const
 {
 	int32 NumInstances = 0;
diff --git a/Engine/Plugins/Experimental/Mutable/Source/CustomizableObject/Private/MuCO/LogBenchmarkUtil.cpp b/Engine/Plugins/Experimental/Mutable/Source/CustomizableObject/Private/MuCO/LogBenchmarkUtil.cpp
--- a/Engine/Plugins/Experimental/Mutable/Source/CustomizableObject/Private/MuCO/LogBenchmarkUtil.cpp
+++ b/Engine/Plugins/Experimental/Mutable/Source/CustomizableObject/Private/MuCO/LogBenchmarkUtil.cpp
@@ -157,7 +157,7 @@ void FLogBenchmarkUtil::UpdateStats()
 	NumInstancesLOD2 = LocalNumInstancesLOD2;
 	NumAllocatedSkeletalMeshes = LocalNumAllocatedSkeletalMeshes;
 
-	NumPendingInstanceUpdates = System->GetPrivate()->MutablePendingInstanceWork.Num();
+	NumPendingInstanceUpdates = System->GetNumPendingInstances();
 
 	// Textures
 	uint32 LocalNumAllocatedTextures = 0;
