[mutable] Added a new node to directly generate textures from float channels instead of having to use several intermediate nodes.
#rnx
#tests Editor

diff --git a/Engine/Plugins/Experimental/Mutable/Source/CustomizableObjectEditor/Classes/MuCOE/Nodes/CustomizableObjectNodeTextureFromFloats.h b/Engine/Plugins/Experimental/Mutable/Source/CustomizableObjectEditor/Classes/MuCOE/Nodes/CustomizableObjectNodeTextureFromFloats.h
new file mode 100644
--- /dev/null
+++ b/Engine/Plugins/Experimental/Mutable/Source/CustomizableObjectEditor/Classes/MuCOE/Nodes/CustomizableObjectNodeTextureFromFloats.h
@@ -0,0 +1,50 @@
+// Copyright Epic Games, Inc. All Rights Reserved.
+
+#pragma once
+
+#include "MuCOE/Nodes/CustomizableObjectNode.h"
+
+#include "CustomizableObjectNodeTextureFromFloats.generated.h"
+
+namespace ENodeTitleType { enum Type : int; }
+
+class UCustomizableObjectNodeRemapPins;
+class UEdGraphPin;
+class UObject;
+
+
+UCLASS()
+class CUSTOMIZABLEOBJECTEDITOR_API UCustomizableObjectNodeTextureFromFloats : public UCustomizableObjectNode
+{
+public:
+	GENERATED_BODY()
+
+	// Begin EdGraphNode interface
+	FText GetNodeTitle(ENodeTitleType::Type TitleType) const override;
+	FLinearColor GetNodeTitleColor() const override;
+	FText GetTooltipText() const override;
+
+	// UCustomizableObjectNode interface
+	virtual void AllocateDefaultPins(UCustomizableObjectNodeRemapPins* RemapPins) override;
+
+	UEdGraphPin* RPin() const
+	{
+		return FindPin(TEXT("R"));
+	}
+
+	UEdGraphPin* GPin() const
+	{
+		return FindPin(TEXT("G"));
+	}
+
+	UEdGraphPin* BPin() const
+	{
+		return FindPin(TEXT("B"));
+	}
+
+	UEdGraphPin* APin() const
+	{
+		return FindPin(TEXT("A"));
+	}
+};
+
diff --git a/Engine/Plugins/Experimental/Mutable/Source/CustomizableObjectEditor/Private/MuCOE/EdGraphSchema_CustomizableObject.cpp b/Engine/Plugins/Experimental/Mutable/Source/CustomizableObjectEditor/Private/MuCOE/EdGraphSchema_CustomizableObject.cpp
--- a/Engine/Plugins/Experimental/Mutable/Source/CustomizableObjectEditor/Private/MuCOE/EdGraphSchema_CustomizableObject.cpp
+++ b/Engine/Plugins/Experimental/Mutable/Source/CustomizableObjectEditor/Private/MuCOE/EdGraphSchema_CustomizableObject.cpp
@@ -60,6 +60,7 @@
 #include "MuCOE/Nodes/CustomizableObjectNodeTextureColourMap.h"
 #include "MuCOE/Nodes/CustomizableObjectNodeTextureFromChannels.h"
 #include "MuCOE/Nodes/CustomizableObjectNodeTextureFromColor.h"
+#include "MuCOE/Nodes/CustomizableObjectNodeTextureFromFloats.h"
 #include "MuCOE/Nodes/CustomizableObjectNodeTextureInterpolate.h"
 #include "MuCOE/Nodes/CustomizableObjectNodeTextureInvert.h"
 #include "MuCOE/Nodes/CustomizableObjectNodeTextureLayer.h"
@@ -430,6 +431,7 @@ void UEdGraphSchema_CustomizableObject::GetGraphContextActions(FGraphContextMenu
 			ContextMenuBuilder.CreateTemplateNode<UCustomizableObjectNodeTextureToChannels>(),
 			ContextMenuBuilder.CreateTemplateNode<UCustomizableObjectNodeTextureFromChannels>(),
 			ContextMenuBuilder.CreateTemplateNode<UCustomizableObjectNodeTextureFromColor>(),
+			ContextMenuBuilder.CreateTemplateNode<UCustomizableObjectNodeTextureFromFloats>(),
 			ContextMenuBuilder.CreateTemplateNode<UCustomizableObjectNodeTextureProject>(),
 			ContextMenuBuilder.CreateTemplateNode<UCustomizableObjectNodeTextureParameter>(),
 			ContextMenuBuilder.CreateTemplateNode<UCustomizableObjectNodeTextureInvert>(),
diff --git a/Engine/Plugins/Experimental/Mutable/Source/CustomizableObjectEditor/Private/MuCOE/GenerateMutableSource/GenerateMutableSourceImage.cpp b/Engine/Plugins/Experimental/Mutable/Source/CustomizableObjectEditor/Private/MuCOE/GenerateMutableSource/GenerateMutableSourceImage.cpp
--- a/Engine/Plugins/Experimental/Mutable/Source/CustomizableObjectEditor/Private/MuCOE/GenerateMutableSource/GenerateMutableSourceImage.cpp
+++ b/Engine/Plugins/Experimental/Mutable/Source/CustomizableObjectEditor/Private/MuCOE/GenerateMutableSource/GenerateMutableSourceImage.cpp
@@ -18,6 +18,7 @@
 #include "MuCOE/Nodes/CustomizableObjectNodeTextureColourMap.h"
 #include "MuCOE/Nodes/CustomizableObjectNodeTextureFromChannels.h"
 #include "MuCOE/Nodes/CustomizableObjectNodeTextureFromColor.h"
+#include "MuCOE/Nodes/CustomizableObjectNodeTextureFromFloats.h"
 #include "MuCOE/Nodes/CustomizableObjectNodeTextureInterpolate.h"
 #include "MuCOE/Nodes/CustomizableObjectNodeTextureInvert.h"
 #include "MuCOE/Nodes/CustomizableObjectNodeTextureLayer.h"
@@ -40,6 +41,7 @@
 #include "MuT/NodeImageLayerColour.h"
 #include "MuT/NodeImageParameter.h"
 #include "MuT/NodeImagePlainColour.h"
+#include "MuT/NodeColourFromScalars.h"
 #include "MuT/NodeImageProject.h"
 #include "MuT/NodeImageResize.h"
 #include "MuT/NodeImageSaturate.h"
@@ -442,6 +444,48 @@ mu::NodeImagePtr GenerateMutableSourceImage(const UEdGraphPin* Pin, FMutableGrap
 		}
 	}
 
+	else if (const UCustomizableObjectNodeTextureFromFloats* TypedNodeFromFloats = Cast<UCustomizableObjectNodeTextureFromFloats>(Node))
+	{		
+		mu::NodeColourFromScalarsPtr Color = new mu::NodeColourFromScalars();
+
+		if (const UEdGraphPin* RPin = FollowInputPin(*TypedNodeFromFloats->RPin()))
+		{
+			mu::NodeScalarPtr R = GenerateMutableSourceFloat(RPin, GenerationContext);
+			Color->SetX(R);
+		}
+
+		if (const UEdGraphPin* GPin = FollowInputPin(*TypedNodeFromFloats->GPin()))
+		{
+			mu::NodeScalarPtr G = GenerateMutableSourceFloat(GPin, GenerationContext);
+			Color->SetY(G);
+		}
+
+		if (const UEdGraphPin* BPin = FollowInputPin(*TypedNodeFromFloats->BPin()))
+		{
+			mu::NodeScalarPtr B = GenerateMutableSourceFloat(BPin, GenerationContext);
+			Color->SetZ(B);
+		}
+
+		if (const UEdGraphPin* APin = FollowInputPin(*TypedNodeFromFloats->APin()))
+		{
+			mu::NodeScalarPtr A = GenerateMutableSourceFloat(APin, GenerationContext);
+			Color->SetW(A);
+		}
+
+		mu::NodeImagePlainColourPtr ImageFromColour = new mu::NodeImagePlainColour;
+		Result = ImageFromColour;
+
+		if (Color)
+		{
+			ImageFromColour->SetColour(Color);
+		}
+
+		if (ReferenceTextureSize > 0)
+		{
+			ImageFromColour->SetSize(ReferenceTextureSize, ReferenceTextureSize);
+		}
+	}
+
 	else if (const UCustomizableObjectNodeTextureFromChannels* TypedNodeFrom = Cast<UCustomizableObjectNodeTextureFromChannels>(Node))
 	{
 		mu::NodeImagePtr RNode, GNode, BNode, ANode;
diff --git a/Engine/Plugins/Experimental/Mutable/Source/CustomizableObjectEditor/Private/MuCOE/Nodes/CustomizableObjectNodeTextureFromFloats.cpp b/Engine/Plugins/Experimental/Mutable/Source/CustomizableObjectEditor/Private/MuCOE/Nodes/CustomizableObjectNodeTextureFromFloats.cpp
new file mode 100644
--- /dev/null
+++ b/Engine/Plugins/Experimental/Mutable/Source/CustomizableObjectEditor/Private/MuCOE/Nodes/CustomizableObjectNodeTextureFromFloats.cpp
@@ -0,0 +1,57 @@
+// Copyright Epic Games, Inc. All Rights Reserved.
+
+#include "MuCOE/Nodes/CustomizableObjectNodeTextureFromFloats.h"
+
+#include "MuCO/CustomizableObjectCustomVersion.h"
+#include "MuCOE/EdGraphSchema_CustomizableObject.h"
+
+class UCustomizableObjectNodeRemapPins;
+
+#define LOCTEXT_NAMESPACE "CustomizableObjectEditor"
+
+
+void UCustomizableObjectNodeTextureFromFloats::AllocateDefaultPins(UCustomizableObjectNodeRemapPins* RemapPins)
+{
+	const UEdGraphSchema_CustomizableObject* Schema = GetDefault<UEdGraphSchema_CustomizableObject>();
+
+	FString PinName = TEXT("Texture");
+	UEdGraphPin* OutputPin = CustomCreatePin(EGPD_Output, Schema->PC_Image, FName(*PinName));
+	OutputPin->bDefaultValueIsIgnored = true;
+
+	PinName = TEXT("R");
+	UEdGraphPin* RPin = CustomCreatePin(EGPD_Input, Schema->PC_Float, FName(*PinName));
+	RPin->bDefaultValueIsIgnored = true;
+
+	PinName = TEXT("G");
+	UEdGraphPin* GPin = CustomCreatePin(EGPD_Input, Schema->PC_Float, FName(*PinName));
+	GPin->bDefaultValueIsIgnored = true;
+
+	PinName = TEXT("B");
+	UEdGraphPin* BPin = CustomCreatePin(EGPD_Input, Schema->PC_Float, FName(*PinName));
+	BPin->bDefaultValueIsIgnored = true;
+
+	PinName = TEXT("A");
+	UEdGraphPin* APin = CustomCreatePin(EGPD_Input, Schema->PC_Float, FName(*PinName));
+	APin->bDefaultValueIsIgnored = true;
+}
+
+
+FText UCustomizableObjectNodeTextureFromFloats::GetNodeTitle(ENodeTitleType::Type TitleType) const
+{
+	return LOCTEXT("Texture_From_Floats", "Texture From Float Channels");
+}
+
+
+FLinearColor UCustomizableObjectNodeTextureFromFloats::GetNodeTitleColor() const
+{
+	const UEdGraphSchema_CustomizableObject* Schema = GetDefault<UEdGraphSchema_CustomizableObject>();
+	return Schema->GetPinTypeColor(Schema->PC_Image);
+}
+
+
+FText UCustomizableObjectNodeTextureFromFloats::GetTooltipText() const
+{
+	return LOCTEXT("Texture_From_Floats_Tooltip", "Creates a flat color texture from the float channels provided.");
+}
+
+#undef LOCTEXT_NAMESPACE
