[mutable] Now parameter nodes show the name of the parameter on the title of the node.
+ Added methods to rename the nodes from the graph editor (double click or pressing F2)

#rnx

diff --git a/Engine/Plugins/Experimental/Mutable/Source/CustomizableObject/Internal/MuCO/CustomizableObjectCustomVersion.h b/Engine/Plugins/Experimental/Mutable/Source/CustomizableObject/Internal/MuCO/CustomizableObjectCustomVersion.h
--- a/Engine/Plugins/Experimental/Mutable/Source/CustomizableObject/Internal/MuCO/CustomizableObjectCustomVersion.h
+++ b/Engine/Plugins/Experimental/Mutable/Source/CustomizableObject/Internal/MuCO/CustomizableObjectCustomVersion.h
@@ -116,6 +116,8 @@ struct CUSTOMIZABLEOBJECT_API FCustomizableObjectCustomVersion
 
 		StateUIMetadata,
 
+		AddedRenameOptionToParameterNodes,
+
 		// -----<new versions can be added above this line>-------------------------------------------------
 		VersionPlusOne,
 		LatestVersion = VersionPlusOne - 1
diff --git a/Engine/Plugins/Experimental/Mutable/Source/CustomizableObjectEditor/Classes/MuCOE/Nodes/CustomizableObjectNodeColorParameter.h b/Engine/Plugins/Experimental/Mutable/Source/CustomizableObjectEditor/Classes/MuCOE/Nodes/CustomizableObjectNodeColorParameter.h
--- a/Engine/Plugins/Experimental/Mutable/Source/CustomizableObjectEditor/Classes/MuCOE/Nodes/CustomizableObjectNodeColorParameter.h
+++ b/Engine/Plugins/Experimental/Mutable/Source/CustomizableObjectEditor/Classes/MuCOE/Nodes/CustomizableObjectNodeColorParameter.h
@@ -31,9 +31,12 @@ public:
 	FText GetNodeTitle(ENodeTitleType::Type TitleType) const override;
 	FLinearColor GetNodeTitleColor() const override;
 	FText GetTooltipText() const override;
+	virtual void OnRenameNode(const FString& NewName) override;
+	virtual bool GetCanRenameNode() const override { return true; }
 
 	// UCustomizableObjectNode interface
 	void AllocateDefaultPins(UCustomizableObjectNodeRemapPins* RemapPins) override;
 	bool IsAffectedByLOD() const override { return false; }
+	virtual void BackwardsCompatibleFixup() override;
 };
 
diff --git a/Engine/Plugins/Experimental/Mutable/Source/CustomizableObjectEditor/Classes/MuCOE/Nodes/CustomizableObjectNodeEnumParameter.h b/Engine/Plugins/Experimental/Mutable/Source/CustomizableObjectEditor/Classes/MuCOE/Nodes/CustomizableObjectNodeEnumParameter.h
--- a/Engine/Plugins/Experimental/Mutable/Source/CustomizableObjectEditor/Classes/MuCOE/Nodes/CustomizableObjectNodeEnumParameter.h
+++ b/Engine/Plugins/Experimental/Mutable/Source/CustomizableObjectEditor/Classes/MuCOE/Nodes/CustomizableObjectNodeEnumParameter.h
@@ -51,9 +51,13 @@ public:
 	FText GetNodeTitle(ENodeTitleType::Type TitleType) const override;
 	FLinearColor GetNodeTitleColor() const override;
 	FText GetTooltipText() const override;
+	virtual void OnRenameNode(const FString& NewName) override;
+	virtual bool GetCanRenameNode() const override { return true; }
 
 	// UCustomizableObjectNode interface
 	void AllocateDefaultPins(UCustomizableObjectNodeRemapPins* RemapPins) override;
 	virtual bool IsAffectedByLOD() const override { return false; }
+	virtual void BackwardsCompatibleFixup() override;
+
 };
 
diff --git a/Engine/Plugins/Experimental/Mutable/Source/CustomizableObjectEditor/Classes/MuCOE/Nodes/CustomizableObjectNodeFloatParameter.h b/Engine/Plugins/Experimental/Mutable/Source/CustomizableObjectEditor/Classes/MuCOE/Nodes/CustomizableObjectNodeFloatParameter.h
--- a/Engine/Plugins/Experimental/Mutable/Source/CustomizableObjectEditor/Classes/MuCOE/Nodes/CustomizableObjectNodeFloatParameter.h
+++ b/Engine/Plugins/Experimental/Mutable/Source/CustomizableObjectEditor/Classes/MuCOE/Nodes/CustomizableObjectNodeFloatParameter.h
@@ -48,11 +48,13 @@ public:
 	FText GetNodeTitle(ENodeTitleType::Type TitleType) const override;
 	FLinearColor GetNodeTitleColor() const override;
 	FText GetTooltipText() const override;
+	virtual bool GetCanRenameNode() const override { return true; }
 
 	// UCustomizableObjectNode interface
 	void AllocateDefaultPins(UCustomizableObjectNodeRemapPins* RemapPins) override;
 	bool IsAffectedByLOD() const override;
 	void BackwardsCompatibleFixup() override;
+	virtual void OnRenameNode(const FString& NewName) override;
 
 };
 
diff --git a/Engine/Plugins/Experimental/Mutable/Source/CustomizableObjectEditor/Classes/MuCOE/Nodes/CustomizableObjectNodeProjectorParameter.h b/Engine/Plugins/Experimental/Mutable/Source/CustomizableObjectEditor/Classes/MuCOE/Nodes/CustomizableObjectNodeProjectorParameter.h
--- a/Engine/Plugins/Experimental/Mutable/Source/CustomizableObjectEditor/Classes/MuCOE/Nodes/CustomizableObjectNodeProjectorParameter.h
+++ b/Engine/Plugins/Experimental/Mutable/Source/CustomizableObjectEditor/Classes/MuCOE/Nodes/CustomizableObjectNodeProjectorParameter.h
@@ -54,6 +54,8 @@ public:
 	FText GetNodeTitle(ENodeTitleType::Type TitleType) const override;
 	FLinearColor GetNodeTitleColor() const override;
 	FText GetTooltipText() const override;
+	virtual bool GetCanRenameNode() const override { return true; }
+	virtual void OnRenameNode(const FString& NewName) override;
 
 	// UCustomizableObjectNode interface
 	void AllocateDefaultPins(UCustomizableObjectNodeRemapPins* RemapPins) override;
diff --git a/Engine/Plugins/Experimental/Mutable/Source/CustomizableObjectEditor/Classes/MuCOE/Nodes/CustomizableObjectNodeTable.h b/Engine/Plugins/Experimental/Mutable/Source/CustomizableObjectEditor/Classes/MuCOE/Nodes/CustomizableObjectNodeTable.h
--- a/Engine/Plugins/Experimental/Mutable/Source/CustomizableObjectEditor/Classes/MuCOE/Nodes/CustomizableObjectNodeTable.h
+++ b/Engine/Plugins/Experimental/Mutable/Source/CustomizableObjectEditor/Classes/MuCOE/Nodes/CustomizableObjectNodeTable.h
@@ -248,6 +248,8 @@ public:
 	FLinearColor GetNodeTitleColor() const override;
 	FText GetTooltipText() const override;
 	virtual void PinConnectionListChanged(UEdGraphPin* Pin) override;
+	virtual void OnRenameNode(const FString& NewName) override;
+	virtual bool GetCanRenameNode() const override { return true; }
 	
 	// UCustomizableObjectNode interface
 	virtual void PostBackwardsCompatibleFixup() override;
diff --git a/Engine/Plugins/Experimental/Mutable/Source/CustomizableObjectEditor/Classes/MuCOE/Nodes/CustomizableObjectNodeTextureParameter.h b/Engine/Plugins/Experimental/Mutable/Source/CustomizableObjectEditor/Classes/MuCOE/Nodes/CustomizableObjectNodeTextureParameter.h
--- a/Engine/Plugins/Experimental/Mutable/Source/CustomizableObjectEditor/Classes/MuCOE/Nodes/CustomizableObjectNodeTextureParameter.h
+++ b/Engine/Plugins/Experimental/Mutable/Source/CustomizableObjectEditor/Classes/MuCOE/Nodes/CustomizableObjectNodeTextureParameter.h
@@ -45,6 +45,8 @@ public:
 	virtual FText GetNodeTitle(ENodeTitleType::Type TitleType) const override;
 	virtual FLinearColor GetNodeTitleColor() const override;
 	virtual FText GetTooltipText() const override;
+	virtual void OnRenameNode(const FString& NewName) override;
+	virtual bool GetCanRenameNode() const override { return true; }
 
 	// UCustomizableObjectNode interface
 	virtual void BackwardsCompatibleFixup() override;
diff --git a/Engine/Plugins/Experimental/Mutable/Source/CustomizableObjectEditor/Private/MuCOE/CustomizableObjectEditor.cpp b/Engine/Plugins/Experimental/Mutable/Source/CustomizableObjectEditor/Private/MuCOE/CustomizableObjectEditor.cpp
--- a/Engine/Plugins/Experimental/Mutable/Source/CustomizableObjectEditor/Private/MuCOE/CustomizableObjectEditor.cpp
+++ b/Engine/Plugins/Experimental/Mutable/Source/CustomizableObjectEditor/Private/MuCOE/CustomizableObjectEditor.cpp
@@ -545,6 +545,10 @@ void FCustomizableObjectEditor::CreateGraphEditorWidget(UEdGraph* InGraph)
 		FExecuteAction::CreateSP(this, &FCustomizableObjectEditor::DuplicateSelectedNodes),
 		FCanExecuteAction::CreateSP(this, &FCustomizableObjectEditor::CanDuplicateSelectedNodes));
 
+	GraphEditorCommands->MapAction(FGenericCommands::Get().Rename,
+		FExecuteAction::CreateSP(this, &FCustomizableObjectEditor::OnRenameNode),
+		FCanExecuteAction::CreateSP(this, &FCustomizableObjectEditor::CanRenameNodes));
+
 	GraphEditorCommands->MapAction(FCustomizableObjectEditorNodeContextCommands::Get().CreateComment,
 		FExecuteAction::CreateSP(this, &FCustomizableObjectEditor::CreateCommentBoxFromKey));
 
@@ -2475,5 +2479,44 @@ void FCustomizableObjectEditor::OnCustomizableObjectStatusChanged(FCustomizableO
 }
 
 
+void FCustomizableObjectEditor::OnRenameNode()
+{
+	if (GraphEditor.IsValid())
+	{
+		const FGraphPanelSelectionSet& SelectedNodes = GraphEditor->GetSelectedNodes();
+
+		for (FGraphPanelSelectionSet::TConstIterator NodeIt(SelectedNodes); NodeIt; ++NodeIt)
+		{
+			// Rename only the first valid selected node
+			UEdGraphNode* SelectedNode = Cast<UEdGraphNode>(*NodeIt);
+			if (SelectedNode && SelectedNode->GetCanRenameNode())
+			{
+				GraphEditor->IsNodeTitleVisible(SelectedNode, true);
+				break;
+			}
+		}
+	}
+}
+
+
+bool FCustomizableObjectEditor::CanRenameNodes() const
+{
+	if (GraphEditor.IsValid())
+	{
+		const FGraphPanelSelectionSet& SelectedNodes = GraphEditor->GetSelectedNodes();
+
+		for (FGraphPanelSelectionSet::TConstIterator SelectedIter(SelectedNodes); SelectedIter; ++SelectedIter)
+		{
+			if (UEdGraphNode* Node = Cast<UEdGraphNode>(*SelectedIter))
+			{
+				return Node->GetCanRenameNode();
+			}
+		}
+	}
+
+	return false;
+}
+
+
 
 #undef LOCTEXT_NAMESPACE
diff --git a/Engine/Plugins/Experimental/Mutable/Source/CustomizableObjectEditor/Private/MuCOE/CustomizableObjectEditor.h b/Engine/Plugins/Experimental/Mutable/Source/CustomizableObjectEditor/Private/MuCOE/CustomizableObjectEditor.h
--- a/Engine/Plugins/Experimental/Mutable/Source/CustomizableObjectEditor/Private/MuCOE/CustomizableObjectEditor.h
+++ b/Engine/Plugins/Experimental/Mutable/Source/CustomizableObjectEditor/Private/MuCOE/CustomizableObjectEditor.h
@@ -329,6 +329,12 @@ private:
 	/** Whether we are able to cut the currently selected nodes */
 	bool CanCutNodes() const;
 
+	/** On starting to rename node */
+	void OnRenameNode();
+
+	/** Check if node can be renamed */
+	bool CanRenameNodes() const;
+
 	virtual UEdGraphNode* CreateCommentBox(const FVector2D& NodePos) override;
 
 	/** Updates the visibility of PreviewSkeletalMeshComponent */
diff --git a/Engine/Plugins/Experimental/Mutable/Source/CustomizableObjectEditor/Private/MuCOE/Nodes/CustomizableObjectNodeColorParameter.cpp b/Engine/Plugins/Experimental/Mutable/Source/CustomizableObjectEditor/Private/MuCOE/Nodes/CustomizableObjectNodeColorParameter.cpp
--- a/Engine/Plugins/Experimental/Mutable/Source/CustomizableObjectEditor/Private/MuCOE/Nodes/CustomizableObjectNodeColorParameter.cpp
+++ b/Engine/Plugins/Experimental/Mutable/Source/CustomizableObjectEditor/Private/MuCOE/Nodes/CustomizableObjectNodeColorParameter.cpp
@@ -3,6 +3,7 @@
 #include "MuCOE/Nodes/CustomizableObjectNodeColorParameter.h"
 
 #include "MuCOE/EdGraphSchema_CustomizableObject.h"
+#include "MuCO/CustomizableObjectCustomVersion.h"
 
 class UCustomizableObjectNodeRemapPins;
 
@@ -20,7 +21,18 @@ void UCustomizableObjectNodeColorParameter::AllocateDefaultPins(UCustomizableObj
 
 FText UCustomizableObjectNodeColorParameter::GetNodeTitle(ENodeTitleType::Type TitleType) const
 {
-	return LOCTEXT("Color Parameter", "Color Parameter");
+	if (TitleType == ENodeTitleType::ListView || ParameterName.IsEmpty())
+	{
+		return LOCTEXT("Color_Parameter", "Color Parameter");
+	}
+	else if (TitleType == ENodeTitleType::EditableTitle)
+	{
+		return FText::Format(LOCTEXT("Color_Parameter_EditableTitle", "{0}"), FText::FromString(ParameterName));
+	}
+	else
+	{
+		return FText::Format(LOCTEXT("Color_Parameter_Title", "{0}\nColor Parameter"), FText::FromString(ParameterName));
+	}
 }
 
 
@@ -37,5 +49,26 @@ FText UCustomizableObjectNodeColorParameter::GetTooltipText() const
 }
 
 
+void UCustomizableObjectNodeColorParameter::OnRenameNode(const FString& NewName)
+{
+	if (!NewName.IsEmpty())
+	{
+		ParameterName = NewName;
+	}
+}
+
+
+void UCustomizableObjectNodeColorParameter::BackwardsCompatibleFixup()
+{
+	Super::BackwardsCompatibleFixup();
+
+	const int32 CustomizableObjectCustomVersion = GetLinkerCustomVersion(FCustomizableObjectCustomVersion::GUID);
+
+	if (CustomizableObjectCustomVersion < FCustomizableObjectCustomVersion::AddedRenameOptionToParameterNodes)
+	{
+		ReconstructNode();
+	}
+}
+
 #undef LOCTEXT_NAMESPACE
 
diff --git a/Engine/Plugins/Experimental/Mutable/Source/CustomizableObjectEditor/Private/MuCOE/Nodes/CustomizableObjectNodeEnumParameter.cpp b/Engine/Plugins/Experimental/Mutable/Source/CustomizableObjectEditor/Private/MuCOE/Nodes/CustomizableObjectNodeEnumParameter.cpp
--- a/Engine/Plugins/Experimental/Mutable/Source/CustomizableObjectEditor/Private/MuCOE/Nodes/CustomizableObjectNodeEnumParameter.cpp
+++ b/Engine/Plugins/Experimental/Mutable/Source/CustomizableObjectEditor/Private/MuCOE/Nodes/CustomizableObjectNodeEnumParameter.cpp
@@ -3,6 +3,7 @@
 #include "MuCOE/Nodes/CustomizableObjectNodeEnumParameter.h"
 
 #include "MuCOE/EdGraphSchema_CustomizableObject.h"
+#include "MuCO/CustomizableObjectCustomVersion.h"
 
 class UCustomizableObjectNodeRemapPins;
 
@@ -32,7 +33,18 @@ void UCustomizableObjectNodeEnumParameter::AllocateDefaultPins(UCustomizableObje
 
 FText UCustomizableObjectNodeEnumParameter::GetNodeTitle(ENodeTitleType::Type TitleType) const
 {
-	return LOCTEXT("Enum_Parameter", "Enum Parameter");
+	if (TitleType == ENodeTitleType::ListView || ParameterName.IsEmpty())
+	{
+		return LOCTEXT("Enum_Parameter", "Enum Parameter");
+	}
+	else if (TitleType == ENodeTitleType::EditableTitle)
+	{
+		return FText::Format(LOCTEXT("Enum_Parameter_EditableTitle", "{0}"), FText::FromString(ParameterName));
+	}
+	else
+	{
+		return FText::Format(LOCTEXT("Enum_Parameter_Title", "{0}\nEnum Parameter"), FText::FromString(ParameterName));
+	}
 }
 
 
@@ -50,6 +62,28 @@ FText UCustomizableObjectNodeEnumParameter::GetTooltipText() const
 }
 
 
+void UCustomizableObjectNodeEnumParameter::OnRenameNode(const FString& NewName)
+{
+	if (!NewName.IsEmpty())
+	{
+		ParameterName = NewName;
+	}
+}
+
+
+void UCustomizableObjectNodeEnumParameter::BackwardsCompatibleFixup()
+{
+	Super::BackwardsCompatibleFixup();
+
+	const int32 CustomizableObjectCustomVersion = GetLinkerCustomVersion(FCustomizableObjectCustomVersion::GUID);
+
+	if (CustomizableObjectCustomVersion < FCustomizableObjectCustomVersion::AddedRenameOptionToParameterNodes)
+	{
+		ReconstructNode();
+	}
+}
+
+
 
 #undef LOCTEXT_NAMESPACE
 
diff --git a/Engine/Plugins/Experimental/Mutable/Source/CustomizableObjectEditor/Private/MuCOE/Nodes/CustomizableObjectNodeFloatParameter.cpp b/Engine/Plugins/Experimental/Mutable/Source/CustomizableObjectEditor/Private/MuCOE/Nodes/CustomizableObjectNodeFloatParameter.cpp
--- a/Engine/Plugins/Experimental/Mutable/Source/CustomizableObjectEditor/Private/MuCOE/Nodes/CustomizableObjectNodeFloatParameter.cpp
+++ b/Engine/Plugins/Experimental/Mutable/Source/CustomizableObjectEditor/Private/MuCOE/Nodes/CustomizableObjectNodeFloatParameter.cpp
@@ -40,7 +40,18 @@ bool UCustomizableObjectNodeFloatParameter::IsAffectedByLOD() const
 
 FText UCustomizableObjectNodeFloatParameter::GetNodeTitle(ENodeTitleType::Type TitleType) const
 {
-	return LOCTEXT("Float_Parameter", "Float Parameter");
+	if (TitleType == ENodeTitleType::ListView || ParameterName.IsEmpty())
+	{
+		return LOCTEXT("Float_Parameter", "Float Parameter");
+	}
+	else if (TitleType == ENodeTitleType::EditableTitle)
+	{
+		return FText::Format(LOCTEXT("Float_Parameter_EditableTitle", "{0}"), FText::FromString(ParameterName));
+	}
+	else
+	{
+		return FText::Format(LOCTEXT("Float_Parameter_Title", "{0}\nFloat Parameter"), FText::FromString(ParameterName));
+	}
 }
 
 
@@ -57,6 +68,15 @@ FText UCustomizableObjectNodeFloatParameter::GetTooltipText() const
 }
 
 
+void UCustomizableObjectNodeFloatParameter::OnRenameNode(const FString& NewName)
+{
+	if (!NewName.IsEmpty())
+	{
+		ParameterName = NewName;
+	}
+}
+
+
 void UCustomizableObjectNodeFloatParameter::BackwardsCompatibleFixup()
 {
 	Super::BackwardsCompatibleFixup();
@@ -67,6 +87,11 @@ void UCustomizableObjectNodeFloatParameter::BackwardsCompatibleFixup()
 	{
 		ReconstructNode();
 	}
+
+	if (CustomizableObjectCustomVersion < FCustomizableObjectCustomVersion::AddedRenameOptionToParameterNodes)
+	{
+		ReconstructNode();
+	}
 }
 
 
diff --git a/Engine/Plugins/Experimental/Mutable/Source/CustomizableObjectEditor/Private/MuCOE/Nodes/CustomizableObjectNodeGroupProjectorParameter.cpp b/Engine/Plugins/Experimental/Mutable/Source/CustomizableObjectEditor/Private/MuCOE/Nodes/CustomizableObjectNodeGroupProjectorParameter.cpp
--- a/Engine/Plugins/Experimental/Mutable/Source/CustomizableObjectEditor/Private/MuCOE/Nodes/CustomizableObjectNodeGroupProjectorParameter.cpp
+++ b/Engine/Plugins/Experimental/Mutable/Source/CustomizableObjectEditor/Private/MuCOE/Nodes/CustomizableObjectNodeGroupProjectorParameter.cpp
@@ -86,7 +86,18 @@ TArray<FGroupProjectorParameterImage> UCustomizableObjectNodeGroupProjectorParam
 
 FText UCustomizableObjectNodeGroupProjectorParameter::GetNodeTitle(ENodeTitleType::Type TitleType) const
 {
-	return LOCTEXT("Group_Projector_Parameter", "Group Projector Parameter");
+	if (TitleType == ENodeTitleType::ListView || ParameterName.IsEmpty())
+	{
+		return LOCTEXT("Group_Projector_Parameter", "Group Projector Parameter");
+	}
+	else if (TitleType == ENodeTitleType::EditableTitle)
+	{
+		return FText::Format(LOCTEXT("Group_Projector_Parameter_EditableTitle", "{0}"), FText::FromString(ParameterName));
+	}
+	else
+	{
+		return FText::Format(LOCTEXT("Group_Projector_Parameter_Title", "{0}\nGroup Projector Parameter"), FText::FromString(ParameterName));
+	}
 }
 
 
@@ -121,6 +132,11 @@ void UCustomizableObjectNodeGroupProjectorParameter::BackwardsCompatibleFixup()
 	{
 		ReconstructNode();		
 	}
+
+	if (CustomizableObjectCustomVersion < FCustomizableObjectCustomVersion::AddedRenameOptionToParameterNodes)
+	{
+		ReconstructNode();
+	}
 }
 
 
diff --git a/Engine/Plugins/Experimental/Mutable/Source/CustomizableObjectEditor/Private/MuCOE/Nodes/CustomizableObjectNodeProjectorParameter.cpp b/Engine/Plugins/Experimental/Mutable/Source/CustomizableObjectEditor/Private/MuCOE/Nodes/CustomizableObjectNodeProjectorParameter.cpp
--- a/Engine/Plugins/Experimental/Mutable/Source/CustomizableObjectEditor/Private/MuCOE/Nodes/CustomizableObjectNodeProjectorParameter.cpp
+++ b/Engine/Plugins/Experimental/Mutable/Source/CustomizableObjectEditor/Private/MuCOE/Nodes/CustomizableObjectNodeProjectorParameter.cpp
@@ -10,6 +10,7 @@ class UCustomizableObjectNodeRemapPins;
 
 #define LOCTEXT_NAMESPACE "CustomizableObjectEditor"
 
+
 void UCustomizableObjectNodeProjectorParameter::PostEditChangeProperty(FPropertyChangedEvent& PropertyChangedEvent)
 {
 	const FName PropertyName = PropertyChangedEvent.GetPropertyName();
@@ -49,6 +50,11 @@ void UCustomizableObjectNodeProjectorParameter::BackwardsCompatibleFixup()
 	{
 		DefaultValue.ProjectionType = ProjectionType_DEPRECATED;
 	}
+
+	if (CustomizableObjectCustomVersion < FCustomizableObjectCustomVersion::AddedRenameOptionToParameterNodes)
+	{
+		ReconstructNode();
+	}
 }
 
 
@@ -120,7 +126,18 @@ void UCustomizableObjectNodeProjectorParameter::SetProjectorDefaultAngle(float A
 
 FText UCustomizableObjectNodeProjectorParameter::GetNodeTitle(ENodeTitleType::Type TitleType) const
 {
-	return LOCTEXT("Projector_Parameter", "Projector Parameter");
+	if (TitleType == ENodeTitleType::ListView || ParameterName.IsEmpty())
+	{
+		return LOCTEXT("Projector_Parameter", "Projector Parameter");
+	}
+	else if (TitleType == ENodeTitleType::EditableTitle)
+	{
+		return FText::Format(LOCTEXT("Projector_Parameter_EditableTitle", "{0}"), FText::FromString(ParameterName));
+	}
+	else
+	{
+		return FText::Format(LOCTEXT("Projector_Parameter_Title", "{0}\nProjector Parameter"), FText::FromString(ParameterName));
+	}
 }
 
 
@@ -136,4 +153,13 @@ FText UCustomizableObjectNodeProjectorParameter::GetTooltipText() const
 	return LOCTEXT("Projector_Parameter_Tooltip", "Exposes a runtime modifiable projector parameter from the Customizable Object.");
 }
 
+
+void UCustomizableObjectNodeProjectorParameter::OnRenameNode(const FString& NewName)
+{
+	if (!NewName.IsEmpty())
+	{
+		ParameterName = NewName;
+	}
+}
+
 #undef LOCTEXT_NAMESPACE
diff --git a/Engine/Plugins/Experimental/Mutable/Source/CustomizableObjectEditor/Private/MuCOE/Nodes/CustomizableObjectNodeTable.cpp b/Engine/Plugins/Experimental/Mutable/Source/CustomizableObjectEditor/Private/MuCOE/Nodes/CustomizableObjectNodeTable.cpp
--- a/Engine/Plugins/Experimental/Mutable/Source/CustomizableObjectEditor/Private/MuCOE/Nodes/CustomizableObjectNodeTable.cpp
+++ b/Engine/Plugins/Experimental/Mutable/Source/CustomizableObjectEditor/Private/MuCOE/Nodes/CustomizableObjectNodeTable.cpp
@@ -182,6 +182,11 @@ void UCustomizableObjectNodeTable::BackwardsCompatibleFixup()
 			}
 		}
 	}
+
+	if (CustomizableObjectCustomVersion < FCustomizableObjectCustomVersion::AddedRenameOptionToParameterNodes)
+	{
+		ReconstructNode();
+	}
 }
 
 
@@ -229,22 +234,32 @@ void UCustomizableObjectNodeTable::PostEditChangeChainProperty(FPropertyChangedC
 
 FText UCustomizableObjectNodeTable::GetNodeTitle(ENodeTitleType::Type TitleType) const
 {
-	if (Table)
-	{
-		FFormatNamedArguments Args;
-		Args.Add(TEXT("TableName"), FText::FromString(Table->GetName()));
+	FFormatNamedArguments Args;
+	Args.Add(TEXT("ParamName"), FText::FromString(ParameterName));
 
-		return FText::Format(LOCTEXT("TableNode_Title_DataTable", "{TableName}\nData Table"), Args);
+	if (TitleType == ENodeTitleType::ListView && ParameterName.IsEmpty() && (Table || Structure))
+	{
+		return LOCTEXT("Mutable_Table_Title", "Table");
 	}
-	else if (Structure)
+	else if (TitleType == ENodeTitleType::EditableTitle)
 	{
-		FFormatNamedArguments Args;
-		Args.Add(TEXT("StructureName"), FText::FromString(Structure->GetName()));
-
-		return FText::Format(LOCTEXT("TableNode_Title_ScriptedStruct", "{StructureName}\nScript Struct"), Args);
+		return FText::Format(LOCTEXT("TableNode_EditableTitle", "{ParamName}"), Args);
 	}
-	
-	return LOCTEXT("Mutable Table", "Table");
+	else
+	{
+		if (Table)
+		{
+			Args.Add(TEXT("TableName"), FText::FromString(Table->GetName()));
+			return FText::Format(LOCTEXT("TableNode_Title_DataTable", "{ParamName}\n{TableName} - Data Table"), Args);
+		}
+		else if (Structure)
+		{
+			Args.Add(TEXT("StructureName"), FText::FromString(Structure->GetName()));
+			return FText::Format(LOCTEXT("TableNode_Title_ScriptedStruct", "{ParamName}\n{StructureName} - Script Struct"), Args);
+		}
+	}
+
+	return LOCTEXT("Mutable_Table", "Table");
 }
 
 
@@ -261,6 +276,15 @@ FText UCustomizableObjectNodeTable::GetTooltipText() const
 }
 
 
+void UCustomizableObjectNodeTable::OnRenameNode(const FString& NewName)
+{
+	if (!NewName.IsEmpty())
+	{
+		ParameterName = NewName;
+	}
+}
+
+
 void UCustomizableObjectNodeTable::PinConnectionListChanged(UEdGraphPin* Pin)
 {
 	if (Pin && !Pin->LinkedTo.Num() && 
diff --git a/Engine/Plugins/Experimental/Mutable/Source/CustomizableObjectEditor/Private/MuCOE/Nodes/CustomizableObjectNodeTextureParameter.cpp b/Engine/Plugins/Experimental/Mutable/Source/CustomizableObjectEditor/Private/MuCOE/Nodes/CustomizableObjectNodeTextureParameter.cpp
--- a/Engine/Plugins/Experimental/Mutable/Source/CustomizableObjectEditor/Private/MuCOE/Nodes/CustomizableObjectNodeTextureParameter.cpp
+++ b/Engine/Plugins/Experimental/Mutable/Source/CustomizableObjectEditor/Private/MuCOE/Nodes/CustomizableObjectNodeTextureParameter.cpp
@@ -22,7 +22,18 @@ void UCustomizableObjectNodeTextureParameter::AllocateDefaultPins(UCustomizableO
 
 FText UCustomizableObjectNodeTextureParameter::GetNodeTitle(ENodeTitleType::Type TitleType) const
 {
-	return LOCTEXT("Texture Parameter", "Texture Parameter");
+	if (TitleType == ENodeTitleType::ListView || ParameterName.IsEmpty())
+	{
+		return LOCTEXT("Texture_Parameter", "Texture Parameter");
+	}
+	else if (TitleType == ENodeTitleType::EditableTitle)
+	{
+		return FText::Format(LOCTEXT("Texture_Parameter_EditableTitle", "{0}"), FText::FromString(ParameterName));
+	}
+	else
+	{
+		return FText::Format(LOCTEXT("Texture_Parameter_Title", "{0}\nTexture Parameter"), FText::FromString(ParameterName));
+	}
 }
 
 
@@ -39,6 +50,15 @@ FText UCustomizableObjectNodeTextureParameter::GetTooltipText() const
 }
 
 
+void UCustomizableObjectNodeTextureParameter::OnRenameNode(const FString& NewName)
+{
+	if (!NewName.IsEmpty())
+	{
+		ParameterName = NewName;
+	}
+}
+
+
 void UCustomizableObjectNodeTextureParameter::BackwardsCompatibleFixup()
 {
 	Super::BackwardsCompatibleFixup();
@@ -49,7 +69,12 @@ void UCustomizableObjectNodeTextureParameter::BackwardsCompatibleFixup()
 	{
 		ReferenceValue = DefaultValue;
 		DefaultValue = {};
-	}	
+	}
+
+	if (CustomizableObjectCustomVersion < FCustomizableObjectCustomVersion::AddedRenameOptionToParameterNodes)
+	{
+		ReconstructNode();
+	}
 }
 
 #undef LOCTEXT_NAMESPACE
