[Mutable] Moved all files from Classes folder to Public folder.

#rnx

diff --git a/Engine/Plugins/Experimental/Mutable/Source/CustomizableObject/Public/MuCO/CustomizableInstanceLODManagement.h b/Engine/Plugins/Experimental/Mutable/Source/CustomizableObject/Public/MuCO/CustomizableInstanceLODManagement.h
new file mode 100644
--- /dev/null
+++ b/Engine/Plugins/Experimental/Mutable/Source/CustomizableObject/Public/MuCO/CustomizableInstanceLODManagement.h
@@ -0,0 +1,68 @@
+// Copyright Epic Games, Inc. All Rights Reserved.
+
+#pragma once
+#include "GameFramework/Actor.h"
+
+#include "CustomizableInstanceLODManagement.generated.h"
+
+typedef TMap<const class UCustomizableObjectInstance*, class FMutableUpdateCandidate> FMutableInstanceUpdateMap;
+
+// This is an abstract base class, override it to create a new Instance LOD management system and register with UCustomizableObjectSystem::SetInstanceLODManagement
+UCLASS(Abstract, Blueprintable, BlueprintType)
+class CUSTOMIZABLEOBJECT_API UCustomizableInstanceLODManagementBase : public UObject
+{
+	GENERATED_BODY()
+public:
+	UCustomizableInstanceLODManagementBase() : UObject() {};
+	virtual ~UCustomizableInstanceLODManagementBase() {};
+
+	// WARNING! The following methods must be overriden in derived classes
+	virtual void UpdateInstanceDistsAndLODs(FMutableInstanceUpdateMap& InOutRequestedUpdates) PURE_VIRTUAL(UCustomizableInstanceLODManagementBase::UpdateInstanceDistsAndLODs, return;);
+	virtual int32 GetNumGeneratedInstancesLimitFullLODs() const PURE_VIRTUAL(UCustomizableInstanceLODManagementBase::GetNumGeneratedInstancesLimitFullLODs, return 0;);
+	virtual int32 GetNumGeneratedInstancesLimitLOD1() const PURE_VIRTUAL(UCustomizableInstanceLODManagementBase::GetNumGeneratedInstancesLimitLOD1, return 0;);
+	virtual int32 GetNumGeneratedInstancesLimitLOD2() const PURE_VIRTUAL(UCustomizableInstanceLODManagementBase::GetNumGeneratedInstancesLimitLOD2, return 0;);
+	virtual float GetOnlyUpdateCloseCustomizableObjectsDist() const PURE_VIRTUAL(UCustomizableInstanceLODManagementBase::GetOnlyUpdateCloseCustomizableObjectsDist, return 0.f;);
+	virtual bool IsOnlyUpdateCloseCustomizableObjectsEnabled() const PURE_VIRTUAL(UCustomizableInstanceLODManagementBase::IsOnlyUpdateCloseCustomizableObjectsEnabled, return false;);
+	virtual bool IsOnlyGenerateRequestedLODLevelsEnabled() const PURE_VIRTUAL(UCustomizableInstanceLODManagementBase::IsOnlyGenerateRequestedLODLevelsEnabled, return false;);
+};
+
+
+UCLASS(Blueprintable, BlueprintType)
+class CUSTOMIZABLEOBJECT_API UCustomizableInstanceLODManagement : public UCustomizableInstanceLODManagementBase
+{
+	GENERATED_BODY()
+public:
+	UCustomizableInstanceLODManagement();
+	virtual ~UCustomizableInstanceLODManagement();
+
+	virtual void UpdateInstanceDistsAndLODs(FMutableInstanceUpdateMap& InOutRequestedUpdates) override;
+
+	virtual int32 GetNumGeneratedInstancesLimitFullLODs() const override;
+	virtual int32 GetNumGeneratedInstancesLimitLOD1() const override;
+	virtual int32 GetNumGeneratedInstancesLimitLOD2() const override;
+
+	// Used to define the actors that will be considered the centers of LOD distance calculations. If left empty it will revert to the default, the first player controller actor
+	void AddViewCenter(const AActor* const InCentralActor) { ViewCenters.Add(TWeakObjectPtr<const AActor>(InCentralActor)); };
+	void RemoveViewCenter(const AActor* const InCentralActor) { ViewCenters.Remove(TWeakObjectPtr<const AActor>(InCentralActor)); };
+	void ClearViewCenters() { ViewCenters.Empty(); }
+
+	// Sets how many of the nearest instances to the player will have updates with priority over LOD changes, only works if LODs are enabled with b.NumGeneratedInstancesLimitLODs
+	void SetNumberOfPriorityUpdateInstances(int32 NumPriorityUpdateInstances);
+	int32 GetNumberOfPriorityUpdateInstances() const;
+
+	virtual void SetCustomizableObjectsUpdateDistance(float Distance);
+	virtual float GetOnlyUpdateCloseCustomizableObjectsDist() const override;
+
+	virtual bool IsOnlyUpdateCloseCustomizableObjectsEnabled() const override;
+	virtual bool IsOnlyGenerateRequestedLODLevelsEnabled() const override { return true; };
+
+private:
+	// The list of actors that define a view radius (CloseDist) around them used for LOD management
+	TSet<TWeakObjectPtr<const AActor>> ViewCenters;
+	
+	// Sets how many of the nearest instances to the player will have updates with priority over LOD changes
+	int32 NumPriorityUpdateInstances = 1;
+
+	float CloseCustomizableObjectsDist;	
+};
+
diff --git a/Engine/Plugins/Experimental/Mutable/Source/CustomizableObject/Public/MuCO/CustomizableObject.h b/Engine/Plugins/Experimental/Mutable/Source/CustomizableObject/Public/MuCO/CustomizableObject.h
new file mode 100644
--- /dev/null
+++ b/Engine/Plugins/Experimental/Mutable/Source/CustomizableObject/Public/MuCO/CustomizableObject.h
@@ -0,0 +1,705 @@
+// Copyright Epic Games, Inc. All Rights Reserved.
+ 
+#pragma once
+
+#include "RHIDefinitions.h"
+#include "MuCO/CustomizableObjectCompilerTypes.h"
+#include "MuCO/CustomizableObjectClothingTypes.h"
+#include "MuCO/CustomizableObjectResourceData.h"
+#include "MuCO/CustomizableObjectResourceDataTypes.h"
+#include "MuCO/CustomizableObjectStreamedResourceData.h"
+#include "MuCO/CustomizableObjectParameterTypeDefinitions.h"
+#include "MuCO/CustomizableObjectUIData.h"
+#include "Templates/TypeCompatibleBytes.h"
+
+#include "CustomizableObject.generated.h"
+
+class FReply;
+class FObjectPreSaveContext;
+class FText;
+class IAsyncReadFileHandle;
+class ITargetPlatform;
+class UCustomizableObject;
+class UEdGraph;
+class USkeletalMesh;
+class UCustomizableObjectPrivate;
+struct FFrame;
+struct FStreamableHandle;
+template <typename FuncType> class TFunctionRef;
+ 
+// Forward declaration of Mutable classes necessary for the interface
+namespace mu
+{
+	class Model;
+	class Parameters;
+}
+
+
+CUSTOMIZABLEOBJECT_API DECLARE_LOG_CATEGORY_EXTERN(LogMutable, Log, All);
+
+
+USTRUCT()
+struct FFParameterOptionsTags
+{
+	GENERATED_USTRUCT_BODY()
+
+	/** List of tags of a Parameter Options */
+	UPROPERTY(Category = "CustomizablePopulation", EditAnywhere)
+	TArray<FString> Tags;
+};
+
+
+USTRUCT()
+struct FParameterTags
+{
+	GENERATED_USTRUCT_BODY()
+
+	/** List of tags of a parameter */
+	UPROPERTY(Category = "CustomizablePopulation", EditAnywhere)
+	TArray<FString> Tags;
+
+	/** Map of options available for a parameter can have and their tags */
+	UPROPERTY(Category = "CustomizablePopulation", EditAnywhere)
+	TMap<FString, FFParameterOptionsTags> ParameterOptions;
+};
+
+
+USTRUCT()
+struct FProfileParameterDat
+{
+	GENERATED_USTRUCT_BODY()
+
+	UPROPERTY()
+	FString ProfileName;
+
+	UPROPERTY()
+	TArray<FCustomizableObjectBoolParameterValue> BoolParameters;
+
+	////
+	UPROPERTY()
+	TArray<FCustomizableObjectIntParameterValue> IntParameters;
+
+	////
+	UPROPERTY()
+	TArray<FCustomizableObjectFloatParameterValue> FloatParameters;
+
+	////
+	UPROPERTY()
+	TArray<FCustomizableObjectTextureParameterValue> TextureParameters;
+
+	////
+	UPROPERTY()
+	TArray<FCustomizableObjectVectorParameterValue> VectorParameters;
+
+	////
+	UPROPERTY()
+	TArray<FCustomizableObjectProjectorParameterValue> ProjectorParameters;
+};
+
+// A USTRUCT version of FMeshToMeshVertData in SkeletalMeshTypes.h
+// We are taking advantage of the padding data to store from which asset this data comes from
+// maintaining the same memory footprint than the original.
+USTRUCT()
+struct FCustomizableObjectMeshToMeshVertData
+{
+	GENERATED_USTRUCT_BODY()
+
+	FCustomizableObjectMeshToMeshVertData() = default;
+
+	explicit FCustomizableObjectMeshToMeshVertData(const FMeshToMeshVertData& Original)
+		: Weight(Original.Weight)
+	{
+		for (int32 i = 0; i < 4; ++i)
+		{
+			PositionBaryCoordsAndDist[i] = Original.PositionBaryCoordsAndDist[i];
+			NormalBaryCoordsAndDist[i] = Original.NormalBaryCoordsAndDist[i];
+			TangentBaryCoordsAndDist[i] = Original.TangentBaryCoordsAndDist[i];
+			SourceMeshVertIndices[i] = Original.SourceMeshVertIndices[i];
+		}
+	}
+		
+
+	explicit operator FMeshToMeshVertData() const
+	{
+		FMeshToMeshVertData ReturnValue;
+
+		for (int32 i = 0; i < 4; ++i)
+		{
+			ReturnValue.PositionBaryCoordsAndDist[i] = PositionBaryCoordsAndDist[i];
+			ReturnValue.NormalBaryCoordsAndDist[i] = NormalBaryCoordsAndDist[i];
+			ReturnValue.TangentBaryCoordsAndDist[i] = TangentBaryCoordsAndDist[i];
+			ReturnValue.SourceMeshVertIndices[i] = SourceMeshVertIndices[i];
+		}
+		ReturnValue.Weight = Weight;
+		ReturnValue.Padding = 0;
+
+		return ReturnValue;
+	}
+
+	
+	// Barycentric coords and distance along normal for the position of the final vert
+	UPROPERTY()
+	float PositionBaryCoordsAndDist[4] = {0.f};
+
+	// Barycentric coords and distance along normal for the location of the unit normal endpoint
+	// Actual normal = ResolvedNormalPosition - ResolvedPosition
+	UPROPERTY()
+	float NormalBaryCoordsAndDist[4] = {0.f};
+
+	// Barycentric coords and distance along normal for the location of the unit Tangent endpoint
+	// Actual normal = ResolvedNormalPosition - ResolvedPosition
+	UPROPERTY()
+	float TangentBaryCoordsAndDist[4] = {0.f};
+
+	// Contains the 3 indices for verts in the source mesh forming a triangle, the last element
+	// is a flag to decide how the skinning works, 0xffff uses no simulation, and just normal
+	// skinning, anything else uses the source mesh and the above skin data to get the final position
+	UPROPERTY()
+	uint16	 SourceMeshVertIndices[4] = {0u, 0u, 0u, 0u};
+
+	UPROPERTY()
+	float Weight = 0.0f;
+
+	// Non serialized, unused padding. This is present in the FMeshToMeshVertData struct as Padding for alignment.
+	uint32 UnusedPadding = 0;
+
+	/**
+	 * Serializer
+	 *
+	 * @param Ar - archive to serialize with
+	 * @param V - vertex to serialize
+	 * @return archive that was used
+	 */
+	friend FArchive& operator<<(FArchive& Ar, FCustomizableObjectMeshToMeshVertData& V)
+	{
+		Ar << V.PositionBaryCoordsAndDist[0];
+		Ar << V.PositionBaryCoordsAndDist[1];
+		Ar << V.PositionBaryCoordsAndDist[2];
+		Ar << V.PositionBaryCoordsAndDist[3];
+		Ar << V.NormalBaryCoordsAndDist[0];
+		Ar << V.NormalBaryCoordsAndDist[1];
+		Ar << V.NormalBaryCoordsAndDist[2];
+		Ar << V.NormalBaryCoordsAndDist[3];
+		Ar << V.TangentBaryCoordsAndDist[0];
+		Ar << V.TangentBaryCoordsAndDist[1];
+		Ar << V.TangentBaryCoordsAndDist[2];
+		Ar << V.TangentBaryCoordsAndDist[3];
+
+		Ar << V.SourceMeshVertIndices[0];
+		Ar << V.SourceMeshVertIndices[1];
+		Ar << V.SourceMeshVertIndices[2];
+		Ar << V.SourceMeshVertIndices[3];
+
+		Ar << V.Weight;
+
+		return Ar;
+	}
+};
+static_assert(sizeof(FCustomizableObjectMeshToMeshVertData) == sizeof(float)*4*3 + sizeof(uint16)*4 + sizeof(float) + sizeof(uint32));
+template<> struct TCanBulkSerialize<FCustomizableObjectMeshToMeshVertData> { enum { Value = true }; };
+
+USTRUCT()
+struct CUSTOMIZABLEOBJECT_API FMutableStreamableBlock
+{
+	GENERATED_USTRUCT_BODY()
+
+	UPROPERTY()
+	uint32 FileId = 0;
+
+	UPROPERTY()
+	uint32 Size = 0;
+	
+	UPROPERTY()
+	uint64 Offset = 0;
+
+	friend FArchive& operator<<(FArchive& Ar, FMutableStreamableBlock& Data)
+	{
+		Ar << Data.FileId;
+		Ar << Data.Size;
+		Ar << Data.Offset;
+
+		return Ar;
+	}
+};
+template<> struct TCanBulkSerialize<FMutableStreamableBlock> { enum { Value = true }; };
+static_assert(sizeof(FMutableStreamableBlock) == 8*2);
+
+USTRUCT()
+struct FMutableLODSettings
+{
+	GENERATED_BODY()
+
+	/** Minimum LOD to render per Platform. */
+	UPROPERTY(EditAnywhere, Category = LODSettings, meta = (DisplayName = "Minimum LOD"))
+	FPerPlatformInt MinLOD;
+
+	/** Minimum LOD to render per Quality level.*/
+	UPROPERTY(EditAnywhere, Category = LODSettings, meta = (DisplayName = "Quality Level Minimum LOD"))
+	FPerQualityLevelInt MinQualityLevelLOD;
+
+#if WITH_EDITORONLY_DATA
+
+	/** Override the LOD Streaming settings from the reference skeletal meshes.*/
+	UPROPERTY(EditAnywhere, Category = LODSettings, meta = (DisplayName = "Override LOD Streaming Settings"))
+	bool bOverrideLODStreamingSettings = true;
+
+	/** Enabled: streaming LODs will trigger automatic updates to generate and discard LODs. Streaming may decrease the amount of memory used, but will stress the CPU and Streaming of resources.
+	  *	Keep in mind that, even though updates may be faster depending on the amount of LODs to generate, there may be more updates to process.
+	  * 
+	  * Disabled: all LODs will be generated at once. It may increase the amount of memory used by the meshes and the generation may take longer, but less updates will be required. */
+	UPROPERTY(EditAnywhere, Category = LODSettings, meta = (EditCondition = "bOverrideLODStreamingSettings", DisplayName = "Enable LOD Streaming"))
+	FPerPlatformBool bEnableLODStreaming = true;
+
+	/** Limit the number of LODs to stream. A value of 0 is the same as disabling streaming of LODs.*/
+	UPROPERTY(EditAnywhere, Category = LODSettings, meta = (EditCondition = "bOverrideLODStreamingSettings"))
+	FPerPlatformInt NumMaxStreamedLODs = MAX_MESH_LOD_COUNT;
+
+#endif
+};
+
+
+UCLASS( config=Engine )
+class CUSTOMIZABLEOBJECT_API UCustomizableObjectBulk : public UObject
+{
+public:
+	GENERATED_BODY()
+
+	//~ Begin UObject Interface
+	virtual void PostLoad() override;
+	//~ End UObject Interface
+
+	/**  */
+	const FString& GetBulkFilePrefix() const { return BulkFilePrefix; }
+	
+	TUniquePtr<IAsyncReadFileHandle> OpenFileAsyncRead(uint32 FileId) const;
+
+#if WITH_EDITOR
+
+	//~ Begin UObject Interface
+	virtual void CookAdditionalFilesOverride(const TCHAR*, const ITargetPlatform*, TFunctionRef<void(const TCHAR*, void*, int64)> ) override;
+	//~ End UObject Interface
+
+	/** Compute the number of files and sizes the BulkData will be split into and fix up 
+	 * the HashToStreamableBlock's FileIds and Offsets. 
+	 */	
+	void PrepareBulkData(UCustomizableObject* InCustomizableObject, const ITargetPlatform* TargetPlatform);
+
+#endif
+	
+private:
+
+#if WITH_EDITOR
+
+	enum class EDataType : uint8
+	{
+		None = 0,
+		Model,
+		RealTimeMorph,
+		Clothing,
+	};
+
+	struct FBlock
+	{
+		/** Data Type*/
+		EDataType DataType;
+	
+		/** Used on some data types as the index to the block stored in the CustomizableObject */
+		uint32 Id;
+
+		/** Size of the data block. */
+		uint32 Size;
+
+		/** Offset in the full source streamed data file that is created when compiling. */
+		uint64 Offset;
+	};
+
+	struct FFile
+	{
+		EDataType DataType;
+
+		/** Id generated from a hash of the file content + offset to avoid collisions. */
+		uint32 Id;
+
+		/** List of blocks that are contained in the file, in order. */
+		TArray<FBlock> Blocks;
+	};
+
+	/** Helper to store the size of each BulkData partition. Only valid while cooking */
+	TArray<FFile> BulkDataFiles;
+
+	/** Helper to retrieve the BulkData from within the CookAdditionalFilesOverride */
+	TObjectPtr<UCustomizableObject> CustomizableObject;
+#endif
+
+	/** Prefix to locate bulkfiles for loading, using the file ids in each FMutableStreamableBlock. */
+	FString BulkFilePrefix;
+};
+
+
+UCLASS( BlueprintType, config=Engine )
+class CUSTOMIZABLEOBJECT_API UCustomizableObject : public UObject
+{
+	friend UCustomizableObjectPrivate;
+	
+public:
+	GENERATED_BODY()
+
+	UCustomizableObject();
+
+#if WITH_EDITORONLY_DATA
+private:
+	UPROPERTY()
+	TObjectPtr<USkeletalMesh> ReferenceSkeletalMesh_DEPRECATED;
+
+public:
+	/** All the SkeletalMeshes generated for this CustomizableObject instances will use the Reference Skeletal Mesh
+	* properties for everything that Mutable doesn't create or modify. This includes data like LOD distances, Physics
+	* properties, Bounding Volumes, Skeleton, etc.
+	*
+	* While a CustomizableObject instance is being created for the first time, and in some situation with lots of
+	* objects this may require some seconds, the Reference Skeletal Mesh is used for the actor. This works as a better
+	* solution than the alternative of not showing anything, although this can be disabled with the function
+	* "SetReplaceDiscardedWithReferenceMeshEnabled" (See the c++ section). */
+	UPROPERTY(EditAnywhere, Category = CustomizableObject)
+	TArray<TObjectPtr<USkeletalMesh>> ReferenceSkeletalMeshes;
+
+	/**
+	  * The optional VersionBridge asset, which must implement the ICustomizableObjectVersionBridgeInterface, will be used to
+	  * decide which Mutable child CustomizableObjects and table rows must be included in a compilation/cook depending on its 
+	  * version struct/column by comparing it to the game-specific version system.
+	*/
+	UPROPERTY(EditAnywhere, Category = Versioning)
+	TObjectPtr<UObject> VersionBridge;
+
+	/**
+	  * This optional struct is used to define which version this child CustomizableObject belongs to. It will be used during
+	  * cook/compilation to decide whether this CO should be included or not in the final compiled CO. To be used, the root
+	  * CO must have defined the VersionBridge property, which must implement the ICustomizableObjectVersionBridgeInterface
+	*/
+	UPROPERTY(EditAnywhere, Category = Versioning)
+	FInstancedStruct VersionStruct;
+#endif
+	
+	UPROPERTY(EditAnywhere, Category = CustomizableObject, meta = (DisplayName = "LOD Settings"))
+	FMutableLODSettings LODSettings;
+	
+private:
+	// mu::ExtensionData::Index is an index into this array when mu::ExtensionData::Origin is ConstantAlwaysLoaded
+	UPROPERTY()
+	TArray<FCustomizableObjectResourceData> AlwaysLoadedExtensionData;
+
+	// mu::ExtensionData::Index is an index into this array when mu::ExtensionData::Origin is ConstantStreamed
+	UPROPERTY()
+	TArray<FCustomizableObjectStreamedResourceData> StreamedExtensionData;
+
+	// Constant Resources streamed in on demand when generating meshes
+	UPROPERTY()
+	TArray<FCustomizableObjectStreamedResourceData> StreamedResourceData;
+
+public:
+	/** Use the SkeletalMesh of reference as a placeholder until the custom mesh is ready to use.
+	  * Note: If disabled, a null mesh will be used to replace the discarded mesh due to 'ReplaceDiscardedWithReferenceMesh' being enabled. */
+	UPROPERTY(EditAnywhere, Category = CustomizableObject)
+	bool bEnableUseRefSkeletalMeshAsPlaceholder = true;
+
+	/** Use the Instance MinLOD, and RequestedLODs in the descriptor when performing the initial generation (ignore LOD Management). */
+	UPROPERTY(Category = "CustomizableObject", EditAnywhere, DisplayName = "Preserve User LODs On First Generation")
+	bool bPreserveUserLODsOnFirstGeneration = false;
+
+	/** If true, reuse previously generated USkeletalMesh (if still valid and the the number of LOD have not changed)
+	 * USkeletalMeshes are only reused between the same CO. */
+	UPROPERTY(EditAnywhere, Category = CustomizableObject)
+	bool bEnableMeshCache = false;
+
+	/** If true, Mesh LODs will be streamed on demand. It requires streaming of SkeletalMeshes and Mutable.StreamMeshLODsEnabled to be enabled.
+	 *  Does not support Clothing, Morphs, and alternative SkinWeightProfiles yet. */
+	UPROPERTY(EditAnywhere, Category = CustomizableObject)
+	bool bEnableMeshStreaming = false;
+	
+#if WITH_EDITORONLY_DATA
+	// Compilation options to use in editor and for packaging for this object.
+	UPROPERTY()
+	FCompilationOptions CompileOptions;
+
+	UPROPERTY(EditAnywhere, Category = CompileOptions)
+	bool bEnableRealTimeMorphTargets = false;
+
+	UPROPERTY(EditAnywhere, Category = CompileOptions)
+	bool bEnableClothing = false;
+
+	// TODO: Enable 16 bit weights 
+	UPROPERTY(VisibleAnywhere, Category = CompileOptions)
+	bool bEnable16BitBoneWeights = false;
+
+	UPROPERTY(EditAnywhere, Category = CompileOptions)
+	bool bEnableAltSkinWeightProfiles = false;
+
+	UPROPERTY(EditAnywhere, Category = CompileOptions)
+	bool bEnablePhysicsAssetMerge = false;
+
+	UPROPERTY(EditAnywhere, Category = CompileOptions)
+	bool bEnableAnimBpPhysicsAssetsManipualtion = false;
+
+	// When this is enabled generated meshes will merge the AssetUserData from all of its constituent mesh parts
+	UPROPERTY(EditAnywhere, Category = CompileOptions)
+	bool bEnableAssetUserDataMerge = false;
+
+	// Disabling the Table Materials parent material check will let the user use any material regardless of its parent when connecting a material from a table column to a material node.
+	// Warning! But it will not check if the table material channels connected to the Material node exist in the actual material used in the Instance, and will fail silently at runtime 
+	// when setting the value of those channels if they don't exist.
+	UPROPERTY(EditAnywhere, Category = CompileOptions)
+	bool bDisableTableMaterialsParentCheck = false;
+
+	// Options when compiling this customizable object (see EMutableCompileMeshType declaration for info)
+	UPROPERTY(EditAnywhere, Category = CompileOptions)
+	EMutableCompileMeshType MeshCompileType = EMutableCompileMeshType::LocalAndChildren;
+
+	// Array of elements to use with compile option CompileType = WorkingSet
+	UPROPERTY(EditAnywhere, Category = CompileOptions)
+	TArray<TSoftObjectPtr<UCustomizableObject>> WorkingSet;
+
+private:
+	// Editor graph
+	UPROPERTY()
+	TObjectPtr<UEdGraph> Source;
+
+	// Used to verify the derived data matches this version of the Customizable Object.
+	UPROPERTY()
+	FGuid VersionId;
+
+	UPROPERTY()
+	TArray<FProfileParameterDat> InstancePropertiesProfiles;
+#endif // WITH_EDITORONLY_DATA
+
+public:
+	/** Get the number of components this Customizable Object has. */
+	UFUNCTION(BlueprintCallable, Category = CustomizableObject)
+	int32 GetComponentCount() const;
+	
+	/** Get the number of parameters available in instances of this object. */
+	UFUNCTION(BlueprintCallable, Category = CustomizableObject)
+	int32 GetParameterCount() const;
+
+	/** Get the index of a parameter from its name. Return -1 if the parameter is not found. */
+	UFUNCTION(BlueprintCallable, Category = CustomizableObject)
+	int32 FindParameter(const FString& Name) const;
+
+	/** Get the type of a parameter from its index. */
+	UFUNCTION(BlueprintCallable, Category = CustomizableObject)
+	EMutableParameterType GetParameterType(int32 ParamIndex) const;
+
+	/** Get the type of a parameter from its name. */
+	UFUNCTION(BlueprintCallable, Category = CustomizableObject)
+	EMutableParameterType GetParameterTypeByName(const FString& Name) const;
+
+	/** Get the name of a parameter from its index. */
+	UFUNCTION(BlueprintCallable, Category = CustomizableObject)
+	const FString& GetParameterName(int32 ParamIndex) const;
+
+	/** Deprecated. It will always return 0. */
+	UFUNCTION(BlueprintCallable, Category = CustomizableObject, meta = (DeprecatedFunction, DeprecationMessage = "Parameter decorations have been removed. This method will be removed in future versions."))
+	int32 GetParameterDescriptionCount(const FString& ParamName) const;
+
+	/** Returns how many possible options an int parameter has, if the parameter is an enumeration. Otherwise return 0. */
+	UFUNCTION(BlueprintCallable, Category = CustomizableObject)
+	int32 GetIntParameterNumOptions(int32 ParamIndex) const;
+
+	/** Gets the Name of the option at position K in the list of available options for the int parameter.
+	  * Useful to enumerate the int parameter's possible options (Ex: "Hat1", "Hat2", "Cap", "Nothing") */
+	UFUNCTION(BlueprintCallable, Category = CustomizableObject)
+	const FString& GetIntParameterAvailableOption(int32 ParamIndex, int32 K) const;
+
+	int32 FindIntParameterValue( int32 ParamIndex, const FString& Value ) const;
+	FString FindIntParameterValueName(int32 ParamIndex, int32 ParamValue) const;
+
+	USkeletalMesh* GetRefSkeletalMesh(int32 ComponentIndex = 0) const;
+
+	/** Get the default value of a parameter of type Float.
+	  * @param InParameterName The name of the Float parameter to get the default value of.
+	  * @return The default value of the provided parameter name. */
+	UFUNCTION(BlueprintCallable, Category = CustomizableObject)
+	float GetFloatParameterDefaultValue(UPARAM(DisplayName = "Parameter Name") const FString& InParameterName) const;
+	
+	/** Get the default value of a parameter of type Int. 
+	  * @param InParameterName The name of the Int parameter to get the default value of.
+	  * @return The default value of the provided parameter name. */
+	UFUNCTION(BlueprintCallable, Category = CustomizableObject)
+	int32 GetIntParameterDefaultValue(UPARAM(DisplayName = "Parameter Name") const FString& InParameterName) const;
+ 
+	/** Get the default value of a parameter of type Bool.
+	  * @param InParameterName The name of the Bool parameter to get the default value of.
+	  * @return The default value of the provided parameter name. */
+	UFUNCTION(BlueprintCallable, Category = CustomizableObject)
+	bool GetBoolParameterDefaultValue(UPARAM(DisplayName = "Parameter Name") const FString& InParameterName) const;
+
+	/** Get the default value of a parameter of type Color.
+	  * @param InParameterName The name of the Color parameter to get the default value of.
+	  * @return The default value of the provided parameter name. */
+	UFUNCTION(BlueprintCallable, Category = CustomizableObject)
+	FLinearColor GetColorParameterDefaultValue(UPARAM(DisplayName = "Parameter Name") const FString& InParameterName) const;
+	
+	/** Get the default value of a parameter of type Projector.
+	  * @param InParameterName The name of the Projector parameter to get the default value of.
+	  * @param OutPos The default position of the Projector.
+	  * @param OutDirection The default projection direction of the Projector.
+	  * @param OutUp The default up vector of the Projector.
+	  * @param OutScale The default scale of the Projector.
+	  * @param OutAngle The default angle of the Projector.
+	  * @param OutType The default type of the Projector. */
+	UFUNCTION(BlueprintCallable, Category = CustomizableObject)
+	void GetProjectorParameterDefaultValue(
+		UPARAM(DisplayName = "Parameter Name") const FString& InParameterName,
+		UPARAM(DisplayName = "Position") FVector3f& OutPos,
+		UPARAM(DisplayName = "Direction") FVector3f& OutDirection,
+		UPARAM(DisplayName = "Up") FVector3f& OutUp,
+		UPARAM(DisplayName = "Scale") FVector3f& OutScale,
+		UPARAM(DisplayName = "Angle") float& OutAngle,
+		UPARAM(DisplayName = "Type") ECustomizableObjectProjectorType& OutType) const;
+
+	/** Get the default value of a projector with the provided name
+	  * @param InParameterName The name of the parameter to get the default value of.
+	  * @return A data structure containing all the default data for the targeted projector parameter. */
+	FCustomizableObjectProjector GetProjectorParameterDefaultValue(const FString& InParameterName) const;
+	
+	/** Get the default value of a parameter of type Texture.
+	  * @param InParameterName The name of the Projector parameter to get the default value of.
+	  * @return An id representing the default parameter's texture. */
+	FName GetTextureParameterDefaultValue(const FString& InParameterName) const;
+
+	/** Return true if the parameter at the index provided is multidimensional.
+	  * @param InParamIndex The index of the parameter to check.
+	  * @return True if the parameter is multidimensional and false if it is not. */
+	bool IsParameterMultidimensional(const int32& InParamIndex) const;
+	
+	/** Return true if the parameter at the index provided is multidimensional.
+	  * @param InParameterName The name of the parameter to check.
+	  * @return True if the parameter is multidimensional and false if it is not. */
+	UFUNCTION(BlueprintCallable, Category = CustomizableObject)
+	bool IsParameterMultidimensional(UPARAM(DisplayName = "Parameter Name") const FString& InParameterName) const;
+
+	// Begin UObject interface.
+	void PostLoad() override;
+	void Serialize(FArchive& Ar) override;
+#if WITH_EDITOR
+	
+	/** Compile the object if Automatic Compilation is enabled and has not been already compiled.
+	  * Automatic compilation can be enabled/disabled in the Mutable's Plugin Settings.
+	  * @return true if compiled */
+	bool ConditionalAutoCompile();
+	
+	virtual void GetAssetRegistryTags(FAssetRegistryTagsContext Context) const override;
+	UE_DEPRECATED(5.4, "Implement the version that takes FAssetRegistryTagsContext instead.")
+	virtual void GetAssetRegistryTags(TArray<FAssetRegistryTag>& OutTags) const override;
+
+	/** Returns a one line description of an object for viewing in the thumbnail view of the generic browser */
+	FString GetDesc() override;
+	bool IsEditorOnly() const override;
+	void PreSave(FObjectPreSaveContext ObjectSaveContext) override;
+	void PostSaveRoot(FObjectPostSaveRootContext ObjectSaveContext) override;
+	void PostRename(UObject* OldOuter, const FName OldName) override;	
+	void BeginCacheForCookedPlatformData(const ITargetPlatform* TargetPlatform) override;
+	bool IsCachedCookedPlatformDataLoaded(const ITargetPlatform* TargetPlatform) override;
+#endif
+	// End UObject interface.
+
+	int32 FindState(const FString& Name) const;
+
+	/** Return the number of object states that are defined in the CustomizableObject. */
+	UFUNCTION(BlueprintCallable, Category = CustomizableObject)
+	int32 GetStateCount() const;
+
+	/** Return the name of an object state from its index. */
+	UFUNCTION(BlueprintCallable, Category = CustomizableObject)
+	FString GetStateName(int32 StateIndex) const;
+
+	int32 GetStateParameterCount(int32 StateIndex) const;
+	int32 GetStateParameterIndex(int32 StateIndex, int32 ParameterIndex) const;
+
+	/** Return the number of parameters that are editable at runtime for a specific state. */
+	UFUNCTION(BlueprintCallable, Category = CustomizableObject)
+	int32 GetStateParameterCount(const FString& StateName) const;
+
+	/** Return the name of one of the state's runtime parameters, by its index (from 0 to GetStateParameterCount-1). */
+	UFUNCTION(BlueprintCallable, Category = CustomizableObject)
+	FString GetStateParameterName(const FString& StateName, int32 ParameterIndex) const;
+	FString GetStateParameterName(int32 StateIndex, int32 ParameterIndex) const;
+	
+	/** Return the metadata associated to a parameter. */
+	UFUNCTION(BlueprintCallable, Category = CustomizableObject)
+	FMutableParamUIMetadata GetParameterUIMetadata(const FString& ParamName) const;
+
+	/** Return the metadata associated to an int parameter option. */
+	UFUNCTION(BlueprintCallable, Category = CustomizableObject)
+	FMutableParamUIMetadata GetIntParameterOptionUIMetadata(const FString& ParamName, const FString& OptionName) const;
+	
+	/** Returns the group type of the given integer parameter */
+	UFUNCTION(BlueprintCallable, Category = CustomizableObject)
+	ECustomizableObjectGroupType GetIntParameterGroupType(const FString& ParamName) const;
+
+	/** Return the metadata associated to a state. */
+	UFUNCTION(BlueprintCallable, Category = CustomizableObject)
+	FMutableStateUIMetadata GetStateUIMetadata(const FString& StateName) const;
+
+	
+private:
+	/** Textures marked as low priority will generate defaulted resident mips (if texture streaming is enabled).
+	  * Generating defaulted resident mips greatly reduce initial generation times. */
+	UPROPERTY(EditAnywhere, Category = CustomizableObject)
+	TArray<FName> LowPriorityTextures;
+
+	// Customizable Object Population data start ------------------------------------------------------
+	/** Array to store the selected Population Class tags for this Customizable Object */
+	UPROPERTY()
+	TArray<FString> CustomizableObjectClassTags;
+	
+	/** Array to store all the Population Class tags */
+	UPROPERTY()
+	TArray<FString> PopulationClassTags;
+
+	/** Map of parameters available for the Customizable Object and their tags */
+	UPROPERTY()
+	TMap<FString, FParameterTags> CustomizableObjectParametersTags;
+	// Customizable Object Population data end --------------------------------------------------------
+
+#if WITH_EDITORONLY_DATA
+	/** True if this object references a parent object. This is used basically to exclude this object
+	  * from cooking. This is actually derived from the source graph object node pointing to another
+	  * object or not, but it needs to be cached here because the source graph is not always available.
+	  * For old objects this may be false even if they are child objects until they are resaved, but 
+	  * that is the conservative case and shouldn't cause a problem. */
+	UPROPERTY()
+	bool bIsChildObject = false;
+#endif
+	
+public:
+#if WITH_EDITORONLY_DATA
+	FPostCompileDelegate& GetPostCompileDelegate() const;
+#endif
+
+	/** Create a new instance of this object. The instance parameters will be initialized with the object default values. */
+	UFUNCTION(BlueprintCallable, Category = CustomizableObject)
+	UCustomizableObjectInstance* CreateInstance();
+
+	UCustomizableObjectPrivate* GetPrivate() const;
+
+	/** Check if the CustomizableObject asset has been compiled. This will always be true in a packaged game, but it could be false in the editor. */
+	UFUNCTION(BlueprintCallable, Category = CustomizableObject)
+	bool IsCompiled() const;
+
+	int32 GetNumLODs() const;
+
+	bool IsChildObject() const;
+
+private:
+	/** BulkData that stores all in-game resources used by Mutable when generating instances.
+	  * Only valid in packaged builds */
+	UPROPERTY()
+	TObjectPtr<UCustomizableObjectBulk> BulkData;
+
+	UPROPERTY()
+	TObjectPtr<UCustomizableObjectPrivate> Private;
+};
+
+#if UE_ENABLE_INCLUDE_ORDER_DEPRECATED_IN_5_2
+#include "Input/Reply.h"
+#endif
diff --git a/Engine/Plugins/Experimental/Mutable/Source/CustomizableObject/Public/MuCO/CustomizableObjectInstance.h b/Engine/Plugins/Experimental/Mutable/Source/CustomizableObject/Public/MuCO/CustomizableObjectInstance.h
new file mode 100644
--- /dev/null
+++ b/Engine/Plugins/Experimental/Mutable/Source/CustomizableObject/Public/MuCO/CustomizableObjectInstance.h
@@ -0,0 +1,723 @@
+// Copyright Epic Games, Inc. All Rights Reserved.
+
+#pragma once
+
+#include "Async/TaskGraphFwd.h"
+#include "MuCO/CustomizableObjectInstanceDescriptor.h"
+#include "Templates/SubclassOf.h"
+#include "Math/RandomStream.h"
+
+#include "CustomizableObjectInstance.generated.h"
+
+namespace mu
+{
+	class Image;
+}
+
+class UCustomizableObjectSystemPrivate;
+class USkeletalMesh;
+class AActor;
+class FProperty;
+class UAnimInstance;
+class UCustomizableObject;
+class UTexture2D;
+class FUpdateContextPrivate;
+class UCustomizableObjectInstanceUsage;
+class UMaterialInterface;
+struct FFrame;
+struct FGameplayTagContainer;
+struct FPropertyChangedEvent;
+struct FTexturePlatformData;
+struct FMutableModelImageProperties;
+
+
+
+/**
+ * Represents what kind of saving procedure was performed to save the package
+ */
+UENUM()
+enum class EPackageSaveResolutionType : uint8
+{
+	None = 0,
+
+	/** The package got saved as a new file. */
+	NewFile,
+
+	/** The package was already present on disk so the old package was deleted and a new one was saved on its place */
+	Overriden,
+
+	/** Error type : An override was required but due to an error or lack of user permission it could not be done. */
+	UnableToOverride
+};
+
+
+/**
+ * Data structure that exposes the path to a baked package and also what type of save was performed (an override, a standard save with a new file...)
+ */
+USTRUCT(BlueprintType, Blueprintable)
+struct FBakedResourceData
+{
+	GENERATED_BODY()
+
+	/**
+	 * The way the package represented by this object was saved onto disk.
+	 */
+	UPROPERTY(BlueprintReadOnly, Category = CustomizableObjectInstanceBaker)
+	EPackageSaveResolutionType SaveType = EPackageSaveResolutionType::None;
+
+	/**
+	 * The path used by the saved package.
+	 */
+	UPROPERTY(BlueprintReadOnly, Category = CustomizableObjectInstanceBaker)
+	FString AssetPath;
+};
+
+
+/**
+ * Structure returned as output of the baking operation. May contain a filled collection of FBakedResourceData objects and also the success end state of the
+ * baking operation. 
+ */
+USTRUCT(BlueprintType, Blueprintable)
+struct FCustomizableObjectInstanceBakeOutput
+{
+	GENERATED_BODY()
+
+	/**
+	 * Success state for the baking operation. True for success and false for failure.
+	 */
+	UPROPERTY(BlueprintReadOnly, Category = CustomizableObjectInstanceBaker)
+	bool bWasBakeSuccessful = false;
+
+	/**
+	 * Collection of FBakedResourceData representing all saved packages during the baking operation. It may be empty if the operation failed.
+	 */
+	UPROPERTY(BlueprintReadOnly, Category = CustomizableObjectInstanceBaker)
+	TArray<FBakedResourceData> SavedPackages;
+};
+
+
+DECLARE_DYNAMIC_DELEGATE_OneParam(FBakeOperationCompletedDelegate, const FCustomizableObjectInstanceBakeOutput, BakeOperationOutput);
+
+
+/**
+ * Configuration data structure designed to serve as variable container for the customizable object instance baking methods.
+ */
+USTRUCT(BlueprintType, Blueprintable)
+struct FBakingConfiguration
+{
+	GENERATED_BODY()
+
+	/**
+	 * The path where to save the baked resources. EX /Game/MyBakingTest
+	 */
+	UPROPERTY(BlueprintReadWrite, Blueprintable, Category = CustomizableObjectInstanceBaker)
+	FString OutputPath = TEXT("/Game");
+	
+	/**
+	 * The name to be used as base (prefix) during the naming of the exported resources
+	 */
+	UPROPERTY(BlueprintReadWrite, Blueprintable, Category = CustomizableObjectInstanceBaker)
+	FString OutputFilesBaseName;
+
+	/**
+	 * Determines if we want a full or partial export
+	 */
+	UPROPERTY(BlueprintReadWrite, Blueprintable, Category = CustomizableObjectInstanceBaker)
+	bool bExportAllResourcesOnBake = false;
+	
+	/**
+	 *  Determines if we want (or not) to generate constant material instances for each of the material instances found in the mutable instance
+	 */
+	UPROPERTY(BlueprintReadWrite, Blueprintable, Category = CustomizableObjectInstanceBaker)
+	bool bGenerateConstantMaterialInstancesOnBake = false;
+	
+	/**
+	 * Flag that determines if we should override already exported files or if we should not. If we encounter files to override and we have not permission to override them then the baking operation will fail.
+	 */
+	UPROPERTY(BlueprintReadWrite, Blueprintable, Category = CustomizableObjectInstanceBaker)
+	bool bAllowOverridingOfFiles = false;
+
+	/**
+	 * Callback executed once the baking operation gets completed. It will return the end success state and also some data about the assets saved.
+	 */
+	UPROPERTY(BlueprintReadWrite, Blueprintable, Category = CustomizableObjectInstanceBaker)
+	FBakeOperationCompletedDelegate OnBakeOperationCompletedCallback;
+};
+
+
+// Priority for the mutable update queue, Low is the normal distance-based priority, High is normally used for discards and Mid for LOD downgrades
+enum class EQueuePriorityType : uint8 { High, Med, Med_Low, Low };
+
+/** Result of all the checks just before beginning an update. */
+enum class EUpdateRequired : uint8
+{
+	NoUpdate, // No work required.
+	Update, // Normal update.
+	Discard // Discard instead of update.
+};
+
+
+/** Instance Update Result. */
+UENUM()
+enum class EUpdateResult : uint8
+{
+	Success, // Update finished without issues.
+	Warning, // Generic warning. Update finished but with warnings.
+	
+	Error, // Generic error.
+	ErrorOptimized, // The update was skipped since its result would have been the same as the current customization.
+	ErrorReplaced, // The update was replaced by a newer update request.
+	ErrorDiscarded, // The update was not finished since due to the LOD management discarding the data.
+	Error16BitBoneIndex // The update finish unsuccessfully due to Instance not supporting 16 Bit Bone Indexing required by the Engine.
+};
+
+
+/** Indicates the status of the generated Skeletal Mesh. */
+enum class ESkeletalMeshStatus : uint8
+{
+	NotGenerated, // Set only when loading the Instance for the first time or after compiling. Any generation, successful or not, can not end up in this state.
+	Success, // Generated successfully.
+	Error // Generated with errors.
+};
+
+
+/** Instance Update Context.
+ * Used to avoid changing the delegate signature in the future.  */
+USTRUCT()
+struct FUpdateContext
+{	
+	GENERATED_BODY()
+
+	UPROPERTY()
+	EUpdateResult UpdateResult = EUpdateResult::Success;
+};
+
+
+DECLARE_DYNAMIC_DELEGATE_OneParam(FInstanceUpdateDelegate, const FUpdateContext&, Result);
+DECLARE_MULTICAST_DELEGATE_OneParam(FInstanceUpdateNativeDelegate, const FUpdateContext&);
+
+/* When creating new delegates use the following conventions:
+ *
+ * - All delegates must be multicast.
+ * - If the delegate is exposed to the API create both, dynamic and native versions (non-dynamic).
+ * - Dynamic delegates should not be transient. Use the native version if you do not want it to be saved.
+ * - Native delegates names should end with "NativeDelegate".
+ * - Dynamic delegates broadcast before native delegates. */
+
+/** Broadcast when an Instance update has completed.
+ * Notice that Mutable internally can also start an Instance update. */
+DECLARE_DYNAMIC_MULTICAST_DELEGATE_OneParam(FObjectInstanceUpdatedDelegate, UCustomizableObjectInstance*, Instance);
+DECLARE_MULTICAST_DELEGATE_OneParam(FObjectInstanceUpdatedNativeDelegate, UCustomizableObjectInstance*);
+
+DECLARE_DELEGATE_OneParam(FProjectorStateChangedDelegate, FString);
+
+DECLARE_DYNAMIC_DELEGATE_TwoParams(FEachComponentAnimInstanceClassDelegate, FName, SlotIndex, TSubclassOf<UAnimInstance>, AnimInstClass);
+
+DECLARE_DELEGATE_TwoParams(FEachComponentAnimInstanceClassNativeDelegate, FName /*SlotIndex*/, TSubclassOf<UAnimInstance> /*AnimInstClass*/);
+
+
+UCLASS( Blueprintable, BlueprintType, HideCategories=(CustomizableObjectInstance) )
+class CUSTOMIZABLEOBJECT_API UCustomizableObjectInstance : public UObject
+{
+	GENERATED_BODY()
+
+	// Friends
+	friend UCustomizableInstancePrivate;
+	friend FMutableUpdateCandidate;
+
+public:
+	UCustomizableObjectInstance();
+	
+	/** Broadcast when the Customizable Object Instance is updated. */
+	UPROPERTY(Transient, BlueprintAssignable, Category = CustomizableObjectInstance)
+	FObjectInstanceUpdatedDelegate UpdatedDelegate;
+
+	/** Broadcast when the Customizable Object Instance is updated. */
+	FObjectInstanceUpdatedNativeDelegate UpdatedNativeDelegate;
+
+	// UObject interface.
+#if WITH_EDITOR
+	virtual void PostEditChangeProperty(FPropertyChangedEvent& PropertyChangedEvent) override;
+	virtual bool CanEditChange( const FProperty* InProperty ) const override;
+	virtual void PostTransacted(const FTransactionObjectEvent& TransactionEvent) override;
+#endif //WITH_EDITOR
+
+	virtual void Serialize(FArchive& Ar) override;
+	virtual void PostLoad() override;
+	virtual void BeginDestroy() override;
+	virtual bool IsReadyForFinishDestroy() override;
+	virtual FString GetDesc() override;
+	virtual bool IsEditorOnly() const override;
+
+	/** Set the CustomizableObject this instance will be generated from. 
+	  * It is usually not necessary to call this since instances are already generated from a CustomizableObject. */
+	UFUNCTION(BlueprintCallable, Category = CustomizableObjectInstance)
+	void SetObject(UCustomizableObject* InObject);
+
+	/** Get the CustomizableObject that this is an instance of. */
+	UFUNCTION(BlueprintCallable, Category = CustomizableObjectInstance)
+	UCustomizableObject* GetCustomizableObject() const;
+
+	/** Return true if the parameter relevancy will be updated when this instance is generated. */
+	UFUNCTION(BlueprintCallable, Category = CustomizableObjectInstance)
+	bool GetBuildParameterRelevancy() const;
+
+	/** Set the flag that controls if parameter relevancy will be updated when this instance is generated. */
+	UFUNCTION(BlueprintCallable, Category = CustomizableObjectInstance)
+	void SetBuildParameterRelevancy(bool Value);
+
+	/** Return the name of the current CustomizableObject state this is instance is set to. */
+	UFUNCTION(BlueprintCallable, Category = CustomizableObjectInstance)
+	FString GetCurrentState() const;
+
+	/** Set the CustomizableObject state that this instance will be generated into. */
+	UFUNCTION(BlueprintCallable, Category = CustomizableObjectInstance)
+	void SetCurrentState(const FString& StateName);
+
+	/** Get the skeletal mesh generated for this instance. 
+	  * If the object has multiple components, an index of the component can be specified. */
+	UFUNCTION(BlueprintCallable, Category = CustomizableObjectInstance)
+	USkeletalMesh* GetSkeletalMesh(int32 ComponentIndex = 0) const;
+	
+	/** Return true if a skeletal mesh has been generated for any component of this instance. */
+	UFUNCTION(BlueprintCallable, Category = CustomizableObjectInstance)
+	bool HasAnySkeletalMesh() const;
+	
+	/** Return true if the instance has any parameters. */
+	UFUNCTION(BlueprintCallable, Category = CustomizableObjectInstance)
+	bool HasAnyParameters() const;
+	
+	/** Set random values to the parameters. */
+	UFUNCTION(BlueprintCallable, Category = CustomizableObjectInstance)
+	void SetRandomValues();
+
+	/** Set random values to the parameters using a stream. */
+	UFUNCTION(BlueprintCallable, Category = CustomizableObjectInstance)
+	void SetRandomValuesFromStream(const FRandomStream& InStream);
+
+	/** Set all parameters to their default value. */
+	UFUNCTION(BlueprintCallable, Category = CustomizableObjectInstance)
+	void SetDefaultValues();
+	
+	/** Returns the AssetUserData that was gathered from all the constituent mesh parts during the last update. 
+	  * It requires that the CustomizableObject had the bEnableAssetUserDataMerge set to true during compilation. */
+	UFUNCTION(BlueprintCallable, Category = CustomizableObjectInstance)
+	TSet<UAssetUserData*> GetMergedAssetUserData(int32 ComponentIndex) const;
+	
+	/** Return true if the instance is not locked and if it's compiled. */
+	bool CanUpdateInstance() const;
+
+	/** Generate the instance with the current parameters and update all the components Skeletal Meshes asynchronously. */
+	UFUNCTION(BlueprintCallable, Category = CustomizableObjectInstance)
+	void UpdateSkeletalMeshAsync(bool bIgnoreCloseDist = false, bool bForceHighPriority = false);
+		
+	/** Generate the instance with the current parameters and update all the components Skeletal Meshes asynchronously.
+	  * Callback will be called once the update finishes, even if it fails. */
+	UFUNCTION(BlueprintCallable, Category = CustomizableObjectInstance)
+	void UpdateSkeletalMeshAsyncResult(FInstanceUpdateDelegate Callback, bool bIgnoreCloseDist = false, bool bForceHighPriority = false);
+
+	void UpdateSkeletalMeshAsyncResult(FInstanceUpdateNativeDelegate Callback, bool bIgnoreCloseDist = false, bool bForceHighPriority = false);
+
+	/** Clones the instance creating a new identical transient instance. */
+	UFUNCTION(BlueprintCallable, Category = CustomizableObjectInstance)
+	UCustomizableObjectInstance* Clone();
+
+	/** Clones the instance creating a new identical static instance with the given Outer. */
+	UFUNCTION(BlueprintCallable, Category = CustomizableObjectInstance)
+	UCustomizableObjectInstance* CloneStatic(UObject* Outer);
+
+	/** Copy parameters from the given Instance. */
+	void CopyParametersFromInstance(UCustomizableObjectInstance* Instance);
+
+	/** Immediately destroy the Mutable Core Live Update Instance attached to this (if exists). */
+	void DestroyLiveUpdateInstance();
+
+	/** Return true if changing the parameter would affect the Instance given its current generation. */
+	bool IsParameterRelevant(int32 ParameterIndex) const;
+
+	/** Return true if the given parameter has any effect in the current object state, and considering the current values of the other parameters. */
+	UFUNCTION(BlueprintCallable, Category = CustomizableObjectInstance)
+	bool IsParameterRelevant(const FString& ParamName) const;
+
+	/** Return true if the parameter has changed but the Instance has not yet been updated. */
+	UFUNCTION(BlueprintCallable, Category = CustomizableObjectInstance)
+	bool IsParameterDirty(const FString& ParamName, int32 RangeIndex = -1) const;
+	
+	/** For multidimensional parameters, return the number of dimensions that the given projector parameter supports. */
+	UFUNCTION(BlueprintCallable, Category = CustomizableObjectInstance)
+	int32 GetProjectorValueRange(const FString& ParamName) const;
+
+	/** For multidimensional parameters, return the number of dimensions that the given int parameter supports. */
+	UFUNCTION(BlueprintCallable, Category = CustomizableObjectInstance)
+	int32 GetIntValueRange(const FString& ParamName) const;
+
+	/** For multidimensional parameters, return the number of dimensions that the given float parameter supports. */
+	UFUNCTION(BlueprintCallable, Category = CustomizableObjectInstance)
+	int32 GetFloatValueRange(const FString& ParamName) const;
+
+	/** For multidimensional parameters, return the number of dimensions that the given texture parameter supports. */
+	UFUNCTION(BlueprintCallable, Category = CustomizableObjectInstance)
+	int32 GetTextureValueRange(const FString& ParamName) const;
+
+	/** Return the name of the option currently set in the given parameter. */
+	UFUNCTION(BlueprintCallable, Category = CustomizableObjectInstance)
+	const FString& GetIntParameterSelectedOption(const FString& ParamName, int32 RangeIndex = -1) const;
+
+	/** Set the currently selected option value for the given parameter, by parameter index and option name. */
+	void SetIntParameterSelectedOption(int32 IntParamIndex, const FString& SelectedOption, int32 RangeIndex = -1);
+
+	/** Set the currently selected option value for the given parameter, by parameter name and option name. */
+	UFUNCTION(BlueprintCallable, Category = CustomizableObjectInstance)
+	void SetIntParameterSelectedOption(const FString& ParamName, const FString& SelectedOptionName, int32 RangeIndex = -1);
+
+	/** Gets the value of a float parameter with name "FloatParamName". */
+	UFUNCTION(BlueprintCallable, Category = CustomizableObjectInstance)
+	float GetFloatParameterSelectedOption(const FString& FloatParamName, int32 RangeIndex = -1) const;
+
+	/** Sets the float value "FloatValue" of a float parameter with index "FloatParamIndex". */
+	UFUNCTION(BlueprintCallable, Category = CustomizableObjectInstance)
+	void SetFloatParameterSelectedOption(const FString& FloatParamName, float FloatValue, int32 RangeIndex = -1);
+	
+	/** Gets the value of a texture parameter with name "TextureParamName". */
+	UFUNCTION(BlueprintCallable, Category = CustomizableObjectInstance)
+	FName GetTextureParameterSelectedOption(const FString& TextureParamName, int32 RangeIndex = -1) const;
+
+	/** Sets the texture value "TextureValue" of a texture parameter with index "TextureParamIndex". */
+	UFUNCTION(BlueprintCallable, Category = CustomizableObjectInstance)
+	void SetTextureParameterSelectedOption(const FString& TextureParamName, const FString& TextureValue, int32 RangeIndex = -1);
+
+	/** Gets the value of a color parameter with name "ColorParamName". */
+	UFUNCTION(BlueprintCallable, Category = CustomizableObjectInstance)
+	FLinearColor GetColorParameterSelectedOption(const FString& ColorParamName) const;
+
+	/** Sets the color value "ColorValue" of a color parameter with index "ColorParamIndex". */
+	UFUNCTION(BlueprintCallable, Category = CustomizableObjectInstance)
+	void SetColorParameterSelectedOption(const FString& ColorParamName, const FLinearColor& ColorValue);
+
+	/** Sets the bool value "BoolValue" of a bool parameter with name "BoolParamName". */
+	UFUNCTION(BlueprintCallable, Category = CustomizableObjectInstance)
+	bool GetBoolParameterSelectedOption(const FString& BoolParamName) const;
+
+	/** Sets the bool value "BoolValue" of a bool parameter with name "BoolParamName". */
+	UFUNCTION(BlueprintCallable, Category = CustomizableObjectInstance)
+	void SetBoolParameterSelectedOption(const FString& BoolParamName, bool BoolValue);
+
+	/** Sets the vector value "VectorValue" of a bool parameter with index "VectorParamIndex". */
+	UFUNCTION(BlueprintCallable, Category = CustomizableObjectInstance)
+	void SetVectorParameterSelectedOption(const FString& VectorParamName, const FLinearColor& VectorValue);
+
+	/** Sets the projector values of a projector parameter with index "ProjectorParamIndex". */
+	UFUNCTION(BlueprintCallable, Category = CustomizableObjectInstance)
+	void SetProjectorValue(const FString& ProjectorParamName,
+		const FVector& OutPos, const FVector& OutDirection, const FVector& OutUp, const FVector& OutScale,
+		float OutAngle,
+		int32 RangeIndex = -1);
+
+	/** Set only the projector position keeping the rest of values. */
+	UFUNCTION(BlueprintCallable, Category = CustomizableObjectInstance)
+	void SetProjectorPosition(const FString& ProjectorParamName, const FVector& Pos, int32 RangeIndex = -1);
+	
+	/** Set only the projector direction vector keeping the rest of values. */
+	UFUNCTION(BlueprintCallable, Category = CustomizableObjectInstance)
+	void SetProjectorDirection(const FString& ProjectorParamName, const FVector& Direction, int32 RangeIndex = -1);
+	
+	/** Set only the projector up vector keeping the rest of values. */
+	UFUNCTION(BlueprintCallable, Category = CustomizableObjectInstance)
+	void SetProjectorUp(const FString& ProjectorParamName, const FVector& Up, int32 RangeIndex = -1);
+
+	/** Set only the projector scale keeping the rest of values. */
+	UFUNCTION(BlueprintCallable, Category = CustomizableObjectInstance)
+	void SetProjectorScale(const FString& ProjectorParamName, const FVector& Scale, int32 RangeIndex = -1);
+
+	/** Set only the cylindrical projector angle keeping the rest of values. */
+	UFUNCTION(BlueprintCallable, Category = CustomizableObjectInstance)
+	void SetProjectorAngle(const FString& ProjectorParamName, float Angle, int32 RangeIndex = -1);
+	
+	/** Get the projector values of a projector parameter with index "ProjectorParamIndex". */
+	UFUNCTION(BlueprintCallable, Category = CustomizableObjectInstance)
+	void GetProjectorValue(const FString& ProjectorParamName,
+		FVector& OutPos, FVector& OutDirection, FVector& OutUp, FVector& OutScale,
+		float& OutAngle, ECustomizableObjectProjectorType& OutType,
+		int32 RangeIndex = -1) const;
+
+	void GetProjectorValueF(const FString& ProjectorParamName,
+		FVector3f& Pos, FVector3f& Direction, FVector3f& Up, FVector3f& Scale,
+		float& Angle, ECustomizableObjectProjectorType& Type,
+		int32 RangeIndex = -1) const;
+
+	/** Get the current projector position for the parameter with the given name. */
+	UFUNCTION(BlueprintCallable, Category = CustomizableObjectInstance)
+	FVector GetProjectorPosition(const FString & ParamName, int32 RangeIndex = -1) const;
+
+	/** Get the current projector direction vector for the parameter with the given name. */
+	UFUNCTION(BlueprintCallable, Category = CustomizableObjectInstance)
+	FVector GetProjectorDirection(const FString & ParamName, int32 RangeIndex = -1) const;
+
+	/** Get the current projector up vector for the parameter with the given name. */
+	UFUNCTION(BlueprintCallable, Category = CustomizableObjectInstance)
+	FVector GetProjectorUp(const FString & ParamName, int32 RangeIndex = -1) const;
+
+	/** Get the current projector scale for the parameter with the given name. */
+	UFUNCTION(BlueprintCallable, Category = CustomizableObjectInstance)
+	FVector GetProjectorScale(const FString & ParamName, int32 RangeIndex = -1) const;
+
+	/** Get the current cylindrical projector angle for the parameter with the given name. */
+	UFUNCTION(BlueprintCallable, Category = CustomizableObjectInstance)
+	float GetProjectorAngle(const FString& ParamName, int32 RangeIndex = -1) const;
+
+	/** Get the current projector type for the parameter with the given name. */
+	UFUNCTION(BlueprintCallable, Category = CustomizableObjectInstance)
+	ECustomizableObjectProjectorType GetProjectorParameterType(const FString& ParamName, int32 RangeIndex = -1) const;
+
+	/** Get the current projector for the parameter with the given name. */
+	FCustomizableObjectProjector GetProjector(const FString& ParamName, int32 RangeIndex) const;
+	
+	/** Finds in IntParameters a parameter with name ParamName, returns the index if found, -1 otherwise. */
+	UFUNCTION(BlueprintCallable, Category = CustomizableObjectInstance)
+	int32 FindIntParameterNameIndex(const FString& ParamName) const;
+
+	/** Finds in FloatParameters a parameter with name ParamName, returns the index if found, -1 otherwise. */
+	UFUNCTION(BlueprintCallable, Category = CustomizableObjectInstance)
+	int32 FindFloatParameterNameIndex(const FString& ParamName) const;
+
+	/** Finds in BoolParameters a parameter with name ParamName, returns the index if found, -1 otherwise. */
+	UFUNCTION(BlueprintCallable, Category = CustomizableObjectInstance)
+	int32 FindBoolParameterNameIndex(const FString& ParamName) const;
+
+	/** Finds in VectorParameters a parameter with name ParamName, returns the index if found, -1 otherwise. */
+	UFUNCTION(BlueprintCallable, Category = CustomizableObjectInstance)
+	int32 FindVectorParameterNameIndex(const FString& ParamName) const;
+
+	/** Finds in ProjectorParameters a parameter with name ParamName, returns the index if found, -1 otherwise. */
+	UFUNCTION(BlueprintCallable, Category = CustomizableObjectInstance)
+	int32 FindProjectorParameterNameIndex(const FString& ParamName) const;
+
+	/** Increases the range of values of the integer with ParamName and returns the index of the new integer value, -1 otherwise.
+	  * The added value is initialized with the first integer option and is the last one of the range. */
+	UFUNCTION(BlueprintCallable, Category = CustomizableObjectInstance)
+	int32 AddValueToIntRange(const FString& ParamName);
+
+	/** Increases the range of values of the float with ParamName, returns the index of the new float value, -1 otherwise.
+	  * The added value is initialized with 0.5f and is the last one of the range. */
+	UFUNCTION(BlueprintCallable, Category = CustomizableObjectInstance)
+	int32 AddValueToFloatRange(const FString& ParamName);
+
+	/** Increases the range of values of the projector with ParamName, returns the index of the new projector value, -1 otherwise.
+	  * The added value is initialized with the default projector as set up in the editor and is the last one of the range. */
+	UFUNCTION(BlueprintCallable, Category = CustomizableObjectInstance)
+	int32 AddValueToProjectorRange(const FString& ParamName);
+
+	/** Remove the last of the integer range of values from the parameter ParamName, returns the index of the last valid integer, -1 if no values left. */
+	UFUNCTION(BlueprintCallable, Category = CustomizableObjectInstance)
+	int32 RemoveValueFromIntRange(const FString& ParamName);
+
+	/** Remove the RangeIndex element of the integer range of values from the parameter ParamName, returns the index of the last valid integer, -1 if no values left. */
+	int32 RemoveValueFromIntRange(const FString& ParamName, int32 RangeIndex);
+
+	/** Remove the last of the float range of values from the parameter ParamName, returns the index of the last valid float, -1 if no values left. */
+	UFUNCTION(BlueprintCallable, Category = CustomizableObjectInstance)
+	int32 RemoveValueFromFloatRange(const FString& ParamName);
+
+	/** Remove the RangeIndex element of the float range of values from the parameter ParamName, returns the index of the last valid float, -1 if no values left. */
+	int32 RemoveValueFromFloatRange(const FString& ParamName, int32 RangeIndex);
+
+	/** Remove the last of the projector range of values from the parameter ParamName, returns the index of the last valid projector, -1 if no values left. */
+	UFUNCTION(BlueprintCallable, Category = CustomizableObjectInstance)
+	int32 RemoveValueFromProjectorRange(const FString& ParamName);
+
+	/** Remove the RangeIndex element of the projector range of values from the parameter ParamName, returns the index of the last valid projector, -1 if no values left. */
+	int32 RemoveValueFromProjectorRange(const FString& ParamName, int32 RangeIndex);
+
+	// ------------------------------------------------------------
+	// Multilayer Projectors
+	// ------------------------------------------------------------
+	
+	/** Given Multilayer Projector name, create a new Multilayer Projector Helper (if non-existent). See FMultilayerProjector.
+	  * @return true if successfully created (or was already created). */
+	UFUNCTION(BlueprintCallable, Category = CustomizableObjectInstance)
+	bool CreateMultiLayerProjector(const FName& ProjectorParamName);
+	
+	/** Given Multilayer Projector name, remove a Multilayer Projector Helper. See FMultilayerProjector. */
+	UFUNCTION(BlueprintCallable, Category = CustomizableObjectInstance)
+	void RemoveMultilayerProjector(const FName& ProjectorParamName);
+	
+	// Layers
+
+	/** See FMultilayerProjector::NumLayers. */
+	UFUNCTION(BlueprintCallable, Category = CustomizableObjectInstance)
+	int32 MultilayerProjectorNumLayers(const FName& ProjectorParamName) const;
+
+	/** See FMultilayerProjector::CreateLayer. */
+	UFUNCTION(BlueprintCallable, Category = CustomizableObjectInstance)
+	void MultilayerProjectorCreateLayer(const FName& ProjectorParamName, int32 Index);
+
+	/** See FMultilayerProjector::RemoveLayerAt. */
+	UFUNCTION(BlueprintCallable, Category = CustomizableObjectInstance)
+	void MultilayerProjectorRemoveLayerAt(const FName& ProjectorParamName, int32 Index);
+
+	/** See FMultilayerProjector::GetLayer. */
+	UFUNCTION(BlueprintCallable, Category = CustomizableObjectInstance)
+	FMultilayerProjectorLayer MultilayerProjectorGetLayer(const FName& ProjectorParamName, int32 Index) const;
+
+	/** See FMultilayerProjector::UpdateLayer. */
+	UFUNCTION(BlueprintCallable, Category = CustomizableObjectInstance)
+	void MultilayerProjectorUpdateLayer(const FName& ProjectorParamName, int32 Index, const FMultilayerProjectorLayer& Layer);
+
+	// Virtual layers
+
+	/** See FMultilayerProjector::GetVirtualLayers. */
+	UFUNCTION(BlueprintCallable, Category = CustomizableObjectInstance)
+	TArray<FName> MultilayerProjectorGetVirtualLayers(const FName& ProjectorParamName) const;
+	
+	/** See FMultilayerProjector::VirtualLayer. */
+	UFUNCTION(BlueprintCallable, Category = CustomizableObjectInstance)
+	void MultilayerProjectorCreateVirtualLayer(const FName& ProjectorParamName, const FName& Id);
+
+	/** See FMultilayerProjector::FindOrCreateVirtualLayer. */
+	UFUNCTION(BlueprintCallable, Category = CustomizableObjectInstance)
+	FMultilayerProjectorVirtualLayer MultilayerProjectorFindOrCreateVirtualLayer(const FName& ProjectorParamName, const FName& Id);
+
+	/** See FMultilayerProjector::RemoveVirtualLayer. */
+	UFUNCTION(BlueprintCallable, Category = CustomizableObjectInstance)
+	void MultilayerProjectorRemoveVirtualLayer(const FName& ProjectorParamName, const FName& Id);
+
+	/** See FMultilayerProjector::GetVirtualLayer. */
+	UFUNCTION(BlueprintCallable, Category = CustomizableObjectInstance)
+	FMultilayerProjectorVirtualLayer MultilayerProjectorGetVirtualLayer(const FName& ProjectorParamName, const FName& Id) const;
+
+	/** See FMultilayerProjector::UpdateVirtualLayer. */
+	UFUNCTION(BlueprintCallable, Category = CustomizableObjectInstance)
+	void MultilayerProjectorUpdateVirtualLayer(const FName& ProjectorParamName, const FName& Id, const FMultilayerProjectorVirtualLayer& Layer);
+	
+	// ------------------------------------------------------------
+	
+	/** Returns the animation BP for the parameter component and slot, gathered from all the meshes that compose this instance. */
+	UFUNCTION(BlueprintCallable, Category = CustomizableObjectInstance)
+	TSubclassOf<UAnimInstance> GetAnimBP(int32 ComponentIndex, const FName& Slot) const;
+
+	/** Return the list of tags for this instance. */
+	UFUNCTION(BlueprintCallable, Category = CustomizableObjectInstance)
+	const FGameplayTagContainer& GetAnimationGameplayTags() const;
+	
+	/** Execute a delegate for each animation instance involved in this customizable object instance. */
+	UFUNCTION(BlueprintCallable, Category = CustomizableObjectInstance)
+	void ForEachAnimInstance(int32 ComponentIndex, FEachComponentAnimInstanceClassDelegate Delegate) const;
+
+	void ForEachAnimInstance(int32 ComponentIndex, FEachComponentAnimInstanceClassNativeDelegate Delegate) const;
+
+	/** Check if the given UAnimInstance class requires to be fixed up. */
+	bool AnimInstanceNeedsFixup(TSubclassOf<UAnimInstance> AnimInstance) const;
+
+	/** Fix the given UAnimInstance instance. */
+	void AnimInstanceFixup(UAnimInstance* AnimInstance) const;
+
+	/** See FCustomizableObjectInstanceDescriptor::SaveDescriptor. */
+	void SaveDescriptor(FArchive &CustomizableObjectDescriptor, bool bUseCompactDescriptor);
+
+	/** See FCustomizableObjectInstanceDescriptor::LoadDescriptor. */
+	void LoadDescriptor(FArchive &CustomizableObjectDescriptor);
+
+	/** Enable physics asset replacement so that generated skeletal meshes have the merged physics assets of their skeletal mesh parts and reference mesh. */
+	UFUNCTION(BlueprintCallable, Category = CustomizableObjectInstance)
+	void SetReplacePhysicsAssets(bool bReplaceEnabled);
+
+	/** Enables the reuse of all possible textures when the instance is updated without any changes in geometry or state (the first update after creation doesn't reuse any)
+	  * It will only work if the textures aren't compressed, so set the instance to a Mutable state with texture compression disabled
+	  * WARNING! If texture reuse is enabled, do NOT keep external references to the textures of the instance. The instance owns the textures. */
+	void SetReuseInstanceTextures(bool bTextureReuseEnabled);
+	
+	/** If enabled, low-priority textures will generate resident mipmaps too. */
+	void SetForceGenerateResidentMips(bool bForceGenerateResidentMips);
+
+	const TArray<TObjectPtr<UMaterialInterface>>* GetOverrideMaterials(int32 ComponentIndex) const;
+	
+	// The following methods should only be used in an LOD management class
+	void SetIsBeingUsedByComponentInPlay(bool bIsUsedByComponent);
+	bool GetIsBeingUsedByComponentInPlay() const;
+	void SetIsDiscardedBecauseOfTooManyInstances(bool bIsDiscarded);
+	bool GetIsDiscardedBecauseOfTooManyInstances() const;
+	void SetIsPlayerOrNearIt(bool NewValue);
+	float GetMinSquareDistToPlayer() const;
+	void SetMinSquareDistToPlayer(float NewValue);
+
+	/** Return the number of components generated in this instance. */
+	UFUNCTION(BlueprintCallable, Category = CustomizableObjectInstance)
+	int32 GetNumComponents() const;
+
+	/** Return the min LOD that will be used in the next Instance update. */
+	int32 GetMinLODToLoad() const;
+
+	int32 GetNumLODsAvailable() const;
+
+	/** Return the min LOD this Instance is using. */
+	int32 GetCurrentMinLOD() const;
+
+	/** Sets an array of LODs to generate per component. Mutable will generate those plus the currently generated LODs (if any).
+	  * Requires mutable.EnableOnlyGenerateRequestedLODs and CurrentInstanceLODManagement->IsOnlyGenerateRequestedLODLevelsEnabled() to be true.
+	  * @param InMinLOD - MinLOD to generate.
+	  * @param InMaxLOD - MaxLOD to generate - DEPRECATED.
+	  * @param InRequestedLODsPerComponent - Array with bitmasks of requested LODs per component with range from [0 .. CO->GetComponentCount()].
+	  * @param InOutRequestedUpdates - Map from Instance to Update data that stores a request for the Instance to be updated, which will be either processed or discarded by priority (to be rerequested the next tick) */
+	void SetRequestedLODs(int32 InMinLOD, int32 InMaxLOD, const TArray<uint16>& InRequestedLODsPerComponent, FMutableInstanceUpdateMap& InOutRequestedUpdates);
+
+	const TArray<uint16>& GetRequestedLODsPerComponent() const;
+
+#if WITH_EDITOR
+	/**
+	 * Performs the baking of the instance resources in an async fashion. Bind yourself to the callback present in InBakingConfiguration to get notified in case it fails
+	 * @param InBakingConfiguration The configuration to be using for the baking operation
+	 */
+	UFUNCTION(BlueprintCallable, Category = CustomizableObjectInstanceBaking)
+	void Bake(const FBakingConfiguration& InBakingConfiguration);
+#endif
+	
+	UCustomizableInstancePrivate* GetPrivate() const;
+
+private:
+	UPROPERTY()
+	FCustomizableObjectInstanceDescriptor Descriptor;
+
+	UPROPERTY()
+	TObjectPtr<UCustomizableInstancePrivate> PrivateData;
+
+public:
+#if WITH_EDITORONLY_DATA
+	/** Textures which can used as values in Texture Parameters. */
+	UPROPERTY(EditAnywhere, Category = TextureParameter)
+	TArray<TObjectPtr<UTexture2D>> TextureParameterDeclarations;
+#endif
+
+private:
+	// Deprecated properties	
+	UPROPERTY()
+	TObjectPtr<UCustomizableObject> CustomizableObject_DEPRECATED;
+	
+	UPROPERTY()
+	TArray<FCustomizableObjectBoolParameterValue> BoolParameters_DEPRECATED;
+
+	UPROPERTY()
+	TArray<FCustomizableObjectIntParameterValue> IntParameters_DEPRECATED;
+
+	UPROPERTY()
+	TArray<FCustomizableObjectFloatParameterValue> FloatParameters_DEPRECATED;
+
+	UPROPERTY()																																						
+	TArray<FCustomizableObjectTextureParameterValue> TextureParameters_DEPRECATED;
+
+	UPROPERTY()
+	TArray<FCustomizableObjectVectorParameterValue> VectorParameters_DEPRECATED;
+
+	UPROPERTY()
+	TArray<FCustomizableObjectProjectorParameterValue> ProjectorParameters_DEPRECATED;
+	
+   	UPROPERTY()
+   	TMap<FName, FMultilayerProjector> MultilayerProjectors_DEPRECATED;
+
+	bool bBuildParameterRelevancy_DEPRECATED = false;
+};
+
+
+#if UE_ENABLE_INCLUDE_ORDER_DEPRECATED_IN_5_2
+#include "Async/TaskGraphInterfaces.h"
+#include "MuCO/CustomizableObjectParameterTypeDefinitions.h"
+#endif
diff --git a/Engine/Plugins/Experimental/Mutable/Source/CustomizableObject/Public/MuCO/CustomizableObjectSystem.h b/Engine/Plugins/Experimental/Mutable/Source/CustomizableObject/Public/MuCO/CustomizableObjectSystem.h
new file mode 100644
--- /dev/null
+++ b/Engine/Plugins/Experimental/Mutable/Source/CustomizableObject/Public/MuCO/CustomizableObjectSystem.h
@@ -0,0 +1,368 @@
+// Copyright Epic Games, Inc. All Rights Reserved.
+
+#pragma once
+
+#include "HAL/IConsoleManager.h"
+#include "Engine/StreamableManager.h"
+#include "AssetRegistry/AssetData.h"
+#include "MuCO/CustomizableObjectInstance.h"
+#include "MuCO/CustomizableObjectInstanceDescriptor.h"
+#include "MuR/Parameters.h"
+#include "MuR/Image.h"
+#include "MuR/Types.h"
+
+#include "CustomizableObjectSystem.generated.h"
+
+class IConsoleVariable;
+class UCustomizableInstanceLODManagementBase;
+class ITargetPlatform;
+class UCustomizableObject;
+class UDefaultImageProvider;
+class USkeletalMesh;
+class UMaterialInterface;
+class UTexture2D;
+class UCustomizableObjectSystemPrivate; // This is used to hide Mutable SDK members in the public headers.
+class FUpdateContextPrivate;
+struct FFrame;
+struct FGuid;
+
+
+constexpr uint64 KEY_OFFSET_COMPILATION_OUT_OF_DATE = 1;
+
+extern TAutoConsoleVariable<bool> CVarClearWorkingMemoryOnUpdateEnd;
+
+extern TAutoConsoleVariable<bool> CVarReuseImagesBetweenInstances;
+
+extern TAutoConsoleVariable<bool> CVarPreserveUserLODsOnFirstGeneration;
+
+extern TAutoConsoleVariable<bool> CVarEnableMeshCache;
+
+extern TAutoConsoleVariable<bool> CVarRollbackFixModelDiskStreamerDataRace;
+
+extern TAutoConsoleVariable<bool> CVarEnableNewSplitMutableTask;
+
+#if WITH_EDITOR
+
+// Struct used to keep a copy of the EditorSettings needed to compile Customizable Objects.
+struct FEditorCompileSettings 
+{
+	// General case
+	bool bIsMutableEnabled = true;
+
+	// Auto Compile 
+	bool bEnableAutomaticCompilation = true;
+	bool bCompileObjectsSynchronously = true;
+	bool bCompileRootObjectsOnStartPIE = false;
+};
+
+#endif
+
+namespace EMutableProfileMetric
+{
+	typedef uint8 Type;
+
+	constexpr Type BuiltInstances = 1;
+	constexpr Type UpdateOperations = 2;
+	constexpr Type Count = 4;
+
+};
+
+
+USTRUCT()
+struct FPendingReleaseMaterialsInfo
+{
+	GENERATED_USTRUCT_BODY()
+
+	UPROPERTY()
+	TArray<TObjectPtr<UMaterialInterface>> Materials;
+
+	UPROPERTY()
+	int32 TicksUntilRelease = 0;
+};
+
+
+
+/** Key to identify an image inside a generated mutable runtime instance. */
+struct FMutableImageReference
+{
+	/** Original image ID. Once generated it will be unique. However, future updates of the image may return a different ID for the
+	* same image, if many other resources have been built in the middle. For this reason the rest of the data in the struct is what
+	* must be used to request the additional mips.	*/
+	uint32 ImageID = 0;
+
+	uint32 SurfaceId = 0;
+
+	uint8 LOD = 0;
+	uint8 Component = 0;
+	uint8 Image = 0;
+
+	uint8 BaseMip = 0;
+};
+
+
+
+// Opaque representation of a possible registered value for texture parameters.
+struct FCustomizableObjectExternalTexture
+{
+	FCustomizableObjectExternalTexture() = default;
+
+	FString Name;
+	FName Value;
+};
+
+
+/** Base class for Image provider. */
+UCLASS(abstract)
+class CUSTOMIZABLEOBJECT_API UCustomizableSystemImageProvider : public UObject
+{
+public:
+	GENERATED_BODY()
+
+	enum class ValueType : uint8
+	{
+		// This texture is not provided by this provider.
+		None,
+
+		// Data will be provided with size and pointer
+		Raw,
+
+		// Data will be provided from an unreal texture, loaded in the game thread and kept in memory
+		Unreal,
+
+		// Data will be provided from an unreal texture, and will only be loaded when actually needed in the Mutable thread
+		Unreal_Deferred,
+
+		// Number of elements of this enum.
+		Count
+	};
+
+	// Query that Mutable will run to find out if a texture will be provided as an Unreal UTexture2D,
+	// or as a raw data blob.
+	virtual ValueType HasTextureParameterValue(const FName& ID) { return ValueType::None; }
+
+	// In case IsTextureParameterValueUnreal returns false, this will be used to query the texture size data.
+	virtual FIntVector GetTextureParameterValueSize(const FName& ID) { return FIntVector(0, 0, 0); }
+
+	// In case IsTextureParameterValueUnreal returns false, this will be used to query the texture data that must
+	// be copied in the preallocated buffer. The pixel format is assumed to be 4-channel RGBA, uint8_t per channel.
+	virtual void GetTextureParameterValueData(const FName& ID, uint8* OutData) {}
+
+	// In case IsTextureParameterValueUnreal returns true, this will be used to query the texture.
+	virtual UTexture2D* GetTextureParameterValue(const FName& ID) { return nullptr; }
+
+	// Used in the editor to show the list of available options.
+	// Only necessary if the images are required in editor previews.
+	virtual void GetTextureParameterValues(TArray<FCustomizableObjectExternalTexture>& OutValues) {}
+};
+
+
+UCLASS(Blueprintable, BlueprintType)
+class CUSTOMIZABLEOBJECT_API UCustomizableObjectSystem : public UObject
+{
+	// Friends
+	friend class UCustomizableObjectSystemPrivate;
+
+public:
+	GENERATED_BODY()
+
+	UCustomizableObjectSystem() = default;
+	void InitSystem();
+
+	/** Get the singleton object. It will be created if it doesn't exist yet. */
+	UFUNCTION(BlueprintCallable, BlueprintPure, Category = Status)
+	static UCustomizableObjectSystem* GetInstance();
+	
+	UFUNCTION(BlueprintCallable, BlueprintPure, Category = Status)
+	static UCustomizableObjectSystem* GetInstanceChecked();
+
+	/** Determines if the result of the instance update is valid or not.
+	 * @return true if the result is successful or has warnings, false if the result is from the Error category */
+	UFUNCTION(BlueprintCallable, Category = Status)
+	static bool IsUpdateResultValid(const EUpdateResult UpdateResult);
+	
+	// Return true if the singleton has been created. It is different than GetInstance in that GetInstance will create it if it doesn't exist.
+	static bool IsCreated();
+
+	/** Returns the current status of Mutable. Only when active is it possible to compile COs, generate instances, and stream textures.
+	  * @return True if Mutable is enabled. */
+	static bool IsActive();
+
+	// Begin UObject interface.
+	virtual void BeginDestroy() override;
+	virtual FString GetDesc() override;
+	// End UObject interface.
+
+	bool IsReplaceDiscardedWithReferenceMeshEnabled() const;
+	void SetReplaceDiscardedWithReferenceMeshEnabled(bool bIsEnabled);
+
+	// Unprotect the resources in the instances of this object of being garbage collector
+	// while an instance is being built or updated, so that they can be reused.
+	void ClearResourceCacheProtected();
+
+#if WITH_EDITOR
+	// Lock a CustomizableObjects, preventing the generation or update of any of its instances
+	// Will return true if successful, false if it fails to lock because an update is already underway
+	// This is usually only used in the editor
+	bool LockObject(const UCustomizableObject*);
+	void UnlockObject(const UCustomizableObject*);
+
+	/** Checks if there are any outstanding disk or mip update operations in flight for the parameter Customizable Object that may
+	* make it unsafe to compile at the moment.
+	* @return true if there are operations in flight and it's not safe to compile */
+	bool CheckIfDiskOrMipUpdateOperationsPending(const UCustomizableObject& Object) const;
+	
+	// Called whenever the Mutable Editor Settings change, copying the new value of the current needed settings to the Customizable Object System
+	void EditorSettingsChanged(const FEditorCompileSettings& InEditorSettings);
+
+	// If true, uncompiled Customizable Objects will be compiled whenever an instance update is required
+	bool IsAutoCompileEnabled() const;
+
+	/** Return true if inside commandlets uncompiled Customizable Objects will be compiled whenever an instance update is required. */
+	bool IsAutoCompileCommandletEnabled() const;
+
+	/** Set if inside commandlets uncompiled Customizable Objects will be compiled whenever an instance update is required. */
+	void SetAutoCompileCommandletEnabled(bool bValue);
+	
+	// If true, uncompiled Customizable Objects will be compiled synchronously
+	bool IsAutoCompilationSync() const;
+#endif
+	
+	// Return the current MinLodQualityLevel for skeletal meshes.
+	int32 GetSkeletalMeshMinLODQualityLevel() const;
+
+	bool IsSupport16BitBoneIndexEnabled() const;
+
+	bool IsProgressiveMipStreamingEnabled() const;
+	void SetProgressiveMipStreamingEnabled(bool bIsEnabled);
+
+	bool IsOnlyGenerateRequestedLODsEnabled() const;
+	void SetOnlyGenerateRequestedLODsEnabled(bool bIsEnabled);
+
+#if WITH_EDITOR
+	void SetImagePixelFormatOverride(const mu::FImageOperator::FImagePixelFormatFunc&);
+#endif
+
+	/** [Texture Parameters] Get a list of all the possible values for external texture parameters according to the various providers registered with RegisterImageProvider. */
+	TArray<FCustomizableObjectExternalTexture> GetTextureParameterValues();
+
+	/** [Texture Parameters] Add a new image provider to the CustomizableObject System. This will be queried for when an external image ID is provided to mutable in an Texture Parameter node. */
+	void RegisterImageProvider(UCustomizableSystemImageProvider* Provider);
+
+	/** [Texture Parameters] Remove a previously registered provider. */
+	void UnregisterImageProvider(UCustomizableSystemImageProvider* Provider);
+
+	/** [Texture Parameters] Interface to actually cache Images in the Mutable system and make them available at run-time.
+		Any cached image has to be registered by an Image provider before caching it.
+		Have in mind that once an image has been cached, it will spend memory according to its size, except in the case
+		of images of type UCustomizableSystemImageProvider::ValueType::Unreal_Deferred, where only a very small amount of
+		memory is used and the real texel data is loaded when needed during an update and then immediately discarded */
+
+	/** [Texture Parameters] Cache an image which has to have been previously registered by an Image provider with the parameter id. */
+	void CacheImage(FName ImageId);
+	/** [Texture Parameters] Remove an image from the cache. */
+	void UnCacheImage(FName ImageId);
+	/** [Texture Parameters] Remove all images from the cache. */
+	void ClearImageCache();
+
+	// Show a warning on-screen and via a notification (if in Editor) and log an error when a CustomizableObject is
+	// being used and it's not compiled.  Callers can add additional information to the error log.
+	void AddUncompiledCOWarning(const UCustomizableObject& InObject, FString const* OptionalLogInfo = nullptr);
+
+	// Enables the collection of internal Mutable performance data. It has a performance cost.
+	void EnableBenchmark();
+	// Writes the benchmark results
+	void EndBenchmark();
+
+	// Show data about all UCustomizableObjectInstance existing elements
+	void LogShowData(bool bFullInfo, bool ShowMaterialInfo) const;
+
+	// Give access to the internal object data.
+	UCustomizableObjectSystemPrivate* GetPrivate();
+	const UCustomizableObjectSystemPrivate* GetPrivate() const;
+
+	UCustomizableInstanceLODManagementBase* GetInstanceLODManagement() const;
+
+	// Pass a null ptr to reset to the default InstanceLODManagement
+	void SetInstanceLODManagement(UCustomizableInstanceLODManagementBase* NewInstanceLODManagement);
+
+	// Find out the version of the plugin
+	UFUNCTION(BlueprintCallable, Category = Status)
+	FString GetPluginVersion() const;
+
+	// Get the number of instances built and alive.
+	UFUNCTION(BlueprintCallable, Category = Status)
+	int32 GetNumInstances() const;
+
+	// Get the number of instances waiting to be updated.
+	UFUNCTION(BlueprintCallable, Category = Status)
+	int32 GetNumPendingInstances() const;
+
+	// Get the total number of instances including built and not built.
+	UFUNCTION(BlueprintCallable, Category = Status)
+	int32 GetTotalInstances() const;
+
+	// Get the amount of GPU memory in use in bytes for textures generated by mutable.
+	UFUNCTION(BlueprintCallable, Category = Status)
+	int64 GetTextureMemoryUsed() const;
+
+	// Return the average build/update time of an instance in ms.
+	UFUNCTION(BlueprintCallable, Category = Status)
+	int32 GetAverageBuildTime() const;
+
+	// If set to true, Mutable will release Mutable-generated textures immediately when they are no longer used without waiting for GC
+	// IMPORTANT!!! Do NOT keep references to any Mutable generated textures or skeletal meshes if this is enabled,
+	// they are owned by Mutable and will be destroyed without notice
+	UFUNCTION(BlueprintCallable, Category = Status)
+	void SetReleaseMutableTexturesImmediately(bool bReleaseTextures);
+
+	bool IsMutableAnimInfoDebuggingEnabled() const;
+
+#if WITH_EDITOR
+	// Get the maximum size a chunk can have on a specific platform. If unspecified return MUTABLE_STREAMED_DATA_MAXCHUNKSIZE
+	uint64 GetMaxChunkSizeForPlatform(const ITargetPlatform* InTargetPlatform);
+#endif
+
+	void ClearCurrentMutableOperation();
+
+private:
+	// Most of the work in this plugin happens here.
+	bool Tick(float DeltaTime);
+
+	/** Returns the number of remaining operations. */
+	int32 TickInternal();
+
+	// If there is an on-going operation, advance it.
+	void AdvanceCurrentOperation();
+
+	void DiscardInstances();
+	void ReleaseInstanceIDs();
+
+public:
+	/** Return true if the instance is being updated. */
+	UFUNCTION(BlueprintCallable, BlueprintPure, Category = CustomizableObjectSystem)
+	bool IsUpdating(const UCustomizableObjectInstance* Instance) const;
+
+	/** Set Mutable's working memory limit (bytes). Mutable will flush internal caches to try to keep its memory consumption below the WorkingMemory (i.e., it is not a hard limit).
+	 * The working memory limit will especially reduce the memory required to perform Instance Updates and Texture Streaming.
+ 	 * Notice that Mutable does not track all its memory (e.g., UObjects memory is no tracked).
+	 * This value can also be set using "mutable.WorkingMemory" CVar. */
+	void SetWorkingMemory(int32 Bytes);
+
+	/** Get Mutable's working memory limit (bytes). See SetWorkingMemory(int32). */
+	int32 GetWorkingMemory() const;
+
+#if WITH_EDITOR
+	// Copy of the Mutable Editor Settings tied to CO compilation. They are updated whenever changed
+	FEditorCompileSettings EditorSettings;
+#endif
+	
+private:
+	UPROPERTY(Transient)
+	TObjectPtr<UCustomizableObjectSystemPrivate> Private = nullptr;
+};
+
+
+#if UE_ENABLE_INCLUDE_ORDER_DEPRECATED_IN_5_2
+#include "Engine/EngineTypes.h"
+#include "PixelFormat.h"
+#endif
diff --git a/Engine/Plugins/Experimental/Mutable/Source/CustomizableObject/Public/MuCO/CustomizableSkeletalComponent.h b/Engine/Plugins/Experimental/Mutable/Source/CustomizableObject/Public/MuCO/CustomizableSkeletalComponent.h
new file mode 100644
--- /dev/null
+++ b/Engine/Plugins/Experimental/Mutable/Source/CustomizableObject/Public/MuCO/CustomizableSkeletalComponent.h
@@ -0,0 +1,70 @@
+// Copyright Epic Games, Inc. All Rights Reserved.
+
+#pragma once
+
+#include "Components/SceneComponent.h"
+
+#include "MuCO/CustomizableObjectInstance.h"
+
+#include "CustomizableSkeletalComponent.generated.h"
+
+
+DECLARE_DELEGATE(FCustomizableSkeletalComponentUpdatedDelegate);
+
+UCLASS(Blueprintable, BlueprintType, ClassGroup = (CustomizableObject), meta = (BlueprintSpawnableComponent))
+class CUSTOMIZABLEOBJECT_API UCustomizableSkeletalComponent : public USceneComponent
+{
+public:
+	GENERATED_BODY()
+
+	UPROPERTY(BlueprintReadWrite, EditAnywhere, Category = CustomizableSkeletalMesh)
+	TObjectPtr<UCustomizableObjectInstance> CustomizableObjectInstance;
+
+	UPROPERTY(BlueprintReadWrite, EditAnywhere, Category = CustomizableSkeletalMesh)
+	int32 ComponentIndex;
+
+	// Used to replace the SkeletalMesh of the parent component by the ReferenceSkeletalMesh or the generated SkeletalMesh 
+	bool bPendingSetSkeletalMesh = false;
+
+	// Used to avoid replacing the SkeletalMesh of the parent component by the ReferenceSkeletalMesh if bPendingSetSkeletalMesh is true
+	bool bSkipSetReferenceSkeletalMesh = false;
+
+	FCustomizableSkeletalComponentUpdatedDelegate UpdatedDelegate;
+
+	/** Common end point of all updates. Even those which failed. */
+	void Callbacks() const;
+	
+	USkeletalMesh* GetSkeletalMesh() const;
+	USkeletalMesh* GetAttachedSkeletalMesh() const;
+
+	/** Update Skeletal Mesh asynchronously. */
+	UFUNCTION(BlueprintCallable, Category = CustomizableObject)
+	void UpdateSkeletalMeshAsync(bool bNeverSkipUpdate = false);
+
+	/** Update Skeletal Mesh asynchronously. Callback will be called once the update finishes, even if it fails. */
+	UFUNCTION(BlueprintCallable, Category = CustomizableObjectInstance)
+	void UpdateSkeletalMeshAsyncResult(FInstanceUpdateDelegate Callback, bool bIgnoreCloseDist = false, bool bForceHighPriority = false);
+
+	void SetSkeletalMesh(USkeletalMesh* SkeletalMesh);
+	void SetPhysicsAsset(class UPhysicsAsset* PhysicsAsset);
+	void UpdateDistFromComponentToPlayer(const AActor* const Pawn, bool bForceEvenIfNotBegunPlay = false);
+
+	void SetVisibilityOfSkeletalMeshSectionWithMaterialName(bool bVisible, const FString& MaterialName, int32 LOD);
+
+#if WITH_EDITOR
+	// Used to generate instances outside the CustomizableObject editor and PIE
+	void UpdateDistFromComponentToLevelEditorCamera(const FVector& CameraPosition);
+	void EditorUpdateComponent();
+#endif
+
+private:
+	UPROPERTY(Transient)
+	TObjectPtr<UCustomizableObjectInstanceUsage> CustomizableObjectInstanceUsage;
+
+	void OnAttachmentChanged() override;
+	void CreateCustomizableObjectInstanceUsage();
+
+protected:
+	virtual void PostInitProperties() override;
+};
+
diff --git a/Engine/Plugins/Experimental/Mutable/Source/CustomizableObject/Public/MuCO/CustomizableSkeletalMeshActor.h b/Engine/Plugins/Experimental/Mutable/Source/CustomizableObject/Public/MuCO/CustomizableSkeletalMeshActor.h
new file mode 100644
--- /dev/null
+++ b/Engine/Plugins/Experimental/Mutable/Source/CustomizableObject/Public/MuCO/CustomizableSkeletalMeshActor.h
@@ -0,0 +1,38 @@
+// Copyright Epic Games, Inc. All Rights Reserved.
+
+#pragma once
+
+#include "Animation/SkeletalMeshActor.h"
+
+#include "CustomizableSkeletalMeshActor.generated.h"
+
+class UCustomizableSkeletalComponent;
+
+class UObject;
+
+UCLASS(ClassGroup = CustomizableObject, Blueprintable, ComponentWrapperClass, ConversionRoot, meta = (ChildCanTick))
+class CUSTOMIZABLEOBJECT_API ACustomizableSkeletalMeshActor : public ASkeletalMeshActor
+{
+	GENERATED_UCLASS_BODY()
+public:
+	
+	class UCustomizableSkeletalComponent* GetCustomizableSkeletalComponent(int32 Index = 0) const { return CustomizableSkeletalComponents[Index]; }
+	
+	class USkeletalMeshComponent* GetSkeletalMeshComponentAt(int32 Index = 0) const { return SkeletalMeshComponents[Index]; }
+
+	void AttachNewComponent();
+
+	int32 GetNumComponents() { return CustomizableSkeletalComponents.Num(); }
+
+private:
+	
+	UPROPERTY(Category = CustomizableSkeletalMesh, VisibleAnywhere, BlueprintReadOnly, meta = (AllowPrivateAccess = "true"))
+	TArray< TObjectPtr<class UCustomizableSkeletalComponent> > CustomizableSkeletalComponents;
+
+	UPROPERTY(Category = CustomizableSkeletalMesh, VisibleAnywhere, BlueprintReadOnly, meta = (AllowPrivateAccess = "true"))
+	TArray< TObjectPtr<class USkeletalMeshComponent>> SkeletalMeshComponents;
+
+	// TODO: This is a temporary fix to not break the demos, we should update the demos to support the arrays of components
+	UPROPERTY(Category = CustomizableSkeletalMesh, VisibleAnywhere, BlueprintReadOnly, meta = (AllowPrivateAccess = "true"))
+	TObjectPtr<class UCustomizableSkeletalComponent> CustomizableSkeletalComponent;
+};
diff --git a/Engine/Plugins/Experimental/Mutable/Source/CustomizableObjectEditor/Public/MuCOE/CustomizableObjectGraph.h b/Engine/Plugins/Experimental/Mutable/Source/CustomizableObjectEditor/Public/MuCOE/CustomizableObjectGraph.h
new file mode 100644
--- /dev/null
+++ b/Engine/Plugins/Experimental/Mutable/Source/CustomizableObjectEditor/Public/MuCOE/CustomizableObjectGraph.h
@@ -0,0 +1,42 @@
+// Copyright Epic Games, Inc. All Rights Reserved.
+
+#pragma once
+
+#include "EdGraph/EdGraph.h"
+
+#include "CustomizableObjectGraph.generated.h"
+
+class UObject;
+
+
+UCLASS()
+class CUSTOMIZABLEOBJECTEDITOR_API UCustomizableObjectGraph : public UEdGraph
+{
+public:
+	GENERATED_BODY()
+
+	UCustomizableObjectGraph();
+
+	// UObject Interface
+	virtual void PostLoad() override;
+	void PostRename(UObject * OldOuter, const FName OldName) override;
+
+	// Own Interface
+	void NotifyNodeIdChanged(const FGuid& OldGuid, const FGuid& NewGuid);
+
+	FGuid RequestNotificationForNodeIdChange(const FGuid& OldGuid, const FGuid& NodeToNotifyGuid);
+
+	void PostDuplicate(bool bDuplicateForPIE) override;
+
+	/** Adds the necessary nodes for a CO to work */
+	void AddEssentialGraphNodes();
+
+private:
+
+	// Request Node Id Update Map
+	TMap<FGuid, TSet<FGuid>> NodesToNotifyMap;
+
+	// Guid map with the key beeing the old Guid and the Value the new one, filled after duplicating COs
+	TMap<FGuid, FGuid> NotifiedNodeIdsMap;
+};
+
diff --git a/Engine/Plugins/Experimental/Mutable/Source/CustomizableObjectEditor/Public/MuCOE/EdGraphSchema_CustomizableObject.h b/Engine/Plugins/Experimental/Mutable/Source/CustomizableObjectEditor/Public/MuCOE/EdGraphSchema_CustomizableObject.h
new file mode 100644
--- /dev/null
+++ b/Engine/Plugins/Experimental/Mutable/Source/CustomizableObjectEditor/Public/MuCOE/EdGraphSchema_CustomizableObject.h
@@ -0,0 +1,178 @@
+// Copyright Epic Games, Inc. All Rights Reserved.
+
+#pragma once
+
+#include "EdGraphSchema_K2_Actions.h"
+
+#include "EdGraphSchema_CustomizableObject.generated.h"
+
+enum class EMaterialParameterType : uint8;
+
+class UEdGraph;
+class UEdGraphNode;
+class UEdGraphPin;
+class UObject;
+struct FAssetData;
+struct FEdGraphPinType;
+class ICustomizableObjectEditor;
+
+/** Action to add a node to the graph */
+USTRUCT()
+struct CUSTOMIZABLEOBJECTEDITOR_API FCustomizableObjectSchemaAction_NewNode : public FEdGraphSchemaAction
+{
+	GENERATED_USTRUCT_BODY();
+
+	/** Template of node we want to create */
+	UPROPERTY()
+	TObjectPtr<class UEdGraphNode> NodeTemplate;
+
+
+	FCustomizableObjectSchemaAction_NewNode() 
+		: FEdGraphSchemaAction()
+		, NodeTemplate(NULL)
+	{}
+
+	FCustomizableObjectSchemaAction_NewNode(const FString& InNodeCategory, const FText& InMenuDesc, const FText& InToolTip, const int32 InGrouping, const FText& InKeywords = FText(), int32 InSectionID = 0)
+		: FEdGraphSchemaAction(FText::FromString(InNodeCategory), InMenuDesc, InToolTip, InGrouping, InKeywords, InSectionID)
+		, NodeTemplate(NULL)
+	{}
+
+	// FEdGraphSchemaAction interface
+	virtual UEdGraphNode* PerformAction(class UEdGraph* ParentGraph, UEdGraphPin* FromPin, const FVector2D Location, bool bSelectNewNode = true) override;
+	virtual UEdGraphNode* PerformAction(class UEdGraph* ParentGraph, TArray<UEdGraphPin*>& FromPins, const FVector2D Location, bool bSelectNewNode = true) override;
+	virtual void AddReferencedObjects( FReferenceCollector& Collector ) override;
+
+	/** Reimplementaiton of EdGraphSchema::CreateNode(...). Performs the overlap calculation before calling AutowireNewNode(...). AutowireNewNode can induce a call to a ReconstrucNode() which removes pins required for the calculation. */
+	static UEdGraphNode* CreateNode(class UEdGraph* ParentGraph, UEdGraphPin* FromPin, const FVector2D Location, class UEdGraphNode* InNodeTemplate);
+	// End of FEdGraphSchemaAction interface
+
+	template <typename NodeType>
+	static NodeType* InstantSpawn(class UEdGraph* ParentGraph, NodeType* InTemplateNode, const FVector2D Location)
+	{
+		FEdGraphSchemaAction_K2NewNode Action;
+		Action.NodeTemplate = InTemplateNode;
+
+		return Cast<NodeType>(Action.PerformAction(ParentGraph, NULL, Location));
+	}
+};
+
+
+/** Action to paste clipboard contents into the graph */
+USTRUCT()
+struct CUSTOMIZABLEOBJECTEDITOR_API FCustomizableObjectSchemaAction_Paste : public FEdGraphSchemaAction
+{
+	GENERATED_USTRUCT_BODY();
+
+	// Simple type info
+	static FName StaticGetTypeId() { static FName Type("FCustomizableObjectSchemaAction_Paste"); return Type; }
+	virtual FName GetTypeId() const override { return StaticGetTypeId(); }
+
+	FCustomizableObjectSchemaAction_Paste()
+		: FEdGraphSchemaAction()
+	{}
+
+	FCustomizableObjectSchemaAction_Paste(const FText& InNodeCategory, const FText& InMenuDesc, const FString& InToolTip, const int32 InGrouping)
+		: FEdGraphSchemaAction(InNodeCategory, InMenuDesc, FText::FromString(InToolTip), InGrouping)
+	{}
+
+	// FEdGraphSchemaAction interface
+	virtual UEdGraphNode* PerformAction(class UEdGraph* ParentGraph, UEdGraphPin* FromPin, const FVector2D Location, bool bSelectNewNode = true) override;
+	// End of FEdGraphSchemaAction interface
+};
+
+
+UCLASS()
+class CUSTOMIZABLEOBJECTEDITOR_API UEdGraphSchema_CustomizableObject : public UEdGraphSchema
+{
+public:
+	GENERATED_BODY()
+
+	// Allowable PinType.PinCategory values
+	static const FName PC_Object;
+	static const FName PC_Material;
+	static const FName PC_Mesh;
+	static const FName PC_Layout;
+	static const FName PC_Image;
+	static const FName PC_PassThroughImage;
+	static const FName PC_Projector;
+	static const FName PC_GroupProjector;
+	static const FName PC_Color;
+	static const FName PC_Float;
+	static const FName PC_Bool;
+	static const FName PC_Enum;
+	static const FName PC_Stack;
+	static const FName PC_MaterialAsset;
+	static const FName PC_Wildcard;
+
+	// EdGraphSchema interface
+	virtual void GetAssetsGraphHoverMessage(const TArray<FAssetData>& Assets, const UEdGraph* HoverGraph, FString& OutTooltipText, bool& OutOkIcon) const override;
+	virtual void GetGraphContextActions(FGraphContextMenuBuilder& ContextMenuBuilder) const override;
+	virtual const FPinConnectionResponse CanCreateConnection(const UEdGraphPin* A, const UEdGraphPin* B) const override;
+	virtual FLinearColor GetPinTypeColor(const FEdGraphPinType& PinType) const override;
+	virtual bool ShouldHidePinDefaultValue(UEdGraphPin* Pin) const override;
+	virtual void DroppedAssetsOnGraph(const TArray<FAssetData>& Assets, const FVector2D& GraphPosition, UEdGraph* Graph) const override;
+	virtual void OnPinConnectionDoubleCicked(UEdGraphPin* PinA, UEdGraphPin* PinB, const FVector2D& GraphPosition) const override;
+
+	FLinearColor GetPinTypeColor(const FName& PinType) const;
+
+	void GetContextMenuActionsReconstructAllChildNodes(UToolMenu* Menu, TWeakObjectPtr<UGraphNodeContextMenuContext> Context) const;
+	virtual void GetContextMenuActions(UToolMenu* Menu, UGraphNodeContextMenuContext* Context) const override;
+
+	virtual void BreakPinLinks(UEdGraphPin& TargetPin, bool bSendsNodeNotification) const override;
+	virtual void BreakSinglePinLink(UEdGraphPin* SourcePin, UEdGraphPin* TargetPin) const override;
+
+	//  Own interface
+
+	/** Utility to create a node creating action */
+	static TSharedPtr<FCustomizableObjectSchemaAction_NewNode> AddNewNodeAction(FGraphActionListBuilderBase& ContextMenuBuilder, const FString& Category, const FText& MenuDesc, const FText& Tooltip, const int32 Grouping = 0, const FString& Keywords = FString());
+
+	/** */
+	void GetBreakLinkToSubMenuActions( class FMenuBuilder& MenuBuilder, UEdGraphPin* InGraphPin ) const;
+	
+private: 
+/** Enum containing all the object types that we are able to convert onto a node when dragging
+ * and dropping and asset of that type onto the CO graph.
+ * Each value of this enumeration will, in practice, have a CO node to be represented by.
+ */
+// TODO: Replace the usage of this enum class with something similar to typeID (not casting)
+enum class ESpawnableObjectType : int32
+{
+	None = -1,		// Invalid value
+	
+	UTexture2D,
+	USkeletalMesh,
+	UStaticMesh,
+	UMaterialInterface,
+};
+
+public:
+
+	/**
+	 * Provided an asset object check if mutable does have a relative node to be able to host this asset's data.
+	 *  @param InAsset - The asset object we want to check if can be used to generate a new mutable node containing it.
+	 *  @param OutObjectType - The type of object of the asset.
+	 *  @return True if the asset can be used to spawn a mutable node, false if not.
+	 *  @note All possible object types that we can convert onto mutable nodes must be defined here.
+	 */
+	bool IsSpawnableAsset(const FAssetData& InAsset, ESpawnableObjectType& OutObjectType) const;
+
+	static TSharedPtr<class ICustomizableObjectEditor> GetCustomizableObjectEditor(const class UEdGraph* ParentGraph);
+	UEdGraphNode* AddComment(class UEdGraph* ParentGraph, UEdGraphPin* FromPin, const FVector2D Location, bool bSelectNewNode = true) const;
+
+	/**
+	 * Break non allowed connections once the connection has been created.
+	 * Allows the nodes to have greater connectivity granularity than simply forbidding to create a connection with UEdGraphSchema_CustomizableObject::CanCreateConnection.
+	 */ 
+	virtual bool TryCreateConnection(UEdGraphPin* PinA, UEdGraphPin* PinB) const override;
+
+	/** Given a pin category name, get its friendly name (user readable). */
+	static FText GetPinCategoryName(const FName& PinCategory);
+
+	/** Given a Material parameter type, returns the Mutable pin category. */
+	static const FName& GetPinCategory(EMaterialParameterType Type);
+};
+
+
+#if UE_ENABLE_INCLUDE_ORDER_DEPRECATED_IN_5_2
+#include "MaterialTypes.h"
+#endif
