[Mutable] Moved all files from Classes folder to Public folder.

diff --git a/Engine/Plugins/Experimental/Mutable/Source/CustomizableObjectEditor/Public/MuCOE/Nodes/CustomizableObjectNode.h b/Engine/Plugins/Experimental/Mutable/Source/CustomizableObjectEditor/Public/MuCOE/Nodes/CustomizableObjectNode.h
new file mode 100644
--- /dev/null
+++ b/Engine/Plugins/Experimental/Mutable/Source/CustomizableObjectEditor/Public/MuCOE/Nodes/CustomizableObjectNode.h
@@ -0,0 +1,348 @@
+// Copyright Epic Games, Inc. All Rights Reserved.
+
+#pragma once
+
+#include "EdGraph/EdGraph.h" // IWYU pragma: keep
+#include "EdGraph/EdGraphPin.h"
+#include "MuCO/CustomizableObject.h"
+#include "MuCO/CustomizableObjectPrivate.h"
+
+#include "CustomizableObjectNode.generated.h"
+
+class FArchive;
+class SWidget;
+class UCustomizableObjectNodeObject;
+class UCustomizableObjectNodeRemapPins;
+class UCustomizableObjectNodeRemapPinsByName;
+class UCustomizableObjectNodeRemapPinsByPosition;
+struct FPropertyChangedEvent;
+
+/** Abstract base class for all pin data. */
+UCLASS(Abstract)
+class UCustomizableObjectNodePinData : public UObject
+{
+	GENERATED_BODY()
+
+public:
+	UCustomizableObjectNodePinData();
+
+	/** False if inherited types are different. */
+	bool operator==(const UCustomizableObjectNodePinData& Other) const;
+	
+	virtual bool operator!=(const UCustomizableObjectNodePinData& Other) const;
+
+	/** Virtual function used to copy pin data when remapping pins. */
+	virtual void Copy(const UCustomizableObjectNodePinData& Other) {}
+
+protected:
+	/** Virtual function used to perform comparision between different UCustomizableObjectNodePinData inherited types. */
+	virtual bool Equals(const UCustomizableObjectNodePinData& Other) const;
+};
+
+/** Encapsulation of parameters for the FPostEditChangePropertyDelegate delegate function.
+ *
+ * The delegate requires the function being a UFUNCTION which requires all parameters to be UCLASS, USTRUC or UENUM.
+ */
+USTRUCT()
+struct FPostEditChangePropertyDelegateParameters
+{
+	GENERATED_BODY()
+
+	UObject* Node;
+	FPropertyChangedEvent* FPropertyChangedEvent;
+};
+
+DECLARE_DYNAMIC_MULTICAST_DELEGATE_OneParam(FPostEditChangePropertyDelegate, FPostEditChangePropertyDelegateParameters&, Parameters); // Deprecated, use non-dynamic version.
+DECLARE_MULTICAST_DELEGATE_TwoParams(FPostEditChangePropertyRegularDelegate, UObject*, FPropertyChangedEvent&);
+
+DECLARE_MULTICAST_DELEGATE(FPostReconstructNodeDelegate);
+
+DECLARE_MULTICAST_DELEGATE(FNodeConnectionListChangedDelegate);
+
+DECLARE_MULTICAST_DELEGATE_OneParam(FPinConnectionListChangedDelegate, UEdGraphPin*);
+
+DECLARE_MULTICAST_DELEGATE(FDestroyNodeDelegate);
+
+using FRemapPinsDelegateParameter = TMap<UEdGraphPin*, UEdGraphPin*>; // Required for the delegate macro.
+DECLARE_MULTICAST_DELEGATE_OneParam(FRemapPinsDelegate, const FRemapPinsDelegateParameter&);
+
+/** Base class of all Customizable Object nodes.
+ * 
+ * The Customizable Object node system is build on top of the following premises. To avoid breaking this system, these premise must be hold!
+ * 
+ * PREMISES:
+ * 1. A node can only be modified when a node is reconstructed (i.e., inside the ReconstructNode, AllocateDefaultPins functions). Exceptionally, a pin can also be created inside the BeginConstruct.
+ *    We consider that a node has been modified if one of the following things happen:
+ *        - Creation, destruction and modification of pins.
+ *        - Modification of any data that would modify any pin.
+ */
+UCLASS(abstract)
+class CUSTOMIZABLEOBJECTEDITOR_API UCustomizableObjectNode : public UEdGraphNode
+{
+public:
+	GENERATED_BODY()
+
+	// UObject interface
+	virtual void Serialize(FArchive& Ar) override;
+	virtual void PostLoad() override; // Do work at PostBackwardsCompatibleFixup.
+
+	// UEdGraphNode interface
+	virtual void AutowireNewNode(UEdGraphPin* FromPin) override;
+	virtual void NodeConnectionListChanged() override;
+	virtual void PinConnectionListChanged(UEdGraphPin* Pin) override;
+	virtual void DestroyNode() override;
+
+	/** Allocates the default pins using the empty remap pins action. Usually called from CreateNode. */
+	virtual void AllocateDefaultPins() override final; // Final. Override AllocateDefaultPins(UCustomizableObjectNodeRemapPins* RemapPins) instead
+	
+	/** Reconstructs the node using its default remap pins action. */
+	virtual void ReconstructNode() override final; // Final. Override ReconstructNode(UCustomizableObjectNodeRemapPins* RemapPins) instead
+	virtual void PostInitProperties() override;
+
+	// Own interface
+	/** Called at the beginning of the node lifecycle.*/
+	virtual void BeginConstruct() {};
+
+	/** Add backwards compatible code here.
+	 * When called, it is guaranteed that all nodes in this graph will have executed the PostLoad function. */
+	virtual void BackwardsCompatibleFixup();
+
+	/** Add post load work here.
+	 * When called, it is guaranteed that all nodes in this graph will have executed the BackwardsCompatibleFixup function. */
+	virtual void PostBackwardsCompatibleFixup() {};
+	
+	/** Virtual implementation of RemovePin. Allows to do work before removing a pin.
+	 * Use this function instead of RemovePin. RemovePin does not removes possible attached PinData. */
+	virtual bool CustomRemovePin(UEdGraphPin& Pin);
+
+	/**
+	 * Subclasses should override this to return true and set OutCategory if this node should be
+	 * auto-added to the right-click context menu in the graph editor.
+	 * 
+	 * Some nodes are added manually in the graph editor code and don't need to do this.
+	 */
+	virtual bool ShouldAddToContextMenu(FText& OutCategory) const;
+
+	void GetInputPins(TArray<class UEdGraphPin*>& OutInputPins) const;
+	void GetOutputPins(TArray<class UEdGraphPin*>& OutOutputPins) const;
+	UEdGraphPin* GetOutputPin(int32 OutputIndex) const;
+	
+	virtual bool IsNodeOutDatedAndNeedsRefresh() { return false; }
+	virtual FString GetRefreshMessage() const { return "Refresh Node."; }
+	void SetRefreshNodeWarning();
+	void RemoveWarnings();
+
+	// Used to replace Ids of referenced nodes by their new ids after duplicating the CustomizableObject.
+	virtual void UpdateReferencedNodeId(const FGuid& NewGuid) {};
+
+	TSharedPtr<class ICustomizableObjectEditor> GetGraphEditor() const;
+
+	virtual bool ProvidesCustomPinRelevancyTest() const { return false; } // Override and return true if IsPinRelevant is overridden.
+	virtual bool IsPinRelevant(const UEdGraphPin*) const { return true; } // Override if default pins allocated with AllocateDefaultPins do not provide enough information.
+
+	struct FAttachedErrorDataView
+	{
+		TArrayView<const float> UnassignedUVs;
+	};
+	
+	/** 
+	 * Check if two pins can be connected. 
+	 * 
+	 * @param	InOwnedInputPin	Input pin which belongs to this node.
+	 * @param	InOutputPin		Output pin which belongs to another node. If the node connects to iself, it could belong to this node.
+	 * @param	bOutIsOtherNodeBlocklisted		Is the other node of a type we are not allowed to connect?
+	 * @param	bOutArePinsCompatible		Does InOutputPin pin share the same type as the InOwnedInputPin?
+	 * @return	True if the pin types match and if the other node is not one of the Blocklisted types.
+	 */
+	virtual bool CanConnect( const UEdGraphPin* InOwnedInputPin, const UEdGraphPin* InOutputPin, bool& bOutIsOtherNodeBlocklisted, bool& bOutArePinsCompatible) const;
+
+	// Used during compilation process to cache the node for all LODs, or generate it specifically for each of them.
+	virtual bool IsAffectedByLOD() const { return true; }
+
+	// Get the CustomizableObject graph that owns this node
+	class UCustomizableObjectGraph* GetCustomizableObjectGraph() const;
+
+	/** If returns true, only a single link can go out of every output pin. False by default. */
+	virtual bool IsSingleOutputNode() const;
+
+	// Wrapper of the pin creation method, to cope with multiple Unreal Engine versions
+	UEdGraphPin* CustomCreatePin( EEdGraphPinDirection Direction, const FName& Type, const FName& Name, bool bIsArray=false );
+
+	/** Create pin with attached data.
+	 *
+     * @param PinData Pin data to be saved. */
+	UEdGraphPin* CustomCreatePin(EEdGraphPinDirection Direction, const FName& Type, const FName& Name, UCustomizableObjectNodePinData* PinData);
+
+	
+	virtual void AddAttachedErrorData(const FAttachedErrorDataView& ) {};
+	virtual void ResetAttachedErrorData() {};
+
+	/** 
+	 * When creating a new connection, break all previous connections. This method can be overridden.
+	 * 
+	 * @return false by default.
+	 */
+	virtual bool ShouldBreakExistingConnections(const UEdGraphPin* InputPin, const UEdGraphPin* OutputPin) const;
+
+	/** Called when a new connection is perform. Used to break an specific existing connectoin. Empty default implementation. This method can be overridden.
+	 *
+	 * @param	InputPin	New connection input pin.
+	 * @param	OutputPin	New connection output pin.
+	 */
+	virtual void BreakExistingConnectionsPostConnection(UEdGraphPin* InputPin, UEdGraphPin* OutputPin) {};
+
+	void PostEditChangeProperty(FPropertyChangedEvent& PropertyChangedEvent) override;
+
+	/** Custom post duplicate function. Called at the beginning of duplication, before the nodes have their Guid updated. */
+	virtual void BeginPostDuplicate(bool bDuplicateForPIE) {};
+	
+	/** Gets all non orphan pins this node owns. */
+	TArray<UEdGraphPin*> GetAllNonOrphanPins() const;
+
+	/** Gets all orphan pins this node owns. */
+	TArray<UEdGraphPin*> GetAllOrphanPins() const;
+
+	/**
+	 * Specialization of ReconstructNode UEdGraphNode function.
+	 *
+	 * @param RemapPins pointer to a remap pins action. Remap pins action which will be used once the node has been reconstructed. Can be null.
+	 */
+	virtual void ReconstructNode(UCustomizableObjectNodeRemapPins* RemapPinsMode);
+
+	/** Return the LOD which this node belongs to. Return -1 if it is not connected (directly or indirectly) to a LOD pin. */
+	int32 GetLOD() const;
+
+	/** Get all parent Object nodes of this node which are parents of this node.
+	 * 
+	 * If a Object node is connected through a LOD pin, it has to match the given LOD.
+	 */
+	TArray<UCustomizableObjectNodeObject*> GetParentObjectNodes(int LOD) const;
+	
+	/** Set a pin to be hidden or not. */
+	void SetPinHidden(UEdGraphPin& Pin, bool bHidden);
+
+	/** Set an array of pins to be hidden or not. 
+	 *
+	 * @param Pins Input pins. All pins must be non null.
+	 */
+	void SetPinHidden(const TArray<UEdGraphPin*>& Pins, bool bHidden);
+
+	/** Return true if a pin can be hidden. Override it to have a custom behaivour. */
+	virtual bool CanPinBeHidden(const UEdGraphPin& Pin) const;
+
+	/** Returns pin custom details. Override if required. */
+	virtual TSharedPtr<SWidget> CustomizePinDetails(UEdGraphPin& Pin);
+	
+	/** See GetPinData(const FEdGraphPinReference&). */
+	UCustomizableObjectNodePinData* GetPinData(const UEdGraphPin &Pin) const;
+	
+	/** Given a pin which has attached data, get the typed pin data.
+	 * @param Pin Pin which has attached data.
+	 * @return Pin data already typed. */
+	template<typename T>
+	T& GetPinData(const UEdGraphPin& Pin) const
+	{
+		return *CastChecked<T>(PinsDataId[Pin.PinId]);
+	}
+	
+	UPROPERTY()
+	FPostEditChangePropertyDelegate PostEditChangePropertyDelegate;
+
+	FPostEditChangePropertyRegularDelegate PostEditChangePropertyRegularDelegate;
+
+	FPostReconstructNodeDelegate PostReconstructNodeDelegate;
+
+	FNodeConnectionListChangedDelegate NodeConnectionListChangedDelegate;
+
+	FPinConnectionListChangedDelegate PinConnectionListChangedDelegate;
+
+	FDestroyNodeDelegate DestroyNodeDelegate;
+
+	FRemapPinsDelegate RemapPinsDelegate;
+
+protected:
+	/**
+	 * Specialization of AllocateDefaultPins UEdGraphNode function. Override.
+	 * 
+	 * Given the node context (this), create all the required pins.
+	 *
+	 * @param RemapPins remap pins action. Use this object to save any information required to later perfrom the remap pins action. Can be null.
+	 */
+	virtual void AllocateDefaultPins(UCustomizableObjectNodeRemapPins* RemapPins) {};
+
+	/**
+	 * Creates and returns the node default node remap pins action.
+	 * 
+	 * By default all nodes remap by name. Override if a node requires a diferent default remap pins action.
+	 */
+	virtual UCustomizableObjectNodeRemapPins* CreateRemapPinsDefault() const;
+
+	/**
+	 * Creates and returns a remap pins by name action. 
+	 * 
+	 * Override a if a node requires an specialized remap pins by name.
+	 */
+	virtual UCustomizableObjectNodeRemapPinsByName* CreateRemapPinsByName() const;
+
+	/**
+	 * Creates and returns a remap pins by position action.
+	 *
+	 * Override a if a node requires an specialized remap pins by position.
+	 */
+	virtual UCustomizableObjectNodeRemapPinsByPosition* CreateRemapPinsByPosition() const;
+	
+	/** Allows to perform work when remapping a pin.
+	  * Copies pin data from old pin to new pin. Keeps the id of the new pin. */
+	virtual void RemapPin(UEdGraphPin& NewPin, const UEdGraphPin& OldPin);
+
+	/** Allows to perform work when remapping a pin. */
+	virtual void RemapPins(const TMap<UEdGraphPin*, UEdGraphPin*>& PinsToRemap);
+
+	/*** Allows to perform work when remapping the pin data. */
+	virtual void RemapPinsData(const TMap<UEdGraphPin*, UEdGraphPin*>& PinsToRemap);
+	
+	/** Attach data to a given pin. */
+	void AddPinData(const UEdGraphPin& Pin, UCustomizableObjectNodePinData& PinData);
+
+private:
+	/** Data attached to a given pin. Not all pins contain data.
+	 * Data must be eliminated once the pins is removed. This is done automatically on CustomRemovePin.
+	 * Do not use RemovePin since it will not remove the data. */
+	UPROPERTY()
+	TMap<FGuid, TObjectPtr<UCustomizableObjectNodePinData>> PinsDataId;
+
+	// Deprectated properties
+	UPROPERTY()
+	TMap<FEdGraphPinReference, TObjectPtr<UCustomizableObjectNodePinData>> PinsData_DEPRECATED;
+};
+
+
+/** Return the specified node located in another Customizable Object graph.
+ *
+ * @param Customizable Object object to look for the node. If nulltpr, it returns nullptr.
+ * @param NodeGuid node guid. If invalid, it returns nullptr.
+ * 
+ * Return nullptr if not found.
+ */
+template<class FNodeType>
+FNodeType* GetCustomizableObjectExternalNode(UCustomizableObject* Object, const FGuid& NodeGuid)
+{
+	FNodeType* Result = nullptr;
+
+	if (Object && Object->GetPrivate()->GetSource() && NodeGuid.IsValid())
+	{
+		TArray<FNodeType*> GroupNodes;
+		Object->GetPrivate()->GetSource()->GetNodesOfClass<FNodeType>(GroupNodes);
+
+		for (FNodeType* GroupNode : GroupNodes)
+		{
+			if (NodeGuid == GroupNode->NodeGuid)
+			{
+				Result = GroupNode;
+				break;
+			}
+		}
+	}
+
+	return Result;
+}
diff --git a/Engine/Plugins/Experimental/Mutable/Source/CustomizableObjectEditor/Public/MuCOE/Nodes/CustomizableObjectNodeAnimationPose.h b/Engine/Plugins/Experimental/Mutable/Source/CustomizableObjectEditor/Public/MuCOE/Nodes/CustomizableObjectNodeAnimationPose.h
new file mode 100644
--- /dev/null
+++ b/Engine/Plugins/Experimental/Mutable/Source/CustomizableObjectEditor/Public/MuCOE/Nodes/CustomizableObjectNodeAnimationPose.h
@@ -0,0 +1,45 @@
+// Copyright Epic Games, Inc. All Rights Reserved.
+
+#pragma once
+
+#include "MuCOE/Nodes/CustomizableObjectNode.h"
+
+#include "CustomizableObjectNodeAnimationPose.generated.h"
+
+namespace ENodeTitleType { enum Type : int; }
+
+class FString;
+class UCustomizableObjectNodeRemapPins;
+class UEdGraphPin;
+class UObject;
+class UPoseAsset;
+struct FPropertyChangedEvent;
+
+
+UCLASS()
+class CUSTOMIZABLEOBJECTEDITOR_API UCustomizableObjectNodeAnimationPose : public UCustomizableObjectNode
+{
+public:
+	GENERATED_BODY()
+
+	UCustomizableObjectNodeAnimationPose();
+
+	UPROPERTY(EditAnywhere, Category=CustomizableObject)
+	TObjectPtr<UPoseAsset> PoseAsset;
+
+	// UObject interface.
+	virtual void PostEditChangeProperty(FPropertyChangedEvent& PropertyChangedEvent) override;
+
+	// Begin EdGraphNode interface
+	UEdGraphPin* GetInputMeshPin() const;
+	virtual FText GetNodeTitle(ENodeTitleType::Type TitleType) const override;
+	virtual FLinearColor GetNodeTitleColor() const override;
+	// End EdGraphNode interface
+
+	// UCustomizableObjectNode interface
+	void AllocateDefaultPins(UCustomizableObjectNodeRemapPins* RemapPins) override;
+
+	// Retrieve the pose information from the PoseAsset
+	static void StaticRetrievePoseInformation(UPoseAsset* PoseAsset, class USkeletalMesh* RefSkeletalMesh, TArray<FName>& OutArrayBoneName, TArray<FTransform>& OutArrayTransform);
+};
+
diff --git a/Engine/Plugins/Experimental/Mutable/Source/CustomizableObjectEditor/Public/MuCOE/Nodes/CustomizableObjectNodeColorArithmeticOp.h b/Engine/Plugins/Experimental/Mutable/Source/CustomizableObjectEditor/Public/MuCOE/Nodes/CustomizableObjectNodeColorArithmeticOp.h
new file mode 100644
--- /dev/null
+++ b/Engine/Plugins/Experimental/Mutable/Source/CustomizableObjectEditor/Public/MuCOE/Nodes/CustomizableObjectNodeColorArithmeticOp.h
@@ -0,0 +1,65 @@
+// Copyright Epic Games, Inc. All Rights Reserved.
+
+#pragma once
+
+#include "MuCOE/Nodes/CustomizableObjectNode.h"
+
+#include "CustomizableObjectNodeColorArithmeticOp.generated.h"
+
+namespace ENodeTitleType { enum Type : int; }
+
+class UCustomizableObjectNodeRemapPins;
+class UEdGraphPin;
+class UObject;
+struct FPropertyChangedEvent;
+
+
+UENUM(BlueprintType)
+enum class EColorArithmeticOperation : uint8
+{
+	E_Add 	UMETA(DisplayName = "+"),
+	E_Sub 	UMETA(DisplayName = "-"),
+	E_Mul	UMETA(DisplayName = "x"),
+	E_Div	UMETA(DisplayName = "/")
+};
+
+UCLASS()
+class CUSTOMIZABLEOBJECTEDITOR_API UCustomizableObjectNodeColorArithmeticOp : public UCustomizableObjectNode
+{
+public:
+	GENERATED_BODY()
+
+	UCustomizableObjectNodeColorArithmeticOp();
+
+	// UObject interface.
+	void PostEditChangeProperty(FPropertyChangedEvent& PropertyChangedEvent) override;
+
+	// Begin EdGraphNode interface
+	FText GetNodeTitle(ENodeTitleType::Type TitleType) const override;
+	FLinearColor GetNodeTitleColor() const override;
+	FText GetTooltipText() const override;
+	bool ShouldOverridePinNames() const override { return true; }
+
+	// UCustomizableObjectNode interface
+	void AllocateDefaultPins(UCustomizableObjectNodeRemapPins* RemapPins) override;
+
+	// First Operand
+	UEdGraphPin* XPin() const
+	{
+		return FindPin(TEXT("A"));
+	}
+
+	// Second Operand
+	UEdGraphPin* YPin() const
+	{
+		return FindPin(TEXT("B"));
+	}
+
+	UPROPERTY(EditAnywhere, Category = ColorArithmeticOperation)
+	EColorArithmeticOperation Operation;
+
+	// UCustomizableObjectNode interface
+	virtual bool IsAffectedByLOD() const { return false; }
+
+};
+
diff --git a/Engine/Plugins/Experimental/Mutable/Source/CustomizableObjectEditor/Public/MuCOE/Nodes/CustomizableObjectNodeColorFromFloats.h b/Engine/Plugins/Experimental/Mutable/Source/CustomizableObjectEditor/Public/MuCOE/Nodes/CustomizableObjectNodeColorFromFloats.h
new file mode 100644
--- /dev/null
+++ b/Engine/Plugins/Experimental/Mutable/Source/CustomizableObjectEditor/Public/MuCOE/Nodes/CustomizableObjectNodeColorFromFloats.h
@@ -0,0 +1,59 @@
+// Copyright Epic Games, Inc. All Rights Reserved.
+
+#pragma once
+
+#include "MuCOE/Nodes/CustomizableObjectNode.h"
+
+#include "CustomizableObjectNodeColorFromFloats.generated.h"
+
+namespace ENodeTitleType { enum Type : int; }
+
+class UCustomizableObjectNodeRemapPins;
+class UEdGraphPin;
+class UObject;
+struct FPropertyChangedEvent;
+
+
+UCLASS()
+class CUSTOMIZABLEOBJECTEDITOR_API UCustomizableObjectNodeColorFromFloats : public UCustomizableObjectNode
+{
+public:
+	GENERATED_BODY()
+
+	UCustomizableObjectNodeColorFromFloats();
+
+	// UObject interface.
+	void PostEditChangeProperty(FPropertyChangedEvent& PropertyChangedEvent) override;
+
+	// Begin EdGraphNode interface
+	FText GetNodeTitle(ENodeTitleType::Type TitleType) const override;
+	FLinearColor GetNodeTitleColor() const override;
+	FText GetTooltipText() const override;
+
+	// UCustomizableObjectNode interface
+	void AllocateDefaultPins(UCustomizableObjectNodeRemapPins* RemapPins) override;
+
+	UEdGraphPin* RPin() const
+	{
+		return FindPin(TEXT("R"));
+	}
+
+	UEdGraphPin* GPin() const
+	{
+		return FindPin(TEXT("G"));
+	}
+
+	UEdGraphPin* BPin() const
+	{
+		return FindPin(TEXT("B"));
+	}
+
+	UEdGraphPin* APin() const
+	{
+		return FindPin(TEXT("A"));
+	}
+
+	// UCustomizableObjectNode interface
+	virtual bool IsAffectedByLOD() const override { return false; }
+};
+
diff --git a/Engine/Plugins/Experimental/Mutable/Source/CustomizableObjectEditor/Public/MuCOE/Nodes/CustomizableObjectNodeColorParameter.h b/Engine/Plugins/Experimental/Mutable/Source/CustomizableObjectEditor/Public/MuCOE/Nodes/CustomizableObjectNodeColorParameter.h
new file mode 100644
--- /dev/null
+++ b/Engine/Plugins/Experimental/Mutable/Source/CustomizableObjectEditor/Public/MuCOE/Nodes/CustomizableObjectNodeColorParameter.h
@@ -0,0 +1,42 @@
+// Copyright Epic Games, Inc. All Rights Reserved.
+
+#pragma once
+
+#include "MuCOE/Nodes/CustomizableObjectNode.h"
+
+#include "CustomizableObjectNodeColorParameter.generated.h"
+
+namespace ENodeTitleType { enum Type : int; }
+
+class UCustomizableObjectNodeRemapPins;
+class UObject;
+
+
+UCLASS()
+class CUSTOMIZABLEOBJECTEDITOR_API UCustomizableObjectNodeColorParameter : public UCustomizableObjectNode
+{
+public:
+	GENERATED_BODY()
+
+	UPROPERTY(EditAnywhere, Category=CustomizableObject, meta = (DontUpdateWhileEditing))
+	FLinearColor DefaultValue = FLinearColor(1, 1, 1, 1);
+
+	UPROPERTY(EditAnywhere, Category=CustomizableObject)
+	FString ParameterName = "Default Name";
+
+	UPROPERTY(EditAnywhere, Category = UI, meta = (DisplayName = "Parameter UI Metadata"))
+	FMutableParamUIMetadata ParamUIMetadata;
+
+	// Begin EdGraphNode interface
+	FText GetNodeTitle(ENodeTitleType::Type TitleType) const override;
+	FLinearColor GetNodeTitleColor() const override;
+	FText GetTooltipText() const override;
+	virtual void OnRenameNode(const FString& NewName) override;
+	virtual bool GetCanRenameNode() const override { return true; }
+
+	// UCustomizableObjectNode interface
+	void AllocateDefaultPins(UCustomizableObjectNodeRemapPins* RemapPins) override;
+	bool IsAffectedByLOD() const override { return false; }
+	virtual void BackwardsCompatibleFixup() override;
+};
+
diff --git a/Engine/Plugins/Experimental/Mutable/Source/CustomizableObjectEditor/Public/MuCOE/Nodes/CustomizableObjectNodeColorSwitch.h b/Engine/Plugins/Experimental/Mutable/Source/CustomizableObjectEditor/Public/MuCOE/Nodes/CustomizableObjectNodeColorSwitch.h
new file mode 100644
--- /dev/null
+++ b/Engine/Plugins/Experimental/Mutable/Source/CustomizableObjectEditor/Public/MuCOE/Nodes/CustomizableObjectNodeColorSwitch.h
@@ -0,0 +1,22 @@
+// Copyright Epic Games, Inc. All Rights Reserved.
+
+#pragma once
+
+#include "MuCOE/Nodes/CustomizableObjectNodeSwitchBase.h"
+
+#include "CustomizableObjectNodeColorSwitch.generated.h"
+
+
+UCLASS()
+class CUSTOMIZABLEOBJECTEDITOR_API UCustomizableObjectNodeColorSwitch : public UCustomizableObjectNodeSwitchBase
+{
+public:
+	GENERATED_BODY()
+
+	// UCustomizableObjectNode interface
+	virtual void BackwardsCompatibleFixup() override;
+
+	// UCustomizableObjectNodeSwitchBase interface
+	virtual FName GetCategory() const override;
+};
+
diff --git a/Engine/Plugins/Experimental/Mutable/Source/CustomizableObjectEditor/Public/MuCOE/Nodes/CustomizableObjectNodeColorVariation.h b/Engine/Plugins/Experimental/Mutable/Source/CustomizableObjectEditor/Public/MuCOE/Nodes/CustomizableObjectNodeColorVariation.h
new file mode 100644
--- /dev/null
+++ b/Engine/Plugins/Experimental/Mutable/Source/CustomizableObjectEditor/Public/MuCOE/Nodes/CustomizableObjectNodeColorVariation.h
@@ -0,0 +1,28 @@
+// Copyright Epic Games, Inc. All Rights Reserved.
+
+#pragma once
+
+#include "MuCOE/Nodes/CustomizableObjectNodeVariation.h"
+
+#include "CustomizableObjectNodeColorVariation.generated.h"
+
+struct FCustomizableObjectColorVariation;
+
+
+UCLASS()
+class CUSTOMIZABLEOBJECTEDITOR_API UCustomizableObjectNodeColorVariation : public UCustomizableObjectNodeVariation
+{
+public:
+	GENERATED_BODY()
+
+	// Deprecated properties
+	UPROPERTY()
+	TArray<FCustomizableObjectColorVariation> Variations_DEPRECATED;
+
+	// UCustomizableObjectNode interface
+	virtual void BackwardsCompatibleFixup() override;
+	
+	// UCustomizableObjectNodeVariation interface
+	virtual FName GetCategory() const override;
+};
+
diff --git a/Engine/Plugins/Experimental/Mutable/Source/CustomizableObjectEditor/Public/MuCOE/Nodes/CustomizableObjectNodeCopyMaterial.h b/Engine/Plugins/Experimental/Mutable/Source/CustomizableObjectEditor/Public/MuCOE/Nodes/CustomizableObjectNodeCopyMaterial.h
new file mode 100644
--- /dev/null
+++ b/Engine/Plugins/Experimental/Mutable/Source/CustomizableObjectEditor/Public/MuCOE/Nodes/CustomizableObjectNodeCopyMaterial.h
@@ -0,0 +1,71 @@
+// Copyright Epic Games, Inc. All Rights Reserved.
+
+#pragma once
+
+#include "MuCOE/Nodes/CustomizableObjectNodeMaterial.h"
+
+#include "CustomizableObjectNodeCopyMaterial.generated.h"
+
+class UCustomizableObjectNodeRemapPins;
+class UCustomizableObjectNodeSkeletalMesh;
+class UEdGraphPin;
+class UObject;
+
+/**
+ * Copy Material node. Duplicates a Material Node. Duplicates all Material node input pins and properties except for the Mesh input pin.
+ * A new Mesh has to be defined through the new Mesh input pin.
+ *
+ * Input pins:
+ * - Mesh: New mesh.
+ * - Material: Material to duplicate.
+ * 
+ * Output pins:
+ * - Material: New material.
+ * 
+ * Properties:
+ *   [NONE]
+ */
+UCLASS()
+class CUSTOMIZABLEOBJECTEDITOR_API UCustomizableObjectNodeCopyMaterial : public UCustomizableObjectNodeMaterial
+{
+public:
+	GENERATED_BODY()
+
+	FText GetNodeTitle(ENodeTitleType::Type TitleType) const override;
+
+	// UCustomizableObjectNode interface
+	void AllocateDefaultPins(UCustomizableObjectNodeRemapPins* RemapPins) override;
+
+	/** Get Mesh input pin. */
+	UEdGraphPin* GetMeshPin() const override;
+
+	/** Get Material input pin. */
+	UEdGraphPin* GetMaterialPin() const;
+
+	/** Get Mesh node. */
+	UCustomizableObjectNodeSkeletalMesh* GetMeshNode() const;
+
+	/** Get the Material node. */
+	UCustomizableObjectNodeMaterial* GetMaterialNode() const;
+
+	/** Input material pin can only be connect to a NodeMaterial output pin or a NodeExternalPin. A connection to a NodeMaterialCopy output pin is not allowed. */
+	virtual bool CanConnect(const UEdGraphPin* InOwnedInputPin, const UEdGraphPin* InOutputPin, bool& bOutIsOtherNodeBlocklisted, bool& bOutArePinsCompatible) const override;
+	
+	bool ShouldBreakExistingConnections(const UEdGraphPin* InputPin, const UEdGraphPin* OutputPin) const override;
+
+	/** Override the NodeMaterial method. Does not break any connection. */
+	void BreakExistingConnectionsPostConnection(UEdGraphPin* InputPin, UEdGraphPin* OutputPin) override {};
+
+	virtual bool IsNodeOutDatedAndNeedsRefresh() override;
+
+	bool ProvidesCustomPinRelevancyTest() const override;
+
+	/** 
+	 * Only are relevant pins:
+	 * - From: NodeMaterial, excluding NodeCopyMaterial 
+	 * - To: NodeBaseObject
+	 */
+	bool IsPinRelevant(const UEdGraphPin* Pin) const override;
+
+	FText GetTooltipText() const override;
+};
diff --git a/Engine/Plugins/Experimental/Mutable/Source/CustomizableObjectEditor/Public/MuCOE/Nodes/CustomizableObjectNodeCurve.h b/Engine/Plugins/Experimental/Mutable/Source/CustomizableObjectEditor/Public/MuCOE/Nodes/CustomizableObjectNodeCurve.h
new file mode 100644
--- /dev/null
+++ b/Engine/Plugins/Experimental/Mutable/Source/CustomizableObjectEditor/Public/MuCOE/Nodes/CustomizableObjectNodeCurve.h
@@ -0,0 +1,55 @@
+// Copyright Epic Games, Inc. All Rights Reserved.
+
+#pragma once
+
+#include "MuCOE/Nodes/CustomizableObjectNode.h"
+
+#include "CustomizableObjectNodeCurve.generated.h"
+
+namespace ENodeTitleType { enum Type : int; }
+
+class FArchive;
+class UCurveBase;
+class UCustomizableObjectNodeRemapPins;
+class UEdGraphPin;
+class UObject;
+struct FPropertyChangedEvent;
+
+
+UCLASS()
+class CUSTOMIZABLEOBJECTEDITOR_API UCustomizableObjectNodeCurve : public UCustomizableObjectNode
+{
+public:
+	GENERATED_BODY()
+
+	UCustomizableObjectNodeCurve();
+
+	// UObject interface.
+	void PostEditChangeProperty(FPropertyChangedEvent& PropertyChangedEvent) override;
+	virtual void Serialize(FArchive& Ar) override;
+
+	// Begin EdGraphNode interface
+	FText GetNodeTitle(ENodeTitleType::Type TitleType) const override;
+	FLinearColor GetNodeTitleColor() const override;
+	FText GetTooltipText() const override;
+
+	// UCustomizableObjectNode interface
+	void AllocateDefaultPins(UCustomizableObjectNodeRemapPins* RemapPins) override;
+
+	UEdGraphPin* InputPin() const;
+	UEdGraphPin* CurvePins(int32 Index) const;
+	int32 GetNumCurvePins() const;
+
+	UPROPERTY(EditAnywhere, Category = Curve)
+	TObjectPtr<UCurveBase> CurveAsset;
+
+	// UCustomizableObjectNode interface
+	virtual bool IsAffectedByLOD() const override { return false; }
+
+	bool ProvidesCustomPinRelevancyTest() const override { return true; }
+	bool IsPinRelevant(const UEdGraphPin* Pin) const override;
+
+	/** Override the default behaivour of remap pins. Use remap pins by position by default. */
+	UCustomizableObjectNodeRemapPins* CreateRemapPinsDefault() const override;
+};
+
diff --git a/Engine/Plugins/Experimental/Mutable/Source/CustomizableObjectEditor/Public/MuCOE/Nodes/CustomizableObjectNodeEditLayoutBlocks.h b/Engine/Plugins/Experimental/Mutable/Source/CustomizableObjectEditor/Public/MuCOE/Nodes/CustomizableObjectNodeEditLayoutBlocks.h
new file mode 100644
--- /dev/null
+++ b/Engine/Plugins/Experimental/Mutable/Source/CustomizableObjectEditor/Public/MuCOE/Nodes/CustomizableObjectNodeEditLayoutBlocks.h
@@ -0,0 +1,21 @@
+// Copyright Epic Games, Inc. All Rights Reserved.
+
+#pragma once
+
+#include "MuCOE/Nodes/CustomizableObjectNodeEditMaterialBase.h"
+
+#include "CustomizableObjectNodeEditLayoutBlocks.generated.h"
+
+
+UCLASS(Abstract)
+class CUSTOMIZABLEOBJECTEDITOR_API UCustomizableObjectNodeEditLayoutBlocks : public UCustomizableObjectNodeEditMaterialBase
+{
+public:
+	GENERATED_BODY()
+
+	// Selected blocks
+	UPROPERTY()
+	TArray<FGuid> BlockIds;
+
+};
+
diff --git a/Engine/Plugins/Experimental/Mutable/Source/CustomizableObjectEditor/Public/MuCOE/Nodes/CustomizableObjectNodeEditMaterial.h b/Engine/Plugins/Experimental/Mutable/Source/CustomizableObjectEditor/Public/MuCOE/Nodes/CustomizableObjectNodeEditMaterial.h
new file mode 100644
--- /dev/null
+++ b/Engine/Plugins/Experimental/Mutable/Source/CustomizableObjectEditor/Public/MuCOE/Nodes/CustomizableObjectNodeEditMaterial.h
@@ -0,0 +1,99 @@
+// Copyright Epic Games, Inc. All Rights Reserved.
+
+#pragma once
+
+#include "MuCOE/Nodes/CustomizableObjectNodeEditLayoutBlocks.h"
+#include "MuCOE/Nodes/CustomizableObjectNodeUseMaterial.h"
+
+#include "CustomizableObjectNodeEditMaterial.generated.h"
+
+namespace ENodeTitleType { enum Type : int; }
+struct FCustomizableObjectNodeEditMaterialImage;
+
+class FArchive;
+class FCustomizableObjectNodeParentedMaterial;
+class UCustomizableObject;
+class UCustomizableObjectNodeRemapPins;
+class UObject;
+
+
+/** Additional data for the Parent Texture Parameter to edit pin. */
+UCLASS()
+class CUSTOMIZABLEOBJECTEDITOR_API UCustomizableObjectNodeEditMaterialPinEditImageData : public UCustomizableObjectNodePinData
+{
+	GENERATED_BODY()
+
+public:
+	UPROPERTY()
+	FGuid ImageId;
+
+	UPROPERTY();
+	FEdGraphPinReference PinMask;
+};
+
+
+UCLASS()
+class CUSTOMIZABLEOBJECTEDITOR_API UCustomizableObjectNodeEditMaterial :
+	public UCustomizableObjectNodeEditLayoutBlocks,
+	public FCustomizableObjectNodeUseMaterial
+{
+public:
+	GENERATED_BODY()
+
+	// UObject interface
+	virtual void Serialize(FArchive& Ar) override;
+	
+	// EdGraphNode interface
+	virtual FText GetNodeTitle(ENodeTitleType::Type TitleType) const override;
+	virtual FLinearColor GetNodeTitleColor() const override;
+	virtual FText GetTooltipText() const override;
+	virtual void PinConnectionListChanged(UEdGraphPin* Pin) override;
+
+	// UCustomizableObjectNode interface
+	virtual void BackwardsCompatibleFixup() override;
+	virtual void PostBackwardsCompatibleFixup() override;
+	virtual void AllocateDefaultPins(UCustomizableObjectNodeRemapPins* RemapPins) override;
+	virtual bool IsNodeOutDatedAndNeedsRefresh() override;
+	virtual FString GetRefreshMessage() const override;
+	virtual bool IsSingleOutputNode() const override;
+	virtual bool CustomRemovePin(UEdGraphPin& Pin) override;
+
+	// UCustomizableObjectNodeEditMaterialBase interface
+	virtual UEdGraphPin* OutputPin() const override;
+	virtual void SetLayoutIndex(const int32 LayoutIndex) override;
+	
+	// ICustomizableObjectNodeParentedMaterial interface
+	virtual void SetParentNode(UCustomizableObject* Object, FGuid NodeId) override;
+
+	// FCustomizableObjectNodeUseMaterial interface
+	virtual UCustomizableObjectNode& GetNode() override;
+	virtual FCustomizableObjectNodeParentedMaterial& GetNodeParentedMaterial() override;
+	virtual TMap<FGuid, FEdGraphPinReference>& GetPinsParameter() override;
+	
+	
+	/** Returns the Image mask pin of the given Image that will be edited.
+	 *
+	 * @returns Always returns a valid pin if EditsImage(const FGuid&) returns true. */
+	const UEdGraphPin* GetUsedImageMaskPin(const FGuid& ImageId) const;
+
+	// Function to select all the layout blocks. Called when a parameter is reset (ParentMaterial, ParentLayoutIndex)
+	void SelectAllLayoutBlocks();
+
+private:
+	/** Relates a Parameter id to a Pin. Only used to improve performance. */
+	UPROPERTY()
+	TMap<FGuid, FEdGraphPinReference> PinsParameter;
+	
+	// Deprecated properties
+	UPROPERTY()
+	TArray<FCustomizableObjectNodeEditMaterialImage> Images_DEPRECATED;
+	
+	// Old layout blocks to patch. Now in parent class.
+	UPROPERTY()
+	TArray<int32> Blocks_DEPRECATED;
+};
+
+
+#if UE_ENABLE_INCLUDE_ORDER_DEPRECATED_IN_5_2
+#include "MuCOE/CustomizableObjectEditor_Deprecated.h"
+#endif
diff --git a/Engine/Plugins/Experimental/Mutable/Source/CustomizableObjectEditor/Public/MuCOE/Nodes/CustomizableObjectNodeEditMaterialBase.h b/Engine/Plugins/Experimental/Mutable/Source/CustomizableObjectEditor/Public/MuCOE/Nodes/CustomizableObjectNodeEditMaterialBase.h
new file mode 100644
--- /dev/null
+++ b/Engine/Plugins/Experimental/Mutable/Source/CustomizableObjectEditor/Public/MuCOE/Nodes/CustomizableObjectNodeEditMaterialBase.h
@@ -0,0 +1,52 @@
+// Copyright Epic Games, Inc. All Rights Reserved.
+
+#pragma once
+
+#include "MuCOE/Nodes/CustomizableObjectNode.h"
+#include "MuCOE/Nodes/CustomizableObjectNodeParentedMaterial.h"
+
+#include "CustomizableObjectNodeEditMaterialBase.generated.h"
+
+class UCustomizableObject;
+class UEdGraphPin;
+class UObject;
+
+
+UCLASS(Abstract)
+class CUSTOMIZABLEOBJECTEDITOR_API UCustomizableObjectNodeEditMaterialBase : public UCustomizableObjectNode, public FCustomizableObjectNodeParentedMaterial
+{
+	GENERATED_BODY()
+	
+public:
+	// FCustomizableObjectNodeParentedMaterial interface
+	virtual void SetParentNode(UCustomizableObject* Object, FGuid NodeId) override;
+	virtual void SaveParentNode(UCustomizableObject* Object, FGuid NodeId) override;
+	virtual UCustomizableObjectNode& GetNode() override;
+	virtual FGuid GetParentNodeId() const override;
+	virtual UCustomizableObject* GetParentObject() const override;
+
+	// UCustomizableObjectNodeEditMaterialBase interface
+	virtual void SetLayoutIndex(const int32 LayoutIndex);
+
+	// Own interface
+	/** Index of the layout to use to patch blocks. */
+	UPROPERTY(EditAnywhere, Category = CustomizableObject)
+	int ParentLayoutIndex = 0;
+
+	void BeginPostDuplicate(bool bDuplicateForPIE) override;
+
+	void UpdateReferencedNodeId(const FGuid& NewGuid) override;
+
+	virtual UEdGraphPin* OutputPin() const
+	{
+		return FindPin(TEXT("Material"));
+	}
+
+private:
+	UPROPERTY(EditAnywhere, Category = CustomizableObject)
+	TObjectPtr<UCustomizableObject> ParentMaterialObject = nullptr;
+
+	UPROPERTY()
+	FGuid ParentMaterialNodeId;
+};
+
diff --git a/Engine/Plugins/Experimental/Mutable/Source/CustomizableObjectEditor/Public/MuCOE/Nodes/CustomizableObjectNodeEnumParameter.h b/Engine/Plugins/Experimental/Mutable/Source/CustomizableObjectEditor/Public/MuCOE/Nodes/CustomizableObjectNodeEnumParameter.h
new file mode 100644
--- /dev/null
+++ b/Engine/Plugins/Experimental/Mutable/Source/CustomizableObjectEditor/Public/MuCOE/Nodes/CustomizableObjectNodeEnumParameter.h
@@ -0,0 +1,63 @@
+// Copyright Epic Games, Inc. All Rights Reserved.
+
+#pragma once
+
+#include "MuCOE/Nodes/CustomizableObjectNode.h"
+
+#include "CustomizableObjectNodeEnumParameter.generated.h"
+
+namespace ENodeTitleType { enum Type : int; }
+
+class UCustomizableObjectNodeRemapPins;
+class UObject;
+struct FPropertyChangedEvent;
+
+
+USTRUCT()
+struct CUSTOMIZABLEOBJECTEDITOR_API FCustomizableObjectNodeEnumValue
+{
+	GENERATED_USTRUCT_BODY()
+
+	UPROPERTY(EditAnywhere, Category=CustomizableObject)
+	FString Name;
+
+	UPROPERTY(EditAnywhere, Category = UI, meta = (DisplayName = "Parameter UI Metadata"))
+	FMutableParamUIMetadata ParamUIMetadata;
+};
+
+
+UCLASS()
+class CUSTOMIZABLEOBJECTEDITOR_API UCustomizableObjectNodeEnumParameter : public UCustomizableObjectNode
+{
+public:
+	GENERATED_BODY()
+
+	UPROPERTY(EditAnywhere, Category=CustomizableObject)
+	int32 DefaultIndex = 0;
+
+	UPROPERTY(EditAnywhere, Category=CustomizableObject)
+	FString ParameterName = "Default Name";
+
+	UPROPERTY(EditAnywhere, Category = UI, meta = (DisplayName = "Parameter UI Metadata"))
+	FMutableParamUIMetadata ParamUIMetadata;
+
+	UPROPERTY(EditAnywhere, Category=CustomizableObject)
+	TArray<FCustomizableObjectNodeEnumValue> Values;
+
+	// UObject interface.
+	void PostEditChangeProperty(FPropertyChangedEvent& PropertyChangedEvent) override;
+
+	// Begin EdGraphNode interface
+	FText GetNodeTitle(ENodeTitleType::Type TitleType) const override;
+	FLinearColor GetNodeTitleColor() const override;
+	FText GetTooltipText() const override;
+	virtual void OnRenameNode(const FString& NewName) override;
+	virtual bool GetCanRenameNode() const override { return true; }
+
+	// UCustomizableObjectNode interface
+	void AllocateDefaultPins(UCustomizableObjectNodeRemapPins* RemapPins) override;
+	virtual bool IsAffectedByLOD() const override { return false; }
+	virtual void BackwardsCompatibleFixup() override;
+
+};
+
diff --git a/Engine/Plugins/Experimental/Mutable/Source/CustomizableObjectEditor/Public/MuCOE/Nodes/CustomizableObjectNodeExposePin.h b/Engine/Plugins/Experimental/Mutable/Source/CustomizableObjectEditor/Public/MuCOE/Nodes/CustomizableObjectNodeExposePin.h
new file mode 100644
--- /dev/null
+++ b/Engine/Plugins/Experimental/Mutable/Source/CustomizableObjectEditor/Public/MuCOE/Nodes/CustomizableObjectNodeExposePin.h
@@ -0,0 +1,57 @@
+// Copyright Epic Games, Inc. All Rights Reserved.
+
+#pragma once
+
+#include "MuCOE/Nodes/CustomizableObjectNode.h"
+
+#include "CustomizableObjectNodeExposePin.generated.h"
+
+namespace ENodeTitleType { enum Type : int; }
+
+class UCustomizableObjectNodeRemapPins;
+class UEdGraphPin;
+class UObject;
+struct FPropertyChangedEvent;
+
+
+DECLARE_MULTICAST_DELEGATE(FOnNameChangedDelegate);
+
+UCLASS()
+class CUSTOMIZABLEOBJECTEDITOR_API UCustomizableObjectNodeExposePin : public UCustomizableObjectNode
+{
+public:
+	GENERATED_BODY()
+
+	// UObject interface
+	void PostEditChangeProperty(FPropertyChangedEvent& PropertyChangedEvent) override;
+
+	// EdGraphNode interface
+	FText GetNodeTitle(ENodeTitleType::Type TitleType) const override;
+	FLinearColor GetNodeTitleColor() const override;
+	FText GetTooltipText() const override;
+
+	virtual bool CanConnect(const UEdGraphPin* InOwnedInputPin, const UEdGraphPin* InOutputPin, bool& bOutIsOtherNodeBlocklisted, bool& bOutArePinsCompatible) const override;
+
+	// UCustomizableObjectNode interface
+	void AllocateDefaultPins(UCustomizableObjectNodeRemapPins* RemapPins) override;
+
+	/** Return the Expose Pin Node expose pin name. */
+	FString GetNodeName() const;
+	
+	// This is actually PinCategory
+	UPROPERTY()
+	FName PinType;
+	
+	UEdGraphPin* InputPin() const
+	{
+		return FindPin(TEXT("Object"));
+	}
+
+	/** Boradcasted when the UPROPERTY Name changes. */
+	FOnNameChangedDelegate OnNameChangedDelegate;
+	
+private:
+	UPROPERTY(Category = CustomizableObject, EditAnywhere)
+	FString Name = "Default Name";
+};
+
diff --git a/Engine/Plugins/Experimental/Mutable/Source/CustomizableObjectEditor/Public/MuCOE/Nodes/CustomizableObjectNodeExtendMaterial.h b/Engine/Plugins/Experimental/Mutable/Source/CustomizableObjectEditor/Public/MuCOE/Nodes/CustomizableObjectNodeExtendMaterial.h
new file mode 100644
--- /dev/null
+++ b/Engine/Plugins/Experimental/Mutable/Source/CustomizableObjectEditor/Public/MuCOE/Nodes/CustomizableObjectNodeExtendMaterial.h
@@ -0,0 +1,90 @@
+// Copyright Epic Games, Inc. All Rights Reserved.
+
+#pragma once
+
+#include "MuCOE/Nodes/CustomizableObjectNodeMaterialBase.h"
+#include "MuCOE/Nodes/CustomizableObjectNodeParentedMaterial.h"
+#include "MuCOE/Nodes/CustomizableObjectNodeUseMaterial.h"
+
+#include "CustomizableObjectNodeExtendMaterial.generated.h"
+
+namespace ENodeTitleType { enum Type : int; }
+
+class FArchive;
+class UCustomizableObject;
+class UCustomizableObjectNode;
+class UCustomizableObjectNodeRemapPins;
+class UEdGraphPin;
+class UObject;
+struct FCustomizableObjectNodeExtendMaterialImage;
+struct FEdGraphPinReference;
+
+
+UCLASS()
+class CUSTOMIZABLEOBJECTEDITOR_API UCustomizableObjectNodeExtendMaterial :
+	public UCustomizableObjectNodeMaterialBase,
+	public FCustomizableObjectNodeParentedMaterial,
+	public FCustomizableObjectNodeUseMaterial
+{
+	GENERATED_BODY()
+
+public:
+	UPROPERTY(EditAnywhere, BlueprintReadWrite, Category = CustomizableObject)
+	TArray<FString> Tags;
+
+	// UObject interface.
+	virtual void Serialize(FArchive& Ar) override;
+	virtual void BeginPostDuplicate(bool bDuplicateForPIE) override;
+	
+	// UCustomizableObjectNode
+	virtual void PostBackwardsCompatibleFixup() override;
+	virtual void BackwardsCompatibleFixup() override;
+	virtual void UpdateReferencedNodeId(const FGuid& NewGuid) override;
+
+	// Begin EdGraphNode interface
+	virtual FText GetNodeTitle(ENodeTitleType::Type TitleType) const override;
+	virtual FLinearColor GetNodeTitleColor() const override;
+	virtual FText GetTooltipText() const override;
+
+	// UCustomizableObjectNode interface
+	virtual void AllocateDefaultPins(UCustomizableObjectNodeRemapPins* RemapPins) override;
+	virtual void PinConnectionListChanged(UEdGraphPin * Pin) override;
+	virtual bool IsNodeOutDatedAndNeedsRefresh() override;
+	virtual FString GetRefreshMessage() const override;
+	virtual bool IsSingleOutputNode() const override;
+	virtual bool CustomRemovePin(UEdGraphPin& Pin) override;
+
+	// UCustomizableObjectNodeMaterialBase interface
+	virtual UEdGraphPin* OutputPin() const override;
+	virtual void SetParentNode(UCustomizableObject* Object, FGuid NodeId) override;
+	TArray<UCustomizableObjectLayout*> GetLayouts() override;
+	
+	// FCustomizableObjectNodeParentMaterial interface
+	virtual void SaveParentNode(UCustomizableObject* Object, FGuid NodeId) override;
+	virtual UCustomizableObjectNode& GetNode() override;
+	virtual FGuid GetParentNodeId() const override;
+	virtual UCustomizableObject* GetParentObject() const override;
+	
+	// FCustomizableObjectNodeUseMaterial interface
+	virtual FCustomizableObjectNodeParentedMaterial& GetNodeParentedMaterial() override;
+	virtual TMap<FGuid, FEdGraphPinReference>& GetPinsParameter() override;
+	
+	// Own interface
+	UEdGraphPin* AddMeshPin() const;
+	
+private:
+	UPROPERTY(EditAnywhere, Category = CustomizableObject)
+	TObjectPtr<UCustomizableObject> ParentMaterialObject = nullptr;
+
+	UPROPERTY()
+	FGuid ParentMaterialNodeId;
+	
+	/** Relates a Parameter id to a Pin. Only used to improve performance. */
+   	UPROPERTY()
+   	TMap<FGuid, FEdGraphPinReference> PinsParameter;
+
+	// Deprecated properties
+	UPROPERTY()
+	TArray<FCustomizableObjectNodeExtendMaterialImage> Images_DEPRECATED;
+};
+
diff --git a/Engine/Plugins/Experimental/Mutable/Source/CustomizableObjectEditor/Public/MuCOE/Nodes/CustomizableObjectNodeExternalPin.h b/Engine/Plugins/Experimental/Mutable/Source/CustomizableObjectEditor/Public/MuCOE/Nodes/CustomizableObjectNodeExternalPin.h
new file mode 100644
--- /dev/null
+++ b/Engine/Plugins/Experimental/Mutable/Source/CustomizableObjectEditor/Public/MuCOE/Nodes/CustomizableObjectNodeExternalPin.h
@@ -0,0 +1,69 @@
+// Copyright Epic Games, Inc. All Rights Reserved.
+
+#pragma once
+
+#include "MuCOE/Nodes/CustomizableObjectNode.h"
+#include "MuCOE/RemapPins/CustomizableObjectNodeRemapPinsByName.h"
+
+#include "CustomizableObjectNodeExternalPin.generated.h"
+
+namespace ENodeTitleType { enum Type : int; }
+
+class FArchive;
+class UCustomizableObject;
+class UCustomizableObjectNodeExposePin;
+class UCustomizableObjectNodeRemapPins;
+class UEdGraphPin;
+class UObject;
+
+
+UCLASS()
+class CUSTOMIZABLEOBJECTEDITOR_API UCustomizableObjectNodeExternalPin : public UCustomizableObjectNode
+{
+public:
+	GENERATED_BODY()
+
+	// UObject interface.
+	virtual void Serialize(FArchive& Ar) override;
+	
+	// EdGraphNode interface 
+	FText GetNodeTitle(ENodeTitleType::Type TitleType) const override;
+	FLinearColor GetNodeTitleColor() const override;
+	FText GetTooltipText() const override;
+	
+	// UCustomizableObjectNode interface
+	virtual void BackwardsCompatibleFixup() override;
+	virtual void PostBackwardsCompatibleFixup() override;
+	void AllocateDefaultPins(UCustomizableObjectNodeRemapPins* RemapPins) override;
+	UCustomizableObjectNodeRemapPins* CreateRemapPinsDefault() const override;
+	void BeginPostDuplicate(bool bDuplicateForPIE) override;
+	virtual bool CanConnect(const UEdGraphPin* InOwnedInputPin, const UEdGraphPin* InOutputPin, bool& bOutIsOtherNodeBlocklisted, bool& bOutArePinsCompatible) const override;
+
+	void UpdateReferencedNodeId(const FGuid& NewGuid) override;
+	
+	/** Set the linked Node Expose Pin node guid. */
+	void SetExternalObjectNodeId(FGuid Guid);
+
+	/** Return the external pin. Can return nullptr. */
+	UEdGraphPin* GetExternalPin() const;
+
+	/** Return the linked Expose Pin node. Return nullptr if not set. */
+	UCustomizableObjectNodeExposePin* GetNodeExposePin() const;
+
+	// This is actually PinCategory
+	UPROPERTY()
+	FName PinType;
+
+	/** External Customizable Object which the linked Node Expose Pin belong to. */
+	UPROPERTY()
+	TObjectPtr<UCustomizableObject> ExternalObject;
+
+private:
+	/** Linked Node Expose Pin node guid. */
+	UPROPERTY()
+	FGuid ExternalObjectNodeId;
+	
+	FDelegateHandle OnNameChangedDelegateHandle;
+	FDelegateHandle DestroyNodeDelegateHandle;
+};
+
diff --git a/Engine/Plugins/Experimental/Mutable/Source/CustomizableObjectEditor/Public/MuCOE/Nodes/CustomizableObjectNodeFloatParameter.h b/Engine/Plugins/Experimental/Mutable/Source/CustomizableObjectEditor/Public/MuCOE/Nodes/CustomizableObjectNodeFloatParameter.h
new file mode 100644
--- /dev/null
+++ b/Engine/Plugins/Experimental/Mutable/Source/CustomizableObjectEditor/Public/MuCOE/Nodes/CustomizableObjectNodeFloatParameter.h
@@ -0,0 +1,60 @@
+// Copyright Epic Games, Inc. All Rights Reserved.
+
+#pragma once
+
+#include "MuCOE/Nodes/CustomizableObjectNode.h"
+
+#include "CustomizableObjectNodeFloatParameter.generated.h"
+
+namespace ENodeTitleType { enum Type : int; }
+
+class FArchive;
+class UCustomizableObjectNodeRemapPins;
+class UEdGraphPin;
+class UObject;
+struct FEdGraphPinReference;
+struct FPropertyChangedEvent;
+
+
+USTRUCT()
+struct CUSTOMIZABLEOBJECTEDITOR_API FCustomizableObjectNodeFloatDescription
+{
+	GENERATED_USTRUCT_BODY()
+
+	UPROPERTY(EditAnywhere, Category=CustomizableObject)
+	FString Name;
+};
+
+
+UCLASS()
+class CUSTOMIZABLEOBJECTEDITOR_API UCustomizableObjectNodeFloatParameter : public UCustomizableObjectNode
+{
+public:
+	GENERATED_BODY()
+
+	UPROPERTY(EditAnywhere, Category=CustomizableObject)
+	float DefaultValue = 1.0f;
+
+	UPROPERTY(EditAnywhere, Category=CustomizableObject)
+	FString ParameterName = "Default Name";
+
+	UPROPERTY(EditAnywhere, Category = UI, meta = (DisplayName = "Parameter UI Metadata"))
+	FMutableParamUIMetadata ParamUIMetadata;
+	
+	// UObject interface.
+	void PostEditChangeProperty(FPropertyChangedEvent& PropertyChangedEvent) override;
+	
+	// Begin EdGraphNode interface
+	FText GetNodeTitle(ENodeTitleType::Type TitleType) const override;
+	FLinearColor GetNodeTitleColor() const override;
+	FText GetTooltipText() const override;
+	virtual bool GetCanRenameNode() const override { return true; }
+
+	// UCustomizableObjectNode interface
+	void AllocateDefaultPins(UCustomizableObjectNodeRemapPins* RemapPins) override;
+	bool IsAffectedByLOD() const override;
+	void BackwardsCompatibleFixup() override;
+	virtual void OnRenameNode(const FString& NewName) override;
+
+};
+
diff --git a/Engine/Plugins/Experimental/Mutable/Source/CustomizableObjectEditor/Public/MuCOE/Nodes/CustomizableObjectNodeFloatSwitch.h b/Engine/Plugins/Experimental/Mutable/Source/CustomizableObjectEditor/Public/MuCOE/Nodes/CustomizableObjectNodeFloatSwitch.h
new file mode 100644
--- /dev/null
+++ b/Engine/Plugins/Experimental/Mutable/Source/CustomizableObjectEditor/Public/MuCOE/Nodes/CustomizableObjectNodeFloatSwitch.h
@@ -0,0 +1,21 @@
+// Copyright Epic Games, Inc. All Rights Reserved.
+
+#pragma once
+
+#include "MuCOE/Nodes/CustomizableObjectNodeSwitchBase.h"
+
+#include "CustomizableObjectNodeFloatSwitch.generated.h"
+
+
+UCLASS()
+class CUSTOMIZABLEOBJECTEDITOR_API UCustomizableObjectNodeFloatSwitch : public UCustomizableObjectNodeSwitchBase
+{
+public:
+	GENERATED_BODY()
+
+	// UCustomizableObjectNode interface
+	virtual void BackwardsCompatibleFixup() override;
+
+	// UCustomizableObjectNodeSwitchBase interface
+	virtual FName GetCategory() const override;
+};
diff --git a/Engine/Plugins/Experimental/Mutable/Source/CustomizableObjectEditor/Public/MuCOE/Nodes/CustomizableObjectNodeFloatVariation.h b/Engine/Plugins/Experimental/Mutable/Source/CustomizableObjectEditor/Public/MuCOE/Nodes/CustomizableObjectNodeFloatVariation.h
new file mode 100644
--- /dev/null
+++ b/Engine/Plugins/Experimental/Mutable/Source/CustomizableObjectEditor/Public/MuCOE/Nodes/CustomizableObjectNodeFloatVariation.h
@@ -0,0 +1,28 @@
+// Copyright Epic Games, Inc. All Rights Reserved.
+
+#pragma once
+
+#include "MuCOE/Nodes/CustomizableObjectNodeVariation.h"
+
+#include "CustomizableObjectNodeFloatVariation.generated.h"
+
+struct FCustomizableObjectFloatVariation;
+
+
+UCLASS()
+class CUSTOMIZABLEOBJECTEDITOR_API UCustomizableObjectNodeFloatVariation : public UCustomizableObjectNodeVariation
+{
+public:
+	GENERATED_BODY()
+	
+	// Deprecated properties
+	UPROPERTY()
+	TArray<FCustomizableObjectFloatVariation> Variations_DEPRECATED;
+
+	// UCustomizableObjectNode interface
+	virtual void BackwardsCompatibleFixup() override;
+	
+	// UCustomizableObjectNodeVariation interface
+	virtual FName GetCategory() const override;
+};
+
diff --git a/Engine/Plugins/Experimental/Mutable/Source/CustomizableObjectEditor/Public/MuCOE/Nodes/CustomizableObjectNodeGroupProjectorParameter.h b/Engine/Plugins/Experimental/Mutable/Source/CustomizableObjectEditor/Public/MuCOE/Nodes/CustomizableObjectNodeGroupProjectorParameter.h
new file mode 100644
--- /dev/null
+++ b/Engine/Plugins/Experimental/Mutable/Source/CustomizableObjectEditor/Public/MuCOE/Nodes/CustomizableObjectNodeGroupProjectorParameter.h
@@ -0,0 +1,111 @@
+// Copyright Epic Games, Inc. All Rights Reserved.
+
+#pragma once
+
+#include "MuCOE/Nodes/CustomizableObjectNodeProjectorParameter.h"
+
+#include "CustomizableObjectNodeGroupProjectorParameter.generated.h"
+
+class FArchive;
+class UCustomizableObjectNodeRemapPins;
+class UDataTable;
+class UObject;
+class UPoseAsset;
+class UTexture2D;
+
+USTRUCT()
+struct FGroupProjectorParameterImage
+{
+	GENERATED_BODY()
+
+	UPROPERTY(EditAnywhere, Category = ProjectorGroup)
+	FString OptionName;
+
+	UPROPERTY(EditAnywhere, Category = ProjectorGroup)
+	TObjectPtr<UTexture2D> OptionImage = nullptr;
+};
+
+
+USTRUCT()
+struct FGroupProjectorParameterPose
+{
+	GENERATED_BODY()
+
+	UPROPERTY(EditAnywhere, Category = ProjectorGroup)
+	FString PoseName;
+
+	UPROPERTY(EditAnywhere, Category = ProjectorGroup)
+	TObjectPtr<UPoseAsset> OptionPose = nullptr;
+};
+
+
+UCLASS()
+class CUSTOMIZABLEOBJECTEDITOR_API UCustomizableObjectNodeGroupProjectorParameter : public UCustomizableObjectNodeProjectorParameter
+{
+public:
+	GENERATED_BODY()
+
+	/** Return array with the sticker name and UTexture2D for projection */
+	TArray<FGroupProjectorParameterImage> GetOptionImagesFromTable() const;
+
+	/** Returns the final option images without repeated elements in the option names,
+	* the data table has preference over elements in the OptionImages array */
+	TArray<FGroupProjectorParameterImage> GetFinalOptionImagesNoRepeat() const;
+
+	UPROPERTY(EditAnywhere, Category = ProjectorGroup, Meta = (ToolTip = "Only used when connected to a Group node. Specifies which material channel in the Group node's child material nodes will be connected to the projection."))
+	FString MaterialChannelNameToConnect;
+
+	UPROPERTY(EditAnywhere, Category = ProjectorGroup, Meta = (ToolTip = "Only used when connected to a Group node. Specifies which material channel will be used to mask out the projection."))
+	FString MaskedOutAreaMaterialChannelName;
+
+	/** Reference Texture used to decide the texture properties of the mutable-generated textures
+	* connected to this material. If null, it will try to be guessed at compile time from
+	* the graph. */
+	UPROPERTY(EditAnywhere, Category = ProjectorGroup)
+	TObjectPtr<UTexture2D> ReferenceTexture = nullptr;
+
+	UPROPERTY(EditAnywhere, Category = ProjectorGroup)
+	int32 ProjectionTextureSize = 512;
+
+	UPROPERTY(EditAnywhere, Category = ProjectorGroup, Meta = (ToolTip = "Specifies at which LOD level the projection texture will not be used and possibly save memory. A negative value means they will never be dropped."))
+	int32 DropProjectionTextureAtLOD = -1;
+
+	UPROPERTY(EditAnywhere, Category = ProjectorGroup, Meta = (ToolTip = "If true, projection textures will be shared between LODs of the same object, and will save memory. Only use if all the LODs share the same UV layout."))
+	bool bShareProjectionTexturesBetweenLODs = false;
+
+	UPROPERTY(EditAnywhere, Category = ProjectorGroup)
+	TArray<FGroupProjectorParameterImage> OptionImages;
+
+	/** Poses will only affect component 0 of the CO */ // TODO UE-206803
+	UPROPERTY(EditAnywhere, Category = ProjectorGroup)
+	TArray<FGroupProjectorParameterPose> OptionPoses;
+
+	/** Name of the column in the Option Images Data Table with the additional option images, (UTexture2D assets). */
+	UPROPERTY(EditAnywhere, Category = ProjectorGroup)
+	FName DataTableTextureColumnName;
+
+	/** Table where additional option images besides Option Images are read. The elements in this table have priority
+	* over elements from Option Images in case of duplicity. Use the "Data Table Texture Column Name" property to specify
+	* the name of the column where textures are read in the table. */
+	UPROPERTY(EditAnywhere, Category = ProjectorGroup)
+	TObjectPtr<UDataTable> OptionImagesDataTable = nullptr;
+
+	UPROPERTY(EditAnywhere, Category = ProjectorGroup)
+	FString AlternateProjectionResolutionStateName;
+
+	UPROPERTY(EditAnywhere, Category = ProjectorGroup)
+	float AlternateProjectionResolutionFactor = 0.0f;
+
+	UPROPERTY(EditAnywhere, Category = ProjectorGroup, Meta = (ClampMin = "0"))
+	int32 UVLayout = 0;
+
+	// EdGraphNode interface.
+	virtual FText GetNodeTitle(ENodeTitleType::Type TitleType) const override;
+	virtual FLinearColor GetNodeTitleColor() const override;
+	virtual FText GetTooltipText() const override;
+
+	// UCustomizableObjectNode interface.
+	virtual void BackwardsCompatibleFixup() override;
+	virtual void AllocateDefaultPins(UCustomizableObjectNodeRemapPins* RemapPins);
+};
+
diff --git a/Engine/Plugins/Experimental/Mutable/Source/CustomizableObjectEditor/Public/MuCOE/Nodes/CustomizableObjectNodeLayoutBlocks.h b/Engine/Plugins/Experimental/Mutable/Source/CustomizableObjectEditor/Public/MuCOE/Nodes/CustomizableObjectNodeLayoutBlocks.h
new file mode 100644
--- /dev/null
+++ b/Engine/Plugins/Experimental/Mutable/Source/CustomizableObjectEditor/Public/MuCOE/Nodes/CustomizableObjectNodeLayoutBlocks.h
@@ -0,0 +1,84 @@
+// Copyright Epic Games, Inc. All Rights Reserved.
+
+#pragma once
+
+#include "MuCOE/Nodes/CustomizableObjectNode.h"
+
+#include "CustomizableObjectNodeLayoutBlocks.generated.h"
+
+class UCustomizableObjectLayout;
+enum class ECustomizableObjectTextureLayoutPackingStrategy : uint8;
+namespace ENodeTitleType { enum Type : int; }
+struct FCustomizableObjectLayoutBlock;
+
+class FArchive;
+class UCustomizableObjectNodeRemapPins;
+class UEdGraphPin;
+class UObject;
+struct FPropertyChangedEvent;
+
+
+UCLASS()
+class CUSTOMIZABLEOBJECTEDITOR_API UCustomizableObjectNodeLayoutBlocks : public UCustomizableObjectNode
+{
+public:
+	GENERATED_BODY()
+
+	UCustomizableObjectNodeLayoutBlocks();
+
+	UPROPERTY()
+	FIntPoint GridSize_DEPRECATED;
+
+	/** Used with the fixed layout strategy. */
+	UPROPERTY()
+	FIntPoint MaxGridSize_DEPRECATED;
+
+	UPROPERTY()
+	TArray<FCustomizableObjectLayoutBlock> Blocks_DEPRECATED;
+
+	UPROPERTY()
+	ECustomizableObjectTextureLayoutPackingStrategy PackingStrategy_DEPRECATED;
+
+	UPROPERTY()
+	TObjectPtr<UCustomizableObjectLayout> Layout = nullptr;
+
+	// EdGraphNode interface
+	FText GetNodeTitle(ENodeTitleType::Type TitleType) const override;
+	FLinearColor GetNodeTitleColor() const override;
+	FText GetTooltipText() const override;
+	void Serialize(FArchive& Ar) override;
+	void PinConnectionListChanged(UEdGraphPin* Pin) override;
+	void PostPasteNode() override;
+
+
+	// UCustomizableObjectNode interface
+	virtual void PostBackwardsCompatibleFixup() override;
+	void AllocateDefaultPins(UCustomizableObjectNodeRemapPins* RemapPins) override;
+
+	UEdGraphPin* OutputPin() const
+	{
+		return FindPin(TEXT("Layout"));
+	}
+
+	void AddAttachedErrorData(const FAttachedErrorDataView& AttachedErrorData) override;
+	void ResetAttachedErrorData() override;
+
+	bool IsSingleOutputNode() const override;
+
+private:
+
+	/** Last static or skeletal mesh connected. Used to remove the callback once disconnected. */
+	TWeakObjectPtr<UCustomizableObjectNode> LastMeshNodeConnected;
+
+	/** Connected NodeStaticMesh or NodeSkeletalMesh Mesh UPROPERTY changed callback function. Sets the layout mesh. */
+	void MeshPostEditChangeProperty(UObject* Node, FPropertyChangedEvent& FPropertyChangedEvent);
+
+	void LinkPostEditChangePropertyDelegate(UEdGraphPin* Pin);
+
+	void SetLayoutSkeletalMesh();
+};
+
+
+#if UE_ENABLE_INCLUDE_ORDER_DEPRECATED_IN_5_2
+#include "MuCOE/CustomizableObjectLayout.h"
+#endif
diff --git a/Engine/Plugins/Experimental/Mutable/Source/CustomizableObjectEditor/Public/MuCOE/Nodes/CustomizableObjectNodeMaterial.h b/Engine/Plugins/Experimental/Mutable/Source/CustomizableObjectEditor/Public/MuCOE/Nodes/CustomizableObjectNodeMaterial.h
new file mode 100644
--- /dev/null
+++ b/Engine/Plugins/Experimental/Mutable/Source/CustomizableObjectEditor/Public/MuCOE/Nodes/CustomizableObjectNodeMaterial.h
@@ -0,0 +1,383 @@
+// Copyright Epic Games, Inc. All Rights Reserved.
+
+#pragma once
+
+#include "MuCOE/Nodes/CustomizableObjectNodeMaterialBase.h"
+#include "MuCOE/RemapPins/CustomizableObjectNodeRemapPinsByName.h"
+
+#include "CustomizableObjectNodeMaterial.generated.h"
+
+enum class EMaterialParameterType : uint8;
+namespace ENodeTitleType { enum Type : int; }
+
+class FArchive;
+class SGraphNode;
+class SWidget;
+class UCustomizableObjectNodeRemapPins;
+class UEdGraphPin;
+class UMaterial;
+class UMaterialInterface;
+class UObject;
+class UTexture2D;
+struct FCustomizableObjectNodeMaterialImage;
+struct FCustomizableObjectNodeMaterialScalar;
+struct FCustomizableObjectNodeMaterialVector;
+struct FEdGraphPinReference;
+struct FFrame;
+struct FPropertyChangedEvent;
+
+
+DECLARE_MULTICAST_DELEGATE(FPostImagePinModeChangedDelegate)
+
+
+/** Custom remap pins by name action.
+ *
+ * Remap pins by Texture Parameter Id. */
+UCLASS()
+class UCustomizableObjectNodeMaterialRemapPinsByName : public UCustomizableObjectNodeRemapPinsByName
+{
+	GENERATED_BODY()
+public:
+	virtual bool Equal(const UCustomizableObjectNode& Node, const UEdGraphPin& OldPin, const UEdGraphPin& NewPin) const override;
+
+	virtual void RemapPins(const UCustomizableObjectNode& Node, const TArray<UEdGraphPin*>& OldPins, const TArray<UEdGraphPin*>& NewPins, TMap<UEdGraphPin*, UEdGraphPin*>& PinsToRemap, TArray<UEdGraphPin*>& PinsToOrphan) override;
+
+	bool HasSavedPinData(const UCustomizableObjectNode& Node, const UEdGraphPin &Pin) const;
+};
+
+
+/** Base class for all Material Parameters. */
+UCLASS()
+class CUSTOMIZABLEOBJECTEDITOR_API UCustomizableObjectNodeMaterialPinDataParameter : public UCustomizableObjectNodePinData
+{
+	GENERATED_BODY()
+
+public:
+	/** Texture Parameter Id. */
+	UPROPERTY()
+	FGuid ParameterId;
+
+	/** Returns true if all properties are in its default state. */
+	virtual bool IsDefault() const;
+};
+
+
+/** Node pin mode. All pins set to EPinMode::Default will use this this mode. */
+UENUM()
+enum class ENodePinMode
+{
+	Mutable UMETA(ToolTip = "All Material Texture Parameters go through Mutable."),
+	Passthrough UMETA(ToolTip = "All Material Texture Parameters are not modified by Mutable.")
+};
+
+
+/** Image pin, pin mode. */
+UENUM()
+enum class EPinMode
+{
+	/** Use the defined node pin mode. Does not override it. */
+	Default,
+	/** Override the node pin mode. Set it to Mutable mode. */
+	Mutable,
+	/** Override the node pin mode. Set it to Pass-through mode. */
+	Passthrough
+};
+
+
+/** Image Pin, UV Layout Mode. */
+UENUM()
+enum class EUVLayoutMode
+{
+	/** Does not override the UV Index specified in the Material. */
+	FromMaterial,
+	/* Texture should not be transformed by any layout. Theses textures will not be reduced automatically for LODs. */
+	Ignore,
+	/** User specified UV Index. */
+	Index
+};
+
+
+/** Enum to FText. */
+FText EPinModeToText(EPinMode PinMode);
+
+UCLASS()
+class CUSTOMIZABLEOBJECTEDITOR_API UCustomizableObjectNodeMaterial : public UCustomizableObjectNodeMaterialBase
+{
+public:
+	GENERATED_BODY()
+
+	UPROPERTY(EditAnywhere, Category=CustomizableObject)
+	TObjectPtr<UMaterialInterface> Material = nullptr;
+
+	UPROPERTY(EditAnywhere, Category=CustomizableObject, Meta = (ToolTip = "Set all Mateiral Texture Parameters to the specified mode. Each Texture Parameter Pin can override this mode."))
+	ENodePinMode TextureParametersMode = ENodePinMode::Passthrough;
+
+	UPROPERTY(EditAnywhere, BlueprintReadWrite, Category = CustomizableObject)
+	TArray<FString> Tags;
+
+	/** Selects which Mesh component of the Instance this material belongs to */
+	UPROPERTY(EditAnywhere, BlueprintReadWrite, Category = CustomizableObject, meta = (ClampMin = "0"))
+	int32 MeshComponentIndex = 0;
+
+	/** Materials will be reused between LODs when possible unless explicitly disabled. */
+	UPROPERTY(EditAnywhere, Category = CustomizableObject)
+	bool bReuseMaterialBetweenLODs = true;
+	
+	// UObject interface.
+	virtual void Serialize(FArchive& Ar) override;
+	virtual void PostLoad() override;
+	void PostEditChangeProperty(FPropertyChangedEvent& PropertyChangedEvent) override;
+
+	// EdGraphNode interface
+	FText GetNodeTitle(ENodeTitleType::Type TitleType) const override;
+	FLinearColor GetNodeTitleColor() const override;
+	FText GetTooltipText() const override;
+	TSharedPtr<SGraphNode> CreateVisualWidget() override;
+	virtual void PinConnectionListChanged(UEdGraphPin* Pin) override;
+	virtual void PostPasteNode() override;
+	virtual bool CanConnect(const UEdGraphPin* InOwnedInputPin, const UEdGraphPin* InOutputPin, bool& bOutIsOtherNodeBlocklisted, bool& bOutArePinsCompatible) const override;
+
+	// UCustomizableObjectNode interface
+	virtual void BackwardsCompatibleFixup() override;
+	virtual void PostBackwardsCompatibleFixup() override;
+	void AllocateDefaultPins(UCustomizableObjectNodeRemapPins* RemapPins) override;
+	bool CanPinBeHidden(const UEdGraphPin& Pin) const override;
+	virtual UCustomizableObjectNodeRemapPinsByName* CreateRemapPinsDefault() const override;
+	bool ProvidesCustomPinRelevancyTest() const override { return true; }
+	bool IsPinRelevant(const UEdGraphPin* Pin) const override;
+	virtual bool CustomRemovePin(UEdGraphPin& Pin) override;
+	virtual void ReconstructNode(UCustomizableObjectNodeRemapPins* RemapPinsMode) override;
+
+	// Own Interface
+	virtual UEdGraphPin* GetMeshPin() const;
+	UEdGraphPin* GetMaterialAssetPin() const;
+
+	/** Allow only one connection to a single BaseObjec node. There can be multiple conections to NodeCopyMaterial nodes. */
+	void BreakExistingConnectionsPostConnection(UEdGraphPin* InputPin, UEdGraphPin* OutputPin) override;
+
+	bool IsNodeOutDatedAndNeedsRefresh() override;
+	FString GetRefreshMessage() const override;
+	virtual TSharedPtr<SWidget> CustomizePinDetails(UEdGraphPin& Pin) override;
+
+	
+	// UCustomizableObjectNodeMaterialBase interface
+	TArray<class UCustomizableObjectLayout*> GetLayouts() override;
+	UEdGraphPin* OutputPin() const override;
+	
+	/** Return true if a Material Parameter has changed on which we had a pin connected or data saved. */
+	bool RealMaterialDataHasChanged() const;
+	
+	// --------------------
+	// ALL PARAMETERS
+	// --------------------
+
+	/** Returns the number of Material Parameters. */
+	int32 GetNumParameters(EMaterialParameterType Type) const;
+
+	/** Returns the Material Parameter id.
+	 *
+	 * @param ParameterIndex Have to be valid. */
+	FGuid GetParameterId(EMaterialParameterType Type, int32 ParameterIndex) const;
+	
+	/** Returns the Material Parameter name.
+	 *
+	 * @param ParameterIndex Have to be valid. */
+	FName GetParameterName(EMaterialParameterType Type, int32 ParameterIndex) const;
+	
+	/** Get the Material Parameter layer index.
+	 *
+	 * @param ParameterIndex Have to be valid.
+	 * @returns INDEX_NONE for global parameters. */
+	int32 GetParameterLayerIndex(EMaterialParameterType Type, int32 ParameterIndex) const;
+	
+	/** Get the Material Parameter layer name.
+	 *
+	 * @param ParameterIndex Have to be valid. */
+	FText GetParameterLayerName(EMaterialParameterType Type, int32 ParameterIndex) const;
+
+	/** Returns true if the Material contains the given Material Parameter. */
+	bool HasParameter(const FGuid& ParameterId) const;
+	
+	/** Get the Vector pin for the given Material Vector Parameter.
+	 * Not all parameters have pin.
+	 *
+	 * @param ParameterIndex Has to be valid.
+	 * @return Can return nullptr. */
+	const UEdGraphPin* GetParameterPin(EMaterialParameterType Type, int32 ParameterIndex) const;
+	
+	// --------------------
+	// IMAGES PARAMETERS
+	// --------------------
+
+	/** Returns true if the Material Texture Parameter goes through Mutable.
+	 *
+	 * @param ImageIndex Have to be valid. */
+	bool IsImageMutableMode(int32 ImageIndex) const;
+
+	/** Given an Image pin, returns true if the Material Texture Parameter goes through Mutable. */
+	bool IsImageMutableMode(const UEdGraphPin& Pin) const;
+
+	/** Update a Material Texture Parameter Mode. */
+	void UpdateImagePinMode(const FGuid ParameterId);
+	
+	/** Update a Material Texture Parameter Mode. */
+	void UpdateImagePinMode(const UEdGraphPin& Pin);
+
+	/** Update all Material Texture Parameters Mode. */
+	void UpdateAllImagesPinMode();
+	
+	/** Returns the reference texture assigned to a Material Texture Parameter.
+	 *
+	 * @param ImageIndex Have to be valid.
+	 * @return nullptr if it does not have one assigned. */
+	UTexture2D* GetImageReferenceTexture(int32 ImageIndex) const;
+
+	/** Returns the Texture set in the Material Texture Parameter.
+	 *
+	 * @param ImageIndex Have to be valid. */
+	UTexture2D* GetImageValue(int32 ImageIndex) const;
+	
+	/** Get the Material Texture Parameter UV Index.
+	 *
+	 * @param ImageIndex Have to be valid.
+	 * @return Return -1 if the UV Index is set to Ignore. Return >= 0 for a valid UV Index.  */
+	int32 GetImageUVLayout(int32 ImageIndex) const;
+
+	/** Delegate called when a Texture Parameter Pin Mode changes. */
+	FPostImagePinModeChangedDelegate PostImagePinModeChangedDelegate;
+	
+private:
+	/** Last static or skeletal mesh connected. Used to remove the callback once disconnected. */
+	TWeakObjectPtr<UCustomizableObjectNode> LastMeshNodeConnected;
+
+	static const TArray<EMaterialParameterType> ParameterTypes;
+
+	/** Relates a Parameter id (key) to a Pin (value). Only used to improve performance. */
+	UPROPERTY()
+	TMap<FGuid, FEdGraphPinReference> PinsParameter;
+
+	/** Relates an Image pin (key) to its Image Pin Mode (value). 
+		Represents the real mode of the pin. Required due to some node configurations can force the mode independently of what the user had previously selected */
+	UPROPERTY()
+	TMap<FGuid, EPinMode> PinsImagePinMode;
+	
+	/** Create the pin data of the given parameter type. */
+	UCustomizableObjectNodeMaterialPinDataParameter* CreatePinData(EMaterialParameterType Type, int32 ParameterIndex);
+
+	/** Allocate a pin for each parameter of the given type. */
+	void AllocateDefaultParameterPins(EMaterialParameterType Type);
+
+	/** Set the default Material from the connected static or skeletal mesh. */
+	void SetDefaultMaterial();
+
+	/** Connected NodeStaticMesh or NodeSkeletalMesh Mesh UPROPERTY changed callback function. Sets the default material. */
+	UFUNCTION()
+	void MeshPostEditChangeProperty(FPostEditChangePropertyDelegateParameters& Parameters);
+
+	/** Format pin name. */
+	FName GetPinName(EMaterialParameterType Type, int32 ParameterIndex) const;
+
+	/** Returns the texture coordinate of the given Material Expression. Returns -1 if not found. */
+	static int32 GetExpressionTextureCoordinate(UMaterial* Material, const FGuid &ImageId);
+
+	/** Converts node NodePinMode (does not include Default mode) to PinMode. */
+	static EPinMode NodePinModeToImagePinMode(ENodePinMode NodePinMode);
+
+	/** Returns the Image Pin Mode the pin should be at. It does not update its mode, to update it call UpdateImagePinMode. */
+	EPinMode GetImagePinMode(const UEdGraphPin& Pin) const;
+
+	/** Get the UV Layout Index defined in the Material. */
+	int32 GetImageUVLayoutFromMaterial(int32 ImageIndex) const;
+	
+	// Deprecated properties
+	/** Set all pins to Mutable mode. Even so, each pin can override its behaviour. */
+	UPROPERTY()
+	bool bDefaultPinModeMutable_DEPRECATED = false;
+	
+	UPROPERTY()
+	TArray<FCustomizableObjectNodeMaterialImage> Images_DEPRECATED;
+	
+	UPROPERTY()
+	TArray<FCustomizableObjectNodeMaterialVector> VectorParams_DEPRECATED;
+
+	UPROPERTY()
+	TArray<FCustomizableObjectNodeMaterialScalar> ScalarParams_DEPRECATED;
+};
+
+
+/** Additional data for a Material Texture Parameter pin. */
+UCLASS()
+class CUSTOMIZABLEOBJECTEDITOR_API UCustomizableObjectNodeMaterialPinDataImage : public UCustomizableObjectNodeMaterialPinDataParameter
+{
+	GENERATED_BODY()
+
+	friend void UCustomizableObjectNodeMaterial::BackwardsCompatibleFixup();
+	
+public:
+	// UObject interface
+	virtual void PostEditChangeProperty(FPropertyChangedEvent& PropertyChangedEvent) override;
+private:
+	virtual void PostLoad() override;
+	
+public:
+	// NodePinDataParameter interface
+	/** Virtual function used to copy pin data when remapping pins. */
+	virtual void Copy(const UCustomizableObjectNodePinData& Other) override;
+	
+	// UCustomizableObjectNodeMaterialPinParameter interface
+	virtual bool IsDefault() const override;
+
+	// Own interface
+	/** Constructor parameters. Should always be called after a NewObject. */
+	void Init(UCustomizableObjectNodeMaterial& InNodeMaterial);
+
+	EPinMode GetPinMode() const;
+
+	void SetPinMode(EPinMode InPinMode);
+	
+private:
+	/** Image pin mode. If is not default, overrides the defined node behaviour. */
+	UPROPERTY(EditAnywhere, Category = NoCategory)
+	EPinMode PinMode = EPinMode::Default;
+
+public:
+	constexpr static int32 UV_LAYOUT_IGNORE = -1;
+
+	UPROPERTY(EditAnywhere, Category = NoCategory)
+	EUVLayoutMode UVLayoutMode = EUVLayoutMode::FromMaterial;
+	
+	/** Index of the UV channel that will be used with this image.It is necessary to apply the proper layout transformations to it. */
+	UPROPERTY(EditAnywhere, Category = NoCategory, meta = (EditCondition = "UVLayoutMode == EUVLayoutMode::Index", EditConditionHides))
+	int32 UVLayout = -2;
+
+	/** Reference Texture used to decide the texture properties of the mutable-generated textures
+	* connected to this material. If null, it will try to be guessed at compile time from
+	* the graph. */
+	UPROPERTY(EditAnywhere, Category = NoCategory) // Required to be EditAnywhere for the selector to work.
+	TObjectPtr<UTexture2D> ReferenceTexture = nullptr;
+
+private:
+	UPROPERTY()
+	TObjectPtr<UCustomizableObjectNodeMaterial> NodeMaterial = nullptr;
+};
+
+
+/** Additional data for a Material Vector Parameter pin. */
+UCLASS()
+class CUSTOMIZABLEOBJECTEDITOR_API UCustomizableObjectNodeMaterialPinDataVector : public UCustomizableObjectNodeMaterialPinDataParameter
+{
+	GENERATED_BODY()
+};
+
+
+/** Additional data for a Material Float Parameter pin. */
+UCLASS()
+class CUSTOMIZABLEOBJECTEDITOR_API UCustomizableObjectNodeMaterialPinDataScalar : public UCustomizableObjectNodeMaterialPinDataParameter
+{
+	GENERATED_BODY()
+};
+
+#if UE_ENABLE_INCLUDE_ORDER_DEPRECATED_IN_5_2
+#include "MaterialTypes.h"
+#include "MuCOE/CustomizableObjectEditor_Deprecated.h"
+#endif
diff --git a/Engine/Plugins/Experimental/Mutable/Source/CustomizableObjectEditor/Public/MuCOE/Nodes/CustomizableObjectNodeMaterialBase.h b/Engine/Plugins/Experimental/Mutable/Source/CustomizableObjectEditor/Public/MuCOE/Nodes/CustomizableObjectNodeMaterialBase.h
new file mode 100644
--- /dev/null
+++ b/Engine/Plugins/Experimental/Mutable/Source/CustomizableObjectEditor/Public/MuCOE/Nodes/CustomizableObjectNodeMaterialBase.h
@@ -0,0 +1,35 @@
+// Copyright Epic Games, Inc. All Rights Reserved.
+
+#pragma once
+
+#include "MuCOE/Nodes/CustomizableObjectNode.h"
+
+#include "CustomizableObjectNodeMaterialBase.generated.h"
+
+class UCustomizableObjectLayout;
+
+class UEdGraphPin;
+class UObject;
+
+
+
+UCLASS(abstract)
+class CUSTOMIZABLEOBJECTEDITOR_API UCustomizableObjectNodeMaterialBase : public UCustomizableObjectNode
+{
+public:
+	GENERATED_BODY()
+
+	// Own interface
+	virtual TArray<class UCustomizableObjectLayout*> GetLayouts() { return TArray<class UCustomizableObjectLayout*>(); }
+
+	// This method should be overidden in all derived classes
+	virtual UEdGraphPin* OutputPin() const
+	{
+		check(false); 
+		return nullptr;
+	}
+
+	/** Allow multiple connections to NodeCopyMaterial but only one connection to a NodeObject. */
+	bool ShouldBreakExistingConnections(const UEdGraphPin* InputPin, const UEdGraphPin* OutputPin) const override;
+};
+
diff --git a/Engine/Plugins/Experimental/Mutable/Source/CustomizableObjectEditor/Public/MuCOE/Nodes/CustomizableObjectNodeMaterialVariation.h b/Engine/Plugins/Experimental/Mutable/Source/CustomizableObjectEditor/Public/MuCOE/Nodes/CustomizableObjectNodeMaterialVariation.h
new file mode 100644
--- /dev/null
+++ b/Engine/Plugins/Experimental/Mutable/Source/CustomizableObjectEditor/Public/MuCOE/Nodes/CustomizableObjectNodeMaterialVariation.h
@@ -0,0 +1,43 @@
+// Copyright Epic Games, Inc. All Rights Reserved.
+
+#pragma once
+
+#include "MuCOE/Nodes/CustomizableObjectNodeVariation.h"
+
+#include "CustomizableObjectNodeMaterialVariation.generated.h"
+
+struct FCustomizableObjectMaterialVariation;
+
+
+UENUM(BlueprintType)
+enum class ECustomizableObjectNodeMaterialVariationType : uint8
+{
+	Tag 		UMETA(DisplayName = "Tag"),
+	State 		UMETA(DisplayName = "State"),
+};
+
+
+UCLASS()
+class CUSTOMIZABLEOBJECTEDITOR_API UCustomizableObjectNodeMaterialVariation : public UCustomizableObjectNodeVariation
+{
+public:
+	GENERATED_BODY()
+
+	UPROPERTY(EditAnywhere, Category = CustomizableObject)
+	ECustomizableObjectNodeMaterialVariationType Type = ECustomizableObjectNodeMaterialVariationType::Tag;
+
+private:
+	// Deprecated properties
+	UPROPERTY()
+	TArray<FCustomizableObjectMaterialVariation> Variations_DEPRECATED;
+
+public:
+	// UCustomizableObjectNode interface
+	virtual void BackwardsCompatibleFixup() override;
+	virtual bool IsSingleOutputNode() const override;
+	
+	// UCustomizableObjectNodeVariation interface
+	virtual FName GetCategory() const override;
+	virtual bool IsInputPinArray() const override;
+};
+
diff --git a/Engine/Plugins/Experimental/Mutable/Source/CustomizableObjectEditor/Public/MuCOE/Nodes/CustomizableObjectNodeMesh.h b/Engine/Plugins/Experimental/Mutable/Source/CustomizableObjectEditor/Public/MuCOE/Nodes/CustomizableObjectNodeMesh.h
new file mode 100644
--- /dev/null
+++ b/Engine/Plugins/Experimental/Mutable/Source/CustomizableObjectEditor/Public/MuCOE/Nodes/CustomizableObjectNodeMesh.h
@@ -0,0 +1,49 @@
+// Copyright Epic Games, Inc. All Rights Reserved.
+
+#pragma once
+
+#include "MuCOE/Nodes/CustomizableObjectNode.h"
+
+#include "CustomizableObjectNodeMesh.generated.h"
+
+class UCustomizableObjectLayout;
+class UEdGraphPin;
+class UObject;
+class UTexture2D;
+
+
+UCLASS(abstract)
+class CUSTOMIZABLEOBJECTEDITOR_API UCustomizableObjectNodeMesh : public UCustomizableObjectNode
+{
+public:
+	GENERATED_BODY()
+
+	// Own interface
+	/** Returns the multiple Layouts of a given Mesh pin. A pin can have multiple Layouts since it can have multiple UVs. Override. */
+	virtual TArray<UCustomizableObjectLayout*> GetLayouts(const UEdGraphPin& OutPin) const PURE_VIRTUAL(UCustomizableObjectNodeMesh::GetLayouts, return {};);
+
+	virtual UTexture2D* FindTextureForPin(const UEdGraphPin* Pin) const PURE_VIRTUAL(UCustomizableObjectNodeMesh::FindTextureForPin, return {};);
+
+	virtual void GetUVChannelForPin(const UEdGraphPin* Pin, TArray<FVector2f>& OutSegments, int32 UVIndex) const PURE_VIRTUAL(UCustomizableObjectNodeMesh::GetUVChannelForPin, );
+
+	/** Returns the Unreal mesh (e.g., USkeletalMesh, UStaticMesh...). */
+	virtual UObject* GetMesh() const PURE_VIRTUAL(UCustomizableObjectNodeMesh::GetMesh, return {};);
+
+	/** Returns the output Mesh pin associated to the given LODIndex and SectionIndex. Override. */
+	virtual UEdGraphPin* GetMeshPin(int32 LOD, int32 SectionIndex) const PURE_VIRTUAL(UCustomizableObjectNodeMesh::GetMeshPin, return {};);
+
+	/** Returns the input Layout pin associated to the given LODIndex, SectionIndex and LayoutIndex. Override. */
+	virtual UEdGraphPin* GetLayoutPin(int32 LODIndex, int32 SectionIndex, int32 LayoutIndex) const PURE_VIRTUAL(UCustomizableObjectNodeMesh::GetLayoutPin, return {};);
+	
+	/** Given a pin, return the Section and Layout index.
+	 *
+	 * Will always return a valid result. The result can be valid but out of sync with respect the Unreal mesh asset.
+	 * In other words, the pin still represents a LOD 3 but the asset no longer has this third LOD.
+	 * 
+	 * @param Pin In. Node's owning pin.
+	 * @param LODIndex Out.
+	 * @param SectionIndex Out.
+	 * @param LayoutIndex Out. */
+	virtual void GetPinSection(const UEdGraphPin& Pin, int32& OutLODIndex, int32& OutSectionIndex, int32& OutLayoutIndex) const PURE_VIRTUAL(UCustomizableObjectNodeMesh::GetPinSection, );
+};
+
diff --git a/Engine/Plugins/Experimental/Mutable/Source/CustomizableObjectEditor/Public/MuCOE/Nodes/CustomizableObjectNodeMeshClipDeform.h b/Engine/Plugins/Experimental/Mutable/Source/CustomizableObjectEditor/Public/MuCOE/Nodes/CustomizableObjectNodeMeshClipDeform.h
new file mode 100644
--- /dev/null
+++ b/Engine/Plugins/Experimental/Mutable/Source/CustomizableObjectEditor/Public/MuCOE/Nodes/CustomizableObjectNodeMeshClipDeform.h
@@ -0,0 +1,50 @@
+// Copyright Epic Games, Inc. All Rights Reserved.
+
+#pragma once
+
+#include "MuCOE/Nodes/CustomizableObjectNodeModifierBase.h"
+#include "MuT/NodeModifier.h"
+
+#include "CustomizableObjectNodeMeshClipDeform.generated.h"
+
+class UCustomizableObjectNodeRemapPins;
+class UEdGraphPin;
+class UObject;
+
+UENUM()
+enum class EShapeBindingMethod : uint32
+{
+	ClosestProject = 0,
+	ClosestToSurface = 1,
+	NormalProject = 2
+};
+
+UCLASS()
+class CUSTOMIZABLEOBJECTEDITOR_API UCustomizableObjectNodeMeshClipDeform : public UCustomizableObjectNodeModifierBase
+{
+public:
+	GENERATED_BODY()
+
+	UPROPERTY(EditAnywhere, BlueprintReadWrite, Category = MeshClipDeform)
+	TArray<FString> Tags;
+	
+	/** Policy to use tags in case more than one is added. */
+	UPROPERTY(EditAnywhere, Category = MeshClipDeform)
+	EMutableMultipleTagPolicy MultipleTagPolicy = EMutableMultipleTagPolicy::OnlyOneRequired;
+
+	UPROPERTY(EditAnywhere, Category = MeshClipDeform)
+	EShapeBindingMethod BindingMethod;
+	
+	// EdGraphNode interface
+	virtual FText GetNodeTitle(ENodeTitleType::Type TitleType) const override;
+	virtual FLinearColor GetNodeTitleColor() const override;
+	virtual FText GetTooltipText() const override;
+
+	// UCustomizableObjectNode interface
+	virtual void AllocateDefaultPins(UCustomizableObjectNodeRemapPins* RemapPins) override;
+	virtual UEdGraphPin* OutputPin() const override;
+
+	// Own interface
+	UEdGraphPin* ClipShapePin() const;
+};
+
diff --git a/Engine/Plugins/Experimental/Mutable/Source/CustomizableObjectEditor/Public/MuCOE/Nodes/CustomizableObjectNodeMeshClipMorph.h b/Engine/Plugins/Experimental/Mutable/Source/CustomizableObjectEditor/Public/MuCOE/Nodes/CustomizableObjectNodeMeshClipMorph.h
new file mode 100644
--- /dev/null
+++ b/Engine/Plugins/Experimental/Mutable/Source/CustomizableObjectEditor/Public/MuCOE/Nodes/CustomizableObjectNodeMeshClipMorph.h
@@ -0,0 +1,102 @@
+// Copyright Epic Games, Inc. All Rights Reserved.
+
+#pragma once
+
+#include "MuCOE/Nodes/CustomizableObjectNodeModifierBase.h"
+#include "MuT/NodeModifier.h"
+
+#include "CustomizableObjectNodeMeshClipMorph.generated.h"
+
+namespace ENodeTitleType { enum Type : int; }
+
+class FArchive;
+class UCustomizableObjectNodeRemapPins;
+class UEdGraphPin;
+class UObject;
+struct FPropertyChangedEvent;
+
+UCLASS()
+class CUSTOMIZABLEOBJECTEDITOR_API UCustomizableObjectNodeMeshClipMorph : public UCustomizableObjectNodeModifierBase
+{
+public:
+	GENERATED_BODY()
+
+	UPROPERTY(EditAnywhere, Category = MeshToClipAndMorph)
+	FName BoneName;
+
+	UPROPERTY(EditAnywhere, BlueprintReadWrite, Category = MeshToClipAndMorph)
+	TArray<FString> Tags;
+
+	/** Policy to use tags in case more than one is added. */
+	UPROPERTY(EditAnywhere, Category = MeshToClipAndMorph)
+	EMutableMultipleTagPolicy MultipleTagPolicy = EMutableMultipleTagPolicy::OnlyOneRequired;
+
+	UPROPERTY(EditAnywhere, Category = MeshToClipAndMorph)
+	uint32 ReferenceSkeletonIndex = 0;
+
+	UPROPERTY(EditAnywhere, Category = MeshClipParameters, meta = (DisplayName="Morph Start Offset", ToolTip="Offset from the origin of the selected bone to the actual start of the morph."))
+	FVector StartOffset;
+
+	UPROPERTY(EditAnywhere, Category = MeshClipParameters, meta = (DisplayName = "Local Start Offset", ToolTip = "Toggles between a local or global start offset."))
+	bool bLocalStartOffset;
+
+	UPROPERTY(EditAnywhere, Category = MeshClipParameters, meta = (DisplayName="Morph length", ToolTip="The length from the morph start to the clip plane."))
+	float B;
+
+	UPROPERTY(EditAnywhere, Category = MeshClipParameters, meta = (DisplayName="Ellipse Radius 1", ToolTip="First radius of the ellipse that the mesh is morphed into."))
+	float Radius;
+
+	UPROPERTY(EditAnywhere, Category = MeshClipParameters, meta = (DisplayName = "Ellipse Radius 2", ToolTip = "Second radius of the ellipse that the mesh is morphed into."))
+	float Radius2;
+
+	UPROPERTY(EditAnywhere, Category = MeshClipParameters, meta = (DisplayName = "Ellipse Rotation", ToolTip = "Ellipse Rotation in degrees around the bone axis."))
+	float RotationAngle;
+
+	UPROPERTY(EditAnywhere, Category = MeshClipParameters, meta = (DisplayName = "Morph Curve Control", ToolTip = "Controls the morph curve shape. A value of 1 is linear, less than 1 is concave and greater than 1 convex."))
+	float Exponent;
+
+	UPROPERTY(EditAnywhere, Category = MeshClipParameters, meta = (DisplayName = "Invert normal direction", ToolTip = "Flag to invert the normal direction"))
+	bool bInvertNormal;
+
+	UPROPERTY()
+	FVector Origin;
+
+	UPROPERTY()
+	FVector Normal;
+
+private:
+	UPROPERTY()
+	bool bOldOffset_DEPRECATED;
+
+public:
+	UPROPERTY(EditAnywhere, Category = MeshClipParameters, meta = (DisplayName = "Max Effect Radius", ToolTip = "The maximum distance from the origin of the widget where vertices will be affected. If negative, there will be no limit."))
+	float MaxEffectRadius;
+
+	bool bUpdateViewportWidget;
+
+	UCustomizableObjectNodeMeshClipMorph();
+
+	// UObject interface.
+	virtual void PostEditChangeProperty(FPropertyChangedEvent& PropertyChangedEvent) override;
+	virtual void Serialize(FArchive& Ar) override;
+
+	// EdGraphNode interface
+	virtual FText GetNodeTitle(ENodeTitleType::Type TitleType) const override;
+	virtual FLinearColor GetNodeTitleColor() const override;
+	virtual FText GetTooltipText() const override;
+
+	// UCustomizableObjectNode interface
+	virtual void AllocateDefaultPins(UCustomizableObjectNodeRemapPins* RemapPins) override;
+	virtual void PinConnectionListChanged(UEdGraphPin* Pin) override;
+	
+	// UCustomizableObjectNodeMaterialBase interface
+	virtual UEdGraphPin* OutputPin() const override;
+
+	// Own interface
+	FVector GetOriginWithOffset() const;
+	void FindLocalAxes(FVector& Right, FVector& Up, FVector& Forward) const;
+
+	// Change StartOffset from World to Local of the other way around
+	void ChangeStartOffsetTransform();
+};
+
diff --git a/Engine/Plugins/Experimental/Mutable/Source/CustomizableObjectEditor/Public/MuCOE/Nodes/CustomizableObjectNodeMeshClipWithMesh.h b/Engine/Plugins/Experimental/Mutable/Source/CustomizableObjectEditor/Public/MuCOE/Nodes/CustomizableObjectNodeMeshClipWithMesh.h
new file mode 100644
--- /dev/null
+++ b/Engine/Plugins/Experimental/Mutable/Source/CustomizableObjectEditor/Public/MuCOE/Nodes/CustomizableObjectNodeMeshClipWithMesh.h
@@ -0,0 +1,69 @@
+// Copyright Epic Games, Inc. All Rights Reserved.
+
+#pragma once
+
+#include "MuCOE/Nodes/CustomizableObjectNodeModifierBase.h"
+#include "MuT/NodeModifier.h"
+
+#include "CustomizableObjectNodeMeshClipWithMesh.generated.h"
+
+namespace ENodeTitleType { enum Type : int; }
+
+class UCustomizableObject;
+class UCustomizableObjectNodeRemapPins;
+class UEdGraphPin;
+class UObject;
+struct FGuid;
+
+UCLASS()
+class CUSTOMIZABLEOBJECTEDITOR_API UCustomizableObjectNodeMeshClipWithMesh : public UCustomizableObjectNodeModifierBase
+{
+public:
+	GENERATED_BODY()
+
+	UCustomizableObjectNodeMeshClipWithMesh();
+
+	UPROPERTY(EditAnywhere, BlueprintReadWrite, Category = MeshToClip)
+	TArray<FString> Tags;
+
+	UPROPERTY(EditAnywhere, Category = CustomizableObjectToClip)
+	EMutableMultipleTagPolicy MultipleTagPolicy = EMutableMultipleTagPolicy::OnlyOneRequired;
+
+	//!< If assigned, then a material inside this CO will be clipped by this node.
+    //!< If several materials with the same name, all are considered (to cover all LOD levels)
+    UPROPERTY(EditAnywhere, Category = CustomizableObjectToClip)
+	TObjectPtr<UCustomizableObject> CustomizableObjectToClipWith;
+
+    //!< Array with the Guids of the nodes with the same material inside the CustomizableObjectToClipWith CO (if any is assigned)
+    UPROPERTY(EditAnywhere, Category = CustomizableObjectToClip)
+    TArray<FGuid> ArrayMaterialNodeToClipWithID;
+
+	/** Mesh Transform*/
+	UPROPERTY(EditAnywhere, Category = CustomizableObject)
+	FTransform Transform;
+
+	// Details view variables
+	// The clipping node uses tags
+	bool bUseTags;
+
+	// The clipping node uses a material name
+	bool bUseMaterials;
+
+	// EdGraphNode interface
+	virtual FText GetNodeTitle(ENodeTitleType::Type TitleType) const override;
+	virtual FLinearColor GetNodeTitleColor() const override;
+	virtual FText GetTooltipText() const override;
+
+	// UCustomizableObjectNode interface
+	virtual void AllocateDefaultPins(UCustomizableObjectNodeRemapPins* RemapPins) override;
+	virtual void PinConnectionListChanged(UEdGraphPin* Pin) override;
+	virtual void UpdateReferencedNodeId(const FGuid& NewGuid) override;
+	virtual void BeginPostDuplicate(bool bDuplicateForPIE) override;
+
+	// UCustomizableObjectNodeMaterialBase interface
+	virtual UEdGraphPin* OutputPin() const override;
+
+	// Own interface
+	UEdGraphPin* ClipMeshPin() const;
+};
+
diff --git a/Engine/Plugins/Experimental/Mutable/Source/CustomizableObjectEditor/Public/MuCOE/Nodes/CustomizableObjectNodeMeshGeometryOperation.h b/Engine/Plugins/Experimental/Mutable/Source/CustomizableObjectEditor/Public/MuCOE/Nodes/CustomizableObjectNodeMeshGeometryOperation.h
new file mode 100644
--- /dev/null
+++ b/Engine/Plugins/Experimental/Mutable/Source/CustomizableObjectEditor/Public/MuCOE/Nodes/CustomizableObjectNodeMeshGeometryOperation.h
@@ -0,0 +1,51 @@
+// Copyright Epic Games, Inc. All Rights Reserved.
+
+#pragma once
+
+#include "MuCOE/Nodes/CustomizableObjectNode.h"
+
+#include "CustomizableObjectNodeMeshGeometryOperation.generated.h"
+
+class UCustomizableObjectNodeRemapPins;
+class UEdGraphPin;
+class UObject;
+
+
+UCLASS()
+class CUSTOMIZABLEOBJECTEDITOR_API UCustomizableObjectNodeMeshGeometryOperation : public UCustomizableObjectNode
+{
+public:
+	GENERATED_BODY()
+
+	UCustomizableObjectNodeMeshGeometryOperation();
+
+	// Begin EdGraphNode interface
+	FText GetNodeTitle(ENodeTitleType::Type TitleType) const override;
+	FLinearColor GetNodeTitleColor() const override;
+	FText GetTooltipText() const override;
+
+	// UCustomizableObjectNode interface
+	void AllocateDefaultPins(UCustomizableObjectNodeRemapPins* RemapPins) override;
+
+	inline UEdGraphPin* MeshAPin() const
+	{
+		return FindPin(TEXT("Mesh A"), EGPD_Input);
+	}
+
+	inline UEdGraphPin* MeshBPin() const
+	{
+		return FindPin(TEXT("Mesh B"), EGPD_Input);
+	}
+
+	inline UEdGraphPin* ScalarAPin() const
+	{
+		return FindPin(TEXT("Scalar A"));
+	}
+
+	inline UEdGraphPin* ScalarBPin() const
+	{
+		return FindPin(TEXT("Scalar B"));
+	}
+	
+};
+
diff --git a/Engine/Plugins/Experimental/Mutable/Source/CustomizableObjectEditor/Public/MuCOE/Nodes/CustomizableObjectNodeMeshMorph.h b/Engine/Plugins/Experimental/Mutable/Source/CustomizableObjectEditor/Public/MuCOE/Nodes/CustomizableObjectNodeMeshMorph.h
new file mode 100644
--- /dev/null
+++ b/Engine/Plugins/Experimental/Mutable/Source/CustomizableObjectEditor/Public/MuCOE/Nodes/CustomizableObjectNodeMeshMorph.h
@@ -0,0 +1,84 @@
+// Copyright Epic Games, Inc. All Rights Reserved.
+
+#pragma once
+
+#include "MuCOE/Nodes/CustomizableObjectNode.h"
+#include "MuCOE/Nodes/CustomizableObjectNodeMeshReshapeCommon.h"
+
+#include "CustomizableObjectNodeMeshMorph.generated.h"
+
+class UCustomizableObjectNodeRemapPins;
+class UEdGraphPin;
+class UObject;
+struct FMeshReshapeBoneReference;
+enum class EBoneDeformSelectionMethod : uint8;
+
+UCLASS()
+class CUSTOMIZABLEOBJECTEDITOR_API UCustomizableObjectNodeMeshMorph : public UCustomizableObjectNode
+{
+public:
+	GENERATED_BODY()
+
+	UCustomizableObjectNodeMeshMorph();
+
+	// Begin EdGraphNode interface
+	FText GetNodeTitle(ENodeTitleType::Type TitleType) const override;
+	FLinearColor GetNodeTitleColor() const override;
+	FText GetTooltipText() const override;
+
+	// UCustomizableObjectNode interface
+	void AllocateDefaultPins(UCustomizableObjectNodeRemapPins* RemapPins) override;
+	bool IsNodeOutDatedAndNeedsRefresh() override;
+	FString GetRefreshMessage() const override;
+	void PinConnectionListChanged(UEdGraphPin* Pin) override;
+
+
+	// Own interface
+	class UCustomizableObjectNodeSkeletalMesh* GetSourceSkeletalMesh() const;
+
+	UEdGraphPin* MeshPin() const
+	{
+		return FindPin(TEXT("Mesh"), EGPD_Input);
+	}
+
+	UEdGraphPin* FactorPin() const
+	{
+		return FindPin(TEXT("Factor"));
+	}
+
+	void Serialize(FArchive& Ar) override;
+
+	UPROPERTY(Category=CustomizableObject, EditAnywhere)
+	FString MorphTargetName;
+
+	/** Enable the deformation of the skeleton of the base mesh. */
+	UPROPERTY(EditAnywhere, Category = CustomizableObject)
+	bool bReshapeSkeleton = false;
+
+	/** Enable the deformation of physics volumes of the base mesh */
+    UPROPERTY(EditAnywhere, Category = CustomizableObject)
+    bool bReshapePhysicsVolumes = false;
+
+	/** Bone Reshape Selection Method */
+	UPROPERTY(EditAnywhere, Category = ReshapeBones, meta = (EditCondition = "bReshapeSkeleton"))
+	EBoneDeformSelectionMethod SelectionMethod;
+	
+	/** Array with selected bones that will be deformed */
+	UPROPERTY(EditAnywhere, Category = ReshapeBones, meta = (EditCondition = "bReshapeSkeleton && (SelectionMethod == EBoneDeformSelectionMethod::ONLY_SELECTED || SelectionMethod == EBoneDeformSelectionMethod::ALL_BUT_SELECTED)"))
+	TArray<FMeshReshapeBoneReference> BonesToDeform;
+
+	UPROPERTY(EditAnywhere, Category = ReshapePhysics, meta = (DisplayName = "Selection Method", EditCondition = "bReshapePhysicsVolumes"))
+	EBoneDeformSelectionMethod PhysicsSelectionMethod;
+
+	/** Array with bones with physics bodies that will be deformed */
+	UPROPERTY(EditAnywhere, Category = ReshapePhysics, meta = (EditCondition = "bReshapePhysicsVolumes && (PhysicsSelectionMethod == EBoneDeformSelectionMethod::ONLY_SELECTED || PhysicsSelectionMethod == EBoneDeformSelectionMethod::ALL_BUT_SELECTED)"))
+	TArray<FMeshReshapeBoneReference> PhysicsBodiesToDeform;
+
+	UPROPERTY()
+	bool bDeformAllBones_DEPRECATED = false;
+
+	UPROPERTY()
+	bool bDeformAllPhysicsBodies_DEPRECATED = false;
+
+};
+
diff --git a/Engine/Plugins/Experimental/Mutable/Source/CustomizableObjectEditor/Public/MuCOE/Nodes/CustomizableObjectNodeMeshMorphStackApplication.h b/Engine/Plugins/Experimental/Mutable/Source/CustomizableObjectEditor/Public/MuCOE/Nodes/CustomizableObjectNodeMeshMorphStackApplication.h
new file mode 100644
--- /dev/null
+++ b/Engine/Plugins/Experimental/Mutable/Source/CustomizableObjectEditor/Public/MuCOE/Nodes/CustomizableObjectNodeMeshMorphStackApplication.h
@@ -0,0 +1,37 @@
+// Copyright Epic Games, Inc. All Rights Reserved.
+
+#pragma once
+
+#include "MuCOE/Nodes/CustomizableObjectNode.h"
+
+#include "CustomizableObjectNodeMeshMorphStackApplication.generated.h"
+
+class UCustomizableObjectNodeRemapPins;
+class UEdGraphPin;
+class UObject;
+
+
+UCLASS()
+class CUSTOMIZABLEOBJECTEDITOR_API UCustomizableObjectNodeMeshMorphStackApplication : public UCustomizableObjectNode
+{
+public:
+	GENERATED_BODY()
+
+	// EdGraphNode interface
+	virtual FText GetNodeTitle(ENodeTitleType::Type TittleType)const override;
+	virtual FLinearColor GetNodeTitleColor() const override;
+	virtual FText GetTooltipText() const override;
+
+	// UCustomizableObjectNode interface
+	virtual void AllocateDefaultPins(UCustomizableObjectNodeRemapPins* RemapPins) override;
+
+	// Own interface
+	/** Fills the list with all the morphs */
+	TArray<FString> GetMorphList() const;
+	
+	/** Returns the mesh pin. */
+	UEdGraphPin* GetMeshPin() const;
+
+	/** Returns the stack pin. */
+	UEdGraphPin* GetStackPin() const;
+};
diff --git a/Engine/Plugins/Experimental/Mutable/Source/CustomizableObjectEditor/Public/MuCOE/Nodes/CustomizableObjectNodeMeshMorphStackDefinition.h b/Engine/Plugins/Experimental/Mutable/Source/CustomizableObjectEditor/Public/MuCOE/Nodes/CustomizableObjectNodeMeshMorphStackDefinition.h
new file mode 100644
--- /dev/null
+++ b/Engine/Plugins/Experimental/Mutable/Source/CustomizableObjectEditor/Public/MuCOE/Nodes/CustomizableObjectNodeMeshMorphStackDefinition.h
@@ -0,0 +1,51 @@
+// Copyright Epic Games, Inc. All Rights Reserved.
+
+#pragma once
+
+#include "MuCOE/Nodes/CustomizableObjectNode.h"
+
+#include "CustomizableObjectNodeMeshMorphStackDefinition.generated.h"
+
+class UCustomizableObjectNodeRemapPins;
+class UEdGraphPin;
+class UObject;
+
+
+UCLASS()
+class CUSTOMIZABLEOBJECTEDITOR_API UCustomizableObjectNodeMeshMorphStackDefinition : public UCustomizableObjectNode
+{
+public:
+	GENERATED_BODY()
+
+	// EdGraphNode interface
+	virtual FText GetNodeTitle(ENodeTitleType::Type TittleType)const override;
+	virtual FLinearColor GetNodeTitleColor() const override;
+	virtual FText GetTooltipText() const override;
+	virtual void PinConnectionListChanged(UEdGraphPin* Pin) override;
+	virtual bool IsNodeOutDatedAndNeedsRefresh() override;
+
+	// UCustomizableObjectNode interface
+	virtual void AllocateDefaultPins(UCustomizableObjectNodeRemapPins* RemapPins) override;
+	virtual void ReconstructNode(UCustomizableObjectNodeRemapPins* RemapPins) override;
+
+	// Own interface
+	/** Returns the mesh pin of the node. */
+	UEdGraphPin* GetMeshPin() const;
+
+	/** Returns the stack pin of the node. */
+	UEdGraphPin* GetStackPin() const;
+
+	/** Returns the morph pin at Index. */
+	UEdGraphPin* GetMorphPin(int32 Index) const;
+	
+	/** Returns the index of the next connected stack node. Returns -1 if there is none. */
+	int32 NextConnectedPin(int32 Index, TArray<FString> AvailableMorphs)const;
+
+private:
+	/** Fills the list with all the morphs. */
+	void UpdateMorphList();
+
+	/** List with all the morphs of the linked skeletal mesh. */
+	UPROPERTY()
+	TArray<FString> MorphNames;
+};
