[Mutable] Moved all files from Classes folder to Public folder.

diff --git a/Engine/Plugins/Experimental/Mutable/Source/CustomizableObjectEditor/Public/MuCOE/Nodes/CustomizableObjectNodeMeshReshape.h b/Engine/Plugins/Experimental/Mutable/Source/CustomizableObjectEditor/Public/MuCOE/Nodes/CustomizableObjectNodeMeshReshape.h
new file mode 100644
--- /dev/null
+++ b/Engine/Plugins/Experimental/Mutable/Source/CustomizableObjectEditor/Public/MuCOE/Nodes/CustomizableObjectNodeMeshReshape.h
@@ -0,0 +1,122 @@
+// Copyright Epic Games, Inc. All Rights Reserved.
+
+#pragma once
+
+#include "MuCOE/Nodes/CustomizableObjectNode.h"
+
+#include "CustomizableObjectNodeMeshReshape.generated.h"
+
+class UCustomizableObjectNodeRemapPins;
+class UEdGraphPin;
+class UObject;
+struct FMeshReshapeBoneReference;
+enum class EBoneDeformSelectionMethod : uint8;
+
+
+UENUM()
+enum class EMeshReshapeVertexColorChannelUsage
+{
+	None = 0,
+	RigidClusterId = 1,
+	MaskWeight = 2
+};
+
+USTRUCT()
+struct FMeshReshapeColorUsage
+{
+	GENERATED_BODY()
+
+	UPROPERTY(EditAnywhere, Category = CustomizableObject)
+	EMeshReshapeVertexColorChannelUsage R = EMeshReshapeVertexColorChannelUsage::None;
+
+	UPROPERTY(EditAnywhere, Category = CustomizableObject)
+	EMeshReshapeVertexColorChannelUsage G = EMeshReshapeVertexColorChannelUsage::None;
+
+	UPROPERTY(EditAnywhere, Category = CustomizableObject)
+	EMeshReshapeVertexColorChannelUsage B = EMeshReshapeVertexColorChannelUsage::None;
+
+	UPROPERTY(EditAnywhere, Category = CustomizableObject)
+	EMeshReshapeVertexColorChannelUsage A = EMeshReshapeVertexColorChannelUsage::None;
+};
+
+UCLASS()
+class CUSTOMIZABLEOBJECTEDITOR_API UCustomizableObjectNodeMeshReshape : public UCustomizableObjectNode
+{
+public:
+	GENERATED_BODY()
+
+	UCustomizableObjectNodeMeshReshape();
+
+	// Begin EdGraphNode interface
+	FText GetNodeTitle(ENodeTitleType::Type TitleType) const override;
+	FLinearColor GetNodeTitleColor() const override;
+	FText GetTooltipText() const override;
+
+	// UCustomizableObjectNode interface
+	void AllocateDefaultPins(UCustomizableObjectNodeRemapPins* RemapPins) override;
+	FString GetRefreshMessage() const override;
+
+	inline UEdGraphPin* BaseMeshPin() const
+	{
+		return FindPin(TEXT("Base Mesh"), EGPD_Input);
+	}
+
+	inline UEdGraphPin* BaseShapePin() const
+	{
+		return FindPin(TEXT("Base Shape"), EGPD_Input);
+	}
+
+	inline UEdGraphPin* TargetShapePin() const
+	{
+		return FindPin(TEXT("Target Shape"), EGPD_Input);
+	}
+
+	void Serialize(FArchive& Ar) override;
+
+	/** Enable the deformation of the vertices of the base mesh. */
+	UPROPERTY(EditAnywhere, Category = CustomizableObject)
+	bool bReshapeVertices = true;
+
+	UPROPERTY(EditAnywhere, Category = CustomizableObject)
+	bool bRecomputeNormals = false;
+
+	/** Enable laplacian smoothing to the result of the base mesh reshape. */
+	UPROPERTY(EditAnywhere, Category = CustomizableObject, meta = (EditCondition = "bReshapeVertices"))
+	bool bApplyLaplacianSmoothing = false;
+
+	/** Enable the deformation of the skeleton of the base mesh. */
+	UPROPERTY(EditAnywhere, Category = CustomizableObject)
+	bool bReshapeSkeleton = false;
+
+	/** Enable the deformation of physics volumes of the base mesh */
+    UPROPERTY(EditAnywhere, Category = CustomizableObject)
+    bool bReshapePhysicsVolumes = false;
+	
+	UPROPERTY()
+	bool bEnableRigidParts_DEPRECATED = false;
+
+	UPROPERTY(EditAnywhere, Category = CustomizableObject)
+	FMeshReshapeColorUsage VertexColorUsage;
+
+	/** Bone Reshape Selection Method */
+	UPROPERTY(EditAnywhere, Category = ReshapeBones, meta = (EditCondition = "bReshapeSkeleton"))
+	EBoneDeformSelectionMethod SelectionMethod;
+
+	/** Array with selected bones that will be deformed */
+	UPROPERTY(EditAnywhere, Category = ReshapeBones, meta = (EditCondition = "bReshapeSkeleton && (SelectionMethod == EBoneDeformSelectionMethod::ONLY_SELECTED || SelectionMethod == EBoneDeformSelectionMethod::ALL_BUT_SELECTED)"))
+	TArray<FMeshReshapeBoneReference> BonesToDeform;
+
+	UPROPERTY(EditAnywhere, Category = ReshapePhysics, meta = (DisplayName = "Selection Method", EditCondition = "bReshapePhysicsVolumes"))
+	EBoneDeformSelectionMethod PhysicsSelectionMethod;
+
+	/** Array with bones with physics bodies that will be deformed */
+	UPROPERTY(EditAnywhere, Category = ReshapePhysics, meta = (EditCondition = "bReshapePhysicsVolumes && (PhysicsSelectionMethod == EBoneDeformSelectionMethod::ONLY_SELECTED || PhysicsSelectionMethod == EBoneDeformSelectionMethod::ALL_BUT_SELECTED)"))
+	TArray<FMeshReshapeBoneReference> PhysicsBodiesToDeform;
+
+	UPROPERTY()
+	bool bDeformAllBones_DEPRECATED = false;
+
+	UPROPERTY()
+	bool bDeformAllPhysicsBodies_DEPRECATED = false;
+};
+
diff --git a/Engine/Plugins/Experimental/Mutable/Source/CustomizableObjectEditor/Public/MuCOE/Nodes/CustomizableObjectNodeMeshReshapeCommon.h b/Engine/Plugins/Experimental/Mutable/Source/CustomizableObjectEditor/Public/MuCOE/Nodes/CustomizableObjectNodeMeshReshapeCommon.h
new file mode 100644
--- /dev/null
+++ b/Engine/Plugins/Experimental/Mutable/Source/CustomizableObjectEditor/Public/MuCOE/Nodes/CustomizableObjectNodeMeshReshapeCommon.h
@@ -0,0 +1,33 @@
+// Copyright Epic Games, Inc. All Rights Reserved.
+
+#pragma once
+
+
+#include "CustomizableObjectNodeMeshReshapeCommon.generated.h"
+
+UENUM()
+enum class EBoneDeformSelectionMethod : uint8
+{
+	// Only selected bones will be deform
+	ONLY_SELECTED = 0 UMETA(DisplayName = "Only Selected"),
+
+	// All bones will be deform except the selected ones
+	ALL_BUT_SELECTED = 1 UMETA(DisplayName = "All But Selected"),
+
+	// Deform only the bones of the reference skeleton
+	DEFORM_REF_SKELETON = 2 UMETA(DisplayName = "Deform Only Ref.Seleton Bones"),
+
+	// Deform only the bones that are not in the reference skeleton
+	DEFORM_NONE_REF_SKELETON = 3 UMETA(DisplayName = "Deform All But not Ref.Skeleton Bones")
+};
+
+
+USTRUCT()
+struct FMeshReshapeBoneReference
+{
+	GENERATED_USTRUCT_BODY()
+
+	/** Name of the bone that will be deformed */
+	UPROPERTY(EditAnywhere, Category = CustomizableObject)
+	FName BoneName;
+};
diff --git a/Engine/Plugins/Experimental/Mutable/Source/CustomizableObjectEditor/Public/MuCOE/Nodes/CustomizableObjectNodeMeshSwitch.h b/Engine/Plugins/Experimental/Mutable/Source/CustomizableObjectEditor/Public/MuCOE/Nodes/CustomizableObjectNodeMeshSwitch.h
new file mode 100644
--- /dev/null
+++ b/Engine/Plugins/Experimental/Mutable/Source/CustomizableObjectEditor/Public/MuCOE/Nodes/CustomizableObjectNodeMeshSwitch.h
@@ -0,0 +1,22 @@
+// Copyright Epic Games, Inc. All Rights Reserved.
+
+#pragma once
+
+#include "MuCOE/Nodes/CustomizableObjectNodeSwitchBase.h"
+
+#include "CustomizableObjectNodeMeshSwitch.generated.h"
+
+
+UCLASS()
+class CUSTOMIZABLEOBJECTEDITOR_API UCustomizableObjectNodeMeshSwitch : public UCustomizableObjectNodeSwitchBase
+{
+public:
+	GENERATED_BODY()
+
+	// UCustomizableObjectNode interface
+	virtual void BackwardsCompatibleFixup() override;
+	
+	// UCustomizableObjectNodeSwitchBase interface
+	virtual FName GetCategory() const override;
+};
+
diff --git a/Engine/Plugins/Experimental/Mutable/Source/CustomizableObjectEditor/Public/MuCOE/Nodes/CustomizableObjectNodeMeshVariation.h b/Engine/Plugins/Experimental/Mutable/Source/CustomizableObjectEditor/Public/MuCOE/Nodes/CustomizableObjectNodeMeshVariation.h
new file mode 100644
--- /dev/null
+++ b/Engine/Plugins/Experimental/Mutable/Source/CustomizableObjectEditor/Public/MuCOE/Nodes/CustomizableObjectNodeMeshVariation.h
@@ -0,0 +1,28 @@
+// Copyright Epic Games, Inc. All Rights Reserved.
+
+#pragma once
+
+#include "MuCOE/Nodes/CustomizableObjectNodeVariation.h"
+
+#include "CustomizableObjectNodeMeshVariation.generated.h"
+
+struct FCustomizableObjectMeshVariation;
+
+
+UCLASS()
+class CUSTOMIZABLEOBJECTEDITOR_API UCustomizableObjectNodeMeshVariation : public UCustomizableObjectNodeVariation
+{
+public:
+	GENERATED_BODY()
+
+	// Deprecated properties
+	UPROPERTY()
+	TArray<FCustomizableObjectMeshVariation> Variations_DEPRECATED;
+
+	// UCustomizableObjectNode interface
+	virtual void BackwardsCompatibleFixup() override;
+	
+	// UCustomizableObjectNodeVariation interface
+	virtual FName GetCategory() const override;
+};
+
diff --git a/Engine/Plugins/Experimental/Mutable/Source/CustomizableObjectEditor/Public/MuCOE/Nodes/CustomizableObjectNodeModifierBase.h b/Engine/Plugins/Experimental/Mutable/Source/CustomizableObjectEditor/Public/MuCOE/Nodes/CustomizableObjectNodeModifierBase.h
new file mode 100644
--- /dev/null
+++ b/Engine/Plugins/Experimental/Mutable/Source/CustomizableObjectEditor/Public/MuCOE/Nodes/CustomizableObjectNodeModifierBase.h
@@ -0,0 +1,21 @@
+// Copyright Epic Games, Inc. All Rights Reserved.
+
+#pragma once
+
+#include "MuCOE/Nodes/CustomizableObjectNodeMaterialBase.h"
+
+#include "CustomizableObjectNodeModifierBase.generated.h"
+
+class UObject;
+
+
+UCLASS(abstract)
+class CUSTOMIZABLEOBJECTEDITOR_API UCustomizableObjectNodeModifierBase : public UCustomizableObjectNodeMaterialBase
+{
+public:
+	GENERATED_BODY()
+
+	UCustomizableObjectNodeModifierBase();
+
+};
+
diff --git a/Engine/Plugins/Experimental/Mutable/Source/CustomizableObjectEditor/Public/MuCOE/Nodes/CustomizableObjectNodeMorphMaterial.h b/Engine/Plugins/Experimental/Mutable/Source/CustomizableObjectEditor/Public/MuCOE/Nodes/CustomizableObjectNodeMorphMaterial.h
new file mode 100644
--- /dev/null
+++ b/Engine/Plugins/Experimental/Mutable/Source/CustomizableObjectEditor/Public/MuCOE/Nodes/CustomizableObjectNodeMorphMaterial.h
@@ -0,0 +1,48 @@
+// Copyright Epic Games, Inc. All Rights Reserved.
+
+#pragma once
+
+#include "MuCOE/Nodes/CustomizableObjectNodeEditMaterialBase.h"
+
+#include "CustomizableObjectNodeMorphMaterial.generated.h"
+
+namespace ENodeTitleType { enum Type : int; }
+
+class UCustomizableObjectNodeRemapPins;
+class UEdGraphPin;
+class UObject;
+struct FPropertyChangedEvent;
+
+
+UCLASS()
+class CUSTOMIZABLEOBJECTEDITOR_API UCustomizableObjectNodeMorphMaterial : public UCustomizableObjectNodeEditMaterialBase
+{
+public:
+
+	GENERATED_BODY()
+	
+	// Begin EdGraphNode interface
+	FText GetNodeTitle(ENodeTitleType::Type TitleType) const override;
+	FLinearColor GetNodeTitleColor() const override;
+	FText GetTooltipText() const override;
+
+	// UCustomizableObjectNode interface
+	virtual void BackwardsCompatibleFixup() override;
+	void AllocateDefaultPins(UCustomizableObjectNodeRemapPins* RemapPins) override;
+	void PinConnectionListChanged(UEdGraphPin* Pin) override;
+	void PostEditChangeProperty(FPropertyChangedEvent& PropertyChangedEvent) override;
+
+	bool IsNodeOutDatedAndNeedsRefresh() override;
+	FString GetRefreshMessage() const override;
+
+	UPROPERTY(Category=CustomizableObject, EditAnywhere)
+	FString MorphTargetName;
+
+	bool IsSingleOutputNode() const override;
+
+	UEdGraphPin* FactorPin() const
+	{
+		return FindPin(TEXT("Factor"));
+	}
+};
+
diff --git a/Engine/Plugins/Experimental/Mutable/Source/CustomizableObjectEditor/Public/MuCOE/Nodes/CustomizableObjectNodeObject.h b/Engine/Plugins/Experimental/Mutable/Source/CustomizableObjectEditor/Public/MuCOE/Nodes/CustomizableObjectNodeObject.h
new file mode 100644
--- /dev/null
+++ b/Engine/Plugins/Experimental/Mutable/Source/CustomizableObjectEditor/Public/MuCOE/Nodes/CustomizableObjectNodeObject.h
@@ -0,0 +1,293 @@
+// Copyright Epic Games, Inc. All Rights Reserved.
+
+#pragma once
+
+#include "MuCOE/Nodes/CustomizableObjectNode.h"
+#include "MuR/System.h"
+
+#include "CustomizableObjectNodeObject.generated.h"
+
+namespace ENodeTitleType { enum Type : int; }
+
+class UCustomizableObject;
+class UCustomizableObjectNodeMaterial;
+class UCustomizableObjectNodeRemapPins;
+class UObject;
+struct FPropertyChangedEvent;
+struct FSoftObjectPath;
+
+USTRUCT()
+struct FBoneToRemove
+{
+	GENERATED_USTRUCT_BODY()
+
+	UPROPERTY(EditAnywhere, Category = CustomizableObject)
+	bool bOnlyRemoveChildren = false;
+
+	UPROPERTY(EditAnywhere, Category = CustomizableObject)
+	FName BoneName;
+};
+
+
+USTRUCT()
+struct FLODReductionSettings
+{
+	GENERATED_USTRUCT_BODY()
+
+	/** Selects which bones will be removed from the final skeleton
+	* BoneName: Name of the bone that will be removed. Its children will be removed too.
+	* Remove Only Children: If true, only the children of the selected bone will be removed. The selected bone will remain.
+	*/
+	UPROPERTY(EditAnywhere, Category = CustomizableObject)
+	TArray<FBoneToRemove> BonesToRemove;
+};
+
+
+USTRUCT()
+struct FComponentSettings
+{
+	GENERATED_USTRUCT_BODY()
+	
+	UPROPERTY(EditAnywhere, Category = CustomizableObject)
+	TArray<FLODReductionSettings> LODReductionSettings;
+};
+
+
+USTRUCT()
+struct CUSTOMIZABLEOBJECTEDITOR_API FCustomizableObjectState
+{
+	GENERATED_USTRUCT_BODY()
+
+	friend UCustomizableObjectNodeObject;
+	
+	UPROPERTY(EditAnywhere, Category=CustomizableObject)
+	FString Name;
+
+	UPROPERTY(EditAnywhere, Category=CustomizableObject)
+	TArray<FString> RuntimeParameters;
+
+	/** Special treatment of texture compression for this state. */
+	UPROPERTY(EditAnywhere, Category = CustomizableObject)
+	ETextureCompressionStrategy TextureCompressionStrategy = ETextureCompressionStrategy::None;
+
+	/** If this is enabled, texture streaming won't be used for this state, and full images will be generated when an instance is first updated. */
+	UPROPERTY(EditAnywhere, Category = CustomizableObject)
+	bool bDisableTextureStreaming = false;
+
+	/** LiveUpdateMode will reuse instance temp. data between updates and speed up update times, but spend much more memory. Good for customization screens, not for actual gameplay modes. */
+	UPROPERTY(EditAnywhere, Category = CustomizableObject)
+	bool bLiveUpdateMode = false;
+
+	// Enables the reuse of all possible textures when the instance is updated without any changes in geometry or state (the first update after creation doesn't reuse any)
+	// It will only work if the textures aren't compressed, so set the instance to a Mutable state with texture compression disabled
+	// WARNING! If texture reuse is enabled, do NOT keep external references to the textures of the instance. The instance owns the textures.
+	UPROPERTY(EditAnywhere, Category = CustomizableObject)
+	bool bReuseInstanceTextures = false;
+
+	UPROPERTY(EditAnywhere, Category = CustomizableObject)
+	bool bBuildOnlyFirstLOD = false;
+
+	/** If there's an entry for a specific platform, when compiling for that platform Num LODs will be built after bBuildOnlyFirstLOD */
+	UPROPERTY(EditAnywhere, Category = CustomizableObject)
+	TMap<FString, int32> NumExtraLODsToBuildPerPlatform;
+
+	UPROPERTY(EditAnywhere, Category = CustomizableObject)
+	TMap<FString, FString> ForcedParameterValues;
+
+	UPROPERTY(EditAnywhere, Category = UI, meta = (DisplayName = "State UI Metadata"))
+	FMutableStateUIMetadata UIMetadata;
+
+private:
+	// Deprecated
+	UPROPERTY()
+	FMutableParamUIMetadata StateUIMetadata_DEPRECATED;
+	
+	/** This is now TextureCompressionStrategy.  */
+	UPROPERTY()
+	bool bDontCompressRuntimeTextures_DEPRECATED = false;
+};
+
+
+UENUM()
+enum class ECustomizableObjectAutomaticLODStrategy : uint8
+{
+	// Use the same strategy than the parent object. If root, then use "Manual".
+	Inherited = 0 UMETA(DisplayName = "Inherit from parent object"),
+	// Don't try to generate LODs automatically for the child nodes. Only the ones tha explicitely define them will be used.
+	Manual = 1 UMETA(DisplayName = "Only manually created LODs"),
+	// Try to generate the same material structure than LOd 0 if the source meshes have LODs.
+	AutomaticFromMesh = 2 UMETA(DisplayName = "Automatic from mesh")
+};
+
+UENUM()
+enum class ECustomizableObjectSelectionOverride : uint8
+{
+    NoOverride = 0 UMETA(DisplayName = "No Override"),
+    Disable    = 1 UMETA(DisplayName = "Disable"    ),
+    Enable     = 2 UMETA(DisplayName = "Enable"     )
+};
+
+USTRUCT()
+struct FRealTimeMorphSelectionOverride
+{
+	GENERATED_BODY()
+
+    FRealTimeMorphSelectionOverride()
+    {
+    }
+
+    FRealTimeMorphSelectionOverride(const FName& InMorphName)
+        : MorphName(InMorphName)
+    {
+    }
+
+    UPROPERTY()
+    FName MorphName;
+
+    UPROPERTY()
+    ECustomizableObjectSelectionOverride SelectionOverride = 
+        ECustomizableObjectSelectionOverride::NoOverride;
+
+    UPROPERTY()
+    TArray<FName> SkeletalMeshesNames;
+
+    UPROPERTY()
+    TArray<ECustomizableObjectSelectionOverride> Override;
+};
+
+UCLASS()
+class CUSTOMIZABLEOBJECTEDITOR_API UCustomizableObjectNodeObject : public UCustomizableObjectNode
+{
+public:
+	GENERATED_BODY()
+
+	UCustomizableObjectNodeObject();
+
+	UPROPERTY(EditAnywhere, Category=CustomizableObject)
+	FString ObjectName;
+
+	UPROPERTY(EditAnywhere, Category = UI, meta = (DisplayName = "Parameter UI Metadata"))
+	FMutableParamUIMetadata ParamUIMetadata;
+
+	UPROPERTY(EditAnywhere, Category = CustomizableObject)
+	int32 NumLODs;
+
+	UPROPERTY(EditAnywhere, Category = CustomizableObject)
+	ECustomizableObjectAutomaticLODStrategy AutoLODStrategy = ECustomizableObjectAutomaticLODStrategy::AutomaticFromMesh;
+
+	UPROPERTY(EditAnywhere, Category = CustomizableObject, meta = (ClampMin = "1"))
+	int32 NumMeshComponents = 1;
+
+	UPROPERTY(EditAnywhere, Category=CustomizableObject ,meta = (TitleProperty = "Name"))
+	TArray<FCustomizableObjectState> States;
+
+	UPROPERTY(EditAnywhere, Category = AttachedToExternalObject)
+	TObjectPtr<UCustomizableObject> ParentObject;
+
+	UPROPERTY(EditAnywhere, Category = AttachedToExternalObject)
+	FGuid ParentObjectGroupId;
+
+	UPROPERTY()
+	FGuid Identifier;
+
+    // Soft references SkeletalMeshes found in the previous compilation.
+    // Only populated if the node is the root.
+    UPROPERTY()
+    TArray<FSoftObjectPath> ReferencedSkeletalMeshes;
+
+    // Information about the realtime morph targets usage. It indexes to ReferncedSkeletakMeshes array
+    // so it is need to keep them syncronized. 
+    // This overrides the per skeletal mesh node selection 
+    UPROPERTY()
+    TArray<FRealTimeMorphSelectionOverride> RealTimeMorphSelectionOverrides;
+
+	// Array of bones to remove from the mesh.All influences assigned to these bones will be transferred to the closest valid bone.
+	// Selected per component and LOD. Bones will be accumulated down the line.
+	UPROPERTY(EditAnywhere, Category = CustomizableObject)
+	TArray<FComponentSettings> ComponentSettings;
+	
+	// UObject interface.
+	void PostEditChangeProperty(FPropertyChangedEvent& PropertyChangedEvent) override;
+
+	// EdGraphNode interface
+	FText GetNodeTitle(ENodeTitleType::Type TitleType) const override;
+	FLinearColor GetNodeTitleColor() const override;
+	FText GetTooltipText() const override;
+	void PrepareForCopying() override;
+	void PostPasteNode() override;
+	void PostDuplicate(bool bDuplicateForPIE) override;
+
+	// UCustomizableObjectNode interface
+	virtual void BackwardsCompatibleFixup() override;
+	virtual void PostBackwardsCompatibleFixup() override;
+	void AllocateDefaultPins(UCustomizableObjectNodeRemapPins* RemapPins) override;
+
+	// Own Interface
+	void SetParentObject(UCustomizableObject* CustomizableParentObject);
+
+	// Own interface
+	UPROPERTY()
+	bool bIsBase;
+
+	UEdGraphPin* LODPin(int32 LODIndex) const
+	{
+		FString LODName = FString::Printf(TEXT("%s%d "), LODPinNamePrefix, LODIndex);
+		return FindPin(LODName);
+	}
+
+	int32 GetNumLODPins() const
+	{
+		int32 Count = 0;
+
+		for (UEdGraphPin* Pin : GetAllNonOrphanPins())
+		{
+			if (Pin->GetName().StartsWith(LODPinNamePrefix))
+			{
+				Count++;
+			}
+		}
+
+		return Count;
+	}
+
+	UEdGraphPin* ChildrenPin() const
+	{
+		return FindPin(ChildrenPinName);
+	}
+
+	UEdGraphPin* OutputPin() const
+	{
+		return FindPin(OutputPinName);
+	}
+
+	/** Return the LOD which a LOD pin references to. Retrun -1 if a pin does not belong to any LOD. */
+	int32 GetLOD(UEdGraphPin* Pin) const;
+
+	virtual bool CanUserDeleteNode() const override;
+	virtual bool CanDuplicateNode() const override;
+
+	/** Get all Material Nodes int his Customizable Object which belong to the given LOD.
+	 *
+	 * @param LOD LOD which materials have to belong to.
+	 */
+	TArray<UCustomizableObjectNodeMaterial*> GetMaterialNodes(int LOD) const;
+
+	bool IsSingleOutputNode() const override;
+
+	void SetMeshComponentNumFromParent(int32 MeshComponentNum)
+	{
+		NumMeshComponents = MeshComponentNum;
+	}
+
+	// Node Details Support
+	int32 CurrentComponent = 0;
+	int32 CurrentLOD = 0;
+
+private:
+	static const FName ChildrenPinName;
+	static const FName OutputPinName;
+	static const TCHAR* LODPinNamePrefix;
+
+	static bool IsBuiltInPin(FName PinName);
+};
+
diff --git a/Engine/Plugins/Experimental/Mutable/Source/CustomizableObjectEditor/Public/MuCOE/Nodes/CustomizableObjectNodeObjectChild.h b/Engine/Plugins/Experimental/Mutable/Source/CustomizableObjectEditor/Public/MuCOE/Nodes/CustomizableObjectNodeObjectChild.h
new file mode 100644
--- /dev/null
+++ b/Engine/Plugins/Experimental/Mutable/Source/CustomizableObjectEditor/Public/MuCOE/Nodes/CustomizableObjectNodeObjectChild.h
@@ -0,0 +1,28 @@
+// Copyright Epic Games, Inc. All Rights Reserved.
+
+#pragma once
+
+#include "MuCOE/Nodes/CustomizableObjectNodeObject.h"
+
+#include "CustomizableObjectNodeObjectChild.generated.h"
+
+class UObject;
+
+
+UCLASS()
+class CUSTOMIZABLEOBJECTEDITOR_API UCustomizableObjectNodeObjectChild : public UCustomizableObjectNodeObject
+{
+public:
+	GENERATED_BODY()
+
+	UCustomizableObjectNodeObjectChild();
+
+public:
+	// UEdGraphNode interface
+	FText GetNodeTitle(ENodeTitleType::Type TitleType) const override;
+	void PrepareForCopying() override;
+	FText GetTooltipText() const override;
+	bool CanUserDeleteNode() const override;
+	bool CanDuplicateNode() const override;
+};
+
diff --git a/Engine/Plugins/Experimental/Mutable/Source/CustomizableObjectEditor/Public/MuCOE/Nodes/CustomizableObjectNodeObjectGroup.h b/Engine/Plugins/Experimental/Mutable/Source/CustomizableObjectEditor/Public/MuCOE/Nodes/CustomizableObjectNodeObjectGroup.h
new file mode 100644
--- /dev/null
+++ b/Engine/Plugins/Experimental/Mutable/Source/CustomizableObjectEditor/Public/MuCOE/Nodes/CustomizableObjectNodeObjectGroup.h
@@ -0,0 +1,67 @@
+// Copyright Epic Games, Inc. All Rights Reserved.
+
+#pragma once
+
+#include "MuCOE/Nodes/CustomizableObjectNode.h"
+
+#include "CustomizableObjectNodeObjectGroup.generated.h"
+
+enum class ECustomizableObjectGroupType : uint8;
+namespace ENodeTitleType { enum Type : int; }
+
+class FArchive;
+class UCustomizableObjectNodeRemapPins;
+class UEdGraphPin;
+class UObject;
+struct FPropertyChangedEvent;
+
+
+UCLASS()
+class CUSTOMIZABLEOBJECTEDITOR_API UCustomizableObjectNodeObjectGroup : public UCustomizableObjectNode
+{
+public:
+	GENERATED_BODY()
+
+	UCustomizableObjectNodeObjectGroup();
+
+	UPROPERTY(EditAnywhere, Category=CustomizableObject)
+	FString GroupName;
+
+	UPROPERTY(EditAnywhere, Category=CustomizableObject)
+	ECustomizableObjectGroupType GroupType;
+	
+	UPROPERTY(EditAnywhere, Category = CustomizableObject)
+	FString DefaultValue;
+
+	UPROPERTY(EditAnywhere, Category = UI, meta = (DisplayName = "Parameter UI Metadata"))
+	FMutableParamUIMetadata ParamUIMetadata;
+
+	// The sockets defined in meshes deriving from this node will inherit this socket priority. When in the generated merged mesh there
+	// are clashes with socket names, the one with higher priority will be kept and the other discarded.
+	UPROPERTY(EditAnywhere, Category = MeshSockets)
+	int32 SocketPriority = 0;
+
+	// UObject interface.
+	void Serialize(FArchive& Ar) override;
+	void PostEditChangeProperty(FPropertyChangedEvent& PropertyChangedEvent) override;
+
+	// EdGraphNode interface
+	FText GetNodeTitle(ENodeTitleType::Type TitleType) const override;
+	FLinearColor GetNodeTitleColor() const override;
+	FText GetTooltipText() const override;
+
+	// UCustomizableObjectNode interface
+	void AllocateDefaultPins(UCustomizableObjectNodeRemapPins* RemapPins) override;
+	bool IsSingleOutputNode() const override;
+
+	// Own interface
+	UEdGraphPin* ObjectsPin() const;
+
+	UEdGraphPin* GroupProjectorsPin() const;
+
+	UEdGraphPin* GroupPin() const;
+
+private:
+	FString LastGroupName;
+};
+
diff --git a/Engine/Plugins/Experimental/Mutable/Source/CustomizableObjectEditor/Public/MuCOE/Nodes/CustomizableObjectNodeParentedMaterial.h b/Engine/Plugins/Experimental/Mutable/Source/CustomizableObjectEditor/Public/MuCOE/Nodes/CustomizableObjectNodeParentedMaterial.h
new file mode 100644
--- /dev/null
+++ b/Engine/Plugins/Experimental/Mutable/Source/CustomizableObjectEditor/Public/MuCOE/Nodes/CustomizableObjectNodeParentedMaterial.h
@@ -0,0 +1,31 @@
+// Copyright Epic Games, Inc. All Rights Reserved.
+
+#pragma once
+
+#include "Containers/Array.h"
+#include "MuCOE/Nodes/CustomizableObjectNodeParentedNode.h"
+
+class UCustomizableObjectNode;
+class UCustomizableObjectNodeMaterial;
+
+
+/** Node Material specialization of ICustomizableObjectNodeParentedNode. */
+class FCustomizableObjectNodeParentedMaterial : public ICustomizableObjectNodeParentedNode
+{
+public:
+	// Own interface
+	/** Return the parent material node. */
+	UCustomizableObjectNodeMaterial* GetParentMaterialNode() const;
+	
+	/** Return all possible parent material nodes of the node. */
+	TArray<UCustomizableObjectNodeMaterial*> GetPossibleParentMaterialNodes() const;
+
+	/** Returns the parent material node if there exist a path to it. */
+	UCustomizableObjectNodeMaterial* GetParentMaterialNodeIfPath() const;
+	
+protected:
+	/** Return the node which this interface belongs to. */
+	virtual UCustomizableObjectNode& GetNode() = 0;
+	const UCustomizableObjectNode& GetNode() const;
+};
+
diff --git a/Engine/Plugins/Experimental/Mutable/Source/CustomizableObjectEditor/Public/MuCOE/Nodes/CustomizableObjectNodeParentedNode.h b/Engine/Plugins/Experimental/Mutable/Source/CustomizableObjectEditor/Public/MuCOE/Nodes/CustomizableObjectNodeParentedNode.h
new file mode 100644
--- /dev/null
+++ b/Engine/Plugins/Experimental/Mutable/Source/CustomizableObjectEditor/Public/MuCOE/Nodes/CustomizableObjectNodeParentedNode.h
@@ -0,0 +1,35 @@
+// Copyright Epic Games, Inc. All Rights Reserved.
+
+#pragma once
+
+
+struct FGuid;
+
+class UCustomizableObject;
+class UCustomizableObjectNode;
+
+/** Interface for a node that saves a parent node. */
+class ICustomizableObjectNodeParentedNode
+{
+public:
+	virtual ~ICustomizableObjectNodeParentedNode() {};
+
+	/** Save the parent node reference. */
+	virtual void SetParentNode(UCustomizableObject* Object, FGuid NodeId);
+	
+	/** Return the saved parent material node. Can return nullptr. */
+	UCustomizableObjectNode* GetParentNode() const;
+
+protected:
+	/** Save the parent node reference. Do not call directly, use SetParentNode instead. */
+	virtual void SaveParentNode(UCustomizableObject* Object, FGuid NodeId) = 0;
+
+	virtual FGuid GetParentNodeId() const = 0;
+
+	virtual UCustomizableObject* GetParentObject() const = 0;
+};
+
+#if UE_ENABLE_INCLUDE_ORDER_DEPRECATED_IN_5_2
+#include "HAL/Platform.h"
+#include "Misc/Guid.h"
+#endif
diff --git a/Engine/Plugins/Experimental/Mutable/Source/CustomizableObjectEditor/Public/MuCOE/Nodes/CustomizableObjectNodeProjectorConstant.h b/Engine/Plugins/Experimental/Mutable/Source/CustomizableObjectEditor/Public/MuCOE/Nodes/CustomizableObjectNodeProjectorConstant.h
new file mode 100644
--- /dev/null
+++ b/Engine/Plugins/Experimental/Mutable/Source/CustomizableObjectEditor/Public/MuCOE/Nodes/CustomizableObjectNodeProjectorConstant.h
@@ -0,0 +1,92 @@
+// Copyright Epic Games, Inc. All Rights Reserved.
+
+#pragma once
+
+#include "MuCOE/Nodes/CustomizableObjectNode.h"
+
+#include "CustomizableObjectNodeProjectorConstant.generated.h"
+
+namespace ENodeTitleType { enum Type : int; }
+
+class UCustomizableObjectNodeRemapPins;
+class UEdGraphPin;
+class UObject;
+struct FPropertyChangedEvent;
+
+
+UCLASS()
+class CUSTOMIZABLEOBJECTEDITOR_API UCustomizableObjectNodeProjectorConstant : public UCustomizableObjectNode
+{
+public:
+	GENERATED_BODY()
+
+	UCustomizableObjectNodeProjectorConstant();
+
+	/**  */
+	UPROPERTY(EditAnywhere, Category = CustomizableObject, meta = (ShowOnlyInnerProperties))
+	FCustomizableObjectProjector Value;
+
+	UPROPERTY(EditAnywhere, Category = UI, meta = (DisplayName = "Parameter UI Metadata"))
+	FMutableParamUIMetadata ParamUIMetadata;
+
+	UPROPERTY(EditAnywhere, Category = CustomizableObject, Meta = (DisplayName = "Projection Angle (degrees)"))
+	float ProjectionAngle;
+
+	UPROPERTY(EditAnywhere, Category = ProjectorSnapToBone)
+	uint32 ReferenceSkeletonIndex = 0;
+
+	UPROPERTY(EditAnywhere, Category = ProjectorSnapToBone)
+	FName ProjectorBone;
+
+	/** Temporary variable where to put the location information for bone combo box selection changes (in FCustomizableObjectNodeProjectorParameterDetails)*/
+	FVector BoneComboBoxLocation;
+
+	/** Temporary variable where to put the direction information for bone combo box selection changes (in FCustomizableObjectNodeProjectorParameterDetails)*/
+	FVector BoneComboBoxForwardDirection;
+
+	/** Temporary variable where to put the direction information for bone combo box selection changes (in FCustomizableObjectNodeProjectorParameterDetails)*/
+	FVector BoneComboBoxUpDirection;
+
+	// UObject interface.
+	void PostEditChangeProperty(FPropertyChangedEvent& PropertyChangedEvent) override;
+
+	// Begin EdGraphNode interface
+	FText GetNodeTitle(ENodeTitleType::Type TitleType) const override;
+	FLinearColor GetNodeTitleColor() const override;
+	FText GetTooltipText() const override;
+
+	// UCustomizableObjectNode interface
+	void AllocateDefaultPins(UCustomizableObjectNodeRemapPins* RemapPins) override;
+	virtual void BackwardsCompatibleFixup() override;
+
+	// Own interface
+	UEdGraphPin* ValuePin();
+
+	ECustomizableObjectProjectorType GetProjectorType() const;
+
+	FVector GetProjectorPosition() const;
+
+	void SetProjectorPosition(const FVector& Position);
+
+	FVector GetProjectorDirection() const;
+
+	void SetProjectorDirection(const FVector& Direction);
+
+	FVector GetProjectorUp() const;
+
+	void SetProjectorUp(const FVector& Up);
+
+	FVector GetProjectorScale() const;
+
+	void SetProjectorScale(const FVector& Scale);
+
+	float GetProjectorAngle();
+
+	void SetProjectorAngle(float Angle);
+
+private:
+
+	UPROPERTY()
+	ECustomizableObjectProjectorType ProjectionType_DEPRECATED;
+};
+
diff --git a/Engine/Plugins/Experimental/Mutable/Source/CustomizableObjectEditor/Public/MuCOE/Nodes/CustomizableObjectNodeProjectorParameter.h b/Engine/Plugins/Experimental/Mutable/Source/CustomizableObjectEditor/Public/MuCOE/Nodes/CustomizableObjectNodeProjectorParameter.h
new file mode 100644
--- /dev/null
+++ b/Engine/Plugins/Experimental/Mutable/Source/CustomizableObjectEditor/Public/MuCOE/Nodes/CustomizableObjectNodeProjectorParameter.h
@@ -0,0 +1,98 @@
+// Copyright Epic Games, Inc. All Rights Reserved.
+
+#pragma once
+
+#include "MuCOE/Nodes/CustomizableObjectNode.h"
+
+#include "CustomizableObjectNodeProjectorParameter.generated.h"
+
+namespace ENodeTitleType { enum Type : int; }
+
+class UCustomizableObjectNodeRemapPins;
+class UEdGraphPin;
+class UObject;
+struct FPropertyChangedEvent;
+
+
+UCLASS()
+class CUSTOMIZABLEOBJECTEDITOR_API UCustomizableObjectNodeProjectorParameter : public UCustomizableObjectNode
+{
+public:
+	GENERATED_BODY()
+
+	UPROPERTY(EditAnywhere, Category=CustomizableObject, meta = (ShowOnlyInnerProperties))
+	FCustomizableObjectProjector DefaultValue;
+
+	UPROPERTY(EditAnywhere, Category=CustomizableObject)
+	FString ParameterName = "Default Name";
+
+	UPROPERTY(EditAnywhere, Category = UI, meta = (DisplayName = "Parameter UI Metadata"))
+	FMutableParamUIMetadata ParamUIMetadata;
+
+	UPROPERTY(EditAnywhere, Category = CustomizableObject, Meta = (DisplayName = "Projection Angle (degrees)"))
+	float ProjectionAngle = 360.0f;
+
+	UPROPERTY(EditAnywhere, Category = ProjectorSnapToBone)
+	uint32 ReferenceSkeletonIndex = 0;
+
+	UPROPERTY(EditAnywhere, Category = ProjectorSnapToBone)
+	FName ProjectorBone;
+
+	/** Temporary variable where to put the location information for bone combo box selection changes (in FCustomizableObjectNodeProjectorParameterDetails)*/
+	FVector BoneComboBoxLocation = FVector::ZeroVector;
+
+	/** Temporary variable where to put the direction information for bone combo box selection changes (in FCustomizableObjectNodeProjectorParameterDetails)*/
+	FVector BoneComboBoxForwardDirection = FVector::ZeroVector;
+
+	/** Temporary variable where to put the direction information for bone combo box selection changes (in FCustomizableObjectNodeProjectorParameterDetails)*/
+	FVector BoneComboBoxUpDirection = FVector::ZeroVector;
+
+	// UObject interface.
+	void PostEditChangeProperty(FPropertyChangedEvent& PropertyChangedEvent) override;
+
+	// Begin EdGraphNode interface
+	FText GetNodeTitle(ENodeTitleType::Type TitleType) const override;
+	FLinearColor GetNodeTitleColor() const override;
+	FText GetTooltipText() const override;
+	virtual bool GetCanRenameNode() const override { return true; }
+	virtual void OnRenameNode(const FString& NewName) override;
+
+	// UCustomizableObjectNode interface
+	void AllocateDefaultPins(UCustomizableObjectNodeRemapPins* RemapPins) override;
+	virtual bool IsAffectedByLOD() const override { return false; }
+	virtual void BackwardsCompatibleFixup() override;
+
+	// Own interface
+	UEdGraphPin* ProjectorPin() const
+	{
+		return FindPin(TEXT("Value"));
+	}
+
+	ECustomizableObjectProjectorType GetProjectorType() const;
+
+	FVector GetProjectorDefaultPosition() const;
+	
+	void SetProjectorDefaultPosition(const FVector& Position);
+
+	FVector GetProjectorDefaultDirection() const;
+
+	void SetProjectorDefaultDirection(const FVector& Direction);
+
+	FVector GetProjectorDefaultUp() const;
+
+	void SetProjectorDefaultUp(const FVector& Up);
+
+	FVector GetProjectorDefaultScale() const;
+
+	void SetProjectorDefaultScale(const FVector& Scale);
+
+	float GetProjectorDefaultAngle() const;
+
+	void SetProjectorDefaultAngle(float Angle);
+
+private:
+
+	UPROPERTY()
+	ECustomizableObjectProjectorType ProjectionType_DEPRECATED;
+};
+
diff --git a/Engine/Plugins/Experimental/Mutable/Source/CustomizableObjectEditor/Public/MuCOE/Nodes/CustomizableObjectNodeRemoveMesh.h b/Engine/Plugins/Experimental/Mutable/Source/CustomizableObjectEditor/Public/MuCOE/Nodes/CustomizableObjectNodeRemoveMesh.h
new file mode 100644
--- /dev/null
+++ b/Engine/Plugins/Experimental/Mutable/Source/CustomizableObjectEditor/Public/MuCOE/Nodes/CustomizableObjectNodeRemoveMesh.h
@@ -0,0 +1,37 @@
+// Copyright Epic Games, Inc. All Rights Reserved.
+
+#pragma once
+
+#include "MuCOE/Nodes/CustomizableObjectNodeEditMaterialBase.h"
+
+#include "CustomizableObjectNodeRemoveMesh.generated.h"
+
+namespace ENodeTitleType { enum Type : int; }
+
+class UCustomizableObjectNodeRemapPins;
+class UEdGraphPin;
+class UObject;
+
+UCLASS()
+class CUSTOMIZABLEOBJECTEDITOR_API UCustomizableObjectNodeRemoveMesh : public UCustomizableObjectNodeEditMaterialBase
+{
+public:
+	GENERATED_BODY()
+
+	// Begin EdGraphNode interface
+	FText GetNodeTitle(ENodeTitleType::Type TitleType) const override;
+	FLinearColor GetNodeTitleColor() const override;
+	FText GetTooltipText() const override;
+
+	// UCustomizableObjectNode interface
+	void AllocateDefaultPins(UCustomizableObjectNodeRemapPins* RemapPins) override;
+	void PinConnectionListChanged(UEdGraphPin * Pin);
+
+	UEdGraphPin* RemoveMeshPin() const
+	{
+		return FindPin(TEXT("Remove Mesh"));
+	}
+
+	bool IsSingleOutputNode() const override;
+};
+
diff --git a/Engine/Plugins/Experimental/Mutable/Source/CustomizableObjectEditor/Public/MuCOE/Nodes/CustomizableObjectNodeRemoveMeshBlocks.h b/Engine/Plugins/Experimental/Mutable/Source/CustomizableObjectEditor/Public/MuCOE/Nodes/CustomizableObjectNodeRemoveMeshBlocks.h
new file mode 100644
--- /dev/null
+++ b/Engine/Plugins/Experimental/Mutable/Source/CustomizableObjectEditor/Public/MuCOE/Nodes/CustomizableObjectNodeRemoveMeshBlocks.h
@@ -0,0 +1,43 @@
+// Copyright Epic Games, Inc. All Rights Reserved.
+
+#pragma once
+
+#include "MuCOE/Nodes/CustomizableObjectNodeEditLayoutBlocks.h"
+
+#include "CustomizableObjectNodeRemoveMeshBlocks.generated.h"
+
+namespace ENodeTitleType { enum Type : int; }
+
+class FArchive;
+class UCustomizableObjectNodeRemapPins;
+class UEdGraphPin;
+class UObject;
+struct FPropertyChangedEvent;
+
+UCLASS()
+class CUSTOMIZABLEOBJECTEDITOR_API UCustomizableObjectNodeRemoveMeshBlocks : public UCustomizableObjectNodeEditLayoutBlocks
+{
+public:
+	GENERATED_BODY()
+
+	UPROPERTY()
+	TArray<int32> Blocks_DEPRECATED;
+
+	// UObject interface.
+	void PostEditChangeProperty(FPropertyChangedEvent& PropertyChangedEvent) override;
+	virtual void Serialize(FArchive& Ar) override;
+
+	// Begin EdGraphNode interface
+	FText GetNodeTitle(ENodeTitleType::Type TitleType) const override;
+	FLinearColor GetNodeTitleColor() const override;
+	FText GetTooltipText() const override;
+
+	// UCustomizableObjectNode interface
+	virtual void BackwardsCompatibleFixup() override;
+	void AllocateDefaultPins(UCustomizableObjectNodeRemapPins* RemapPins) override;
+	bool IsNodeOutDatedAndNeedsRefresh() override;
+	FString GetRefreshMessage() const override;
+	void PinConnectionListChanged(UEdGraphPin * Pin);
+	bool IsSingleOutputNode() const override;
+};
+
diff --git a/Engine/Plugins/Experimental/Mutable/Source/CustomizableObjectEditor/Public/MuCOE/Nodes/CustomizableObjectNodeStaticMesh.h b/Engine/Plugins/Experimental/Mutable/Source/CustomizableObjectEditor/Public/MuCOE/Nodes/CustomizableObjectNodeStaticMesh.h
new file mode 100644
--- /dev/null
+++ b/Engine/Plugins/Experimental/Mutable/Source/CustomizableObjectEditor/Public/MuCOE/Nodes/CustomizableObjectNodeStaticMesh.h
@@ -0,0 +1,173 @@
+// Copyright Epic Games, Inc. All Rights Reserved.
+
+#pragma once
+
+#include "MuCOE/Nodes/CustomizableObjectNodeMesh.h"
+#include "SGraphNode.h"
+
+#include "CustomizableObjectNodeStaticMesh.generated.h"
+
+namespace ENodeTitleType { enum Type : int; }
+
+class FArchive;
+class FAssetThumbnail;
+class FAssetThumbnailPool;
+class SOverlay;
+class SVerticalBox;
+class UCustomizableObjectLayout;
+class UCustomizableObjectNodeRemapPins;
+class UCustomizableObjectNodeRemapPinsByName;
+class UMaterialInterface;
+class UObject;
+class UStaticMesh;
+class UTexture2D;
+struct FPropertyChangedEvent;
+struct FSlateBrush;
+
+
+// Class to render the Static Mesh thumbnail of a CustomizableObjectNodeStaticMesh
+class SGraphNodeStaticMesh : public SGraphNode
+{
+public:
+
+	SLATE_BEGIN_ARGS(SGraphNodeStaticMesh) {}
+	SLATE_END_ARGS();
+
+	SGraphNodeStaticMesh() : SGraphNode() {};
+
+	// Builds the SGraphNodeStaticMesh when needed
+	void Construct(const FArguments& InArgs, UEdGraphNode* InGraphNode);
+
+	// Calls the needed functions to build the SGraphNode widgets
+	void UpdateGraphNode();
+
+	// Overriden functions to build the SGraphNode widgets
+	virtual void SetDefaultTitleAreaWidget(TSharedRef<SOverlay> DefaultTitleAreaWidget) override;
+	virtual void CreateBelowPinControls(TSharedPtr<SVerticalBox> MainBox) override;
+	virtual bool ShouldAllowCulling() const override { return false; }
+
+	// Callbacks for the widget
+	void OnExpressionPreviewChanged(const ECheckBoxState NewCheckedState);
+	ECheckBoxState IsExpressionPreviewChecked() const;
+	const FSlateBrush* GetExpressionPreviewArrow() const;
+	EVisibility ExpressionPreviewVisibility() const;
+
+public:
+
+	// Single property that only draws the combo box widget of the static mesh
+	TSharedPtr<class ISinglePropertyView> StaticMeshSelector;
+
+	// Pointer to the NodeStaticMesh that owns this SGraphNode
+	class UCustomizableObjectNodeStaticMesh* NodeStaticMesh;
+
+private:
+
+	// Classes needed to get and render the thumbnail of the Static Mesh
+	TSharedPtr<FAssetThumbnailPool> AssetThumbnailPool;
+	TSharedPtr<FAssetThumbnail> AssetThumbnail;
+
+	//This parameter defines the size of the thumbnail widget inside the Node
+	float WidgetSize;
+
+	// This parameter defines the resolution of the thumbnail
+	uint32 ThumbnailSize;
+
+};
+
+
+USTRUCT()
+struct CUSTOMIZABLEOBJECTEDITOR_API FCustomizableObjectNodeStaticMeshMaterial
+{
+	GENERATED_USTRUCT_BODY()
+
+	UPROPERTY()
+	FString Name;
+
+	UPROPERTY()
+	TObjectPtr<UEdGraphPin_Deprecated> MeshPin_DEPRECATED;
+
+	UPROPERTY()
+	TObjectPtr<UEdGraphPin_Deprecated> LayoutPin_DEPRECATED;
+
+	UPROPERTY()
+	TArray< TObjectPtr<UEdGraphPin_Deprecated>> ImagePins_DEPRECATED;
+
+	UPROPERTY()
+	FEdGraphPinReference MeshPinRef;
+
+	UPROPERTY()
+	FEdGraphPinReference LayoutPinRef;
+
+	UPROPERTY()
+	TArray<FEdGraphPinReference> ImagePinsRef;
+};
+
+
+USTRUCT()
+struct CUSTOMIZABLEOBJECTEDITOR_API FCustomizableObjectNodeStaticMeshLOD
+{
+	GENERATED_USTRUCT_BODY()
+
+	UPROPERTY()
+	TArray<FCustomizableObjectNodeStaticMeshMaterial> Materials;
+};
+
+
+UCLASS()
+class CUSTOMIZABLEOBJECTEDITOR_API UCustomizableObjectNodeStaticMesh : public UCustomizableObjectNodeMesh
+{ 
+public:
+	GENERATED_BODY()
+
+	UPROPERTY(EditAnywhere, Category = CustomizableObject)
+	TObjectPtr<UStaticMesh> StaticMesh;
+
+	/** Images */
+	UPROPERTY()
+	TArray<FCustomizableObjectNodeStaticMeshLOD> LODs;
+
+	/** Default pin when there is no mesh. */
+	UPROPERTY()
+	FEdGraphPinReference DefaultPin;
+	
+	// UObject interface.
+	void PostEditChangeProperty(FPropertyChangedEvent& PropertyChangedEvent) override;
+	virtual void Serialize(FArchive& Ar) override;
+
+	// UEdGraphNode interface
+	FText GetNodeTitle(ENodeTitleType::Type TitleType) const override;
+	FLinearColor GetNodeTitleColor() const override;
+	FText GetTooltipText() const override;
+
+	// UCustomizableObjectNode interface
+	void AllocateDefaultPins(UCustomizableObjectNodeRemapPins* RemapPins) override;
+	bool ProvidesCustomPinRelevancyTest() const override { return true; }
+	bool IsPinRelevant(const UEdGraphPin* Pin) const override;
+	UCustomizableObjectNodeRemapPinsByName* CreateRemapPinsByName() const override;
+
+	bool IsNodeOutDatedAndNeedsRefresh() override;
+	FString GetRefreshMessage() const override;
+
+	// UCustomizableObjectNodeMesh interface
+	virtual UTexture2D* FindTextureForPin(const UEdGraphPin* Pin) const override;
+	virtual void GetUVChannelForPin(const UEdGraphPin* Pin, TArray<FVector2f>& OutSegments, int32 UVIndex) const override;
+	virtual TArray<UCustomizableObjectLayout*> GetLayouts(const UEdGraphPin& OutPin) const override;
+	virtual UObject* GetMesh() const override;
+	virtual UEdGraphPin* GetMeshPin(int32 LOD, int32 SectionIndex) const override;
+	virtual UEdGraphPin* GetLayoutPin(int32 LODIndex, int32 SectionIndex, int32 LayoutIndex) const override;
+	virtual void GetPinSection(const UEdGraphPin& Pin, int32& OutLODIndex, int32& OutSectionIndex, int32& OutLayoutIndex) const override;
+
+	/** Returns the material assossiated to the given output pin. */
+	UMaterialInterface* GetMaterialFor(const UEdGraphPin* Pin) const;
+
+	// Creates the SGraph Node widget for the thumbnail
+	TSharedPtr<SGraphNode> CreateVisualWidget() override;
+
+	// Determines if the Node is collapsed or not
+	bool bCollapsed = true;
+
+	// Pointer to the SGraphNodeStaticMesh
+	TWeakPtr< SGraphNodeStaticMesh > GraphNodeStaticMesh;
+
+};
+
diff --git a/Engine/Plugins/Experimental/Mutable/Source/CustomizableObjectEditor/Public/MuCOE/Nodes/CustomizableObjectNodeSwitchBase.h b/Engine/Plugins/Experimental/Mutable/Source/CustomizableObjectEditor/Public/MuCOE/Nodes/CustomizableObjectNodeSwitchBase.h
new file mode 100644
--- /dev/null
+++ b/Engine/Plugins/Experimental/Mutable/Source/CustomizableObjectEditor/Public/MuCOE/Nodes/CustomizableObjectNodeSwitchBase.h
@@ -0,0 +1,88 @@
+// Copyright Epic Games, Inc. All Rights Reserved.
+
+#pragma once
+
+#include "MuCOE/Nodes/CustomizableObjectNode.h"
+
+#include "CustomizableObjectNodeSwitchBase.generated.h"
+
+namespace ENodeTitleType { enum Type : int; }
+
+class FArchive;
+class UCustomizableObjectNodeRemapPins;
+class UObject;
+struct FFrame;
+
+
+UCLASS(Abstract)
+class CUSTOMIZABLEOBJECTEDITOR_API UCustomizableObjectNodeSwitchBase : public UCustomizableObjectNode
+{
+public:
+	GENERATED_BODY()
+
+	// UObject interface
+	virtual void Serialize(FArchive& Ar) override;
+	
+	// Begin EdGraphNode interface
+	virtual FText GetNodeTitle(ENodeTitleType::Type TitleType) const override;
+	virtual FLinearColor GetNodeTitleColor() const override;
+	virtual FText GetTooltipText() const override;
+	virtual void PinConnectionListChanged(UEdGraphPin* Pin) override;
+	virtual void PostPasteNode() override;
+
+	// UCustomizableObjectNode interface
+	virtual void PostBackwardsCompatibleFixup() override;
+	virtual void AllocateDefaultPins(UCustomizableObjectNodeRemapPins* RemapPins) override;
+	virtual void ReconstructNode(UCustomizableObjectNodeRemapPins* RemapPins) override;
+
+	/** Get the output pin category. Override. */
+	virtual FName GetCategory() const PURE_VIRTUAL(UCustomizableObjectNodeSwitchBase::GetCategory, return {}; );
+
+	UEdGraphPin* OutputPin() const;
+
+	UEdGraphPin* SwitchParameter() const;
+
+	UEdGraphPin* GetElementPin(int32 Index) const
+	{
+		return FindPin(GetPinPrefix(Index));
+	}
+
+	int32 GetNumElements() const;
+
+	/** Links the PostEditChangeProperty delegate */
+	void LinkPostEditChangePropertyDelegate(const UEdGraphPin& Pin);
+
+	/** Get the ouput pin name. Override. */
+	virtual FString GetOutputPinName() const;
+
+private:
+	/** Get the pin prefix. Used for retrocompatibility. Override. */
+	virtual FString GetPinPrefix() const;
+
+protected:
+	UPROPERTY()
+	FEdGraphPinReference OutputPinReference;
+
+private:
+	/** Get the pin prefix with index. Used for retrocompatibility.*/
+	FString GetPinPrefix(int32 Index) const;
+
+	void ReloadEnumParam();
+
+	/** Last NodeEnumParameter connected. Used to remove the callback once disconnected. */
+	TWeakObjectPtr<UCustomizableObjectNode> LastNodeEnumParameterConnected;
+
+	/** NodeEnumParameter property changed callback function. Reconstructs the node. */
+	UFUNCTION()
+	void EnumParameterPostEditChangeProperty(FPostEditChangePropertyDelegateParameters& Parameters);
+
+	/** The node has to be reconstructed. PinConnectionListChanged(...) can not reconstruct the node, flag used to reconstruct the node on NodeConnectionListChanged(). */
+	bool bMarkReconstruct = false;
+
+	UPROPERTY()
+	TArray<FString> ReloadingElementsNames;
+
+	UPROPERTY()
+	FEdGraphPinReference SwitchParameterPinReference;
+};
+
diff --git a/Engine/Plugins/Experimental/Mutable/Source/CustomizableObjectEditor/Public/MuCOE/Nodes/CustomizableObjectNodeTable.h b/Engine/Plugins/Experimental/Mutable/Source/CustomizableObjectEditor/Public/MuCOE/Nodes/CustomizableObjectNodeTable.h
new file mode 100644
--- /dev/null
+++ b/Engine/Plugins/Experimental/Mutable/Source/CustomizableObjectEditor/Public/MuCOE/Nodes/CustomizableObjectNodeTable.h
@@ -0,0 +1,424 @@
+// Copyright Epic Games, Inc. All Rights Reserved.
+
+#pragma once
+
+#include "Engine/DataTable.h"
+#include "Engine/Texture2DArray.h"
+#include "MuCOE/Nodes/CustomizableObjectNode.h"
+#include "MuCOE/RemapPins/CustomizableObjectNodeRemapPinsByName.h"
+
+#include "CustomizableObjectNodeTable.generated.h"
+
+namespace ENodeTitleType { enum Type : int; }
+
+class UCustomizableObjectLayout;
+class UCustomizableObjectNodeRemapPins;
+class UCustomizableObjectNodeRemapPinsByName;
+class UEdGraphPin;
+class UObject;
+class USkeletalMesh;
+class UTexture2D;
+class UTexture2DArray;
+class UAnimInstance;
+struct FGuid;
+struct FSkeletalMaterial;
+
+
+/** Enum class for the different types of image pins */
+UENUM()
+enum class ETableTextureType : uint8
+{
+	PASSTHROUGH_TEXTURE = 0 UMETA(DisplayName = "Passthrough"),
+	MUTABLE_TEXTURE = 1 UMETA(DisplayName = "Mutable")
+};
+
+
+/** Enum class for the different types of pin meshes */
+enum class ETableMeshPinType : uint8
+{
+	NONE = 0,
+	SKELETAL_MESH = 1,
+	STATIC_MESH = 2
+};
+
+
+/** Enum to decide where the data comes from: Struct or Data Table*/
+UENUM()
+enum class ETableDataGatheringSource : uint8
+{
+	/** Gathers the information from a data table */
+	ETDGM_DataTable = 0 UMETA(DisplayName = "Data Table"),
+
+	/** When compiling the CO, it uses the asset registry to gather and generate a data table. It uses all the data tables found in the specified paths that are references of the selected structure.*/
+	ETDGM_AssetRegistry = 1 UMETA(DisplayName = " Struct + Asset Registry")
+};
+
+
+USTRUCT()
+struct FTableNodeColumnData
+{
+	GENERATED_USTRUCT_BODY()
+
+	/** Anim Instance Column name related to this Mesh pin */
+	UPROPERTY()
+	FString AnimInstanceColumnName = "";
+
+	/** Anim Slot Column name related to this Mesh pin */
+	UPROPERTY()
+	FString AnimSlotColumnName = "";
+
+	/** Anim Tag Column name related to this Mesh pin */
+	UPROPERTY()
+	FString AnimTagColumnName = "";
+};
+
+
+/** Base class for all Table Pins. */
+UCLASS()
+class CUSTOMIZABLEOBJECTEDITOR_API UCustomizableObjectNodeTableObjectPinData : public UCustomizableObjectNodePinData
+{
+	GENERATED_BODY()
+
+public:
+
+	/** Id of the property associated to a struct column */
+	UPROPERTY()
+	FGuid StructColumnId;
+
+	/** Name of the data table column related to the pin */
+	UPROPERTY()
+	FString ColumnName = "";
+
+};
+
+
+/** Additional data for Image pins. */
+UCLASS()
+class CUSTOMIZABLEOBJECTEDITOR_API UCustomizableObjectNodeTableImagePinData : public UCustomizableObjectNodeTableObjectPinData
+{
+	GENERATED_BODY()
+
+public:
+
+	bool IsDefaultImageMode() { return bIsDefault; }
+	void SetDefaultImageMode(bool bValue) { bIsDefault = bValue; }
+
+	bool IsNotTexture2D() { return bIsNotTexture2D; }
+	void SetIsNotTexture2D(bool bValue) { bIsNotTexture2D = bValue; }
+
+	void ConvertArrayTextureToAnyTextureFixup();
+
+	// Pin Type
+	UPROPERTY()
+	ETableTextureType ImageMode = ETableTextureType::MUTABLE_TEXTURE;
+
+private:
+
+	UPROPERTY()
+	bool bIsDefault = true;
+
+	// Replaced by the more general bIsNot2DTexture
+	UPROPERTY()
+	bool bIsArrayTexture_DEPRECATED = false;
+
+	UPROPERTY()
+	bool bIsNotTexture2D = false;
+};
+
+
+/** Additional data for Mesh pins. */
+UCLASS()
+class CUSTOMIZABLEOBJECTEDITOR_API UCustomizableObjectNodeTableMeshPinData : public UCustomizableObjectNodeTableObjectPinData
+{
+	GENERATED_BODY()
+
+public:
+
+	/** Anim Instance Column name related to this Mesh pin */
+	UPROPERTY()
+	FString AnimInstanceColumnName_DEPRECATED = "";
+
+	/** Anim Slot Column name related to this Mesh pin */
+	UPROPERTY()
+	FString AnimSlotColumnName_DEPRECATED = "";
+
+	/** Anim Tag Column name related to this Mesh pin */
+	UPROPERTY()
+	FString AnimTagColumnName_DEPRECATED = "";
+
+	UPROPERTY()
+	FString MutableColumnName = "";
+
+	/** LOD of the mesh related to this Mesh pin */
+	UPROPERTY()
+	int32 LOD = 0;
+
+	/** Section Index (Surface Index) of the mesh related to this Mesh pin */
+	UPROPERTY()
+	int32 Material = 0;
+
+	/** Layouts related to this Mesh pin */
+	UPROPERTY()
+	TArray< TObjectPtr<UCustomizableObjectLayout> > Layouts;
+};
+
+
+UCLASS()
+class UCustomizableObjectNodeTableRemapPins : public UCustomizableObjectNodeRemapPinsByName
+{
+	GENERATED_BODY()
+public:
+
+	// Specific method to decide when two pins are equal
+	virtual bool Equal(const UCustomizableObjectNode& Node, const UEdGraphPin& OldPin, const UEdGraphPin& NewPin) const override;
+
+	// Method to use in the RemapPins step of the node reconstruction process
+	virtual void RemapPins(const UCustomizableObjectNode& Node, const TArray<UEdGraphPin*>& OldPins, const TArray<UEdGraphPin*>& NewPins, TMap<UEdGraphPin*, UEdGraphPin*>& PinsToRemap, TArray<UEdGraphPin*>& PinsToOrphan) override;
+};
+
+
+UCLASS()
+class CUSTOMIZABLEOBJECTEDITOR_API UCustomizableObjectNodeTable : public UCustomizableObjectNode
+{
+public:
+	
+	GENERATED_BODY()
+
+	// Name of the property parameter
+	UPROPERTY(EditAnywhere, Category = TableProperties)
+	FString ParameterName = "Default Name";
+
+	/** If true, adds a "None" parameter option */
+	UPROPERTY(EditAnywhere, Category = TableProperties)
+	bool bAddNoneOption;
+
+	/** If there is a bool column in the table, checked rows will not be compiled */
+	UPROPERTY(EditAnywhere, Category = TableProperties)
+	bool bDisableCheckedRows = true;
+
+	/** Source where table gathers the data */
+	UPROPERTY(EditAnywhere, Category = TableProperties)
+	ETableDataGatheringSource TableDataGatheringMode = ETableDataGatheringSource::ETDGM_DataTable;
+
+	// Pointer to the Data Table Asset represented in this node
+	UPROPERTY(EditAnywhere, Category = TableProperties, meta = (DontUpdateWhileEditing, EditCondition = "TableDataGatheringMode == ETableDataGatheringSource::ETDGM_DataTable", EditConditionHides))
+	TObjectPtr<UDataTable> Table = nullptr;
+
+	// Pointer to the Struct Asset represented in this node
+	UPROPERTY(EditAnywhere, Category = TableProperties, meta = (DontUpdateWhileEditing, EditCondition = "TableDataGatheringMode == ETableDataGatheringSource::ETDGM_AssetRegistry", EditConditionHides))
+	TObjectPtr<UScriptStruct> Structure = nullptr;
+
+	UPROPERTY(EditAnywhere, Category = TableProperties, meta = (EditCondition = "TableDataGatheringMode == ETableDataGatheringSource::ETDGM_AssetRegistry", EditConditionHides))
+	TArray<FName> FilterPaths;
+
+	/** Decides the default type of the texture pins (passtrhough or mutable)
+	*   Right click on a non-linked image pin to customize its image mode
+	*/
+	UPROPERTY(EditAnywhere, Category = TableProperties)
+	ETableTextureType DefaultImageMode = ETableTextureType::MUTABLE_TEXTURE;
+
+	/** Name of the column that contains the Version options. */
+	UPROPERTY(EditAnywhere, Category = TableProperties)
+	FName VersionColumn;
+
+	/** Name of the row that will be used as default value. */
+	UPROPERTY(EditAnywhere, Category = TableProperties)
+	FName DefaultRowName;
+
+	UPROPERTY(EditAnywhere, Category = UI, meta = (DisplayName = "Parameter UI Metadata"))
+	FMutableParamUIMetadata ParamUIMetadata;
+
+	/** Name of the column that contains the MutableUIMetadata of the row options. */
+	UPROPERTY(EditAnywhere, Category = UI)
+	FName ParamUIMetadataColumn;
+
+	/** Map to relate a Structure Column with its Data */
+	UPROPERTY()
+	TMap<FGuid, FTableNodeColumnData> ColumnDataMap;
+
+public:
+
+	// UObject interface
+	virtual void PreEditChange(FProperty* PropertyAboutToChange) override;
+	virtual void PostEditChangeChainProperty(FPropertyChangedChainEvent& PropertyChangedEvent) override;
+	virtual void PostEditChangeProperty(FPropertyChangedEvent& PropertyChangedEvent) override;
+	
+	// EdGraphNode interface
+	FText GetNodeTitle(ENodeTitleType::Type TitleType) const override;
+	FLinearColor GetNodeTitleColor() const override;
+	FText GetTooltipText() const override;
+	virtual void PinConnectionListChanged(UEdGraphPin* Pin) override;
+	virtual void OnRenameNode(const FString& NewName) override;
+	virtual bool GetCanRenameNode() const override { return true; }
+	
+	// UCustomizableObjectNode interface
+	virtual void PostBackwardsCompatibleFixup() override;
+	virtual void BackwardsCompatibleFixup() override;
+	void AllocateDefaultPins(UCustomizableObjectNodeRemapPins* RemapPins) override;
+	bool IsNodeOutDatedAndNeedsRefresh() override;
+	FString GetRefreshMessage() const override;
+	virtual bool ProvidesCustomPinRelevancyTest() const override;
+	virtual bool IsPinRelevant(const UEdGraphPin* Pin) const override;
+	UCustomizableObjectNodeTableRemapPins* CreateRemapPinsDefault() const;
+
+	/*** Allows to perform work when remapping the pin data. */
+	virtual void RemapPinsData(const TMap<UEdGraphPin*, UEdGraphPin*>& PinsToRemap) override;
+	
+	// Own interface	
+	// Returns the reference Texture parameter from a Material
+	UTexture2D* FindReferenceTextureParameter(const UEdGraphPin* Pin, FString ParameterImageName) const;
+
+	// Methods to get the UVs of the skeletal mesh 
+	void GetUVChannel(const UCustomizableObjectLayout* CurrentLayout, TArray<FVector2f>& OutSegments) const;
+	void GetUVChannelForPin(const UEdGraphPin* Pin, TArray<FVector2f>& OutSegments, int32 UVChannel = 0) const;
+
+	// Methods to provide the Layouts to the Layout block editors
+	void SetLayoutInLayoutEditor(UCustomizableObjectLayout* CurrentLayout);
+	TArray<UCustomizableObjectLayout*> GetLayouts(const UEdGraphPin* Pin) const;
+
+	// Returns the name of the table column related to a pin
+	FString GetColumnNameByPin(const UEdGraphPin* Pin) const;
+
+	// Returns the LOD of the mesh associated to the input pin
+	void GetPinLODAndSection(const UEdGraphPin* Pin, int32& LODIndex, int32& SectionIndex) const;
+
+	// Get the anim blueprint and anim slot columns related to a mesh
+	void GetAnimationColumns(const FGuid& ColumnId, FString& AnimBPColumnName, FString& AnimSlotColumnName, FString& AnimTagColumnName) const;
+
+	/** Callback called when the Table property or its contents has changed. */
+	void OnTableChanged();
+	
+	template<class T> 
+	T* GetColumnDefaultAssetByType(FString ColumnName) const
+	{
+		T* ObjectToReturn = nullptr;
+		const UScriptStruct* TableStruct = nullptr;
+
+		if (TableDataGatheringMode == ETableDataGatheringSource::ETDGM_AssetRegistry)
+		{
+			TableStruct = Structure;
+		}
+		else
+		{
+			if (Table)
+			{
+				TableStruct = Table->GetRowStruct();
+			}
+		}
+
+		if (TableStruct)
+		{
+			// Getting Default Struct Values
+			// A Script Struct always has at leaset one property
+			TArray<int8> DeafaultDataArray;
+			DeafaultDataArray.SetNumZeroed(TableStruct->GetStructureSize());
+			TableStruct->InitializeStruct(DeafaultDataArray.GetData());
+
+			FProperty* Property = FindTableProperty(TableStruct, FName(*ColumnName));
+
+			if (Property)
+			{
+				if (const FSoftObjectProperty* SoftObjectProperty = CastField<FSoftObjectProperty>(Property))
+				{
+					UObject* Object = nullptr;
+					
+					// Getting default UObject
+					uint8* CellData = SoftObjectProperty->ContainerPtrToValuePtr<uint8>(DeafaultDataArray.GetData());
+
+					if (CellData)
+					{
+						Object = SoftObjectProperty->GetPropertyValue(CellData).LoadSynchronous();
+					}
+
+					if (Object)
+					{
+						if (Object->IsA(T::StaticClass()))
+						{
+							ObjectToReturn = Cast <T>(Object);
+						}
+					}
+				}
+			}
+
+			// Cleaning Default Structure Pointer
+			TableStruct->DestroyStruct(DeafaultDataArray.GetData());
+		}
+
+		return ObjectToReturn;
+	}
+
+	template<class T>
+	T* GetColumnDefaultAssetByType(const UEdGraphPin* Pin) const
+	{
+		if (Pin)
+		{
+			FString ColumnName = GetColumnNameByPin(Pin);
+
+			if (ColumnName != FString())
+			{
+				return GetColumnDefaultAssetByType<T>(ColumnName);
+			}
+		}
+
+		return nullptr;
+	}
+
+	// Generation Mutable Source Methods
+	// We should do this in a template!
+	USkeletalMesh* GetSkeletalMeshAt(const UEdGraphPin* Pin, const UDataTable* DataTable, const FName& RowName) const;
+	TSoftClassPtr<UAnimInstance> GetAnimInstanceAt(const UEdGraphPin* Pin, const UDataTable* DataTable, const FName& RowName) const;
+
+	// Changes the image mode of a pin
+	// bSetDefault param: if true sets the pin to be equal to the default mode (same as node)
+	void ChangeImagePinMode(UEdGraphPin* Pin, bool bSetDefault = false);
+
+	// Returns true if the pin is in the default mode (same as node)
+	bool IsImagePinDefault(const UEdGraphPin* Pin) const;
+
+	// Returns true if the pin is not a UTexture2D but a different kind of UTexture
+	bool IsNotTexture2DPin(const UEdGraphPin* Pin) const;
+
+	// Returns the image mode of the column
+	ETableTextureType GetColumnImageMode(const FString& ColumnName) const;
+
+	// Returns the mesh type of the Pin
+	ETableMeshPinType GetPinMeshType(const UEdGraphPin* Pin) const;
+
+	// Functions to generate the names of a mutable table's column
+	FString GenerateSkeletalMeshMutableColumName(const FString& PinName, int32 LODIndex, int32 MaterialIndex) const;
+	FString GenerateStaticMeshMutableColumName(const FString& PinName, int32 MaterialIndex) const;
+
+	/** Returns the struct pointer used to gather data */
+	const UScriptStruct* GetTableNodeStruct() const;
+
+	// Methods from UDataTable but modified to support ScriptStructs
+	/** Get an array of all the column titles, using the friendly display name from the property */
+	TArray<FString> GetColumnTitles(const UScriptStruct* ScriptStruct) const;
+	
+	/** Returns the column property where PropertyName matches the name of the column property. Returns nullptr if no match is found or the match is not a supported table property */
+	FProperty* FindTableProperty(const UScriptStruct* ScriptStruct, const FName& PropertyName) const;
+
+	/** Returns the column id */
+	FGuid GetColumnIdByName(const FName& ColumnName) const;
+
+	/** Return the list of UDataTable that will be used to compose the final UDataTable. */
+	TArray<FAssetData> GetParentTables() const;
+
+	/** Returns the skeletal material associated to the given skeletal mesh output pin. */
+	FSkeletalMaterial* GetDefaultSkeletalMaterialFor(const UEdGraphPin& MeshPin) const;
+	
+private:
+
+	/** Number of properties to know when the node needs an update */
+	UPROPERTY()
+	int32 NumProperties;
+	
+	FDelegateHandle OnTableChangedDelegateHandle;
+	
+	// Generates a mesh pin for each LOD and Material Surface of the reference SkeletalMesh
+	void GenerateMeshPins(UObject* Mesh, const FString& ColumnName, const FGuid& ColumnId);
+
+	// Checks if a pin already exists and if it has the same type as before the node refresh
+	bool CheckPinUpdated(const FString& PinName, const FName& PinType) const;
+
+};
diff --git a/Engine/Plugins/Experimental/Mutable/Source/CustomizableObjectEditor/Public/MuCOE/Nodes/CustomizableObjectNodeTexture.h b/Engine/Plugins/Experimental/Mutable/Source/CustomizableObjectEditor/Public/MuCOE/Nodes/CustomizableObjectNodeTexture.h
new file mode 100644
--- /dev/null
+++ b/Engine/Plugins/Experimental/Mutable/Source/CustomizableObjectEditor/Public/MuCOE/Nodes/CustomizableObjectNodeTexture.h
@@ -0,0 +1,92 @@
+// Copyright Epic Games, Inc. All Rights Reserved.
+
+#pragma once
+
+#include "MuCOE/Nodes/CustomizableObjectNode.h"
+#include "SGraphNode.h"
+
+#include "CustomizableObjectNodeTexture.generated.h"
+
+namespace ENodeTitleType { enum Type : int; }
+
+class SOverlay;
+class SVerticalBox;
+class UCustomizableObjectNodeRemapPins;
+class UCustomizableObjectNodeTexture;
+class UObject;
+class UTexture2D;
+struct FGeometry;
+
+// Class to render the Texture thumbnail of a CustomizableObjectNodeTexture
+class SGraphNodeTexture : public SGraphNode
+{
+public:
+
+	SLATE_BEGIN_ARGS(SGraphNodeTexture) {}
+	SLATE_END_ARGS();
+
+	SGraphNodeTexture() : SGraphNode() {};
+
+	// Builds the SGraphNodeTexture when needed
+	void Construct(const FArguments& InArgs, UEdGraphNode* InGraphNode);
+
+	// Calls the needed functions to build the SGraphNode widgets
+	void UpdateGraphNode();
+
+	// Overriden functions to build the SGraphNode widgets
+	virtual void SetDefaultTitleAreaWidget(TSharedRef<SOverlay> DefaultTitleAreaWidget) override;
+	virtual void CreateBelowPinControls(TSharedPtr<SVerticalBox> MainBox) override;
+	virtual void Tick(const FGeometry& AllottedGeometry, const double InCurrentTime, const float InDeltaTime) override;
+	virtual bool ShouldAllowCulling() const override { return false; }
+
+	// Callbacks for the widget
+	void OnExpressionPreviewChanged(const ECheckBoxState NewCheckedState);
+	ECheckBoxState IsExpressionPreviewChecked() const;
+	const FSlateBrush* GetExpressionPreviewArrow() const;
+	EVisibility ExpressionPreviewVisibility() const;
+
+public:
+
+	// Pointer to the NodeTexture that owns this SGraphNode
+	UCustomizableObjectNodeTextureBase* NodeTexture;
+
+	// Single property that only draws the combo box widget of the Texture
+	TSharedPtr<class ISinglePropertyView> TextureSelector;
+
+	// Brush to draw the texture to the widget
+	FSlateBrush TextureBrush;
+};
+
+
+UCLASS()
+class CUSTOMIZABLEOBJECTEDITOR_API UCustomizableObjectNodeTextureBase : public UCustomizableObjectNode
+{
+public:
+	GENERATED_BODY()
+
+	UPROPERTY(EditAnywhere, Category = Texture2D)
+	TObjectPtr<UTexture2D> Texture = nullptr;
+
+	// Determines if the Node is collapsed or not
+	bool bCollapsed = true;
+
+	// Creates the SGraph Node widget for the thumbnail
+	TSharedPtr<SGraphNode> CreateVisualWidget() override;
+};
+
+
+UCLASS()
+class CUSTOMIZABLEOBJECTEDITOR_API UCustomizableObjectNodeTexture : public UCustomizableObjectNodeTextureBase
+{
+public:
+	GENERATED_BODY()
+
+	// UCustomizableObjectNode interface
+	virtual void AllocateDefaultPins(UCustomizableObjectNodeRemapPins* RemapPins) override;
+	virtual void BackwardsCompatibleFixup() override;
+
+	// Begin EdGraphNode interface
+	FText GetNodeTitle(ENodeTitleType::Type TitleType) const override;
+	FLinearColor GetNodeTitleColor() const override;
+	FText GetTooltipText() const override;
+};
diff --git a/Engine/Plugins/Experimental/Mutable/Source/CustomizableObjectEditor/Public/MuCOE/Nodes/CustomizableObjectNodeTextureBinarise.h b/Engine/Plugins/Experimental/Mutable/Source/CustomizableObjectEditor/Public/MuCOE/Nodes/CustomizableObjectNodeTextureBinarise.h
new file mode 100644
--- /dev/null
+++ b/Engine/Plugins/Experimental/Mutable/Source/CustomizableObjectEditor/Public/MuCOE/Nodes/CustomizableObjectNodeTextureBinarise.h
@@ -0,0 +1,40 @@
+// Copyright Epic Games, Inc. All Rights Reserved.
+#pragma once
+
+#include "MuCOE/Nodes/CustomizableObjectNode.h"
+
+#include "CustomizableObjectNodeTextureBinarise.generated.h"
+
+namespace ENodeTitleType { enum Type : int; }
+
+class FArchive;
+class UCustomizableObjectNodeRemapPins;
+class UObject;
+
+UCLASS()
+class CUSTOMIZABLEOBJECTEDITOR_API UCustomizableObjectNodeTextureBinarise : public UCustomizableObjectNode
+{
+public:
+	GENERATED_BODY()
+
+	// Begin EdGraphNode interface
+	FText GetNodeTitle(ENodeTitleType::Type TittleType) const override;
+	FLinearColor GetNodeTitleColor() const override;
+	FText GetTooltipText() const override;
+	virtual void Serialize(FArchive& Ar) override;
+	
+	// UCustomizableObjectNode interface
+	virtual void AllocateDefaultPins(UCustomizableObjectNodeRemapPins* RemapPins) override;
+	virtual void BackwardsCompatibleFixup() override;
+
+	UEdGraphPin* GetBaseImagePin() const;
+
+	UEdGraphPin* GetThresholdPin() const
+	{
+		return FindPin(TEXT("Threshold"));
+	}
+
+private:
+	UPROPERTY()
+	FEdGraphPinReference BaseImagePinReference;
+};
diff --git a/Engine/Plugins/Experimental/Mutable/Source/CustomizableObjectEditor/Public/MuCOE/Nodes/CustomizableObjectNodeTextureColourMap.h b/Engine/Plugins/Experimental/Mutable/Source/CustomizableObjectEditor/Public/MuCOE/Nodes/CustomizableObjectNodeTextureColourMap.h
new file mode 100644
--- /dev/null
+++ b/Engine/Plugins/Experimental/Mutable/Source/CustomizableObjectEditor/Public/MuCOE/Nodes/CustomizableObjectNodeTextureColourMap.h
@@ -0,0 +1,46 @@
+// Copyright Epic Games, Inc. All Rights Reserved.
+
+#pragma once
+
+#include "MuCOE/Nodes/CustomizableObjectNode.h"
+
+#include "CustomizableObjectNodeTextureColourMap.generated.h"
+
+namespace ENodeTitleType { enum Type : int; }
+
+class UCustomizableObjectNodeRemapPins;
+class UEdGraphPin;
+class UObject;
+
+
+UCLASS()
+class CUSTOMIZABLEOBJECTEDITOR_API UCustomizableObjectNodeTextureColourMap : public UCustomizableObjectNode
+{
+public:
+	GENERATED_BODY()
+
+	// Begin EdGraphNode interface
+	FText GetNodeTitle(ENodeTitleType::Type TitleType) const override;
+	FLinearColor GetNodeTitleColor() const override;
+	FText GetTooltipText() const override;
+
+	// UCustomizableObjectNode interface
+	virtual void AllocateDefaultPins(UCustomizableObjectNodeRemapPins* RemapPins) override;
+	virtual void BackwardsCompatibleFixup() override;
+
+	UEdGraphPin* GetMapPin() const
+	{
+		return FindPin(TEXT("Map"));
+	}
+
+	UEdGraphPin* GetBasePin() const
+	{
+		return FindPin(TEXT("Base"));
+	}
+
+	UEdGraphPin* GetMaskPin() const
+	{
+		return FindPin(TEXT("Mask"));
+	}
+
+};
diff --git a/Engine/Plugins/Experimental/Mutable/Source/CustomizableObjectEditor/Public/MuCOE/Nodes/CustomizableObjectNodeTextureFromChannels.h b/Engine/Plugins/Experimental/Mutable/Source/CustomizableObjectEditor/Public/MuCOE/Nodes/CustomizableObjectNodeTextureFromChannels.h
new file mode 100644
--- /dev/null
+++ b/Engine/Plugins/Experimental/Mutable/Source/CustomizableObjectEditor/Public/MuCOE/Nodes/CustomizableObjectNodeTextureFromChannels.h
@@ -0,0 +1,51 @@
+// Copyright Epic Games, Inc. All Rights Reserved.
+
+#pragma once
+
+#include "MuCOE/Nodes/CustomizableObjectNode.h"
+
+#include "CustomizableObjectNodeTextureFromChannels.generated.h"
+
+namespace ENodeTitleType { enum Type : int; }
+
+class UCustomizableObjectNodeRemapPins;
+class UEdGraphPin;
+class UObject;
+
+
+UCLASS()
+class CUSTOMIZABLEOBJECTEDITOR_API UCustomizableObjectNodeTextureFromChannels : public UCustomizableObjectNode
+{
+public:
+	GENERATED_BODY()
+
+	// Begin EdGraphNode interface
+	FText GetNodeTitle(ENodeTitleType::Type TitleType) const override;
+	FLinearColor GetNodeTitleColor() const override;
+	FText GetTooltipText() const override;
+
+	// UCustomizableObjectNode interface
+	virtual void AllocateDefaultPins(UCustomizableObjectNodeRemapPins* RemapPins) override;
+	virtual void BackwardsCompatibleFixup() override;
+
+	UEdGraphPin* RPin() const
+	{
+		return FindPin(TEXT("R"));
+	}
+
+	UEdGraphPin* GPin() const
+	{
+		return FindPin(TEXT("G"));
+	}
+
+	UEdGraphPin* BPin() const
+	{
+		return FindPin(TEXT("B"));
+	}
+
+	UEdGraphPin* APin() const
+	{
+		return FindPin(TEXT("A"));
+	}
+};
+
diff --git a/Engine/Plugins/Experimental/Mutable/Source/CustomizableObjectEditor/Public/MuCOE/Nodes/CustomizableObjectNodeTextureFromColor.h b/Engine/Plugins/Experimental/Mutable/Source/CustomizableObjectEditor/Public/MuCOE/Nodes/CustomizableObjectNodeTextureFromColor.h
new file mode 100644
--- /dev/null
+++ b/Engine/Plugins/Experimental/Mutable/Source/CustomizableObjectEditor/Public/MuCOE/Nodes/CustomizableObjectNodeTextureFromColor.h
@@ -0,0 +1,36 @@
+// Copyright Epic Games, Inc. All Rights Reserved.
+
+#pragma once
+
+#include "MuCOE/Nodes/CustomizableObjectNode.h"
+
+#include "CustomizableObjectNodeTextureFromColor.generated.h"
+
+namespace ENodeTitleType { enum Type : int; }
+
+class UCustomizableObjectNodeRemapPins;
+class UEdGraphPin;
+class UObject;
+
+
+UCLASS()
+class CUSTOMIZABLEOBJECTEDITOR_API UCustomizableObjectNodeTextureFromColor : public UCustomizableObjectNode
+{
+public:
+	GENERATED_BODY()
+
+	// Begin EdGraphNode interface
+	FText GetNodeTitle(ENodeTitleType::Type TitleType) const override;
+	FLinearColor GetNodeTitleColor() const override;
+	FText GetTooltipText() const override;
+
+	// UCustomizableObjectNode interface
+	virtual void AllocateDefaultPins(UCustomizableObjectNodeRemapPins* RemapPins) override;
+	virtual void BackwardsCompatibleFixup() override;
+
+	UEdGraphPin* ColorPin() const
+	{
+		return FindPin(TEXT("Color"));
+	}
+};
+
diff --git a/Engine/Plugins/Experimental/Mutable/Source/CustomizableObjectEditor/Public/MuCOE/Nodes/CustomizableObjectNodeTextureInterpolate.h b/Engine/Plugins/Experimental/Mutable/Source/CustomizableObjectEditor/Public/MuCOE/Nodes/CustomizableObjectNodeTextureInterpolate.h
new file mode 100644
--- /dev/null
+++ b/Engine/Plugins/Experimental/Mutable/Source/CustomizableObjectEditor/Public/MuCOE/Nodes/CustomizableObjectNodeTextureInterpolate.h
@@ -0,0 +1,63 @@
+// Copyright Epic Games, Inc. All Rights Reserved.
+
+#pragma once
+
+#include "MuCOE/Nodes/CustomizableObjectNode.h"
+
+#include "CustomizableObjectNodeTextureInterpolate.generated.h"
+
+namespace ENodeTitleType { enum Type : int; }
+
+class UCustomizableObjectNodeRemapPins;
+class UObject;
+struct FPropertyChangedEvent;
+
+
+UCLASS()
+class CUSTOMIZABLEOBJECTEDITOR_API UCustomizableObjectNodeTextureInterpolate : public UCustomizableObjectNode
+{
+public:
+	GENERATED_BODY()
+
+	UCustomizableObjectNodeTextureInterpolate();
+
+	// UObject interface.
+	void PostEditChangeProperty(FPropertyChangedEvent& PropertyChangedEvent) override;
+
+	// Begin EdGraphNode interface
+	FText GetNodeTitle(ENodeTitleType::Type TitleType) const override;
+	FLinearColor GetNodeTitleColor() const override;
+	FText GetTooltipText() const override;
+
+	// UCustomizableObjectNode interface
+	void AllocateDefaultPins(UCustomizableObjectNodeRemapPins* RemapPins) override;
+
+	UEdGraphPin* FactorPin() const
+	{
+		return FindPin(TEXT("Factor"));
+	}
+
+	UEdGraphPin* Targets(int32 Index) const
+	{
+		return FindPin(FString::Printf(TEXT("Target %d"), Index));
+	}
+
+	int32 GetNumTargets() const
+	{
+		int32 Count = 0;
+
+		for (UEdGraphPin* Pin : GetAllNonOrphanPins())
+		{
+			if (Pin->GetName().StartsWith(TEXT("Target ")))
+			{
+				Count++;
+			}
+		}
+
+		return Count;
+	}
+
+	UPROPERTY(EditAnywhere, Category=CustomizableObject)
+	int32 NumTargets;
+};
+
diff --git a/Engine/Plugins/Experimental/Mutable/Source/CustomizableObjectEditor/Public/MuCOE/Nodes/CustomizableObjectNodeTextureInvert.h b/Engine/Plugins/Experimental/Mutable/Source/CustomizableObjectEditor/Public/MuCOE/Nodes/CustomizableObjectNodeTextureInvert.h
new file mode 100644
--- /dev/null
+++ b/Engine/Plugins/Experimental/Mutable/Source/CustomizableObjectEditor/Public/MuCOE/Nodes/CustomizableObjectNodeTextureInvert.h
@@ -0,0 +1,35 @@
+// Copyright Epic Games, Inc. All Rights Reserved.
+#pragma once
+
+#include "MuCOE/Nodes/CustomizableObjectNode.h"
+
+#include "CustomizableObjectNodeTextureInvert.generated.h"
+
+namespace ENodeTitleType { enum Type : int; }
+
+class FArchive;
+class UCustomizableObjectNodeRemapPins;
+class UObject;
+
+UCLASS()
+class CUSTOMIZABLEOBJECTEDITOR_API UCustomizableObjectNodeTextureInvert : public UCustomizableObjectNode
+{
+public:
+	GENERATED_BODY()
+	
+	// EdGraphNode interface
+	virtual FText GetNodeTitle(ENodeTitleType::Type TittleType)const override;
+	virtual FLinearColor GetNodeTitleColor() const override;
+	virtual FText GetTooltipText() const override;
+	
+	// UCustomizableObjectNode interface
+	virtual void AllocateDefaultPins(UCustomizableObjectNodeRemapPins* RemapPins) override;
+	virtual void BackwardsCompatibleFixup() override;
+
+	// Own interface
+	UEdGraphPin* GetBaseImagePin() const;
+
+private:
+	UPROPERTY()
+	FEdGraphPinReference BaseImagePinReference;
+};
diff --git a/Engine/Plugins/Experimental/Mutable/Source/CustomizableObjectEditor/Public/MuCOE/Nodes/CustomizableObjectNodeTextureLayer.h b/Engine/Plugins/Experimental/Mutable/Source/CustomizableObjectEditor/Public/MuCOE/Nodes/CustomizableObjectNodeTextureLayer.h
new file mode 100644
--- /dev/null
+++ b/Engine/Plugins/Experimental/Mutable/Source/CustomizableObjectEditor/Public/MuCOE/Nodes/CustomizableObjectNodeTextureLayer.h
@@ -0,0 +1,126 @@
+// Copyright Epic Games, Inc. All Rights Reserved.
+
+#pragma once
+
+#include "MuCOE/Nodes/CustomizableObjectNode.h"
+
+#include "CustomizableObjectNodeTextureLayer.generated.h"
+
+namespace ENodeTitleType { enum Type : int; }
+
+class FArchive;
+class UCustomizableObjectNodeRemapPins;
+class UEdGraph;
+class UObject;
+struct FPropertyChangedEvent;
+
+
+UENUM()
+enum ECustomizableObjectTextureLayerEffect : int
+{
+	COTLE_MODULATE			= 0 UMETA(DisplayName = "MODULATE"), 
+	COTLE_MULTIPLY				UMETA(DisplayName = "MULTIPLY"), 
+	COTLE_SOFTLIGHT				UMETA(DisplayName = "SOFTLIGHT"),
+	COTLE_HARDLIGHT				UMETA(DisplayName = "HARDLIGHT"),
+	COTLE_DODGE					UMETA(DisplayName = "DODGE"),
+	COTLE_BURN					UMETA(DisplayName = "BURN"),
+	COTLE_SCREEN				UMETA(DisplayName = "SCREEN"),
+	COTLE_OVERLAY				UMETA(DisplayName = "OVERLAY"),
+	COTLE_ALPHA_OVERLAY			UMETA(DisplayName = "LIGHTEN"),
+	COTLE_NORMAL_COMBINE		UMETA(DisplayName = "COMBINE")
+};
+
+
+USTRUCT()
+struct CUSTOMIZABLEOBJECTEDITOR_API FCustomizableObjectTextureLayer
+{
+	GENERATED_USTRUCT_BODY()
+
+	FCustomizableObjectTextureLayer()
+	{
+		Effect = COTLE_SOFTLIGHT;
+	}
+
+	UPROPERTY(EditAnywhere, Category=CustomizableObject)
+	TEnumAsByte<enum ECustomizableObjectTextureLayerEffect> Effect;
+};
+
+
+UCLASS()
+class CUSTOMIZABLEOBJECTEDITOR_API UCustomizableObjectNodeTextureLayer : public UCustomizableObjectNode
+{
+public:
+	GENERATED_BODY()
+
+	UCustomizableObjectNodeTextureLayer();
+
+	/**  */
+	UPROPERTY(EditAnywhere, Category=CustomizableObject)
+	TArray<FCustomizableObjectTextureLayer> Layers;
+
+	// UObject interface.
+	void PostEditChangeProperty(FPropertyChangedEvent& PropertyChangedEvent) override;
+	virtual void Serialize(FArchive& Ar) override;
+	
+	// Begin EdGraphNode interface
+	FText GetNodeTitle(ENodeTitleType::Type TitleType) const override;
+	FLinearColor GetNodeTitleColor() const override;
+	FText GetTooltipText() const override;
+	void PrepareForCopying() override;
+	void PostPasteNode() override;
+
+	// UCustomizableObjectNode interface
+	virtual void AllocateDefaultPins(UCustomizableObjectNodeRemapPins* RemapPins) override;
+	virtual void BackwardsCompatibleFixup() override;
+	virtual bool CanConnect(const UEdGraphPin* InOwnedInputPin, const UEdGraphPin* InOutputPin, bool& bOutIsOtherNodeBlocklisted, bool& bOutArePinsCompatible) const override;
+		
+	UEdGraphPin* BasePin() const
+	{
+		return FindPin(TEXT("Base"));
+	}
+
+	UEdGraphPin* OutputPin() const;
+
+	UEdGraphPin* LayerPin(int32 Index) const
+	{
+		FString PinName = FString::Printf(TEXT("Layer %d "), Index);
+		return FindPin(PinName);
+	}
+
+	UEdGraphPin* MaskPin(int32 Index) const
+	{
+		FString PinName = FString::Printf(TEXT("Mask %d "), Index);
+		return FindPin(PinName);
+	}
+
+	int32 GetNumLayers() const
+	{
+		int32 Count = 0;
+
+		for (UEdGraphPin* Pin : GetAllNonOrphanPins())
+		{
+			if (Pin->GetName().StartsWith(TEXT("Layer ")))
+			{
+				Count++;
+			}
+		}
+
+		return Count;
+	}
+protected:
+	// This struct will store the information of each textue layer node when its destroyed, to be able to recover the UCustomizableObjectNodeTextureLayer::Layers
+	// data in a post paste operation. Otherwise, this information is lost.
+	struct TextureLayerLayersData
+	{
+		FString Name; // Texture layer node name
+		TArray<FCustomizableObjectTextureLayer> Layers;
+		UEdGraph* Graph; // Graph that contained this node
+	};
+	static int32 GetIndexInArrayDestroyedNodes(UCustomizableObjectNodeTextureLayer* Node);
+	static TArray<TextureLayerLayersData> ArrayDestroyedNodes;
+
+private:
+	UPROPERTY()
+	FEdGraphPinReference OutputPinReference;
+};
+
diff --git a/Engine/Plugins/Experimental/Mutable/Source/CustomizableObjectEditor/Public/MuCOE/Nodes/CustomizableObjectNodeTextureParameter.h b/Engine/Plugins/Experimental/Mutable/Source/CustomizableObjectEditor/Public/MuCOE/Nodes/CustomizableObjectNodeTextureParameter.h
new file mode 100644
--- /dev/null
+++ b/Engine/Plugins/Experimental/Mutable/Source/CustomizableObjectEditor/Public/MuCOE/Nodes/CustomizableObjectNodeTextureParameter.h
@@ -0,0 +1,56 @@
+// Copyright Epic Games, Inc. All Rights Reserved.
+
+#pragma once
+
+#include "MuCOE/Nodes/CustomizableObjectNode.h"
+
+#include "CustomizableObjectNodeTextureParameter.generated.h"
+
+namespace ENodeTitleType { enum Type : int; }
+
+class UCustomizableObjectNodeRemapPins;
+class UObject;
+class UTexture2D;
+
+
+UCLASS()
+class CUSTOMIZABLEOBJECTEDITOR_API UCustomizableObjectNodeTextureParameter : public UCustomizableObjectNode
+{
+public:
+	GENERATED_BODY()
+
+	/** Default value of the parameter. */
+	UPROPERTY(EditAnywhere, Category=CustomizableObject)
+	TObjectPtr<UTexture2D> DefaultValue;
+
+	/** Reference Texture where this parameter copies some properties from. */
+	UPROPERTY(EditAnywhere, Category=CustomizableObject)
+	TObjectPtr<UTexture2D> ReferenceValue;
+
+	UPROPERTY(EditAnywhere, Category=CustomizableObject)
+	FString ParameterName = "Default Name";
+
+	UPROPERTY(EditAnywhere, Category = UI, meta = (DisplayName = "Parameter UI Metadata"))
+	FMutableParamUIMetadata ParamUIMetadata;
+
+	/** Set the width of the Texture when there is no texture reference.*/
+	UPROPERTY(EditAnywhere, Category = CustomizableObject)
+	int32 TextureSizeX = 0;
+
+	/** Set the height of the Texture when there is no texture reference.*/
+	UPROPERTY(EditAnywhere, Category = CustomizableObject)
+	int32 TextureSizeY = 0;
+
+	// EdGraphNode interface
+	virtual FText GetNodeTitle(ENodeTitleType::Type TitleType) const override;
+	virtual FLinearColor GetNodeTitleColor() const override;
+	virtual FText GetTooltipText() const override;
+	virtual void OnRenameNode(const FString& NewName) override;
+	virtual bool GetCanRenameNode() const override { return true; }
+
+	// UCustomizableObjectNode interface
+	virtual void BackwardsCompatibleFixup() override;
+	virtual void AllocateDefaultPins(UCustomizableObjectNodeRemapPins* RemapPins) override;
+	virtual bool IsAffectedByLOD() const override { return false; }
+};
+
diff --git a/Engine/Plugins/Experimental/Mutable/Source/CustomizableObjectEditor/Public/MuCOE/Nodes/CustomizableObjectNodeTextureSample.h b/Engine/Plugins/Experimental/Mutable/Source/CustomizableObjectEditor/Public/MuCOE/Nodes/CustomizableObjectNodeTextureSample.h
new file mode 100644
--- /dev/null
+++ b/Engine/Plugins/Experimental/Mutable/Source/CustomizableObjectEditor/Public/MuCOE/Nodes/CustomizableObjectNodeTextureSample.h
@@ -0,0 +1,51 @@
+// Copyright Epic Games, Inc. All Rights Reserved.
+
+#pragma once
+
+#include "MuCOE/Nodes/CustomizableObjectNode.h"
+
+#include "CustomizableObjectNodeTextureSample.generated.h"
+
+namespace ENodeTitleType { enum Type : int; }
+
+class UCustomizableObjectNodeRemapPins;
+class UEdGraphPin;
+class UObject;
+struct FPropertyChangedEvent;
+
+
+UCLASS()
+class CUSTOMIZABLEOBJECTEDITOR_API UCustomizableObjectNodeTextureSample : public UCustomizableObjectNode
+{
+public:
+	GENERATED_BODY()
+
+	UCustomizableObjectNodeTextureSample();
+
+	// UObject interface.
+	void PostEditChangeProperty(FPropertyChangedEvent& PropertyChangedEvent) override;
+
+	// Begin EdGraphNode interface
+	FText GetNodeTitle(ENodeTitleType::Type TitleType) const override;
+	FLinearColor GetNodeTitleColor() const override;
+	FText GetTooltipText() const override;
+
+	// UCustomizableObjectNode interface
+	void AllocateDefaultPins(UCustomizableObjectNodeRemapPins* RemapPins) override;
+
+	const UEdGraphPin* TexturePin() const
+	{
+		return FindPin(TEXT("Texture"));
+	}
+
+	const UEdGraphPin* XPin() const
+	{
+		return FindPin(TEXT("X"));
+	}
+
+	const UEdGraphPin* YPin() const
+	{
+		return FindPin(TEXT("Y"));
+	}
+};
+
diff --git a/Engine/Plugins/Experimental/Mutable/Source/CustomizableObjectEditor/Public/MuCOE/Nodes/CustomizableObjectNodeTextureSwitch.h b/Engine/Plugins/Experimental/Mutable/Source/CustomizableObjectEditor/Public/MuCOE/Nodes/CustomizableObjectNodeTextureSwitch.h
new file mode 100644
--- /dev/null
+++ b/Engine/Plugins/Experimental/Mutable/Source/CustomizableObjectEditor/Public/MuCOE/Nodes/CustomizableObjectNodeTextureSwitch.h
@@ -0,0 +1,20 @@
+// Copyright Epic Games, Inc. All Rights Reserved.
+
+#pragma once
+
+#include "MuCOE/Nodes/CustomizableObjectNodeSwitchBase.h"
+
+#include "CustomizableObjectNodeTextureSwitch.generated.h"
+
+
+UCLASS()
+class CUSTOMIZABLEOBJECTEDITOR_API UCustomizableObjectNodeTextureSwitch : public UCustomizableObjectNodeSwitchBase
+{
+public:
+	GENERATED_BODY()
+	// UCustomizableObjectNode interface
+	virtual void BackwardsCompatibleFixup() override;
+	
+	// UCustomizableObjectNodeSwitchBase interface
+	virtual FName GetCategory() const override;
+};
diff --git a/Engine/Plugins/Experimental/Mutable/Source/CustomizableObjectEditor/Public/MuCOE/Nodes/CustomizableObjectNodeTextureToChannels.h b/Engine/Plugins/Experimental/Mutable/Source/CustomizableObjectEditor/Public/MuCOE/Nodes/CustomizableObjectNodeTextureToChannels.h
new file mode 100644
--- /dev/null
+++ b/Engine/Plugins/Experimental/Mutable/Source/CustomizableObjectEditor/Public/MuCOE/Nodes/CustomizableObjectNodeTextureToChannels.h
@@ -0,0 +1,62 @@
+// Copyright Epic Games, Inc. All Rights Reserved.
+
+#pragma once
+
+#include "MuCOE/Nodes/CustomizableObjectNode.h"
+
+#include "CustomizableObjectNodeTextureToChannels.generated.h"
+
+namespace ENodeTitleType { enum Type : int; }
+
+class FArchive;
+class UCustomizableObjectNodeRemapPins;
+class UObject;
+struct FPropertyChangedEvent;
+
+
+UCLASS()
+class CUSTOMIZABLEOBJECTEDITOR_API UCustomizableObjectNodeTextureToChannels : public UCustomizableObjectNode
+{
+public:
+	GENERATED_BODY()
+	
+	// UObject interface.
+	void PostEditChangeProperty(FPropertyChangedEvent& PropertyChangedEvent) override;
+	virtual void Serialize(FArchive& Ar) override;
+	
+	// Begin EdGraphNode interface
+	FText GetNodeTitle(ENodeTitleType::Type TitleType) const override;
+	FLinearColor GetNodeTitleColor() const override;
+	FText GetTooltipText() const override;
+
+	// UCustomizableObjectNode interface
+	virtual void AllocateDefaultPins(UCustomizableObjectNodeRemapPins* RemapPins) override;
+	virtual void BackwardsCompatibleFixup() override;
+
+	UEdGraphPin* InputPin() const;
+
+	UEdGraphPin* RPin() const
+	{
+		return FindPin(TEXT("R"));
+	}
+
+	UEdGraphPin* GPin() const
+	{
+		return FindPin(TEXT("G"));
+	}
+
+	UEdGraphPin* BPin() const
+	{
+		return FindPin(TEXT("B"));
+	}
+
+	UEdGraphPin* APin() const
+	{
+		return FindPin(TEXT("A"));
+	}
+
+private:
+	UPROPERTY()
+	FEdGraphPinReference InputPinReference;
+};
+
diff --git a/Engine/Plugins/Experimental/Mutable/Source/CustomizableObjectEditor/Public/MuCOE/Nodes/CustomizableObjectNodeTextureVariation.h b/Engine/Plugins/Experimental/Mutable/Source/CustomizableObjectEditor/Public/MuCOE/Nodes/CustomizableObjectNodeTextureVariation.h
new file mode 100644
--- /dev/null
+++ b/Engine/Plugins/Experimental/Mutable/Source/CustomizableObjectEditor/Public/MuCOE/Nodes/CustomizableObjectNodeTextureVariation.h
@@ -0,0 +1,69 @@
+// Copyright Epic Games, Inc. All Rights Reserved.
+
+#pragma once
+
+#include "MuCOE/Nodes/CustomizableObjectNode.h"
+
+#include "CustomizableObjectNodeTextureVariation.generated.h"
+
+namespace ENodeTitleType { enum Type : int; }
+
+class UCustomizableObjectNodeRemapPins;
+class UEdGraphPin;
+class UObject;
+struct FPropertyChangedEvent;
+
+
+USTRUCT()
+struct CUSTOMIZABLEOBJECTEDITOR_API FCustomizableObjectTextureVariation
+{
+	GENERATED_USTRUCT_BODY()
+
+	UPROPERTY(EditAnywhere, Category = CustomizableObject)
+	FString Tag;
+};
+
+
+UCLASS()
+class CUSTOMIZABLEOBJECTEDITOR_API UCustomizableObjectNodeTextureVariation : public UCustomizableObjectNode
+{
+public:
+	GENERATED_BODY()
+
+	UCustomizableObjectNodeTextureVariation();
+
+	UPROPERTY(EditAnywhere, Category = CustomizableObject)
+	TArray<FCustomizableObjectTextureVariation> Variations;
+
+	// UObject interface.
+	void PostEditChangeProperty(FPropertyChangedEvent& PropertyChangedEvent) override;
+
+	// EdGraphNode interface
+	virtual FText GetNodeTitle(ENodeTitleType::Type TitleType) const override;
+	FLinearColor GetNodeTitleColor() const override;
+	FText GetTooltipText() const override;
+	
+	// UCustomizableObjectNode interface
+	virtual void AllocateDefaultPins(UCustomizableObjectNodeRemapPins* RemapPins) override;
+
+	UEdGraphPin* OutputPin() const
+	{
+		return FindPin(TEXT("Texture"));
+	}
+
+	UEdGraphPin* DefaultPin() const
+	{
+		return FindPin(TEXT("Default"));
+	}
+
+	int32 GetNumVariations() const
+	{
+		return Variations.Num();
+	}
+
+	UEdGraphPin* VariationPin(int Index) const
+	{
+		return FindPin(FString::Printf(TEXT("Variation %d"),Index));
+	}
+};
+
diff --git a/Engine/Plugins/Experimental/Mutable/Source/CustomizableObjectEditor/Public/MuCOE/Nodes/CustomizableObjectNodeUseMaterial.h b/Engine/Plugins/Experimental/Mutable/Source/CustomizableObjectEditor/Public/MuCOE/Nodes/CustomizableObjectNodeUseMaterial.h
new file mode 100644
--- /dev/null
+++ b/Engine/Plugins/Experimental/Mutable/Source/CustomizableObjectEditor/Public/MuCOE/Nodes/CustomizableObjectNodeUseMaterial.h
@@ -0,0 +1,67 @@
+// Copyright Epic Games, Inc. All Rights Reserved.
+
+#pragma once
+
+
+#include "Containers/Map.h"
+#include "Delegates/IDelegateInstance.h"
+
+class FCustomizableObjectNodeParentedMaterial;
+class UCustomizableObject;
+class UCustomizableObjectNode;
+class UEdGraphPin;
+struct FEdGraphPinReference;
+struct FGuid;
+
+/** Common functionalities from EditMaterial and ExtendMaterial.
+ *
+ * See must call functions section! These functions should be called in the corresponding inheritor functions. */
+class FCustomizableObjectNodeUseMaterial
+{
+public:
+	virtual ~FCustomizableObjectNodeUseMaterial() = default;
+	
+	/** Returns true if this node uses the given Image. */
+	bool UsesImage(const FGuid& ImageId) const;
+
+	/** Given an Image id, returns its Image pin. Can nullptr if the Image is not used. */
+	const UEdGraphPin* GetUsedImagePin(const FGuid& ImageId) const;
+
+protected:
+	// Begin must call functions
+	void PostBackwardsCompatibleFixupWork();
+	
+	bool IsNodeOutDatedAndNeedsRefreshWork();
+
+	/** Call before SetParentNode. */
+	void PreSetParentNodeWork(UCustomizableObject* Object, const FGuid NodeId);
+
+	/** Call after SetParentNode. */
+	void PostSetParentNodeWork(UCustomizableObject* Object, const FGuid NodeId);
+
+	void PinConnectionListChangedWork(UEdGraphPin* Pin);
+
+	void CustomRemovePinWork(UEdGraphPin& Pin);
+	// End must call functions
+
+	/** Return the node which this interface belongs to. */
+	virtual UCustomizableObjectNode& GetNode() = 0;
+	virtual const UCustomizableObjectNode& GetNode() const;
+
+	/** Returns the Parented Material interface this node belongs to. */
+	virtual FCustomizableObjectNodeParentedMaterial& GetNodeParentedMaterial() = 0;
+	
+	/** Return the PinsParameter map which this interface belongs to. */
+	virtual TMap<FGuid, FEdGraphPinReference>& GetPinsParameter() = 0;
+	virtual const TMap<FGuid, FEdGraphPinReference>& GetPinsParameter() const;
+
+	/** Get the output material pin. */
+	virtual UEdGraphPin* OutputPin() const = 0;
+
+private:
+	/** Handler to the reconstruct node delegate. Used to remove the previously attached delegate. */
+	FDelegateHandle PostReconstructNodeDelegateHandler;
+	
+	/** Handler to the mutable mode changed delegate. Used to remove the previously attached delegate. */
+	FDelegateHandle PostTextureParameterModeChangedDelegateHandle;
+};
diff --git a/Engine/Plugins/Experimental/Mutable/Source/CustomizableObjectEditor/Public/MuCOE/Thumbnails/CustomizableObjectInstanceThumbnailRenderer.h b/Engine/Plugins/Experimental/Mutable/Source/CustomizableObjectEditor/Public/MuCOE/Thumbnails/CustomizableObjectInstanceThumbnailRenderer.h
new file mode 100644
--- /dev/null
+++ b/Engine/Plugins/Experimental/Mutable/Source/CustomizableObjectEditor/Public/MuCOE/Thumbnails/CustomizableObjectInstanceThumbnailRenderer.h
@@ -0,0 +1,40 @@
+// Copyright Epic Games, Inc. All Rights Reserved.
+
+/**
+ *
+ * This thumbnail renderer displays a given Customizable Object Instance
+ */
+
+#pragma once
+#include "ThumbnailRendering/DefaultSizedThumbnailRenderer.h"
+
+#include "CustomizableObjectInstanceThumbnailRenderer.generated.h"
+
+class FCanvas;
+class FRenderTarget;
+class UObject;
+
+
+UCLASS(config=Editor)
+class CUSTOMIZABLEOBJECTEDITOR_API UCustomizableObjectInstanceThumbnailRenderer : public UDefaultSizedThumbnailRenderer
+{
+public:
+	GENERATED_BODY()
+
+	UCustomizableObjectInstanceThumbnailRenderer();
+
+	// Begin UThumbnailRenderer Object
+	void Draw(UObject* Object, int32 X, int32 Y, uint32 Width, uint32 Height, FRenderTarget*, FCanvas* Canvas, bool bAdditionalViewFamily) override;
+	// End UThumbnailRenderer Object
+
+	// UObject implementation
+	void BeginDestroy() override;
+
+private:
+	class FSkeletalMeshThumbnailScene* ThumbnailScene;
+
+	UPROPERTY()
+	TObjectPtr<class UTexture2D> NoImage;
+
+};
+
diff --git a/Engine/Plugins/Experimental/Mutable/Source/CustomizableObjectEditor/Public/MuCOE/Thumbnails/CustomizableObjectThumbnailRenderer.h b/Engine/Plugins/Experimental/Mutable/Source/CustomizableObjectEditor/Public/MuCOE/Thumbnails/CustomizableObjectThumbnailRenderer.h
new file mode 100644
--- /dev/null
+++ b/Engine/Plugins/Experimental/Mutable/Source/CustomizableObjectEditor/Public/MuCOE/Thumbnails/CustomizableObjectThumbnailRenderer.h
@@ -0,0 +1,42 @@
+// Copyright Epic Games, Inc. All Rights Reserved.
+
+/**
+ *
+ * This thumbnail renderer displays a given Customizable Object
+ */
+
+#pragma once
+#include "ThumbnailRendering/DefaultSizedThumbnailRenderer.h"
+
+#include "CustomizableObjectThumbnailRenderer.generated.h"
+
+class FCanvas;
+class FRenderTarget;
+class UObject;
+
+
+UCLASS(config=Editor)
+class CUSTOMIZABLEOBJECTEDITOR_API UCustomizableObjectThumbnailRenderer : public UDefaultSizedThumbnailRenderer
+{
+public:
+	GENERATED_BODY()
+
+	UCustomizableObjectThumbnailRenderer();
+
+	// Begin UThumbnailRenderer Object
+	virtual void Draw(UObject* Object, int32 X, int32 Y, uint32 Width, uint32 Height, FRenderTarget*, FCanvas* Canvas, bool bAdditionalViewFamily) override;
+	// End UThumbnailRenderer Object
+
+	// UObject implementation
+	virtual void BeginDestroy() override;
+
+private:
+	class FSkeletalMeshThumbnailScene* ThumbnailScene;
+
+	UPROPERTY()
+	TObjectPtr<class UCustomizableObjectInstance> CustomizableObjectInstance;
+
+	UPROPERTY()
+	TObjectPtr<class UTexture2D> NoImage;
+};
+
