[mutable] Added controls to enable/disable the validation of modified COs.
- Validation of COs will be set as enabled by default.

#rnx

diff --git a/Engine/Plugins/Experimental/Mutable/Source/MutableValidation/Private/MuV/AssetValidator_CustomizableObjects.cpp b/Engine/Plugins/Experimental/Mutable/Source/MutableValidation/Private/MuV/AssetValidator_CustomizableObjects.cpp
--- a/Engine/Plugins/Experimental/Mutable/Source/MutableValidation/Private/MuV/AssetValidator_CustomizableObjects.cpp
+++ b/Engine/Plugins/Experimental/Mutable/Source/MutableValidation/Private/MuV/AssetValidator_CustomizableObjects.cpp
@@ -4,6 +4,7 @@
 
 #include "DataValidationModule.h"
 #include "Editor.h"
+#include "MutableValidationSettings.h"
 #include "AssetRegistry/AssetData.h"
 #include "Delegates/DelegateSignatureImpl.inl"
 #include "Engine/SkeletalMesh.h"
@@ -11,8 +12,6 @@
 #include "Engine/Texture.h"
 #include "Materials/Material.h"
 #include "MuCO/CustomizableObject.h"
-#include "MuCO/CustomizableObjectPrivate.h"
-#include "MuCO/CustomizableObjectSystem.h"
 #include "MuCO/CustomizableObjectCompilerTypes.h"
 #include "MuCO/ICustomizableObjectEditorModule.h"
 #include "MuCOE/GraphTraversal.h"
@@ -28,6 +27,15 @@ UAssetValidator_CustomizableObjects::UAssetValidator_CustomizableObjects() : Sup
 
 bool UAssetValidator_CustomizableObjects::CanValidateAsset_Implementation(const FAssetData& AssetData, UObject* InAsset, FDataValidationContext& InContext) const
 {
+	// Use module settings to decide if it needs to run or not.
+	if (const UMutableValidationSettings* ValidationSettings = GetDefault<UMutableValidationSettings>())
+	{
+		if (!ValidationSettings->bEnableDirectCOValidation)
+		{
+			return false;
+		}
+	}
+	
 	// Do not run if saving or running a commandlet (we do not want CIS failing due to our warnings and errors)
 	if (InContext.GetValidationUsecase() == EDataValidationUsecase::Save || InContext.GetValidationUsecase() == EDataValidationUsecase::Commandlet)
 	{
diff --git a/Engine/Plugins/Experimental/Mutable/Source/MutableValidation/Private/MuV/AssetValidator_ReferencedCustomizableObjects.cpp b/Engine/Plugins/Experimental/Mutable/Source/MutableValidation/Private/MuV/AssetValidator_ReferencedCustomizableObjects.cpp
--- a/Engine/Plugins/Experimental/Mutable/Source/MutableValidation/Private/MuV/AssetValidator_ReferencedCustomizableObjects.cpp
+++ b/Engine/Plugins/Experimental/Mutable/Source/MutableValidation/Private/MuV/AssetValidator_ReferencedCustomizableObjects.cpp
@@ -28,7 +28,7 @@ bool UAssetValidator_ReferencedCustomizableObjects::CanValidateAsset_Implementat
 	// Use module settings to decide if it needs to run or not.
 	if (const UMutableValidationSettings* ValidationSettings = GetDefault<UMutableValidationSettings>())
 	{
-		if (!ValidationSettings->bEnableIndirectValidation)
+		if (!ValidationSettings->bEnableIndirectCOValidation)
 		{
 			return false;
 		}
diff --git a/Engine/Plugins/Experimental/Mutable/Source/MutableValidation/Private/MuV/MutableValidationSettings.h b/Engine/Plugins/Experimental/Mutable/Source/MutableValidation/Private/MuV/MutableValidationSettings.h
--- a/Engine/Plugins/Experimental/Mutable/Source/MutableValidation/Private/MuV/MutableValidationSettings.h
+++ b/Engine/Plugins/Experimental/Mutable/Source/MutableValidation/Private/MuV/MutableValidationSettings.h
@@ -11,7 +11,12 @@ class UMutableValidationSettings : public UObject
 
 public:
 
-	/** If true, validation of referenced COs from asset subject to data validation, will be run. */
-	UPROPERTY(config, EditAnywhere, Category = Validation)
-	bool bEnableIndirectValidation = false;
+	/** If set to true the checked out Customizable Objects will get validated. */
+	UPROPERTY(meta=(DisplayName="Enable Direct Customizable Object Validation"),config, EditAnywhere, Category = Validation)
+	bool bEnableDirectCOValidation = false;
+
+	/**If set to true the Customizable Objects referenced by the checked out resources will get validated.
+	 *This validation is resource intensive and will require more time to run than the Direct Customizable Object Validation.*/
+	UPROPERTY(meta=(DisplayName="Enable Indirect Customizable Object Validation"), config, EditAnywhere, Category = Validation)
+	bool bEnableIndirectCOValidation = false;
 };
