#Mutable Add more info to a warning

diff --git a/Engine/Plugins/Experimental/Mutable/Source/CustomizableObjectEditor/Private/MuCOE/Nodes/CustomizableObjectNode.cpp b/Engine/Plugins/Experimental/Mutable/Source/CustomizableObjectEditor/Private/MuCOE/Nodes/CustomizableObjectNode.cpp
--- a/Engine/Plugins/Experimental/Mutable/Source/CustomizableObjectEditor/Private/MuCOE/Nodes/CustomizableObjectNode.cpp
+++ b/Engine/Plugins/Experimental/Mutable/Source/CustomizableObjectEditor/Private/MuCOE/Nodes/CustomizableObjectNode.cpp
@@ -253,13 +253,21 @@ void UCustomizableObjectNode::ReconstructNode(UCustomizableObjectNodeRemapPins*
 	}
 
 	bool bOrphanedPin = false;
+	FText FirstOldPin;
 	for (UEdGraphPin* OldPin : OldPins)
 	{
 		OldPin->Modify();
 
 		if (PinsToOrphan.Contains(OldPin))
 		{
-			bOrphanedPin = bOrphanedPin || !OldPin->bOrphanedPin;
+			if (!bOrphanedPin)
+			{
+				bOrphanedPin = !OldPin->bOrphanedPin;
+				if (bOrphanedPin)
+				{
+					FirstOldPin = FText::FromString(OldPin->GetName());
+				}
+			}
 			OrphanPin(*OldPin);
 		
 			// Move pin to the end.
@@ -279,7 +287,7 @@ void UCustomizableObjectNode::ReconstructNode(UCustomizableObjectNodeRemapPins*
 	{
 		if (bOrphanedPin)
 		{
-			FCustomizableObjectEditorLogger::CreateLog(LOCTEXT("OrphanPinsWarningReconstruct", "Failed to remap old pins"))
+			FCustomizableObjectEditorLogger::CreateLog(FText::Format(LOCTEXT("OrphanPinsWarningReconstruct", "Failed to remap old pins. Graph:{0} Node:{1} FirstOldPin:{2}"), FText::FromString(Graph->GetPathName()), FText::FromString(GetPathName()), FirstOldPin))
 			.BaseObject()
 			.Severity(EMessageSeverity::Warning)
 			.Context(*this)
