[mutable] Fixed the new TextureFromFloatChannels editor node and the old NodeImagePlainColour core node. They did not respect the alpha channel set by the user.
#rnx
#tests Editor

diff --git a/Engine/Plugins/Experimental/Mutable/Source/MutableTools/Internal/MuT/NodeImagePlainColourPrivate.h b/Engine/Plugins/Experimental/Mutable/Source/MutableTools/Internal/MuT/NodeImagePlainColourPrivate.h
--- a/Engine/Plugins/Experimental/Mutable/Source/MutableTools/Internal/MuT/NodeImagePlainColourPrivate.h
+++ b/Engine/Plugins/Experimental/Mutable/Source/MutableTools/Internal/MuT/NodeImagePlainColourPrivate.h
@@ -25,7 +25,7 @@ namespace mu
 
 		Ptr<NodeColour> m_pColour;
 		int32 m_sizeX, m_sizeY;
-		EImageFormat Format = EImageFormat::IF_RGB_UBYTE;
+		EImageFormat Format = EImageFormat::IF_RGBA_UBYTE;
 
 	};
 
