[Mutable] Fix out of video memory when using optimization None + MRQ.

#rnx

diff --git a/Engine/Plugins/Experimental/Mutable/Source/CustomizableObject/Private/MuCO/UnrealMutableImageProvider.cpp b/Engine/Plugins/Experimental/Mutable/Source/CustomizableObject/Private/MuCO/UnrealMutableImageProvider.cpp
--- a/Engine/Plugins/Experimental/Mutable/Source/CustomizableObject/Private/MuCO/UnrealMutableImageProvider.cpp
+++ b/Engine/Plugins/Experimental/Mutable/Source/CustomizableObject/Private/MuCO/UnrealMutableImageProvider.cpp
@@ -697,12 +697,19 @@ void FUnrealMutableImageProvider::CacheRuntimeReferencedImages(const TSharedRef<
 	ModelImages.SourceTextures.Reset();
 	for (const TSoftObjectPtr<UTexture2D>& RuntimeReferencedTexture : RuntimeReferencedTextures)
 	{
+		const bool bNotLoaded = !RuntimeReferencedTexture.IsValid();
+		
 		UTexture2D* Texture = RuntimeReferencedTexture.LoadSynchronous();
 		if (!Texture)
 		{
 			UE_LOG(LogMutable, Warning, TEXT("Failed to load texture [%s]."), *RuntimeReferencedTexture->GetPathName());
 			continue;
 		}
+		
+		if (bNotLoaded)
+		{
+			Texture->MarkAsGarbage();
+		}
 
 		ModelImages.SourceTextures.Emplace(*Texture);
 	}
