[Mutable] Fix projection group image properties when no reference texture is specified.

- Fix possible image id truncation in ImageToPlatformDataMap.
- Set FGeneratedImageProperties ImagePropertiesIndex default value to INDEX_NONE.

#rnx

diff --git a/Engine/Plugins/Experimental/Mutable/Source/CustomizableObject/Internal/MuCO/CustomizableObjectSystemPrivate.h b/Engine/Plugins/Experimental/Mutable/Source/CustomizableObject/Internal/MuCO/CustomizableObjectSystemPrivate.h
--- a/Engine/Plugins/Experimental/Mutable/Source/CustomizableObject/Internal/MuCO/CustomizableObjectSystemPrivate.h
+++ b/Engine/Plugins/Experimental/Mutable/Source/CustomizableObject/Internal/MuCO/CustomizableObjectSystemPrivate.h
@@ -548,7 +548,7 @@ public:
 	uint8 FirstLODAvailable = 0;
 	uint8 FirstResidentLOD = 0;
 
-	TMap<uint32, FTexturePlatformData*> ImageToPlatformDataMap;
+	TMap<mu::FResourceID, FTexturePlatformData*> ImageToPlatformDataMap;
 
 	EUpdateResult UpdateResult = EUpdateResult::Success;
 
@@ -594,7 +594,7 @@ public:
 /** Runtime data used during a mutable instance update */
 struct FMutableReleasePlatformOperationData
 {
-	TMap<uint32, FTexturePlatformData*> ImageToPlatformDataMap;
+	TMap<mu::FResourceID, FTexturePlatformData*> ImageToPlatformDataMap;
 };
 
 
diff --git a/Engine/Plugins/Experimental/Mutable/Source/CustomizableObject/Private/MuCO/CustomizableObjectSystem.cpp b/Engine/Plugins/Experimental/Mutable/Source/CustomizableObject/Private/MuCO/CustomizableObjectSystem.cpp
--- a/Engine/Plugins/Experimental/Mutable/Source/CustomizableObject/Private/MuCO/CustomizableObjectSystem.cpp
+++ b/Engine/Plugins/Experimental/Mutable/Source/CustomizableObject/Private/MuCO/CustomizableObjectSystem.cpp
@@ -2732,8 +2732,8 @@ namespace impl
 
 		check(OperationData);
 		
-		TMap<uint32, FTexturePlatformData*>& ImageToPlatformDataMap = OperationData->ImageToPlatformDataMap;
-		for (const TPair<uint32, FTexturePlatformData*>& Pair : ImageToPlatformDataMap)
+		TMap<mu::FResourceID, FTexturePlatformData*>& ImageToPlatformDataMap = OperationData->ImageToPlatformDataMap;
+		for (const TPair<mu::FResourceID, FTexturePlatformData*>& Pair : ImageToPlatformDataMap)
 		{
 			delete Pair.Value; // If this is not null then it must mean it hasn't been used, otherwise they would have taken ownership and nulled it
 		}
@@ -2875,6 +2875,14 @@ namespace impl
 		{
 			TSharedPtr<FMutableReleasePlatformOperationData> ReleaseOperationData = MakeShared<FMutableReleasePlatformOperationData>();
 			check(ReleaseOperationData);
+			
+			static_assert(
+					std::is_same_v<
+						decltype(std::declval<FMutableReleasePlatformOperationData>().ImageToPlatformDataMap), 
+						decltype(std::declval<FUpdateContextPrivate>().ImageToPlatformDataMap)
+					>, 
+					"Cannot move FMutableReleasePlatformOperationData::ImageToPlatformDataMap to FUpdateContextPrivate::ImageToPlatformDataMap, types do not match.");
+
 			ReleaseOperationData->ImageToPlatformDataMap = MoveTemp(OperationData->ImageToPlatformDataMap);
 			CustomizableObjectSystemPrivateData->MutableTaskGraph.AddAnyThreadTask(
 				TEXT("Mutable_ReleasePlatformData"),
diff --git a/Engine/Plugins/Experimental/Mutable/Source/CustomizableObjectEditor/Private/MuCOE/GenerateMutableSource/GenerateMutableSource.h b/Engine/Plugins/Experimental/Mutable/Source/CustomizableObjectEditor/Private/MuCOE/GenerateMutableSource/GenerateMutableSource.h
--- a/Engine/Plugins/Experimental/Mutable/Source/CustomizableObjectEditor/Private/MuCOE/GenerateMutableSource/GenerateMutableSource.h
+++ b/Engine/Plugins/Experimental/Mutable/Source/CustomizableObjectEditor/Private/MuCOE/GenerateMutableSource/GenerateMutableSource.h
@@ -57,7 +57,7 @@ struct FGeneratedImageProperties
 	FString TextureParameterName;
 
 	/** Name in the mu::Surface. */
-	int32 ImagePropertiesIndex;
+	int32 ImagePropertiesIndex = INDEX_NONE;
 
 	TEnumAsByte<TextureCompressionSettings> CompressionSettings = TC_Default;
 
diff --git a/Engine/Plugins/Experimental/Mutable/Source/CustomizableObjectEditor/Private/MuCOE/GenerateMutableSource/GenerateMutableSourceSurface.cpp b/Engine/Plugins/Experimental/Mutable/Source/CustomizableObjectEditor/Private/MuCOE/GenerateMutableSource/GenerateMutableSourceSurface.cpp
--- a/Engine/Plugins/Experimental/Mutable/Source/CustomizableObjectEditor/Private/MuCOE/GenerateMutableSource/GenerateMutableSourceSurface.cpp
+++ b/Engine/Plugins/Experimental/Mutable/Source/CustomizableObjectEditor/Private/MuCOE/GenerateMutableSource/GenerateMutableSourceSurface.cpp
@@ -473,6 +473,7 @@ mu::NodeSurfacePtr GenerateMutableSourceSurface(const UEdGraphPin * Pin, FMutabl
 					mu::Ptr<mu::NodeImage> PassThroughImagePtr = GenerateMutableSourceImage(ConnectedPin, GenerationContext, 0);
 					SurfNode->SetImage(ImageIndex, PassThroughImagePtr);
 
+					check(Props.ImagePropertiesIndex != INDEX_NONE);
 					const FString SurfNodeImageName = FString::Printf(TEXT("%d"), Props.ImagePropertiesIndex);
 					SurfNode->SetImageName(ImageIndex, SurfNodeImageName);
 					SurfNode->SetImageLayoutIndex(ImageIndex, -1);
@@ -504,41 +505,40 @@ mu::NodeSurfacePtr GenerateMutableSourceSurface(const UEdGraphPin * Pin, FMutabl
 					{
 						// Get the reference texture
 						UTexture2D* ReferenceTexture = nullptr;
-
-						GenerationContext.CurrentMaterialTableParameter = ImageName;
-						GenerationContext.CurrentMaterialTableParameterId = ImageId.ToString();
-
-						if (GroupProjectionImg.get() && GroupProjectionReferenceTexture)
 						{
-							ReferenceTexture = GroupProjectionReferenceTexture;
-						}
-						else
-						{
-							ReferenceTexture = TypedNodeMat->GetImageReferenceTexture(ImageIndex);
+							GenerationContext.CurrentMaterialTableParameter = ImageName;
+							GenerationContext.CurrentMaterialTableParameterId = ImageId.ToString();
 
-							if (!ReferenceTexture && !GroupProjectionImg.get())
+							ReferenceTexture = GroupProjectionImg.get() ? GroupProjectionReferenceTexture : nullptr;
+						
+							if (!ReferenceTexture)
 							{
-								if (ImagePin)
-								{
-									if (const UEdGraphPin* ConnectedPin = FollowInputPin(*ImagePin))
-									{
-										ReferenceTexture = FindReferenceImage(ConnectedPin, GenerationContext);
-									}
-								}
-								if (!ReferenceTexture && bTableMaterialPinLinked)
+								ReferenceTexture = TypedNodeMat->GetImageReferenceTexture(ImageIndex);
+							}
+
+							// In case of group projector, don't follow the pin to find the reference texture.
+							if (!GroupProjectionImg.get() && !ReferenceTexture && ImagePin)
+							{
+								if (const UEdGraphPin* ConnectedPin = FollowInputPin(*ImagePin))
 								{
-									if (const UEdGraphPin* ConnectedPin = FollowInputPin(*TypedNodeMat->GetMaterialAssetPin()))
-									{
-										ReferenceTexture = FindReferenceImage(ConnectedPin, GenerationContext);
-									}
+									ReferenceTexture = FindReferenceImage(ConnectedPin, GenerationContext);
 								}
-								if (!ReferenceTexture)
+							}
+
+							if (!ReferenceTexture && bTableMaterialPinLinked)
+							{
+								if (const UEdGraphPin* ConnectedPin = FollowInputPin(*TypedNodeMat->GetMaterialAssetPin()))
 								{
-									ReferenceTexture = TypedNodeMat->GetImageValue(ImageIndex);
+									ReferenceTexture = FindReferenceImage(ConnectedPin, GenerationContext);
 								}
+							}	
+
+							if (!ReferenceTexture)
+							{
+								ReferenceTexture = TypedNodeMat->GetImageValue(ImageIndex);
 							}
 						}
-						
+
 						const FGeneratedImagePropertiesKey PropsKey(TypedNodeMat, ImageIndex);
 						const bool bNewImageProps = !GenerationContext.ImageProperties.Contains(PropsKey);
 
@@ -585,7 +585,7 @@ mu::NodeSurfacePtr GenerateMutableSourceSurface(const UEdGraphPin * Pin, FMutabl
 									GenerationContext.Compiler->CompilerLog(FText::FromString(msg), Node, EMessageSeverity::Info);
 								}
 							}
-							else if (!GroupProjectionImg.get())
+							else
 							{
 								// warning!
 								FString msg = FString::Printf(TEXT("The Reference texture for material image [%s] is not set and it couldn't be found automatically."), *ImageName);
@@ -831,6 +831,7 @@ mu::NodeSurfacePtr GenerateMutableSourceSurface(const UEdGraphPin * Pin, FMutabl
 
 						SurfNode->SetImage(ImageIndex, ImageNode);
 
+						check(Props.ImagePropertiesIndex != INDEX_NONE);
 						const FString SurfNodeImageName = FString::Printf(TEXT("%d"), Props.ImagePropertiesIndex);
 
 						// Encoding material layer in mutable name
