#Mutable In CustomizableObjectNodeTable, make sure SkelMesh is loaded before accessing its members in a codepath that can be called from PostLoad

diff --git a/Engine/Plugins/Experimental/Mutable/Source/CustomizableObjectEditor/Private/MuCOE/Nodes/CustomizableObjectNodeTable.cpp b/Engine/Plugins/Experimental/Mutable/Source/CustomizableObjectEditor/Private/MuCOE/Nodes/CustomizableObjectNodeTable.cpp
--- a/Engine/Plugins/Experimental/Mutable/Source/CustomizableObjectEditor/Private/MuCOE/Nodes/CustomizableObjectNodeTable.cpp
+++ b/Engine/Plugins/Experimental/Mutable/Source/CustomizableObjectEditor/Private/MuCOE/Nodes/CustomizableObjectNodeTable.cpp
@@ -25,6 +25,7 @@
 #include "AssetRegistry/AssetRegistryModule.h"
 #include "AssetRegistry/IAssetRegistry.h"
 #include "MuCO/CustomizableObjectCustomVersion.h"
+#include "UObject/LinkerLoad.h"
 
 class ICustomizableObjectEditor;
 class UCustomizableObjectNodeRemapPins;
@@ -536,6 +537,15 @@ void UCustomizableObjectNodeTable::GenerateMeshPins(UObject* Mesh, const FString
 {
 	if (USkeletalMesh* SkeletalMesh = Cast<USkeletalMesh>(Mesh))
 	{
+		// This can be called from a codepath that starts in PostLoad. The SkeletalMesh may not be serialized at this time so here we will preload it to make sure we can read the LOD data
+		if (SkeletalMesh->HasAnyFlags(RF_NeedLoad))
+		{
+			if (FLinkerLoad* Linker = GetLinker())
+			{
+				Linker->Preload(SkeletalMesh);
+			}
+		}
+
 		const int NumLODs = SkeletalMesh->GetLODNum();
 		const UEdGraphSchema_CustomizableObject* Schema = GetDefault<UEdGraphSchema_CustomizableObject>();
 
