#Mutable When allocating pins for UCustomizableObjectNodeTable, skip soft object properties in the table's struct that could not refer to an object that we'd create pins for.

diff --git a/Engine/Plugins/Experimental/Mutable/Source/CustomizableObjectEditor/Private/MuCOE/Nodes/CustomizableObjectNodeTable.cpp b/Engine/Plugins/Experimental/Mutable/Source/CustomizableObjectEditor/Private/MuCOE/Nodes/CustomizableObjectNodeTable.cpp
--- a/Engine/Plugins/Experimental/Mutable/Source/CustomizableObjectEditor/Private/MuCOE/Nodes/CustomizableObjectNodeTable.cpp
+++ b/Engine/Plugins/Experimental/Mutable/Source/CustomizableObjectEditor/Private/MuCOE/Nodes/CustomizableObjectNodeTable.cpp
@@ -380,6 +380,14 @@ void UCustomizableObjectNodeTable::AllocateDefaultPins(UCustomizableObjectNodeRe
 	DefaultDataArray.SetNumZeroed(TableStruct->GetStructureSize());
 	TableStruct->InitializeStruct(DefaultDataArray.GetData());
 
+	static TArray<TObjectPtr<UClass>> SupportedSoftObjectTypes = {
+		USkeletalMesh::StaticClass(),
+		UStaticMesh::StaticClass(),
+		UTexture2D::StaticClass(),
+		UTexture::StaticClass(),
+		UMaterialInstance::StaticClass()
+	};
+
 	TArray<UEdGraphPin*> OldPins(Pins);
 	
 	for (TFieldIterator<FProperty> It(TableStruct); It; ++It)
@@ -401,83 +409,92 @@ void UCustomizableObjectNodeTable::AllocateDefaultPins(UCustomizableObjectNodeRe
 
 		if (const FSoftObjectProperty* SoftObjectProperty = CastField<FSoftObjectProperty>(ColumnProperty))
 		{
-			UObject* Object = nullptr;
-
-			// Getting default UObject
-			uint8* CellData = SoftObjectProperty->ContainerPtrToValuePtr<uint8>(DefaultDataArray.GetData());
-
-			if (CellData)
+			// Only process object properties that might have pointers to objects of any of the "SupportedSoftObjectTypes"
+			// Object properties for unrelated types do not reference anything we would make pins from, so there is no need to load them or make sure they are non-null
+			const bool bPotentiallySupportedObject = SoftObjectProperty->PropertyClass && SupportedSoftObjectTypes.ContainsByPredicate([SoftObjectProperty](const TObjectPtr<UClass>& Type)
+				{
+					return SoftObjectProperty->PropertyClass->IsChildOf(Type) || Type->IsChildOf(SoftObjectProperty->PropertyClass);
+				});
+			if (bPotentiallySupportedObject)
 			{
-				Object = SoftObjectProperty->GetPropertyValue(CellData).LoadSynchronous();
-			}
+				UObject* Object = nullptr;
 
-			if (Object)
-			{
-				if (Object->IsA(USkeletalMesh::StaticClass()) || Object->IsA(UStaticMesh::StaticClass()))
-				{
-					GenerateMeshPins(Object, ColumnName, ColumnPropertyId);
+				// Getting default UObject
+				uint8* CellData = SoftObjectProperty->ContainerPtrToValuePtr<uint8>(DefaultDataArray.GetData());
 
-					if (FTableNodeColumnData* ColumnData = AuxOldColumnData.Find(ColumnPropertyId))
-					{
-						ColumnDataMap.Add(ColumnPropertyId, *ColumnData);
-					}
+				if (CellData)
+				{
+					Object = SoftObjectProperty->GetPropertyValue(CellData).LoadSynchronous();
 				}
 
-				else if (Object->IsA(UTexture2D::StaticClass()))
+				if (Object)
 				{
-					UCustomizableObjectNodeTableImagePinData* PinData = NewObject<UCustomizableObjectNodeTableImagePinData>(this);
-					PinData->ColumnName = ColumnName;
-					PinData->StructColumnId = ColumnPropertyId;
-					PinData->SetIsNotTexture2D(false);
+					if (Object->IsA(USkeletalMesh::StaticClass()) || Object->IsA(UStaticMesh::StaticClass()))
+					{
+						GenerateMeshPins(Object, ColumnName, ColumnPropertyId);
 
-					FName PinCategory = DefaultImageMode == ETableTextureType::PASSTHROUGH_TEXTURE ? Schema->PC_PassThroughImage : Schema->PC_Image;
+						if (FTableNodeColumnData* ColumnData = AuxOldColumnData.Find(ColumnPropertyId))
+						{
+							ColumnDataMap.Add(ColumnPropertyId, *ColumnData);
+						}
+					}
 
-					for (UEdGraphPin* Pin : OldPins)
+					else if (Object->IsA(UTexture2D::StaticClass()))
 					{
-						// Checking if this column already exist
-						if (UCustomizableObjectNodeTableImagePinData* OldPinData = Cast<UCustomizableObjectNodeTableImagePinData>(GetPinData(*(Pin))))
+						UCustomizableObjectNodeTableImagePinData* PinData = NewObject<UCustomizableObjectNodeTableImagePinData>(this);
+						PinData->ColumnName = ColumnName;
+						PinData->StructColumnId = ColumnPropertyId;
+						PinData->SetIsNotTexture2D(false);
+
+						FName PinCategory = DefaultImageMode == ETableTextureType::PASSTHROUGH_TEXTURE ? Schema->PC_PassThroughImage : Schema->PC_Image;
+
+						for (UEdGraphPin* Pin : OldPins)
 						{
-							if (OldPinData->ColumnName == ColumnName)
+							// Checking if this column already exist
+							if (UCustomizableObjectNodeTableImagePinData* OldPinData = Cast<UCustomizableObjectNodeTableImagePinData>(GetPinData(*(Pin))))
 							{
-								PinCategory = OldPinData->ImageMode == ETableTextureType::PASSTHROUGH_TEXTURE ? Schema->PC_PassThroughImage : Schema->PC_Image;
-								break;
+								if (OldPinData->ColumnName == ColumnName)
+								{
+									PinCategory = OldPinData->ImageMode == ETableTextureType::PASSTHROUGH_TEXTURE ? Schema->PC_PassThroughImage : Schema->PC_Image;
+									break;
+								}
 							}
 						}
+
+						OutPin = CustomCreatePin(EGPD_Output, PinCategory, FName(*PinName), PinData);
 					}
 
-					OutPin = CustomCreatePin(EGPD_Output, PinCategory, FName(*PinName), PinData);
-				}
+					else if (Object->IsA(UTexture::StaticClass()))
+					{
+						UCustomizableObjectNodeTableImagePinData* PinData = NewObject<UCustomizableObjectNodeTableImagePinData>(this);
+						PinData->ColumnName = ColumnName;
+						PinData->StructColumnId = ColumnPropertyId;
+						PinData->ImageMode = ETableTextureType::PASSTHROUGH_TEXTURE;
+						PinData->SetIsNotTexture2D(true);
 
-				else if (Object->IsA(UTexture::StaticClass()))
-				{
-					UCustomizableObjectNodeTableImagePinData* PinData = NewObject<UCustomizableObjectNodeTableImagePinData>(this);
-					PinData->ColumnName = ColumnName;
-					PinData->StructColumnId = ColumnPropertyId;
-					PinData->ImageMode = ETableTextureType::PASSTHROUGH_TEXTURE;
-					PinData->SetIsNotTexture2D(true);
+						OutPin = CustomCreatePin(EGPD_Output, Schema->PC_PassThroughImage, FName(*PinName), PinData);
+					}
 
-					OutPin = CustomCreatePin(EGPD_Output, Schema->PC_PassThroughImage, FName(*PinName), PinData);
-				}
+					else if (Object->IsA(UMaterialInstance::StaticClass()))
+					{
+						UCustomizableObjectNodeTableObjectPinData* PinData = NewObject<UCustomizableObjectNodeTableObjectPinData>(this);
+						PinData->ColumnName = ColumnName;
+						PinData->StructColumnId = ColumnPropertyId;
 
-				else if (Object->IsA(UMaterialInstance::StaticClass()))
+						OutPin = CustomCreatePin(EGPD_Output, Schema->PC_MaterialAsset, FName(*PinName), PinData);
+					}
+				}
+				else
 				{
-					UCustomizableObjectNodeTableObjectPinData* PinData = NewObject<UCustomizableObjectNodeTableObjectPinData>(this);
-					PinData->ColumnName = ColumnName;
-					PinData->StructColumnId = ColumnPropertyId;
+					const FText Text = FText::FromString(FString::Printf(TEXT("Could not find a Default Value in Structure member [%s]"), *ColumnName));
 
-					OutPin = CustomCreatePin(EGPD_Output, Schema->PC_MaterialAsset, FName(*PinName), PinData);
+					FCustomizableObjectEditorLogger::CreateLog(Text)
+						.Category(ELoggerCategory::General)
+						.Severity(EMessageSeverity::Warning)
+						.Context(*this)
+						.Log();
 				}
 			}
-			else
-			{
-				const FText Text = FText::FromString(FString::Printf(TEXT("Could not find a Default Value in Structure member [%s]"),*ColumnName));
-
-				FCustomizableObjectEditorLogger::CreateLog(Text)
-					.Category(ELoggerCategory::General)
-					.Severity(EMessageSeverity::Warning)
-					.Context(*this)
-					.Log();
-			}
 		}
 
 		else if (const FStructProperty* StructProperty = CastField<FStructProperty>(ColumnProperty))
