[Mutable] Remove IntParametersLookupTable stored in the descriptors. Add method to find the parameter index by type to the customizable object.

#rnx

diff --git a/Engine/Plugins/Experimental/Mutable/Source/CustomizableObject/Internal/MuCO/CustomizableObjectInstanceDescriptor.h b/Engine/Plugins/Experimental/Mutable/Source/CustomizableObject/Internal/MuCO/CustomizableObjectInstanceDescriptor.h
--- a/Engine/Plugins/Experimental/Mutable/Source/CustomizableObject/Internal/MuCO/CustomizableObjectInstanceDescriptor.h
+++ b/Engine/Plugins/Experimental/Mutable/Source/CustomizableObject/Internal/MuCO/CustomizableObjectInstanceDescriptor.h
@@ -202,23 +202,8 @@ struct CUSTOMIZABLEOBJECT_API FCustomizableObjectInstanceDescriptor
 	/** Get the current projector for the parameter with the given name. */
 	FCustomizableObjectProjector GetProjector(const FString& ParamName, int32 RangeIndex) const;
 	
-	/** Finds in IntParameters a parameter with name ParamName, returns the index if found, -1 otherwise. */
-	int32 FindIntParameterNameIndex(const FString& ParamName) const;
-
-	/** Finds in FloatParameters a parameter with name ParamName, returns the index if found, -1 otherwise. */
-	int32 FindFloatParameterNameIndex(const FString& ParamName) const;
-
-	/** Finds in TextureParameters a parameter with name ParamName, returns the index if found, -1 otherwise. */
-	int32 FindTextureParameterNameIndex(const FString& ParamName) const;
-
-	/** Finds in BoolParameters a parameter with name ParamName, returns the index if found, -1 otherwise. */
-	int32 FindBoolParameterNameIndex(const FString& ParamName) const;
-
-	/** Finds in VectorParameters a parameter with name ParamName, returns the index if found, -1 otherwise. */
-	int32 FindVectorParameterNameIndex(const FString& ParamName) const;
-
-	/** Finds in ProjectorParameters a parameter with name ParamName, returns the index if found, -1 otherwise. */
-	int32 FindProjectorParameterNameIndex(const FString& ParamName) const;
+	/** Finds the parameter with name ParamName in the array of its type, returns the index if found, INDEX_NONE otherwise. */
+	int32 FindTypedParameterIndex(const FString& ParamName, EMutableParameterType Type) const;
 
 	// Parameter Ranges
 
@@ -388,15 +373,10 @@ private:
 
 	/** Array of RequestedLODs per component to generate, they MUST NOT be used in an update (Mutable thread). */
 	TArray<uint16> RequestedLODLevels;
-
-	/** Lookup of UCustomizableObjectDescriptor::IntParameters. */
-	TMap<FString, int32> IntParametersLookupTable;
 	
 	/** Multilayer Projector helpers. See FMultilayerProjector.*/
 	UPROPERTY()
 	TMap<FName, FMultilayerProjector> MultilayerProjectors;
-
-	void CreateParametersLookupTable();
 	
 	// Friends
 	friend FDescriptorHash;
diff --git a/Engine/Plugins/Experimental/Mutable/Source/CustomizableObject/Internal/MuCO/CustomizableObjectPrivate.h b/Engine/Plugins/Experimental/Mutable/Source/CustomizableObject/Internal/MuCO/CustomizableObjectPrivate.h
--- a/Engine/Plugins/Experimental/Mutable/Source/CustomizableObject/Internal/MuCO/CustomizableObjectPrivate.h
+++ b/Engine/Plugins/Experimental/Mutable/Source/CustomizableObject/Internal/MuCO/CustomizableObjectPrivate.h
@@ -557,6 +557,12 @@ static_assert(sizeof(FMorphTargetVertexData) == sizeof(FVector3f)*2 + sizeof(uin
 template<> struct TCanBulkSerialize<FMorphTargetVertexData> { enum { Value = true }; };
 
 
+struct FMutableParameterIndex
+{
+	int32 Index = INDEX_NONE;
+	int32 TypedIndex = INDEX_NONE;
+};
+
 USTRUCT()
 struct CUSTOMIZABLEOBJECT_API FIntegerParameterUIData
 {
@@ -800,6 +806,7 @@ public:
 	void ApplyStateForcedValuesToParameters(int32 State, mu::Parameters* Parameters);
 
 	int32 FindParameter(const FString& Name) const;
+	int32 FindParameterTyped(const FString& Name, EMutableParameterType Type) const;
 
 	EMutableParameterType GetParameterType(int32 ParamIndex) const;
 
@@ -962,7 +969,7 @@ public:
 
 	// Map of name to index of ParameterProperties.
 	// use this to lookup fast by Name
-	TMap<FString, int32> ParameterPropertiesLookupTable;
+	TMap<FString, FMutableParameterIndex> ParameterPropertiesLookupTable;
 
 	// This is a manual version number for the binary blobs in this asset.
 	// Increasing it invalidates all the previously compiled models.
diff --git a/Engine/Plugins/Experimental/Mutable/Source/CustomizableObject/Private/MuCO/CustomizableObject.cpp b/Engine/Plugins/Experimental/Mutable/Source/CustomizableObject/Private/MuCO/CustomizableObject.cpp
--- a/Engine/Plugins/Experimental/Mutable/Source/CustomizableObject/Private/MuCO/CustomizableObject.cpp
+++ b/Engine/Plugins/Experimental/Mutable/Source/CustomizableObject/Private/MuCO/CustomizableObject.cpp
@@ -1302,9 +1302,7 @@ EMutableParameterType UCustomizableObjectPrivate::GetParameterType(int32 ParamIn
 
 EMutableParameterType UCustomizableObject::GetParameterTypeByName(const FString& Name) const
 {
-	const int* IndexPtr = GetPrivate()->ParameterPropertiesLookupTable.Find(Name);
-	int Index = IndexPtr ? *IndexPtr : INDEX_NONE;
-
+	const int32 Index = FindParameter(Name); 
 	if (GetPrivate()->ParameterProperties.IsValidIndex(Index))
 	{
 		return GetPrivate()->ParameterProperties[Index].Type;
@@ -1348,19 +1346,22 @@ void UCustomizableObjectPrivate::UpdateParameterPropertiesFromModel(const TShare
 	if (Model)
 	{
 		mu::ParametersPtr MutableParameters = mu::Model::NewParameters(Model);
-		int paramCount = MutableParameters->GetCount();
+		const int32 NumParameters = MutableParameters->GetCount();
+
+		TArray<int32> TypedParametersCount;
+		TypedParametersCount.SetNum(static_cast<int32>(mu::PARAMETER_TYPE::T_COUNT));
 
-		ParameterProperties.Reset(paramCount);
-		ParameterPropertiesLookupTable.Empty(paramCount);
-		for (int paramIndex = 0; paramIndex<paramCount; ++paramIndex)
+		ParameterProperties.Reset(NumParameters);
+		ParameterPropertiesLookupTable.Empty(NumParameters);
+		for (int32 Index = 0; Index < NumParameters; ++Index)
 		{
 			FMutableModelParameterProperties Data;
 
-			Data.Name = MutableParameters->GetName(paramIndex);
+			Data.Name = MutableParameters->GetName(Index);
 			Data.Type = EMutableParameterType::None;
 
-			mu::PARAMETER_TYPE mutableType = MutableParameters->GetType(paramIndex);
-			switch (mutableType)
+			mu::PARAMETER_TYPE ParameterType = MutableParameters->GetType(Index);
+			switch (ParameterType)
 			{
 			case mu::PARAMETER_TYPE::T_BOOL:
 			{
@@ -1372,13 +1373,13 @@ void UCustomizableObjectPrivate::UpdateParameterPropertiesFromModel(const TShare
 			{
 				Data.Type = EMutableParameterType::Int;
 
-				int ValueCount = MutableParameters->GetIntPossibleValueCount(paramIndex);
-				for (int i = 0; i<ValueCount; ++i)
+				const int32 ValueCount = MutableParameters->GetIntPossibleValueCount(Index);
+				Data.PossibleValues.Reserve(ValueCount);
+				for (int32 ValueIndex = 0; ValueIndex < ValueCount; ++ValueIndex)
 				{
-					FMutableModelParameterValue ValueData;
-					ValueData.Name = MutableParameters->GetIntPossibleValueName(paramIndex,i);
-					ValueData.Value = MutableParameters->GetIntPossibleValue(paramIndex,i);
-					Data.PossibleValues.Add(ValueData);
+					FMutableModelParameterValue& ValueData = Data.PossibleValues.AddDefaulted_GetRef();
+					ValueData.Name = MutableParameters->GetIntPossibleValueName(Index, ValueIndex);
+					ValueData.Value = MutableParameters->GetIntPossibleValue(Index, ValueIndex);
 				}
 				break;
 			}
@@ -1414,7 +1415,7 @@ void UCustomizableObjectPrivate::UpdateParameterPropertiesFromModel(const TShare
 			}
 
 			ParameterProperties.Add(Data);
-			ParameterPropertiesLookupTable.Add(Data.Name, paramIndex);
+			ParameterPropertiesLookupTable.Add(Data.Name, FMutableParameterIndex(Index, TypedParametersCount[static_cast<int32>(ParameterType)]++));
 		}
 	}
 	else
@@ -1476,12 +1477,26 @@ int32 UCustomizableObject::FindParameter(const FString& Name) const
 
 int32 UCustomizableObjectPrivate::FindParameter(const FString& Name) const
 {
-	const int32 * Found = ParameterPropertiesLookupTable.Find(Name);
-	if (Found == nullptr)
+	if (const FMutableParameterIndex* Found = ParameterPropertiesLookupTable.Find(Name))
+	{
+		return Found->Index;
+	}
+
+	return INDEX_NONE;
+}
+
+
+int32 UCustomizableObjectPrivate::FindParameterTyped(const FString& Name, EMutableParameterType Type) const
+{
+	if (const FMutableParameterIndex* Found = ParameterPropertiesLookupTable.Find(Name))
 	{
-		return INDEX_NONE;
+		if (ParameterProperties[Found->Index].Type == Type)
+		{
+			return Found->TypedIndex;
+		}
 	}
-	return *Found;
+
+	return INDEX_NONE;
 }
 
 
diff --git a/Engine/Plugins/Experimental/Mutable/Source/CustomizableObject/Private/MuCO/CustomizableObjectInstance.cpp b/Engine/Plugins/Experimental/Mutable/Source/CustomizableObject/Private/MuCO/CustomizableObjectInstance.cpp
--- a/Engine/Plugins/Experimental/Mutable/Source/CustomizableObject/Private/MuCO/CustomizableObjectInstance.cpp
+++ b/Engine/Plugins/Experimental/Mutable/Source/CustomizableObject/Private/MuCO/CustomizableObjectInstance.cpp
@@ -2652,31 +2652,31 @@ FCustomizableObjectProjector UCustomizableObjectInstance::GetProjector(const FSt
 
 int32 UCustomizableObjectInstance::FindIntParameterNameIndex(const FString& ParamName) const
 {
-	return Descriptor.FindIntParameterNameIndex(ParamName);
+	return Descriptor.FindTypedParameterIndex(ParamName, EMutableParameterType::Int);
 }
 
 
 int32 UCustomizableObjectInstance::FindFloatParameterNameIndex(const FString& ParamName) const
 {
-	return Descriptor.FindFloatParameterNameIndex(ParamName);
+	return Descriptor.FindTypedParameterIndex(ParamName, EMutableParameterType::Float);
 }
 
 
 int32 UCustomizableObjectInstance::FindBoolParameterNameIndex(const FString& ParamName) const
 {
-	return Descriptor.FindBoolParameterNameIndex(ParamName);
+	return Descriptor.FindTypedParameterIndex(ParamName, EMutableParameterType::Bool);
 }
 
 
 int32 UCustomizableObjectInstance::FindVectorParameterNameIndex(const FString& ParamName) const
 {
-	return Descriptor.FindVectorParameterNameIndex(ParamName);
+	return Descriptor.FindTypedParameterIndex(ParamName, EMutableParameterType::Color);
 }
 
 
 int32 UCustomizableObjectInstance::FindProjectorParameterNameIndex(const FString& ParamName) const
 {
-	return Descriptor.FindProjectorParameterNameIndex(ParamName);
+	return Descriptor.FindTypedParameterIndex(ParamName, EMutableParameterType::Projector);
 }
 
 void UCustomizableObjectInstance::SetRandomValues()
diff --git a/Engine/Plugins/Experimental/Mutable/Source/CustomizableObject/Private/MuCO/CustomizableObjectInstanceDescriptor.cpp b/Engine/Plugins/Experimental/Mutable/Source/CustomizableObject/Private/MuCO/CustomizableObjectInstanceDescriptor.cpp
--- a/Engine/Plugins/Experimental/Mutable/Source/CustomizableObject/Private/MuCO/CustomizableObjectInstanceDescriptor.cpp
+++ b/Engine/Plugins/Experimental/Mutable/Source/CustomizableObject/Private/MuCO/CustomizableObjectInstanceDescriptor.cpp
@@ -114,7 +114,7 @@ void FCustomizableObjectInstanceDescriptor::SaveDescriptor(FArchive& Ar, bool bU
 
 			bool bIsParamMultidimensional = false;
 
-			int32 IntParameterIndex = FindIntParameterNameIndex(Name);
+			const int32 IntParameterIndex = FindTypedParameterIndex(Name, EMutableParameterType::Int);
 			if (IntParameterIndex >= 0)
 			{					
 				const FCustomizableObjectIntParameterValue & P = IntParameters[IntParameterIndex];
@@ -302,7 +302,7 @@ void FCustomizableObjectInstanceDescriptor::LoadDescriptor(FArchive& Ar)
 			TArray<int32> Values;
 			TArray<FString> ValueNames;
 
-			const int32 IntParameterIndex = FindIntParameterNameIndex(Name);
+			const int32 IntParameterIndex = FindTypedParameterIndex(Name, EMutableParameterType::Int);
 			const int32 ParameterIndexInObject = CustomizableObject->FindParameter(IntParameters[IntParameterIndex].ParameterName);
 			const bool bIsParamMultidimensional = CustomizableObject->IsParameterMultidimensional(ParameterIndexInObject);
 
@@ -414,8 +414,6 @@ void FCustomizableObjectInstanceDescriptor::LoadDescriptor(FArchive& Ar)
 			check(false);
 		}
 	}
-
-	CreateParametersLookupTable();
 }
 
 
@@ -1064,8 +1062,6 @@ void FCustomizableObjectInstanceDescriptor::ReloadParameters()
 			break;
 		}
 	}
-	
-	CreateParametersLookupTable();
 }
 
 
@@ -1222,7 +1218,7 @@ const FString& FCustomizableObjectInstanceDescriptor::GetIntParameterSelectedOpt
 	check(CustomizableObject);
 	
 	const int32 ParameterIndexInObject = CustomizableObject->FindParameter(ParamName);
-	const int32 ParameterIndexInInstance = FindIntParameterNameIndex(ParamName);
+	const int32 ParameterIndexInInstance = FindTypedParameterIndex(ParamName, EMutableParameterType::Int);
 
 	if (ParameterIndexInObject >= 0 && ParameterIndexInInstance >= 0)
 	{
@@ -1307,7 +1303,7 @@ void FCustomizableObjectInstanceDescriptor::SetIntParameterSelectedOption(const
 	check(CustomizableObject);
 	RETURN_ON_UNCOMPILED_CO(CustomizableObject, TEXT("Error: Cannot set Int parameter "));
 
-	const int32 ParamIndexInInstance = FindIntParameterNameIndex(ParamName);
+	const int32 ParamIndexInInstance = FindTypedParameterIndex(ParamName, EMutableParameterType::Int);
 	if (ParamIndexInInstance == INDEX_NONE)
 	{
 		// Warn and early out since we could not find the parameter to set.
@@ -1324,7 +1320,7 @@ float FCustomizableObjectInstanceDescriptor::GetFloatParameterSelectedOption(con
 	check(CustomizableObject);
 
 	const int32 ParameterIndexInObject = CustomizableObject->FindParameter(FloatParamName);
-	const int32 FloatParamIndex = FindFloatParameterNameIndex(FloatParamName);
+	const int32 FloatParamIndex = FindTypedParameterIndex(FloatParamName, EMutableParameterType::Float);
 
 	if (ParameterIndexInObject >= 0 && FloatParamIndex >= 0)
 	{
@@ -1356,7 +1352,7 @@ void FCustomizableObjectInstanceDescriptor::SetFloatParameterSelectedOption(cons
 	RETURN_ON_UNCOMPILED_CO(CustomizableObject, TEXT("Error: Cannot set Int parameter "));
 
 	const int32 ParameterIndexInObject = CustomizableObject->FindParameter(FloatParamName);
-	const int32 ParameterIndexInInstance = FindFloatParameterNameIndex(FloatParamName);
+	const int32 ParameterIndexInInstance = FindTypedParameterIndex(FloatParamName, EMutableParameterType::Float);
 
 	if (ParameterIndexInObject < 0 || ParameterIndexInInstance < 0)
 	{
@@ -1393,7 +1389,7 @@ FName FCustomizableObjectInstanceDescriptor::GetTextureParameterSelectedOption(c
 	check(CustomizableObject);
 
 	const int32 ParameterIndexInObject = CustomizableObject->FindParameter(TextureParamName);
-	const int32 ParameterIndexInInstance = FindTextureParameterNameIndex(TextureParamName);
+	const int32 ParameterIndexInInstance = FindTypedParameterIndex(TextureParamName, EMutableParameterType::Texture);
 
 	if (ParameterIndexInObject >= 0 && ParameterIndexInInstance >= 0)
 	{
@@ -1425,7 +1421,7 @@ void FCustomizableObjectInstanceDescriptor::SetTextureParameterSelectedOption(co
 	RETURN_ON_UNCOMPILED_CO(CustomizableObject, TEXT("Error: Cannot set Int parameter "));
 
 	const int32 ParameterIndexInObject = CustomizableObject->FindParameter(TextureParamName);
-	const int32 ParameterIndexInInstance = FindTextureParameterNameIndex(TextureParamName);
+	const int32 ParameterIndexInInstance = FindTypedParameterIndex(TextureParamName, EMutableParameterType::Texture);
 
 	if (ParameterIndexInObject < 0 || ParameterIndexInInstance < 0)
 	{
@@ -1461,7 +1457,7 @@ FLinearColor FCustomizableObjectInstanceDescriptor::GetColorParameterSelectedOpt
 	check(CustomizableObject);
 
 	const int32 ParameterIndexInObject = CustomizableObject->FindParameter(ColorParamName);
-	const int32 ColorParamIndex = FindVectorParameterNameIndex(ColorParamName);
+	const int32 ColorParamIndex = FindTypedParameterIndex(ColorParamName, EMutableParameterType::Color);
 
 	if (ColorParamIndex == INDEX_NONE)
 	{
@@ -1487,7 +1483,7 @@ bool FCustomizableObjectInstanceDescriptor::GetBoolParameterSelectedOption(const
 	check(CustomizableObject);
 
 	const int32 ParameterIndexInObject = CustomizableObject->FindParameter(BoolParamName);
-	const int32 BoolParamIndex = FindBoolParameterNameIndex(BoolParamName);
+	const int32 BoolParamIndex = FindTypedParameterIndex(BoolParamName, EMutableParameterType::Bool);
 
 	if (BoolParamIndex == INDEX_NONE)
 	{
@@ -1505,7 +1501,7 @@ void FCustomizableObjectInstanceDescriptor::SetBoolParameterSelectedOption(const
 	RETURN_ON_UNCOMPILED_CO(CustomizableObject, TEXT("Error: Cannot set Int parameter "));
 
 	const int32 ParameterIndexInObject = CustomizableObject->FindParameter(BoolParamName);
-	const int32 ParameterIndexInInstance = FindBoolParameterNameIndex(BoolParamName);
+	const int32 ParameterIndexInInstance = FindTypedParameterIndex(BoolParamName, EMutableParameterType::Bool);
 
 	if (ParameterIndexInObject < 0 || ParameterIndexInInstance < 0)
 	{
@@ -1524,7 +1520,7 @@ void FCustomizableObjectInstanceDescriptor::SetVectorParameterSelectedOption(con
 	RETURN_ON_UNCOMPILED_CO(CustomizableObject, TEXT("Error: Cannot set Int parameter "));
 
 	const int32 ParameterIndexInObject = CustomizableObject->FindParameter(VectorParamName);
-	const int32 ParameterIndexInInstance = FindVectorParameterNameIndex(VectorParamName);
+	const int32 ParameterIndexInInstance = FindTypedParameterIndex(VectorParamName, EMutableParameterType::Color);
 
 	if (ParameterIndexInObject < 0 || ParameterIndexInInstance < 0)
 	{
@@ -1546,7 +1542,7 @@ void FCustomizableObjectInstanceDescriptor::SetProjectorValue(const FString& Pro
 	RETURN_ON_UNCOMPILED_CO(CustomizableObject, TEXT("Error: Cannot set Projector parameter "))
 
 	const int32 ParameterIndexInObject = CustomizableObject->FindParameter(ProjectorParamName);
-	const int32 ParameterIndexInInstance = FindProjectorParameterNameIndex(ProjectorParamName);
+	const int32 ParameterIndexInInstance = FindTypedParameterIndex(ProjectorParamName, EMutableParameterType::Projector);
 
 	if (ParameterIndexInObject < 0 || ParameterIndexInInstance < 0)
 	{
@@ -1667,7 +1663,7 @@ void FCustomizableObjectInstanceDescriptor::GetProjectorValueF(const FString& Pr
 	check(CustomizableObject);
 
 	const int32 ParameterIndexInObject = CustomizableObject->FindParameter(ProjectorParamName);
-	const int32 ParameterIndexInInstance = FindProjectorParameterNameIndex(ProjectorParamName);
+	const int32 ParameterIndexInInstance = FindTypedParameterIndex(ProjectorParamName, EMutableParameterType::Projector);
 
 	if (ParameterIndexInObject >= 0 && ParameterIndexInInstance >= 0)
 	{
@@ -1710,7 +1706,7 @@ FVector FCustomizableObjectInstanceDescriptor::GetProjectorPosition(const FStrin
 	check(CustomizableObject);
 
 	const int32 ParameterIndexInObject = CustomizableObject->FindParameter(ParamName);
-	const int32 ParameterIndexInInstance = FindProjectorParameterNameIndex(ParamName);
+	const int32 ParameterIndexInInstance = FindTypedParameterIndex(ParamName, EMutableParameterType::Projector);
 
 	if (ParameterIndexInObject >= 0 && ParameterIndexInInstance >= 0)
 	{
@@ -1735,7 +1731,7 @@ FVector FCustomizableObjectInstanceDescriptor::GetProjectorDirection(const FStri
 	check(CustomizableObject);
 
 	const int32 ParameterIndexInObject = CustomizableObject->FindParameter(ParamName);
-	const int32 ParameterIndexInInstance = FindProjectorParameterNameIndex(ParamName);
+	const int32 ParameterIndexInInstance = FindTypedParameterIndex(ParamName, EMutableParameterType::Projector);
 
 	if (ParameterIndexInObject >= 0 && ParameterIndexInInstance >= 0)
 	{
@@ -1760,7 +1756,7 @@ FVector FCustomizableObjectInstanceDescriptor::GetProjectorUp(const FString& Par
 	check(CustomizableObject);
 
 	const int32 ParameterIndexInObject = CustomizableObject->FindParameter(ParamName);
-	const int32 ParameterIndexInInstance = FindProjectorParameterNameIndex(ParamName);
+	const int32 ParameterIndexInInstance = FindTypedParameterIndex(ParamName, EMutableParameterType::Projector);
 
 	if (ParameterIndexInObject >= 0 && ParameterIndexInInstance >= 0)
 	{
@@ -1785,7 +1781,7 @@ FVector FCustomizableObjectInstanceDescriptor::GetProjectorScale(const FString&
 	check(CustomizableObject);
 
 	const int32 ParameterIndexInObject = CustomizableObject->FindParameter(ParamName);
-	const int32 ParameterIndexInInstance = FindProjectorParameterNameIndex(ParamName);
+	const int32 ParameterIndexInInstance = FindTypedParameterIndex(ParamName, EMutableParameterType::Projector);
 
 	if (ParameterIndexInObject >= 0 && ParameterIndexInInstance >= 0)
 	{
@@ -1810,7 +1806,7 @@ float FCustomizableObjectInstanceDescriptor::GetProjectorAngle(const FString& Pa
 	check(CustomizableObject);
 
 	const int32 ParameterIndexInObject = CustomizableObject->FindParameter(ParamName);
-	const int32 ParameterIndexInInstance = FindProjectorParameterNameIndex(ParamName);
+	const int32 ParameterIndexInInstance = FindTypedParameterIndex(ParamName, EMutableParameterType::Projector);
 
 	if (ParameterIndexInObject >= 0 && ParameterIndexInInstance >= 0)
 	{
@@ -1835,7 +1831,7 @@ ECustomizableObjectProjectorType FCustomizableObjectInstanceDescriptor::GetProje
 	check(CustomizableObject);
 
 	const int32 ParameterIndexInObject = CustomizableObject->FindParameter(ParamName);
-	const int32 ParameterIndexInInstance = FindProjectorParameterNameIndex(ParamName);
+	const int32 ParameterIndexInInstance = FindTypedParameterIndex(ParamName, EMutableParameterType::Projector);
 
 	if (ParameterIndexInObject >= 0 && ParameterIndexInInstance >= 0)
 	{
@@ -1857,9 +1853,9 @@ ECustomizableObjectProjectorType FCustomizableObjectInstanceDescriptor::GetProje
 
 FCustomizableObjectProjector FCustomizableObjectInstanceDescriptor::GetProjector(const FString& ParamName, const int32 RangeIndex) const
 {
-	const int32 Index = FindProjectorParameterNameIndex(ParamName);
+	const int32 Index = FindTypedParameterIndex(ParamName, EMutableParameterType::Projector);
 
-	if (Index != -1)
+	if (Index != INDEX_NONE)
 	{
 		if (RangeIndex == -1)
 		{
@@ -1877,94 +1873,17 @@ FCustomizableObjectProjector FCustomizableObjectInstanceDescriptor::GetProjector
 }
 
 
-int32 FCustomizableObjectInstanceDescriptor::FindIntParameterNameIndex(const FString& ParamName) const
+int32 FCustomizableObjectInstanceDescriptor::FindTypedParameterIndex(const FString& ParamName, EMutableParameterType Type) const
 {
-	if (const int32* Found = IntParametersLookupTable.Find(ParamName))
-	{
-		return *Found;
-	}
-	else
-	{
-		return INDEX_NONE;
-	}
-}
-
-
-int32 FCustomizableObjectInstanceDescriptor::FindFloatParameterNameIndex(const FString& ParamName) const
-{
-	for (int32 i = 0; i < FloatParameters.Num(); ++i)
-	{
-		if (FloatParameters[i].ParameterName == ParamName)
-		{
-			return i;
-		}
-	}
-
-	return -1;
-}
-
-
-int32 FCustomizableObjectInstanceDescriptor::FindTextureParameterNameIndex(const FString& ParamName) const
-{
-	for (int32 i = 0; i < TextureParameters.Num(); ++i)
-	{
-		if (TextureParameters[i].ParameterName == ParamName)
-		{
-			return i;
-		}
-	}
-
-	return -1;
-}
-
-
-int32 FCustomizableObjectInstanceDescriptor::FindBoolParameterNameIndex(const FString& ParamName) const
-{
-	for (int32 i = 0; i < BoolParameters.Num(); ++i)
-	{
-		if (BoolParameters[i].ParameterName == ParamName)
-		{
-			return i;
-		}
-	}
-
-	return -1;
+	return CustomizableObject ? CustomizableObject->GetPrivate()->FindParameterTyped(ParamName, Type) : INDEX_NONE;
 }
 
 
-
-int32 FCustomizableObjectInstanceDescriptor::FindVectorParameterNameIndex(const FString& ParamName) const
-{
-	for (int32 i = 0; i < VectorParameters.Num(); ++i)
-	{
-		if (VectorParameters[i].ParameterName == ParamName)
-		{
-			return i;
-		}
-	}
-
-	return -1;
-}
-
-
-int32 FCustomizableObjectInstanceDescriptor::FindProjectorParameterNameIndex(const FString& ParamName) const
-{
-	for (int32 i = 0; i < ProjectorParameters.Num(); ++i)
-	{
-		if (ProjectorParameters[i].ParameterName == ParamName)
-		{
-			return i;
-		}
-	}
-
-	return -1;
-}
-
 int32 FCustomizableObjectInstanceDescriptor::GetProjectorValueRange(const FString& ParamName) const
 {
 	check(CustomizableObject);
 
-	const int32 ProjectorParamIndex = FindProjectorParameterNameIndex(ParamName);
+	const int32 ProjectorParamIndex = FindTypedParameterIndex(ParamName, EMutableParameterType::Projector);
 	if (ProjectorParamIndex < 0)
 	{
 		return -1;
@@ -1977,7 +1896,7 @@ int32 FCustomizableObjectInstanceDescriptor::GetIntValueRange(const FString& Par
 {
 	check(CustomizableObject);
 
-	const int32 IntParamIndex = FindIntParameterNameIndex(ParamName);
+	const int32 IntParamIndex = FindTypedParameterIndex(ParamName, EMutableParameterType::Int);
 	if (IntParamIndex < 0)
 	{
 		return -1;
@@ -1991,7 +1910,7 @@ int32 FCustomizableObjectInstanceDescriptor::GetFloatValueRange(const FString& P
 {
 	check(CustomizableObject);
 
-	const int32 FloatParamIndex = FindFloatParameterNameIndex(ParamName);
+	const int32 FloatParamIndex = FindTypedParameterIndex(ParamName, EMutableParameterType::Float);
 	if (FloatParamIndex < 0)
 	{
 		return -1;
@@ -2005,7 +1924,7 @@ int32 FCustomizableObjectInstanceDescriptor::GetTextureValueRange(const FString&
 {
 	check(CustomizableObject);
 
-	const int32 TextureParamIndex = FindTextureParameterNameIndex(ParamName);
+	const int32 TextureParamIndex = FindTypedParameterIndex(ParamName, EMutableParameterType::Texture);
 	if (TextureParamIndex < 0)
 	{
 		return -1;
@@ -2019,14 +1938,14 @@ int32 FCustomizableObjectInstanceDescriptor::AddValueToIntRange(const FString& P
 {
 	check(CustomizableObject);
 
-	const int32 intParameterIndex = FindIntParameterNameIndex(ParamName);
-	if (intParameterIndex != -1)
+	const int32 IntParameterIndex = FindTypedParameterIndex(ParamName, EMutableParameterType::Int);
+	if (IntParameterIndex != INDEX_NONE)
 	{
-		FCustomizableObjectIntParameterValue& intParameter = IntParameters[intParameterIndex];
-		const int32 ParamIndexInObject = CustomizableObject->FindParameter(intParameter.ParameterName);
+		FCustomizableObjectIntParameterValue& IntParameter = IntParameters[IntParameterIndex];
+		const int32 ParamIndexInObject = CustomizableObject->FindParameter(IntParameter.ParameterName);
 		// TODO: Define the default option in the editor instead of taking the first available, like it's currently defined for GetProjectorDefaultValue()
 		const FString DefaultValue = CustomizableObject->GetIntParameterAvailableOption(ParamIndexInObject, 0);
-		return intParameter.ParameterRangeValueNames.Add(DefaultValue);
+		return IntParameter.ParameterRangeValueNames.Add(DefaultValue);
 	}
 	return -1;
 }
@@ -2034,12 +1953,12 @@ int32 FCustomizableObjectInstanceDescriptor::AddValueToIntRange(const FString& P
 
 int32 FCustomizableObjectInstanceDescriptor::AddValueToFloatRange(const FString& ParamName)
 {
-	const int32 floatParameterIndex = FindFloatParameterNameIndex(ParamName);
-	if (floatParameterIndex != -1)
+	const int32 FloatParameterIndex = FindTypedParameterIndex(ParamName, EMutableParameterType::Float);
+	if (FloatParameterIndex != INDEX_NONE)
 	{
-		FCustomizableObjectFloatParameterValue& floatParameter = FloatParameters[floatParameterIndex];
+		FCustomizableObjectFloatParameterValue& FloatParameter = FloatParameters[FloatParameterIndex];
 		// TODO: Define the default float in the editor instead of [0.5f], like it's currently defined for GetProjectorDefaultValue()
-		return floatParameter.ParameterRangeValues.Add(0.5f);
+		return FloatParameter.ParameterRangeValues.Add(0.5f);
 	}
 	return -1;
 }
@@ -2047,8 +1966,8 @@ int32 FCustomizableObjectInstanceDescriptor::AddValueToFloatRange(const FString&
 
 int32 FCustomizableObjectInstanceDescriptor::AddValueToTextureRange(const FString& ParamName)
 {
-	const int32 TextureParameterIndex = FindTextureParameterNameIndex(ParamName);
-	if (TextureParameterIndex != -1)
+	const int32 TextureParameterIndex = FindTypedParameterIndex(ParamName, EMutableParameterType::Texture);
+	if (TextureParameterIndex != INDEX_NONE)
 	{
 		FCustomizableObjectTextureParameterValue& TextureParameter = TextureParameters[TextureParameterIndex];
 		return TextureParameter.ParameterRangeValues.Add(FName());		
@@ -2062,10 +1981,10 @@ int32 FCustomizableObjectInstanceDescriptor::AddValueToProjectorRange(const FStr
 {
 	check(CustomizableObject);
 
-	const int32 projectorParameterIndex = FindProjectorParameterNameIndex(ParamName);
-	if (projectorParameterIndex != -1)
+	const int32 ProjectorParameterIndex = FindTypedParameterIndex(ParamName, EMutableParameterType::Projector);
+	if (ProjectorParameterIndex != INDEX_NONE)
 	{
-		FCustomizableObjectProjectorParameterValue& ProjectorParameter = ProjectorParameters[projectorParameterIndex];
+		FCustomizableObjectProjectorParameterValue& ProjectorParameter = ProjectorParameters[ProjectorParameterIndex];
 		const FCustomizableObjectProjector Projector = GetCustomizableObject()->GetProjectorParameterDefaultValue(ParamName);
 		return ProjectorParameter.RangeValues.Add(Projector);
 	}
@@ -2076,8 +1995,8 @@ int32 FCustomizableObjectInstanceDescriptor::AddValueToProjectorRange(const FStr
 
 int32 FCustomizableObjectInstanceDescriptor::RemoveValueFromIntRange(const FString& ParamName)
 {
-	const int32 IntParameterIndex = FindIntParameterNameIndex(ParamName);
-	if (IntParameterIndex != -1)
+	const int32 IntParameterIndex = FindTypedParameterIndex(ParamName, EMutableParameterType::Int);
+	if (IntParameterIndex != INDEX_NONE)
 	{
 		FCustomizableObjectIntParameterValue& IntParameter = IntParameters[IntParameterIndex];
 		if (IntParameter.ParameterRangeValueNames.Num() > 0)
@@ -2092,8 +2011,8 @@ int32 FCustomizableObjectInstanceDescriptor::RemoveValueFromIntRange(const FStri
 
 int32 FCustomizableObjectInstanceDescriptor::RemoveValueFromIntRange(const FString& ParamName, const int32 RangeIndex)
 {
-	const int32 IntParameterIndex = FindIntParameterNameIndex(ParamName);
-	if (IntParameterIndex != -1)
+	const int32 IntParameterIndex = FindTypedParameterIndex(ParamName, EMutableParameterType::Int);
+	if (IntParameterIndex != INDEX_NONE)
 	{
 		FCustomizableObjectIntParameterValue& IntParameter = IntParameters[IntParameterIndex];
 		if (IntParameter.ParameterRangeValueNames.Num() > 0)
@@ -2108,8 +2027,8 @@ int32 FCustomizableObjectInstanceDescriptor::RemoveValueFromIntRange(const FStri
 
 int32 FCustomizableObjectInstanceDescriptor::RemoveValueFromFloatRange(const FString& ParamName)
 {
-	const int32 FloatParameterIndex = FindFloatParameterNameIndex(ParamName);
-	if (FloatParameterIndex != -1)
+	const int32 FloatParameterIndex = FindTypedParameterIndex(ParamName, EMutableParameterType::Float);
+	if (FloatParameterIndex != INDEX_NONE)
 	{
 		FCustomizableObjectFloatParameterValue& FloatParameter = FloatParameters[FloatParameterIndex];
 		if (FloatParameter.ParameterRangeValues.Num() > 0)
@@ -2124,8 +2043,8 @@ int32 FCustomizableObjectInstanceDescriptor::RemoveValueFromFloatRange(const FSt
 
 int32 FCustomizableObjectInstanceDescriptor::RemoveValueFromFloatRange(const FString& ParamName, const int32 RangeIndex)
 {
-	const int32 FloatParameterIndex = FindFloatParameterNameIndex(ParamName);
-	if (FloatParameterIndex != -1)
+	const int32 FloatParameterIndex = FindTypedParameterIndex(ParamName, EMutableParameterType::Float);
+	if (FloatParameterIndex != INDEX_NONE)
 	{
 		FCustomizableObjectFloatParameterValue& FloatParameter = FloatParameters[FloatParameterIndex];
 		if (FloatParameter.ParameterRangeValues.Num() > 0)
@@ -2140,8 +2059,8 @@ int32 FCustomizableObjectInstanceDescriptor::RemoveValueFromFloatRange(const FSt
 
 int32 FCustomizableObjectInstanceDescriptor::RemoveValueFromTextureRange(const FString& ParamName)
 {
-	const int32 TextureParameterIndex = FindTextureParameterNameIndex(ParamName);
-	if (TextureParameterIndex != -1)
+	const int32 TextureParameterIndex = FindTypedParameterIndex(ParamName, EMutableParameterType::Texture);
+	if (TextureParameterIndex != INDEX_NONE)
 	{
 		FCustomizableObjectTextureParameterValue& TextureParameter = TextureParameters[TextureParameterIndex];
 		if (TextureParameter.ParameterRangeValues.Num() > 0)
@@ -2156,8 +2075,8 @@ int32 FCustomizableObjectInstanceDescriptor::RemoveValueFromTextureRange(const F
 
 int32 FCustomizableObjectInstanceDescriptor::RemoveValueFromTextureRange(const FString& ParamName, const int32 RangeIndex)
 {
-	const int32 TextureParameterIndex = FindTextureParameterNameIndex(ParamName);
-	if (TextureParameterIndex != -1)
+	const int32 TextureParameterIndex = FindTypedParameterIndex(ParamName, EMutableParameterType::Texture);
+	if (TextureParameterIndex != INDEX_NONE)
 	{
 		FCustomizableObjectTextureParameterValue& TextureParameter = TextureParameters[TextureParameterIndex];
 		if (TextureParameter.ParameterRangeValues.Num() > 0)
@@ -2172,8 +2091,8 @@ int32 FCustomizableObjectInstanceDescriptor::RemoveValueFromTextureRange(const F
 
 int32 FCustomizableObjectInstanceDescriptor::RemoveValueFromProjectorRange(const FString& ParamName)
 {
-	const int32 ProjectorParameterIndex = FindProjectorParameterNameIndex(ParamName);
-	if (ProjectorParameterIndex != -1)
+	const int32 ProjectorParameterIndex = FindTypedParameterIndex(ParamName, EMutableParameterType::Projector);
+	if (ProjectorParameterIndex != INDEX_NONE)
 	{
 		FCustomizableObjectProjectorParameterValue& ProjectorParameter = ProjectorParameters[ProjectorParameterIndex];
 		if (ProjectorParameter.RangeValues.Num() > 0)
@@ -2188,8 +2107,8 @@ int32 FCustomizableObjectInstanceDescriptor::RemoveValueFromProjectorRange(const
 
 int32 FCustomizableObjectInstanceDescriptor::RemoveValueFromProjectorRange(const FString& ParamName, const int32 RangeIndex)
 {
-	const int32 ProjectorParameterIndex = FindProjectorParameterNameIndex(ParamName);
-	if (ProjectorParameterIndex != -1)
+	const int32 ProjectorParameterIndex = FindTypedParameterIndex(ParamName, EMutableParameterType::Projector);
+	if (ProjectorParameterIndex != INDEX_NONE)
 	{
 		FCustomizableObjectProjectorParameterValue& ProjectorParameter = ProjectorParameters[ProjectorParameterIndex];
 		if (ProjectorParameter.RangeValues.Num() > 0)
@@ -2507,24 +2426,4 @@ void FCustomizableObjectInstanceDescriptor::MultilayerProjectorUpdateVirtualLaye
 	MultilayerProjectors[ProjectorParamName].UpdateVirtualLayer(*this, Id, Layer);
 }
 
-
-void FCustomizableObjectInstanceDescriptor::CreateParametersLookupTable()
-{
-	check(CustomizableObject);
-
-	const int32 NumParameters = CustomizableObject->GetParameterCount();
-
-	IntParametersLookupTable.Reset();
-
-	int32 IntParameter = 0;
-	for (int32 ParamIndex = 0; ParamIndex < NumParameters; ++ParamIndex)
-	{
-		if (CustomizableObject->GetParameterType(ParamIndex) == EMutableParameterType::Int)
-		{
-			IntParametersLookupTable.Add(CustomizableObject->GetParameterName(ParamIndex), IntParameter++);
-		}
-	}
-}
-
-
 #undef RETURN_ON_CO_UNCOMPILED
diff --git a/Engine/Plugins/Experimental/Mutable/Source/CustomizableObject/Private/MuCO/MultilayerProjector.cpp b/Engine/Plugins/Experimental/Mutable/Source/CustomizableObject/Private/MuCO/MultilayerProjector.cpp
--- a/Engine/Plugins/Experimental/Mutable/Source/CustomizableObject/Private/MuCO/MultilayerProjector.cpp
+++ b/Engine/Plugins/Experimental/Mutable/Source/CustomizableObject/Private/MuCO/MultilayerProjector.cpp
@@ -24,7 +24,7 @@ void FMultilayerProjectorLayer::Read(const FCustomizableObjectInstanceDescriptor
 	}
 	
 	{
-		const int32 ProjectorParamIndex = Descriptor.FindProjectorParameterNameIndex(ParamName);
+		const int32 ProjectorParamIndex = Descriptor.FindTypedParameterIndex(ParamName, EMutableParameterType::Projector);
 		const FCustomizableObjectProjector& Projector = Descriptor.GetProjectorParameters()[ProjectorParamIndex].RangeValues[Index];
 		Position = static_cast<FVector3d>(Projector.Position);
 		Direction = static_cast<FVector3d>(Projector.Direction);
@@ -34,12 +34,12 @@ void FMultilayerProjectorLayer::Read(const FCustomizableObjectInstanceDescriptor
 	}
 
 	{
-		const int32 ImageParamIndex = Descriptor.FindIntParameterNameIndex(ParamName + FMultilayerProjector::IMAGE_PARAMETER_POSTFIX);
+		const int32 ImageParamIndex = Descriptor.FindTypedParameterIndex(ParamName + FMultilayerProjector::IMAGE_PARAMETER_POSTFIX, EMutableParameterType::Int);
 		Image = Descriptor.GetIntParameters()[ImageParamIndex].ParameterRangeValueNames[Index];
 	}
 
 	{
-		const int32 OpacityParamIndex = Descriptor.FindFloatParameterNameIndex(ParamName + FMultilayerProjector::OPACITY_PARAMETER_POSTFIX);
+		const int32 OpacityParamIndex = Descriptor.FindTypedParameterIndex(ParamName + FMultilayerProjector::OPACITY_PARAMETER_POSTFIX, EMutableParameterType::Float);
 		Opacity = Descriptor.GetFloatParameters()[OpacityParamIndex].ParameterRangeValues[Index];
 	}
 }
@@ -58,7 +58,7 @@ void FMultilayerProjectorLayer::Write(FCustomizableObjectInstanceDescriptor& Des
 	}
 
 	{
-		const int32 ProjectorParamIndex = Descriptor.FindProjectorParameterNameIndex(ParamName);
+		const int32 ProjectorParamIndex = Descriptor.FindTypedParameterIndex(ParamName, EMutableParameterType::Projector);
 		FCustomizableObjectProjector& Projector = Descriptor.GetProjectorParameters()[ProjectorParamIndex].RangeValues[Index];
 		Projector.Position = static_cast<FVector3f>(Position);
 		Projector.Direction = static_cast<FVector3f>(Direction);
@@ -68,12 +68,12 @@ void FMultilayerProjectorLayer::Write(FCustomizableObjectInstanceDescriptor& Des
 	}
 
 	{
-		const int32 ImageParamIndex = Descriptor.FindIntParameterNameIndex(ParamName + FMultilayerProjector::IMAGE_PARAMETER_POSTFIX);
+		const int32 ImageParamIndex = Descriptor.FindTypedParameterIndex(ParamName + FMultilayerProjector::IMAGE_PARAMETER_POSTFIX, EMutableParameterType::Int);
 		Descriptor.GetIntParameters()[ImageParamIndex].ParameterRangeValueNames[Index] = Image;
 	}
 
 	{
-		const int32 OpacityParamIndex = Descriptor.FindFloatParameterNameIndex(ParamName + FMultilayerProjector::OPACITY_PARAMETER_POSTFIX);
+		const int32 OpacityParamIndex = Descriptor.FindTypedParameterIndex(ParamName + FMultilayerProjector::OPACITY_PARAMETER_POSTFIX, EMutableParameterType::Float);
 		Descriptor.GetFloatParameters()[OpacityParamIndex].ParameterRangeValues[Index] = Opacity;
 	}
 }
@@ -115,7 +115,7 @@ int32 FMultilayerProjector::NumLayers(const FCustomizableObjectInstanceDescripto
 {
 	const FString NumLayersParamName = ParamName.ToString() + NUM_LAYERS_PARAMETER_POSTFIX;
 	
-	const int32 FloatParameterIndex = Descriptor.FindFloatParameterNameIndex(NumLayersParamName);
+	const int32 FloatParameterIndex = Descriptor.FindTypedParameterIndex(NumLayersParamName, EMutableParameterType::Float);
 	check(FloatParameterIndex != -1); // Parameter not found.
 
 	return Descriptor.GetFloatParameters()[FloatParameterIndex].ParameterValue;
@@ -140,14 +140,14 @@ void FMultilayerProjector::CreateLayer(FCustomizableObjectInstanceDescriptor& De
 	// Num Layers.
 	{
         const FString NumLayersParamName = ParamNameString + NUM_LAYERS_PARAMETER_POSTFIX;
-		const int32 FloatParameterIndex = Descriptor.FindFloatParameterNameIndex(NumLayersParamName);
+		const int32 FloatParameterIndex = Descriptor.FindTypedParameterIndex(NumLayersParamName, EMutableParameterType::Float);
 
         Descriptor.GetFloatParameters()[FloatParameterIndex].ParameterValue += 1;
     }
 	
 	// Projector Range.
 	{
-		const int32 ProjectorParameterIndex = Descriptor.FindProjectorParameterNameIndex(ParamNameString);
+		const int32 ProjectorParameterIndex = Descriptor.FindTypedParameterIndex(ParamNameString, EMutableParameterType::Projector);
 		
 		FCustomizableObjectProjectorParameterValue& ProjectorParameter = Descriptor.GetProjectorParameters()[ProjectorParameterIndex];
 		const FCustomizableObjectProjector Projector = Descriptor.GetCustomizableObject()->GetProjectorParameterDefaultValue(ParamNameString);
@@ -156,7 +156,7 @@ void FMultilayerProjector::CreateLayer(FCustomizableObjectInstanceDescriptor& De
 	
 	// Selected Image Range.
 	{
-		const int32 IntParameterIndex = Descriptor.FindIntParameterNameIndex(ParamNameString + FMultilayerProjector::IMAGE_PARAMETER_POSTFIX);
+		const int32 IntParameterIndex = Descriptor.FindTypedParameterIndex(ParamNameString + FMultilayerProjector::IMAGE_PARAMETER_POSTFIX, EMutableParameterType::Int);
 
 		FCustomizableObjectIntParameterValue& IntParameter = Descriptor.GetIntParameters()[IntParameterIndex];
 		const int32 ParamIndexInObject = Object->FindParameter(IntParameter.ParameterName);
@@ -167,7 +167,7 @@ void FMultilayerProjector::CreateLayer(FCustomizableObjectInstanceDescriptor& De
 	
 	// Opacity Range.
 	{
-		const int32 FloatParameterIndex = Descriptor.FindFloatParameterNameIndex(ParamNameString + FMultilayerProjector::OPACITY_PARAMETER_POSTFIX);
+		const int32 FloatParameterIndex = Descriptor.FindTypedParameterIndex(ParamNameString + FMultilayerProjector::OPACITY_PARAMETER_POSTFIX, EMutableParameterType::Float);
 
 		FCustomizableObjectFloatParameterValue& FloatParameter = Descriptor.GetFloatParameters()[FloatParameterIndex];
 		FloatParameter.ParameterRangeValues.Insert(0.5, Index); // TODO: Define the default float in the editor instead of [0.5f], like it's currently defined for GetProjectorDefaultValue()
@@ -190,14 +190,14 @@ void FMultilayerProjector::RemoveLayerAt(FCustomizableObjectInstanceDescriptor&
 	// Num Layers.
 	{
 		const FString NumLayersParamName = ParamNameString + NUM_LAYERS_PARAMETER_POSTFIX;
-        const int32 FloatParameterIndex = Descriptor.FindFloatParameterNameIndex(NumLayersParamName);
+        const int32 FloatParameterIndex = Descriptor.FindTypedParameterIndex(NumLayersParamName, EMutableParameterType::Float);
 		
     	Descriptor.GetFloatParameters()[FloatParameterIndex].ParameterValue -= 1;
     }
     
 	// Projector Range.
 	{
-		const int32 ProjectorParameterIndex = Descriptor.FindProjectorParameterNameIndex(ParamNameString);
+		const int32 ProjectorParameterIndex = Descriptor.FindTypedParameterIndex(ParamNameString, EMutableParameterType::Projector);
 		
 		FCustomizableObjectProjectorParameterValue& ProjectorParameter = Descriptor.GetProjectorParameters()[ProjectorParameterIndex];
 		ProjectorParameter.RangeValues.RemoveAt(Index);
@@ -205,7 +205,7 @@ void FMultilayerProjector::RemoveLayerAt(FCustomizableObjectInstanceDescriptor&
 	
 	// Selected Image Range.
 	{
-		const int32 IntParameterIndex = Descriptor.FindIntParameterNameIndex(ParamNameString + IMAGE_PARAMETER_POSTFIX);
+		const int32 IntParameterIndex = Descriptor.FindTypedParameterIndex(ParamNameString + IMAGE_PARAMETER_POSTFIX, EMutableParameterType::Int);
 		
 		FCustomizableObjectIntParameterValue& IntParameter = Descriptor.GetIntParameters()[IntParameterIndex];
 		IntParameter.ParameterRangeValueNames.RemoveAt(Index);
@@ -213,7 +213,7 @@ void FMultilayerProjector::RemoveLayerAt(FCustomizableObjectInstanceDescriptor&
 	
 	// Opacity Range.
 	{
-		const int32 FloatParameterIndex = Descriptor.FindFloatParameterNameIndex(ParamNameString + OPACITY_PARAMETER_POSTFIX);
+		const int32 FloatParameterIndex = Descriptor.FindTypedParameterIndex(ParamNameString + OPACITY_PARAMETER_POSTFIX, EMutableParameterType::Float);
 		
 		FCustomizableObjectFloatParameterValue& FloatParameter = Descriptor.GetFloatParameters()[FloatParameterIndex];
 		FloatParameter.ParameterRangeValues.RemoveAt(Index);
@@ -468,7 +468,7 @@ bool FMultilayerProjector::AreDescriptorParametersValid(const FCustomizableObjec
 	// Num layers.
 	{
 		const FString NumLayersParamName = ParamName + NUM_LAYERS_PARAMETER_POSTFIX;
-		const int32 FloatParameterIndex = Descriptor.FindFloatParameterNameIndex(NumLayersParamName);
+		const int32 FloatParameterIndex = Descriptor.FindTypedParameterIndex(NumLayersParamName, EMutableParameterType::Float);
 		if (FloatParameterIndex != INDEX_NONE)
 		{
 			return false;
@@ -476,21 +476,21 @@ bool FMultilayerProjector::AreDescriptorParametersValid(const FCustomizableObjec
 	}
     
 	// Projector.
-	if (const int32 ProjectorParameterIndex = Descriptor.FindProjectorParameterNameIndex(ParamName);
+	if (const int32 ProjectorParameterIndex = Descriptor.FindTypedParameterIndex(ParamName, EMutableParameterType::Projector);
 		ProjectorParameterIndex != INDEX_NONE)
 	{
 		return false;
 	}
 	
 	// Selected Image.
-	if (const int32 IntParameterIndex = Descriptor.FindIntParameterNameIndex(ParamName + IMAGE_PARAMETER_POSTFIX);
+	if (const int32 IntParameterIndex = Descriptor.FindTypedParameterIndex(ParamName + IMAGE_PARAMETER_POSTFIX, EMutableParameterType::Int);
 		IntParameterIndex != INDEX_NONE)
 	{
 		return false;
 	}
 		
 	// Opacity.	
-	if (const int32 FloatParameterIndex = Descriptor.FindFloatParameterNameIndex(ParamName + OPACITY_PARAMETER_POSTFIX);
+	if (const int32 FloatParameterIndex = Descriptor.FindTypedParameterIndex(ParamName + OPACITY_PARAMETER_POSTFIX, EMutableParameterType::Float);
 		FloatParameterIndex != INDEX_NONE)
 	{
 		return false;
