[mutable] Texture layout block IDs are no longer incremental to make packaged data more deterministic when there are data changes.
- Moved mesh vertex ID to a more general mesh ID also used for layout blocks.
- Layout blocks may be 16 bit relative or 64 bit absolute.
- An additional optimization pass will be needed to some operations as separate tasks.
#rnx

diff --git a/Engine/Plugins/Experimental/Mutable/Source/CustomizableObject/Internal/MuCO/CustomizableObjectPrivate.h b/Engine/Plugins/Experimental/Mutable/Source/CustomizableObject/Internal/MuCO/CustomizableObjectPrivate.h
--- a/Engine/Plugins/Experimental/Mutable/Source/CustomizableObject/Internal/MuCO/CustomizableObjectPrivate.h
+++ b/Engine/Plugins/Experimental/Mutable/Source/CustomizableObject/Internal/MuCO/CustomizableObjectPrivate.h
@@ -981,6 +981,8 @@ public:
 
 		NumRuntimeReferencedTextures,
 
+		DeterminisiticLayoutBlockIds,
+
 		LastCustomizableObjectVersion
 	};
 	static constexpr int32 CurrentSupportedVersion = ECustomizableObjectVersions::LastCustomizableObjectVersion;
diff --git a/Engine/Plugins/Experimental/Mutable/Source/CustomizableObject/Private/MuCO/CustomizableObjectInstanceDescriptor.cpp b/Engine/Plugins/Experimental/Mutable/Source/CustomizableObject/Private/MuCO/CustomizableObjectInstanceDescriptor.cpp
--- a/Engine/Plugins/Experimental/Mutable/Source/CustomizableObject/Private/MuCO/CustomizableObjectInstanceDescriptor.cpp
+++ b/Engine/Plugins/Experimental/Mutable/Source/CustomizableObject/Private/MuCO/CustomizableObjectInstanceDescriptor.cpp
@@ -1220,7 +1220,7 @@ const FString& FCustomizableObjectInstanceDescriptor::GetIntParameterSelectedOpt
 	const int32 ParameterIndexInObject = CustomizableObject->FindParameter(ParamName);
 	const int32 ParameterIndexInInstance = FindTypedParameterIndex(ParamName, EMutableParameterType::Int);
 
-	if (ParameterIndexInObject >= 0 && ParameterIndexInInstance >= 0)
+	if (ParameterIndexInObject >= 0 && IntParameters.IsValidIndex(ParameterIndexInInstance))
 	{
 		if (RangeIndex == -1)
 		{
diff --git a/Engine/Plugins/Experimental/Mutable/Source/CustomizableObjectEditor/Private/MuCOE/SMutableMeshViewer.cpp b/Engine/Plugins/Experimental/Mutable/Source/CustomizableObjectEditor/Private/MuCOE/SMutableMeshViewer.cpp
--- a/Engine/Plugins/Experimental/Mutable/Source/CustomizableObjectEditor/Private/MuCOE/SMutableMeshViewer.cpp
+++ b/Engine/Plugins/Experimental/Mutable/Source/CustomizableObjectEditor/Private/MuCOE/SMutableMeshViewer.cpp
@@ -273,7 +273,7 @@ TSharedRef<SWidget> SMutableMeshViewer::GenerateDataTableSlates()
 	const FText VerticesCountTitle = LOCTEXT("VerticesCountTitle", "Vertex count : ");
 	const FText FacesCountTitle = LOCTEXT("FacesCountTitle", "Face count : ");
 	const FText BonesCountTitle = LOCTEXT("BonesCountTitle", "Bone count : ");
-	const FText VertexIdPrefixTitle = LOCTEXT("VertexIdPrefixTitle", "Vertex ID prefix : ");
+	const FText MeshIdPrefixTitle = LOCTEXT("MeshIdPrefixTitle", "Mesh ID prefix : ");
 	const FText BuffersTitle = LOCTEXT("BuffersTitle", "Buffers");
 	
 	return SNew(SScrollBox)
@@ -377,14 +377,14 @@ TSharedRef<SWidget> SMutableMeshViewer::GenerateDataTableSlates()
 						.AutoWidth()
 						[
 							SNew(STextBlock).
-								Text(VertexIdPrefixTitle)
+								Text(MeshIdPrefixTitle)
 						]
 
 						+ SHorizontalBox::Slot()
 						.AutoWidth()
 						[
 							SNew(STextBlock).
-								Text(this, &SMutableMeshViewer::GetVertexIdPrefix)
+								Text(this, &SMutableMeshViewer::GetMeshIdPrefix)
 						]
 				]
 			]
@@ -689,9 +689,9 @@ FText SMutableMeshViewer::GetBoneCount() const
 	return FText::AsNumber(MutableMesh && MutableMesh->GetSkeleton() ? MutableMesh->GetSkeleton()->GetBoneCount() : 0);
 }
 
-FText SMutableMeshViewer::GetVertexIdPrefix() const
+FText SMutableMeshViewer::GetMeshIdPrefix() const
 {
-	return FText::AsNumber(MutableMesh ? MutableMesh->VertexIDPrefix : 0);
+	return FText::AsNumber(MutableMesh ? MutableMesh->MeshIDPrefix : 0);
 }
 
 
diff --git a/Engine/Plugins/Experimental/Mutable/Source/CustomizableObjectEditor/Private/MuCOE/SMutableMeshViewer.h b/Engine/Plugins/Experimental/Mutable/Source/CustomizableObjectEditor/Private/MuCOE/SMutableMeshViewer.h
--- a/Engine/Plugins/Experimental/Mutable/Source/CustomizableObjectEditor/Private/MuCOE/SMutableMeshViewer.h
+++ b/Engine/Plugins/Experimental/Mutable/Source/CustomizableObjectEditor/Private/MuCOE/SMutableMeshViewer.h
@@ -123,7 +123,7 @@ private:
 	FText GetVertexCount() const;
 	FText GetFaceCount() const;
 	FText GetBoneCount() const;
-	FText GetVertexIdPrefix() const;
+	FText GetMeshIdPrefix() const;
 
 	/*
 	 * Elements used to feed the buffers list (index and buffer channels as an internal list)
diff --git a/Engine/Plugins/Experimental/Mutable/Source/MutableRuntime/Internal/MuR/MeshPrivate.h b/Engine/Plugins/Experimental/Mutable/Source/MutableRuntime/Internal/MuR/MeshPrivate.h
--- a/Engine/Plugins/Experimental/Mutable/Source/MutableRuntime/Internal/MuR/MeshPrivate.h
+++ b/Engine/Plugins/Experimental/Mutable/Source/MutableRuntime/Internal/MuR/MeshPrivate.h
@@ -717,8 +717,6 @@ namespace mu
 			return 
 				Mesh 
 				&&
-				Mesh->HasVertexIds()
-				&&
 				(CurrentIdIndex < Mesh->GetVertexCount())
 				;
 		}
@@ -728,11 +726,13 @@ namespace mu
 			check(Mesh);
 			check(CurrentIdIndex < Mesh->GetVertexCount());
 
+			uint64 Id=Mesh::InvalidVertexId;
+
 			// Is it implicit?
-			if (Mesh->VertexIDPrefix && !BufferIterator.ptr())
+			if (!BufferIterator.ptr())
 			{
 				// The id is just prefix and index
-				return (uint64(Mesh->VertexIDPrefix) << 32) | uint64(CurrentIdIndex);
+				Id = (uint64(Mesh->MeshIDPrefix) << 32) | uint64(CurrentIdIndex);
 			}
 
 			// Is it relative?
@@ -740,19 +740,20 @@ namespace mu
 			{
 				// There is a buffer storing IDs without prefix because it is the same for all vertices.
 				uint32 RelativeId = BufferIterator.GetAsUINT32();
-				return (uint64(Mesh->VertexIDPrefix) << 32) | uint64(RelativeId);
+				Id = (uint64(Mesh->MeshIDPrefix) << 32) | uint64(RelativeId);
 			}
 
 			// Is it explicit?
 			else if (BufferIterator.GetFormat() == MBF_UINT64)
 			{
-				check( Mesh->VertexIDPrefix==0 );
-				uint64 Id = BufferIterator.GetAsUINT64();
-				return Id;
+				Id = BufferIterator.GetAsUINT64();
+			}
+			else
+			{
+				check(false);
 			}
 
-			check(false);
-			return 0;
+			return Id;
 		}
 
 	};
diff --git a/Engine/Plugins/Experimental/Mutable/Source/MutableRuntime/Internal/MuR/OpImageProject.h b/Engine/Plugins/Experimental/Mutable/Source/MutableRuntime/Internal/MuR/OpImageProject.h
--- a/Engine/Plugins/Experimental/Mutable/Source/MutableRuntime/Internal/MuR/OpImageProject.h
+++ b/Engine/Plugins/Experimental/Mutable/Source/MutableRuntime/Internal/MuR/OpImageProject.h
@@ -24,7 +24,7 @@ namespace mu
 		bool bIsRGBFadingEnabled, bool bIsAlphaFadingEnabled,
 		ESamplingMethod SamplingMethod,
 		float FadeStart, float FadeEnd, float MipInterpolationFactor,
-		int Layout, int32 Block,
+		int Layout, uint64 BlockId,
 		UE::Math::TIntVector2<uint16> CropMin, UE::Math::TIntVector2<uint16> UncroppedSize,
 		FScratchImageProject* scratch, bool bUseVectorImplementation = false);
 
@@ -43,7 +43,7 @@ namespace mu
 		bool bIsRGBFadingEnabled, bool bIsAlphaFadingEnabled,
 		ESamplingMethod SamplingMethod,
 		float FadeStart, float FadeEnd, float MipInterpolationFactor,
-		int32 Layout, int32 Block,
+		int32 Layout, uint64 BlockId,
 		UE::Math::TIntVector2<uint16> CropMin, UE::Math::TIntVector2<uint16> UncroppedSize,
 		FScratchImageProject* Scratch, bool bUseVectorImplementation = false);
 
diff --git a/Engine/Plugins/Experimental/Mutable/Source/MutableRuntime/Internal/MuR/OpMeshMerge.h b/Engine/Plugins/Experimental/Mutable/Source/MutableRuntime/Internal/MuR/OpMeshMerge.h
--- a/Engine/Plugins/Experimental/Mutable/Source/MutableRuntime/Internal/MuR/OpMeshMerge.h
+++ b/Engine/Plugins/Experimental/Mutable/Source/MutableRuntime/Internal/MuR/OpMeshMerge.h
@@ -333,7 +333,7 @@ namespace mu
 			MUTABLE_CPUPROFILER_SCOPE(Layouts);
 
 			Result->Layouts.SetNum(pFirst->Layouts.Num());
-			for (int i = 0; i < pFirst->Layouts.Num(); ++i)
+			for (int32 i = 0; i < pFirst->Layouts.Num(); ++i)
 			{
 				const Layout* pF = pFirst->Layouts[i].get();
 				LayoutPtr pR = pF->Clone();
@@ -342,7 +342,7 @@ namespace mu
 				{
 					const Layout* pS = pSecond->Layouts[i].get();
 
-					pR->m_blocks.Append(pS->m_blocks);
+					pR->Blocks.Append(pS->Blocks);
 				}
 
 				Result->Layouts[i] = pR;
@@ -679,6 +679,14 @@ namespace mu
 			}
 		}
 
+		// This affects both vertex IDs and layout block ids.
+		bool bNeedsExplicitIds = pFirst->MeshIDPrefix != pSecond->MeshIDPrefix;
+		if (!bNeedsExplicitIds)
+		{
+			// This is needed in case a mesh merges with itself.
+			Result->MeshIDPrefix = pFirst->MeshIDPrefix;
+		}
+
 		// Vertices
 		//-----------------
 		{
@@ -688,16 +696,9 @@ namespace mu
 			const int32 SecondCount = pSecond->GetVertexBuffers().GetElementCount();
 
 			// TODO: fast path should be per-buffer
-			// TODO: when formats match, which at runtime should be always.
+			// TODO: when formats match, which at runtime should be often.
 			bool bFastPath = pFirst->GetVertexBuffers().HasSameFormat( pSecond->GetVertexBuffers() );
-
-			bool bNeedsExplicitVertexIds = pFirst->VertexIDPrefix != pSecond->VertexIDPrefix;
-			if (!bNeedsExplicitVertexIds)
-			{
-				Result->VertexIDPrefix = pFirst->VertexIDPrefix;
-			}
-
-			bFastPath = bFastPath && !bNeedsExplicitVertexIds;
+			bFastPath = bFastPath && !bNeedsExplicitIds;
 
 			// Check if the format of the BoneIndex buffer has to change
 			bool bChangeBoneIndicesFormat = false;
@@ -705,49 +706,10 @@ namespace mu
 
 			// Iterate all vertex buffers to check if we need to format bone indices
 			{
-				// TODO: deduplicate code. move all to cpp.
+				bChangeBoneIndicesFormat = pFirst->GetVertexBuffers().HasAnySemanticWithDifferentFormat(MBS_BONEINDICES, BoneIndexFormat);
+				if (!bChangeBoneIndicesFormat)
 				{
-					// Check if pFirst requires a reformat of the bone index buffers
-					const FMeshBufferSet& VertexBuffers = pFirst->GetVertexBuffers();
-					for (int32 VertexBufferIndex = 0; !bChangeBoneIndicesFormat && VertexBufferIndex < VertexBuffers.m_buffers.Num(); ++VertexBufferIndex)
-					{
-						const FMeshBuffer& Buffer = VertexBuffers.m_buffers[VertexBufferIndex];
-
-						const int32 elemSize = VertexBuffers.GetElementSize(VertexBufferIndex);
-						const int32 firstSize = FirstCount * elemSize;
-
-						const int32 ChannelsCount = VertexBuffers.GetBufferChannelCount(VertexBufferIndex);
-						for (int32 ChannelIndex = 0; ChannelIndex < ChannelsCount; ++ChannelIndex)
-						{
-							if (Buffer.m_channels[ChannelIndex].m_semantic == MBS_BONEINDICES)
-							{
-								bChangeBoneIndicesFormat = Buffer.m_channels[ChannelIndex].m_format != BoneIndexFormat;
-								break;
-							}
-						}
-					}
-				}
-
-				{
-					// Check if pSecond requires a reformat of the bone index buffers
-					const FMeshBufferSet& VertexBuffers = pSecond->GetVertexBuffers();
-					for (int32 VertexBufferIndex = 0; !bChangeBoneIndicesFormat && VertexBufferIndex < VertexBuffers.m_buffers.Num(); ++VertexBufferIndex)
-					{
-						const FMeshBuffer& Buffer = VertexBuffers.m_buffers[VertexBufferIndex];
-
-						const int32 elemSize = VertexBuffers.GetElementSize(VertexBufferIndex);
-						const int32 firstSize = FirstCount * elemSize;
-
-						const int32 ChannelsCount = VertexBuffers.GetBufferChannelCount(VertexBufferIndex);
-						for (int32 ChannelIndex = 0; ChannelIndex < ChannelsCount; ++ChannelIndex)
-						{
-							if (Buffer.m_channels[ChannelIndex].m_semantic == MBS_BONEINDICES)
-							{
-								bChangeBoneIndicesFormat = Buffer.m_channels[ChannelIndex].m_format != BoneIndexFormat;
-								break;
-							}
-						}
-					}
+					pSecond->GetVertexBuffers().HasAnySemanticWithDifferentFormat(MBS_BONEINDICES, BoneIndexFormat);
 				}
 
 				bFastPath = bFastPath && !bChangeBoneIndicesFormat;
@@ -775,7 +737,7 @@ namespace mu
 					result.m_elementSize = first.m_elementSize;
 
 					// See if we need to enlarge the components of any of the result channels
-					bool resetOffsets = false;
+					bool bResetOffsets = false;
 					for ( int32 c=0; c<result.m_channels.Num(); ++c )
 					{
 						int32 sb = -1;
@@ -794,25 +756,16 @@ namespace mu
 								 >
 								 result.m_channels[c].m_componentCount )
 							{
-								result.m_channels[c].m_componentCount =
-										second.m_channels[sc].m_componentCount;
-								resetOffsets = true;
+								result.m_channels[c].m_componentCount = second.m_channels[sc].m_componentCount;
+								bResetOffsets = true;
 							}
 						}
 					}
 
 					// Reset the channel offsets if necessary
-					if (resetOffsets)
+					if (bResetOffsets)
 					{
-						int32 offset = 0;
-						for ( int32 c=0; c<result.m_channels.Num(); ++c )
-						{
-                            result.m_channels[c].m_offset = (uint8_t)offset;
-							offset += result.m_channels[c].m_componentCount
-									*
-									GetMeshFormatData(result.m_channels[c].m_format).SizeInBytes;
-						}
-						result.m_elementSize = offset;
+						Result->GetVertexBuffers().UpdateOffsets(vb);
 					}
 				}
 
@@ -858,16 +811,7 @@ namespace mu
 								FMeshBuffer& Buffer = VertexBuffers.m_buffers[foundBuffer];
 								Buffer.m_channels.Insert(chan, foundChannel + 1);
 
-								// Update offsets
-								int32 Offset = Buffer.m_channels[foundChannel].m_offset;
-								for (int32 c = foundChannel; c < Buffer.m_channels.Num(); ++c)
-								{
-									Buffer.m_channels[c].m_offset = (uint8)Offset;
-									Offset += Buffer.m_channels[c].m_componentCount
-										*
-										GetMeshFormatData(Buffer.m_channels[c].m_format).SizeInBytes;
-								}
-								Buffer.m_elementSize = Offset;
+								VertexBuffers.UpdateOffsets(foundBuffer);
 							}
 						}
 					}
@@ -908,27 +852,41 @@ namespace mu
 						}
 					}
 				}
-				
+
+				if (bNeedsExplicitIds)
+				{
+					// Make sure the result format is suitable for the explicit IDs
+
+					// Set to a dummy prefix, to ensure conversion happens
+					Result->MeshIDPrefix = 1;
+					Result->MakeIdsExplicit();
+				}
+
 				// Convert the source meshes to the new format
-                if (!bNeedsExplicitVertexIds && pFirst->GetVertexBuffers().HasSameFormat(Result->GetVertexBuffers()))
+                if (!bNeedsExplicitIds && pFirst->GetVertexBuffers().HasSameFormat(Result->GetVertexBuffers()))
                 {
-                    pVFirst = pFirst;
-                }
+					Result->VertexBuffers = pVFirst->VertexBuffers;
+				}
                 else
                 {
 					check(ScratchMeshes.FirstReformat);
 
 					bool bOutSuccess = false;
-					MeshFormat(ScratchMeshes.FirstReformat.get(), pFirst, Result, false, true, false, false, bOutSuccess);
 
-					if (bNeedsExplicitVertexIds)
+					if (bNeedsExplicitIds)
 					{
-						ScratchMeshes.FirstReformat->MakeVertexIndicesExplicit();
+						Ptr<Mesh> TempMesh = pFirst->Clone(EMeshCopyFlags::WithVertexBuffers);
+						TempMesh->MakeIdsExplicit();
+						MeshFormat(ScratchMeshes.FirstReformat.get(), TempMesh.get(), Result, false, true, false, false, bOutSuccess);
+					}
+					else
+					{
+						MeshFormat(ScratchMeshes.FirstReformat.get(), pFirst, Result, false, true, false, false, bOutSuccess);
 					}
 
 					if (bOutSuccess)
 					{
-						pVFirst = ScratchMeshes.FirstReformat;
+						Result->VertexBuffers = ScratchMeshes.FirstReformat->VertexBuffers;
 					}
 					else
 					{
@@ -938,7 +896,7 @@ namespace mu
                 }
 
 
-                if (!bNeedsExplicitVertexIds && pSecond->GetVertexBuffers().HasSameFormat(Result->GetVertexBuffers()))
+                if (!bNeedsExplicitIds && pSecond->GetVertexBuffers().HasSameFormat(Result->GetVertexBuffers()))
                 {
                     pVSecond = pSecond;
                 }
@@ -947,11 +905,16 @@ namespace mu
 					check(ScratchMeshes.SecondReformat);
 
 					bool bOutSuccess = false;
-                    MeshFormat(ScratchMeshes.SecondReformat.get(), pSecond, Result, false, true, false, false, bOutSuccess);
 
-					if (bNeedsExplicitVertexIds)
+					if (bNeedsExplicitIds)
+					{
+						Ptr<Mesh> TempMesh = pSecond->Clone(EMeshCopyFlags::WithVertexBuffers);
+						TempMesh->MakeIdsExplicit();
+						MeshFormat(ScratchMeshes.SecondReformat.get(), TempMesh.get(), Result, false, true, false, false, bOutSuccess);
+					}
+					else
 					{
-						ScratchMeshes.SecondReformat->MakeVertexIndicesExplicit();
+						MeshFormat(ScratchMeshes.SecondReformat.get(), pSecond, Result, false, true, false, false, bOutSuccess);
 					}
 
 					if (bOutSuccess)
@@ -964,17 +927,16 @@ namespace mu
 					}
 				}
 
-				check(pVFirst->GetVertexBuffers().HasSameFormat(pVSecond->GetVertexBuffers()));
+				check(Result->GetVertexBuffers().HasSameFormat(pVSecond->GetVertexBuffers()));
 			}
 			else
 			{
                 MUTABLE_CPUPROFILER_SCOPE(FastPath);
 
-                pVFirst = pFirst;
+				Result->VertexBuffers = pFirst->VertexBuffers;
 				pVSecond = pSecond;
 			}
 
-			Result->VertexBuffers = pVFirst->VertexBuffers;
 
 			// Allocate additional vertices
 			Result->VertexBuffers.SetElementCount( FirstCount + SecondCount );
diff --git a/Engine/Plugins/Experimental/Mutable/Source/MutableRuntime/Internal/MuR/Operations.h b/Engine/Plugins/Experimental/Mutable/Source/MutableRuntime/Internal/MuR/Operations.h
--- a/Engine/Plugins/Experimental/Mutable/Source/MutableRuntime/Internal/MuR/Operations.h
+++ b/Engine/Plugins/Experimental/Mutable/Source/MutableRuntime/Internal/MuR/Operations.h
@@ -666,7 +666,7 @@ namespace mu
         {
             ADDRESS layout, base, blockImage;
             ADDRESS mask;
-            uint32 blockIndex;
+            uint64 BlockId;
         };
 
         struct ImageInterpolateArgs
@@ -754,7 +754,9 @@ namespace mu
 
         struct ImageRasterMeshArgs
         {
-            ADDRESS mesh;
+			uint64 BlockId;
+			
+			ADDRESS mesh;
 
 			//! These are used in case of projected mesh raster.
 			ADDRESS image;
@@ -767,7 +769,6 @@ namespace mu
 			//! like cylindrical projections.
 			ADDRESS projector;
 			
-			int32 blockId;
 			uint16 sizeX, sizeY;
 			uint16 SourceSizeX, SourceSizeY;
 			uint16 CropMinX, CropMinY;
diff --git a/Engine/Plugins/Experimental/Mutable/Source/MutableRuntime/Private/MuR/CodeRunner.cpp b/Engine/Plugins/Experimental/Mutable/Source/MutableRuntime/Private/MuR/CodeRunner.cpp
--- a/Engine/Plugins/Experimental/Mutable/Source/MutableRuntime/Private/MuR/CodeRunner.cpp
+++ b/Engine/Plugins/Experimental/Mutable/Source/MutableRuntime/Private/MuR/CodeRunner.cpp
@@ -1894,9 +1894,9 @@ namespace mu
                 Ptr<const Mesh> Source = LoadMesh(FCacheAddress(source, item));
 
                 // Access with memcpy necessary for unaligned arm issues.
-                uint32 blocks[1024];
-				check(blockCount<1024);
-				FMemory::Memcpy(blocks, data, sizeof(uint32)*FMath::Min(1024,int32(blockCount)));
+                uint64 blocks[512];
+				check(blockCount< 512);
+				FMemory::Memcpy(blocks, data, sizeof(uint64)*FMath::Min(512,int32(blockCount)));
 
 				if (Source)
 				{
@@ -3719,7 +3719,7 @@ namespace mu
 					Data.Resource = const_cast<Layout*>(ComposeLayout.get());
 					int32 DataPos = m_heapData.Add(Data);
 
-					int32 RelBlockIndex = ComposeLayout->FindBlock(Args.blockIndex);
+					int32 RelBlockIndex = ComposeLayout->FindBlock(Args.BlockId);
 
 					if (RelBlockIndex >= 0)
 					{
@@ -3767,7 +3767,7 @@ namespace mu
 					Data.Resource = const_cast<Layout*>(ComposeLayout.get());
 					int32 DataPos = m_heapData.Add(Data);
 
-					int32 RelBlockIndex = ComposeLayout->FindBlock(Args.blockIndex);
+					int32 RelBlockIndex = ComposeLayout->FindBlock(Args.BlockId);
 					if (RelBlockIndex >= 0)
 					{
 						AddOp(FScheduledOp(item.At, item, 2, DataPos),
@@ -4485,7 +4485,7 @@ namespace mu
 					Ptr<Image> ResultImage = CreateImage(SizeX, SizeY, 1, EImageFormat::IF_L_UBYTE, EInitializationType::Black);
 					if (pMesh)
 					{
-						ImageRasterMesh(pMesh.get(), ResultImage.get(), args.LayoutIndex, args.blockId, CropMin, UncroppedSize);
+						ImageRasterMesh(pMesh.get(), ResultImage.get(), args.LayoutIndex, args.BlockId, CropMin, UncroppedSize);
 						Release(pMesh);
 					}
 
@@ -4687,7 +4687,7 @@ namespace mu
 							args.bIsRGBFadingEnabled, args.bIsAlphaFadingEnabled,
 							SamplingMethod,
 							FadeStartRad, FadeEndRad, FMath::Frac(Data.RasterMesh.MipValue),
-							args.LayoutIndex, args.blockId,
+							args.LayoutIndex, args.BlockId,
 							CropMin, UncroppedSize,
 							&Scratch, bUseProjectionVectorImpl);
 						break;
@@ -4698,7 +4698,7 @@ namespace mu
 							args.bIsRGBFadingEnabled, args.bIsAlphaFadingEnabled,
 							SamplingMethod,
 							FadeStartRad, FadeEndRad, FMath::Frac(Data.RasterMesh.MipValue),
-							args.LayoutIndex, args.blockId,
+							args.LayoutIndex, args.BlockId,
 							CropMin, UncroppedSize,
 							&Scratch, bUseProjectionVectorImpl);
 						break;
diff --git a/Engine/Plugins/Experimental/Mutable/Source/MutableRuntime/Private/MuR/CodeRunnerTasked.cpp b/Engine/Plugins/Experimental/Mutable/Source/MutableRuntime/Private/MuR/CodeRunnerTasked.cpp
--- a/Engine/Plugins/Experimental/Mutable/Source/MutableRuntime/Private/MuR/CodeRunnerTasked.cpp
+++ b/Engine/Plugins/Experimental/Mutable/Source/MutableRuntime/Private/MuR/CodeRunnerTasked.cpp
@@ -1900,7 +1900,7 @@ namespace mu
 
 		Ptr<const Image> Base = Runner->LoadImage({ Args.base, Op.ExecutionIndex, Op.ExecutionOptions });
 
-		int32 RelBlockIndex = Layout->FindBlock(Args.blockIndex);
+		int32 RelBlockIndex = Layout->FindBlock(Args.BlockId);
 
 		// Shortcuts
 		if (RelBlockIndex < 0)
diff --git a/Engine/Plugins/Experimental/Mutable/Source/MutableRuntime/Private/MuR/Layout.cpp b/Engine/Plugins/Experimental/Mutable/Source/MutableRuntime/Private/MuR/Layout.cpp
--- a/Engine/Plugins/Experimental/Mutable/Source/MutableRuntime/Private/MuR/Layout.cpp
+++ b/Engine/Plugins/Experimental/Mutable/Source/MutableRuntime/Private/MuR/Layout.cpp
@@ -39,10 +39,10 @@ namespace mu {
 	{
 		LLM_SCOPE_BYNAME(TEXT("MutableRuntime"));
 		LayoutPtr pResult = new Layout();		
-		pResult->m_size = m_size;
-		pResult->m_maxsize = m_maxsize;
-		pResult->m_blocks = m_blocks;
-		pResult->m_strategy = m_strategy;
+		pResult->Size = Size;
+		pResult->MaxSize = MaxSize;
+		pResult->Blocks = Blocks;
+		pResult->Strategy = Strategy;
 		pResult->FirstLODToIgnoreWarnings = FirstLODToIgnoreWarnings;
 		pResult->ReductionMethod = ReductionMethod;
 		return pResult;
@@ -52,10 +52,10 @@ namespace mu {
 	//---------------------------------------------------------------------------------------------
 	bool Layout::operator==( const Layout& o ) const
 	{
-		return  (m_size == o.m_size) &&
-			(m_maxsize == o.m_maxsize) &&
-			(m_blocks == o.m_blocks) &&
-			(m_strategy == o.m_strategy) &&
+		return (Size == o.Size) &&
+			(MaxSize == o.MaxSize) &&
+			(Blocks == o.Blocks) &&
+			(Strategy == o.Strategy) &&
 			// maybe this is not needed
 			(FirstLODToIgnoreWarnings == o.FirstLODToIgnoreWarnings) &&
 			(ReductionMethod==o.ReductionMethod);
@@ -65,52 +65,52 @@ namespace mu {
 	//---------------------------------------------------------------------------------------------
 	int32 Layout::GetDataSize() const
 	{
-		return sizeof(Layout) + m_blocks.GetAllocatedSize();
+		return sizeof(Layout) + Blocks.GetAllocatedSize();
 	}
 
 
 	//---------------------------------------------------------------------------------------------
 	FIntPoint Layout::GetGridSize() const
 	{
-		return FIntPoint(m_size[0], m_size[1]);
+		return FIntPoint(Size[0], Size[1]);
 	}
 
 
 	//---------------------------------------------------------------------------------------------
-	void Layout::SetGridSize( int sizeX, int sizeY )
+	void Layout::SetGridSize( int32 sizeX, int32 sizeY )
 	{
 		check( sizeX>=0 && sizeY>=0 );
-		m_size[0] = (uint16)sizeX;
-		m_size[1] = (uint16)sizeY;
+		Size[0] = (uint16)sizeX;
+		Size[1] = (uint16)sizeY;
 	}
 
 
 	//---------------------------------------------------------------------------------------------
-	void Layout::GetMaxGridSize(int* pSizeX, int* pSizeY) const
+	void Layout::GetMaxGridSize(int32* SizeX, int32* SizeY) const
 	{
-		check(pSizeX && pSizeY);
+		check(SizeX && SizeY);
 
-		if (pSizeX && pSizeY)
+		if (SizeX && SizeY)
 		{
-			*pSizeX = m_maxsize[0];
-			*pSizeY = m_maxsize[1];
+			*SizeX = MaxSize[0];
+			*SizeY = MaxSize[1];
 		}
 	}
 
 
 	//---------------------------------------------------------------------------------------------
-	void Layout::SetMaxGridSize(int sizeX, int sizeY)
+	void Layout::SetMaxGridSize(int32 sizeX, int32 sizeY)
 	{
 		check(sizeX >= 0 && sizeY >= 0);
-		m_maxsize[0] = (uint16)sizeX;
-		m_maxsize[1] = (uint16)sizeY;
+		MaxSize[0] = (uint16)sizeX;
+		MaxSize[1] = (uint16)sizeY;
 	}
 
 
 	//---------------------------------------------------------------------------------------------
 	int32 Layout::GetBlockCount() const
 	{
-		return m_blocks.Num();
+		return Blocks.Num();
 	}
 
 
@@ -119,22 +119,22 @@ namespace mu {
 	{
 		check( n>=0 );
 		LLM_SCOPE_BYNAME(TEXT("MutableRuntime"));
-		m_blocks.SetNum( n );
+		Blocks.SetNum( n );
 	}
 
 
 	//---------------------------------------------------------------------------------------------
-	void Layout::GetBlock( int index, uint16* pMinX, uint16* pMinY, uint16* pSizeX, uint16* pSizeY ) const
+	void Layout::GetBlock( int32 index, uint16* pMinX, uint16* pMinY, uint16* pSizeX, uint16* pSizeY ) const
 	{
-		check( index >=0 && index < m_blocks.Num() );
+		check( index >=0 && index < Blocks.Num() );
 		check( pMinX && pMinY && pSizeX && pSizeY );
 
 		if (pMinX && pMinY && pSizeX && pSizeY)
 		{
-			*pMinX = m_blocks[index].m_min[0];
-			*pMinY = m_blocks[index].m_min[1];
-			*pSizeX = m_blocks[index].m_size[0];
-			*pSizeY = m_blocks[index].m_size[1];
+			*pMinX = Blocks[index].Min[0];
+			*pMinY = Blocks[index].Min[1];
+			*pSizeX = Blocks[index].Size[0];
+			*pSizeY = Blocks[index].Size[1];
 		}
 	}
 
@@ -142,61 +142,61 @@ namespace mu {
 	//---------------------------------------------------------------------------------------------
 	void Layout::GetBlockOptions(int index, int& pPriority, bool& bReduceBothAxes, bool& bReduceByTwo) const
 	{
-		check(index >= 0 && index < m_blocks.Num());
+		check(index >= 0 && index < Blocks.Num());
 
-		pPriority = m_blocks[index].m_priority;
-		bReduceBothAxes = m_blocks[index].bReduceBothAxes;
-		bReduceByTwo = m_blocks[index].bReduceByTwo;
+		pPriority = Blocks[index].Priority;
+		bReduceBothAxes = Blocks[index].bReduceBothAxes;
+		bReduceByTwo = Blocks[index].bReduceByTwo;
 	}
 
 
 	//---------------------------------------------------------------------------------------------
     void Layout::SetBlock( int index, int minx, int miny, int sizex, int sizey )
 	{
-		check( index >=0 && index < m_blocks.Num() );
+		check( index >=0 && index < Blocks.Num() );
 
 		// Keeps the id
-		m_blocks[index].m_min = UE::Math::TIntVector2<uint16>((uint16)minx, (uint16)miny);
-		m_blocks[index].m_size = UE::Math::TIntVector2<uint16>((uint16)sizex, (uint16)sizey);
+		Blocks[index].Min = UE::Math::TIntVector2<uint16>((uint16)minx, (uint16)miny);
+		Blocks[index].Size = UE::Math::TIntVector2<uint16>((uint16)sizex, (uint16)sizey);
 	}
 
 
 	//---------------------------------------------------------------------------------------------
 	void Layout::SetBlockOptions(int index, int priority, bool bReduceBothAxes, bool bReduceByTwo)
 	{
-		check(index >= 0 && index < m_blocks.Num());
+		check(index >= 0 && index < Blocks.Num());
 
-		m_blocks[index].m_priority = priority;
-		m_blocks[index].bReduceBothAxes = bReduceBothAxes;
-		m_blocks[index].bReduceByTwo = bReduceByTwo;
+		Blocks[index].Priority = priority;
+		Blocks[index].bReduceBothAxes = bReduceBothAxes;
+		Blocks[index].bReduceByTwo = bReduceByTwo;
 	}
 
 
 	//---------------------------------------------------------------------------------------------
 	void Layout::SetLayoutPackingStrategy(EPackStrategy _strategy)
 	{
-		m_strategy = _strategy;
+		Strategy = _strategy;
 	}
 
 
 	//---------------------------------------------------------------------------------------------
 	EPackStrategy Layout::GetLayoutPackingStrategy() const
 	{
-		return m_strategy;
+		return Strategy;
 	}
 
 
 	//---------------------------------------------------------------------------------------------
 	void Layout::Serialise(OutputArchive& arch) const
 	{
-		uint32 ver = 6;
+		uint32 ver = 7;
 		arch << ver;
 
-		arch << m_size;
-		arch << m_blocks;
+		arch << Size;
+		arch << Blocks;
 
-		arch << m_maxsize;
-		arch << uint32(m_strategy);
+		arch << MaxSize;
+		arch << uint32(Strategy);
 		arch << FirstLODToIgnoreWarnings;
 		arch << uint32(ReductionMethod);
 	}
@@ -207,56 +207,33 @@ namespace mu {
 	{
 		uint32 ver;
 		arch >> ver;
-		check(ver <= 6);
+		check(ver <= 7);
 
-		arch >> m_size;
-
-		if (ver < 6)
-		{
-			uint32_t Size = 0;
-			arch >> Size;
-			m_blocks.SetNum(Size);
-
-			for (uint32_t BlockIndex = 0; BlockIndex < Size; ++BlockIndex)
-			{
-				// Unserialise taking into account the Layout version
-				m_blocks[BlockIndex].UnserialiseOldVersion(arch, ver);
-			}
-		}
-		else
-		{
-			arch >> m_blocks;
-		}
-
-		arch >> m_maxsize;
+		arch >> Size;
+		arch >> Blocks;
+		arch >> MaxSize;
 
 		uint32 temp;
 		arch >> temp;
-		m_strategy = EPackStrategy(temp);
+		Strategy = EPackStrategy(temp);
 
-		if (ver >= 4)
-		{
-			arch >> FirstLODToIgnoreWarnings;
-		}
+		arch >> FirstLODToIgnoreWarnings;
 
-		if (ver >= 5)
-		{
-			arch >> temp;
-			ReductionMethod = EReductionMethod(temp);
-		}
+		arch >> temp;
+		ReductionMethod = EReductionMethod(temp);
 	}
 
 
 	//---------------------------------------------------------------------------------------------
 	bool Layout::IsSimilar(const Layout& o) const
 	{
-		if (m_size != o.m_size || m_maxsize != o.m_maxsize ||
-			m_blocks.Num() != o.m_blocks.Num() || m_strategy != o.m_strategy)
+		if (Size != o.Size || MaxSize != o.MaxSize ||
+			Blocks.Num() != o.Blocks.Num() || Strategy != o.Strategy)
 			return false;
 
-		for (int32 i = 0; i < m_blocks.Num(); ++i)
+		for (int32 i = 0; i < Blocks.Num(); ++i)
 		{
-			if (!m_blocks[i].IsSimilar(o.m_blocks[i])) return false;
+			if (!Blocks[i].IsSimilar(o.Blocks[i])) return false;
 		}
 
 		return true;
@@ -265,28 +242,26 @@ namespace mu {
 
 
 	//---------------------------------------------------------------------------------------------
-	int32 Layout::FindBlock(int32_t id) const
+	int32 Layout::FindBlock(uint64 Id) const
 	{
-		check(id >= 0);
-		int res = -1;
-		for (int32 i = 0; res < 0 && i < m_blocks.Num(); ++i)
+		for (int32 i = 0; i < Blocks.Num(); ++i)
 		{
-			if (m_blocks[i].m_id == id)
+			if (Blocks[i].Id == Id)
 			{
-				res = i;
+				return i;
 			}
 		}
 
-		return res;
+		return -1;
 	}
 
 
 	//---------------------------------------------------------------------------------------------
 	bool Layout::IsSingleBlockAndFull() const
 	{
-		if (m_blocks.Num() == 1
-			&& m_blocks[0].m_min == UE::Math::TIntVector2<uint16>(0, 0)
-			&& m_blocks[0].m_size == m_size)
+		if (Blocks.Num() == 1
+			&& Blocks[0].Min == UE::Math::TIntVector2<uint16>(0, 0)
+			&& Blocks[0].Size == Size)
 		{
 			return true;
 		}
@@ -325,10 +300,10 @@ namespace mu {
 	//---------------------------------------------------------------------------------------------
 	void Layout::FBlock::Serialise(OutputArchive& arch) const
 	{
-		arch << m_min;
-		arch << m_size;
-		arch << m_id;
-		arch << m_priority;
+		arch << Min;
+		arch << Size;
+		arch << Id;
+		arch << Priority;
 		arch << bReduceBothAxes;
 		arch << bReduceByTwo;
 	}
@@ -337,27 +312,13 @@ namespace mu {
 	//---------------------------------------------------------------------------------------------
 	void Layout::FBlock::Unserialise(InputArchive& arch)
 	{
-		arch >> m_min;
-		arch >> m_size;
-		arch >> m_id;
-		arch >> m_priority;
+		arch >> Min;
+		arch >> Size;
+		arch >> Id;
+		arch >> Priority;
 		arch >> bReduceBothAxes;
 		arch >> bReduceByTwo;
 	}
 
-	
-	//---------------------------------------------------------------------------------------------
-	void Layout::FBlock::UnserialiseOldVersion(InputArchive& Archive, const int32 Version)
-	{
-		Archive >> m_min;
-		Archive >> m_size;
-		Archive >> m_id;
-		Archive >> m_priority;
-
-		if (Version >= 5)
-		{
-			Archive >> bReduceBothAxes;
-		}
-	}
 }
 
diff --git a/Engine/Plugins/Experimental/Mutable/Source/MutableRuntime/Private/MuR/Mesh.cpp b/Engine/Plugins/Experimental/Mutable/Source/MutableRuntime/Private/MuR/Mesh.cpp
--- a/Engine/Plugins/Experimental/Mutable/Source/MutableRuntime/Private/MuR/Mesh.cpp
+++ b/Engine/Plugins/Experimental/Mutable/Source/MutableRuntime/Private/MuR/Mesh.cpp
@@ -54,7 +54,7 @@ MeshPtr Mesh::Clone() const
 	pResult->PhysicsBody = PhysicsBody;
 	pResult->Tags = Tags;
 	pResult->StreamedResources = StreamedResources;
-	pResult->VertexIDPrefix = VertexIDPrefix;
+	pResult->MeshIDPrefix = MeshIDPrefix;
 
     // Clone the main buffers
     pResult->VertexBuffers = VertexBuffers;
@@ -178,7 +178,7 @@ void Mesh::CopyFrom(const Mesh& From, EMeshCopyFlags Flags)
 
     InternalId = From.InternalId;
 	StaticFormatFlags = From.StaticFormatFlags;
-	VertexIDPrefix = From.VertexIDPrefix;
+	MeshIDPrefix = From.MeshIDPrefix;
 
 	if (EnumHasAnyFlags(Flags, EMeshCopyFlags::WithSurfaces))
 	{
@@ -282,26 +282,6 @@ const FMeshBufferSet& Mesh::GetVertexBuffers() const
 }
 
 
-bool Mesh::HasVertexIds() const
-{
-	if (VertexIDPrefix != 0) return true;
-
-	// If there is no prefix, there must be a buffer with full indices
-	int32 BufferIndex = -1;
-	int32 ChannelIndex = -1;
-	VertexBuffers.FindChannel(MBS_VERTEXINDEX, 0, &BufferIndex, &ChannelIndex);
-	if (BufferIndex >= 0 && ChannelIndex >= 0)
-	{
-		if (VertexBuffers.m_buffers[BufferIndex].m_channels[ChannelIndex].m_format == MBF_UINT64)
-		{
-			return true;
-		}
-	}
-
-	return false;
-}
-
-
 bool Mesh::AreVertexIdsImplicit() const
 {
 	// Is there a buffer for vertex ids?
@@ -309,7 +289,7 @@ bool Mesh::AreVertexIdsImplicit() const
 	int32 ChannelIndex = -1;
 	VertexBuffers.FindChannel(MBS_VERTEXINDEX, 0, &BufferIndex, &ChannelIndex);
 
-	return (VertexIDPrefix != 0) && (BufferIndex < 0) && (ChannelIndex < 0);
+	return (MeshIDPrefix != 0) && (BufferIndex < 0) && (ChannelIndex < 0);
 }
 
 
@@ -323,13 +303,13 @@ bool Mesh::AreVertexIdsExplicit() const
 	bool bExplicit = (BufferIndex >= 0) && (ChannelIndex >= 0) && (VertexBuffers.m_buffers[BufferIndex].m_channels[ChannelIndex].m_format == MBF_UINT64);
 	if (bExplicit)
 	{
-		check(VertexIDPrefix == 0);
+		check(MeshIDPrefix == 0);
 	}
 	return bExplicit;
 }
 
 
-void Mesh::MakeVertexIndicesRelative()
+void Mesh::MakeVertexIdsRelative()
 {
 	check(AreVertexIdsImplicit());
 
@@ -351,11 +331,11 @@ void Mesh::MakeVertexIndicesRelative()
 }
 
 
-void Mesh::MakeVertexIndicesExplicit()
+void Mesh::MakeIdsExplicit()
 {
 	MUTABLE_CPUPROFILER_SCOPE(Mesh_MakeVertexIndicesExplicit);
 
-	if (!VertexIDPrefix)
+	if (!MeshIDPrefix)
 	{
 		// We already have explicit vertex IDs, or we don't have any.
 		return;
@@ -363,72 +343,121 @@ void Mesh::MakeVertexIndicesExplicit()
 
 	int32 VertexCount = GetVertexCount();
 
-	bool bHasRelativeVertexIndices = false;
-
-	int32 OldBuf = -1;
-	int32 OldChan = -1;
-	VertexBuffers.FindChannel(MBS_VERTEXINDEX, 0, &OldBuf, &OldChan);
-	bool bHasVertexIndices = (OldBuf >= 0 && OldChan >= 0);
-	if (bHasVertexIndices)
+	// Vertex IDs
 	{
-		check(OldChan == 0 && VertexBuffers.m_buffers[OldBuf].m_channels.Num() == 1);
+		MUTABLE_CPUPROFILER_SCOPE(VertexIDs);
 
-		bool bHasExplicitVertexIndices = VertexBuffers.m_buffers[OldBuf].m_channels[0].m_format == MBF_UINT64;
-		if (bHasExplicitVertexIndices)
+		bool bHasRelativeVertexIndices = false;
+
+		int32 OldBuf = -1;
+		int32 OldChan = -1;
+		VertexBuffers.FindChannel(MBS_VERTEXINDEX, 0, &OldBuf, &OldChan);
+		bool bHasVertexIndices = (OldBuf >= 0 && OldChan >= 0);
+		if (bHasVertexIndices)
 		{
-			// nothing to do
-			return;
+			check(OldChan == 0 && VertexBuffers.m_buffers[OldBuf].m_channels.Num() == 1);
+
+			FMeshBufferChannel& Channel = VertexBuffers.m_buffers[OldBuf].m_channels[0];
+
+			bool bHasExplicitVertexIndices = Channel.m_format == MBF_UINT64;
+			if (bHasExplicitVertexIndices)
+			{
+				// nothing to do
+				return;
+			}
+
+			// The mesh has relative vertex IDs.
+			bHasRelativeVertexIndices = true;
+
+			check(Channel.m_format==MBF_UINT32);
+			const uint32* OldIdData = reinterpret_cast<const uint32*>(VertexBuffers.GetBufferData(OldBuf));
+
+			TMemoryTrackedArray<uint8> NewIds;
+			NewIds.SetNumUninitialized(VertexCount * sizeof(uint64));
+			uint64* NewIdData = reinterpret_cast<uint64*>(NewIds.GetData());
+
+			for (int32 i = 0; i < VertexCount; ++i)
+			{
+				uint32 OldId = *OldIdData++;
+				uint64 Id = (uint64(MeshIDPrefix) << 32) | uint64(OldId);
+				(*NewIdData++) = Id;
+			}
+
+			// 
+			FMeshBuffer& Buffer = VertexBuffers.m_buffers[OldBuf];
+			Swap(Buffer.m_data, NewIds);
+			Buffer.m_channels[0].m_format = MBF_UINT64;
+			Buffer.m_elementSize = sizeof(uint64);
 		}
-				
-		// The mesh has relative vertex IDs.
-		bHasRelativeVertexIndices = true;
-
-		const uint32* OldIdData = reinterpret_cast<const uint32*>(VertexBuffers.GetBufferData(OldBuf));
-
-		// Create a new buffer with explicit ids and remove the old buffer with relative ids
-		int32 NewBuffer = VertexBuffers.GetBufferCount();
-		VertexBuffers.SetBufferCount(NewBuffer + 1);
-		EMeshBufferSemantic Semantic = MBS_VERTEXINDEX;
-		int32 SemanticIndex = 0;
-		EMeshBufferFormat Format = MBF_UINT64;
-		int32 Components = 1;
-		int32 Offset = 0;
-		VertexBuffers.SetBuffer(NewBuffer, sizeof(uint64), 1, &Semantic, &SemanticIndex, &Format, &Components, &Offset);
-		uint64* IdData = reinterpret_cast<uint64*>(VertexBuffers.GetBufferData(NewBuffer));
-
-		for (int32 i = 0; i < VertexCount; ++i)
+
+		if (!bHasRelativeVertexIndices)
 		{
-			uint32 OldId = *OldIdData++;
-			uint64 Id = (uint64(VertexIDPrefix) << 32) | uint64(OldId);
-			(*IdData++) = Id;
+			// The mesh has implicit Ids
+			// Create a new buffer with explicit ids
+			int32 NewBuffer = VertexBuffers.GetBufferCount();
+			VertexBuffers.SetBufferCount(NewBuffer + 1);
+			EMeshBufferSemantic Semantic = MBS_VERTEXINDEX;
+			int32 SemanticIndex = 0;
+			EMeshBufferFormat Format = MBF_UINT64;
+			int32 Components = 1;
+			int32 Offset = 0;
+			VertexBuffers.SetBuffer(NewBuffer, sizeof(uint64), 1, &Semantic, &SemanticIndex, &Format, &Components, &Offset);
+			uint64* IdData = reinterpret_cast<uint64*>(VertexBuffers.GetBufferData(NewBuffer));
+
+			for (int32 VertexIndex = 0; VertexIndex < VertexCount; ++VertexIndex)
+			{
+				uint64 Id = (uint64(MeshIDPrefix) << 32) | uint64(VertexIndex);
+				(*IdData++) = Id;
+			}
 		}
-
-		// TODO: Optimize by overwritting the existing buffer.
-		VertexBuffers.RemoveBuffer(OldBuf);
 	}
 
-	if (!bHasRelativeVertexIndices)
+	// Layout block IDs
 	{
-		// The mesh has implicit Ids
-		// Create a new buffer with explicit ids
-		int32 NewBuffer = VertexBuffers.GetBufferCount();
-		VertexBuffers.SetBufferCount(NewBuffer + 1);
-		EMeshBufferSemantic Semantic = MBS_VERTEXINDEX;
-		int32 SemanticIndex = 0;
-		EMeshBufferFormat Format = MBF_UINT64;
-		int32 Components = 1;
-		int32 Offset = 0;
-		VertexBuffers.SetBuffer(NewBuffer, sizeof(uint64), 1, &Semantic, &SemanticIndex, &Format, &Components, &Offset);
-		uint64* IdData = reinterpret_cast<uint64*>(VertexBuffers.GetBufferData(NewBuffer));
-
-		for (int32 VertexIndex = 0; VertexIndex < VertexCount; ++VertexIndex)
+		MUTABLE_CPUPROFILER_SCOPE(LayoutBlockIDs);
+
+		for (FMeshBuffer& Buffer: VertexBuffers.m_buffers)
 		{
-			uint64 Id = (uint64(VertexIDPrefix) << 32) | uint64(VertexIndex);
-			(*IdData++) = Id;
+			for (FMeshBufferChannel& Channel : Buffer.m_channels)
+			{
+				if (Channel.m_semantic != MBS_LAYOUTBLOCK)
+				{
+					continue;
+				}
+
+				if (Channel.m_format == MBF_UINT64)
+				{
+					continue;
+				}
+
+				check(Buffer.m_channels.Num() == 1);
+				check(Buffer.m_channels[0].m_offset == 0);
+				check(Buffer.m_elementSize == sizeof(uint16));
+
+				check(Channel.m_format == MBF_UINT16);
+				const uint16* OldIdData = reinterpret_cast<const uint16*>(Buffer.m_data.GetData());
+
+				TMemoryTrackedArray<uint8> NewIds;
+				NewIds.SetNumUninitialized(VertexCount *sizeof(uint64));
+				uint64* NewIdData = reinterpret_cast<uint64*>(NewIds.GetData());
+
+				for (int32 i = 0; i < VertexCount; ++i)
+				{
+					uint16 OldId = *OldIdData++;
+					uint64 Id = (uint64(MeshIDPrefix) << 32) | uint64(OldId);
+					(*NewIdData++) = Id;
+				}
+
+				// 
+				Swap(Buffer.m_data, NewIds);
+				Buffer.m_channels[0].m_format = MBF_UINT64;
+				Buffer.m_elementSize = sizeof(uint64);
+			}
 		}
 	}
 
-	VertexIDPrefix = 0;
+	// Final cleanup
+	MeshIDPrefix = 0;
 }
 
 
@@ -1103,7 +1132,7 @@ void Mesh::Serialise(OutputArchive& arch) const
 
 	arch << AdditionalPhysicsBodies;
 
-	arch << VertexIDPrefix;
+	arch << MeshIDPrefix;
 }
 
 
@@ -1235,7 +1264,7 @@ void Mesh::Unserialise(InputArchive& arch)
 
 	if (ver >= 19)
 	{
-		arch >> VertexIDPrefix;
+		arch >> MeshIDPrefix;
 	}
 }
 
@@ -1327,11 +1356,11 @@ void Mesh::CheckIntegrity() const
 			EMeshBufferFormat IdFormat = VertexBuffers.m_buffers[BufferIndex].m_channels[ChannelIndex].m_format;
 			if (IdFormat == MBF_UINT64)
 			{
-				check(VertexIDPrefix==0);
+				check(MeshIDPrefix==0);
 			}
 			else if (IdFormat == MBF_UINT32)
 			{
-				check(VertexIDPrefix != 0);
+				check(MeshIDPrefix != 0);
 			}
 			else
 			{
diff --git a/Engine/Plugins/Experimental/Mutable/Source/MutableRuntime/Private/MuR/MeshBufferSet.cpp b/Engine/Plugins/Experimental/Mutable/Source/MutableRuntime/Private/MuR/MeshBufferSet.cpp
--- a/Engine/Plugins/Experimental/Mutable/Source/MutableRuntime/Private/MuR/MeshBufferSet.cpp
+++ b/Engine/Plugins/Experimental/Mutable/Source/MutableRuntime/Private/MuR/MeshBufferSet.cpp
@@ -662,6 +662,26 @@ namespace mu
 			m_buffers[b].m_elementSize = offset;
 	}
 
+
+	bool FMeshBufferSet::HasAnySemanticWithDifferentFormat(EMeshBufferSemantic Semantic, EMeshBufferFormat ExpectedFormat) const
+	{
+		for (const FMeshBuffer& Buffer : m_buffers)
+		{
+			for (int32 ChannelIndex = 0; ChannelIndex < Buffer.m_channels.Num(); ++ChannelIndex)
+			{
+				if (Buffer.m_channels[ChannelIndex].m_semantic == Semantic)
+				{
+					if (Buffer.m_channels[ChannelIndex].m_format != ExpectedFormat)
+					{
+						return true;
+					}
+				}
+			}
+		}
+
+		return false;
+	}
+
 	
 	//-----------------------------------------------------------------------------------------
 	void FMeshBuffer::Serialise(OutputArchive& arch) const
diff --git a/Engine/Plugins/Experimental/Mutable/Source/MutableRuntime/Private/MuR/OpImageProject.cpp b/Engine/Plugins/Experimental/Mutable/Source/MutableRuntime/Private/MuR/OpImageProject.cpp
--- a/Engine/Plugins/Experimental/Mutable/Source/MutableRuntime/Private/MuR/OpImageProject.cpp
+++ b/Engine/Plugins/Experimental/Mutable/Source/MutableRuntime/Private/MuR/OpImageProject.cpp
@@ -902,7 +902,7 @@ void ImageRasterProjectedPlanar(const Mesh* pMesh, Image* pImage,
 	bool bIsRGBFadingEnabled, bool bIsAlphaFadingEnabled,
 	ESamplingMethod SamplingMethod,
 	float FadeStart, float FadeEnd, float MipInterpolationFactor,
-	int32 Layout, int32 Block,
+	int32 Layout, uint64 BlockId,
 	UE::Math::TIntVector2<uint16> CropMin, UE::Math::TIntVector2<uint16> UncroppedSize,
 	FScratchImageProject* Scratch, bool bUseVectorImplementation)
 {
@@ -958,7 +958,7 @@ void ImageRasterProjectedPlanar(const Mesh* pMesh, Image* pImage,
 			bIsRGBFadingEnabled, bIsAlphaFadingEnabled,
 			SamplingMethod,
 			FadeStart, FadeEnd,  UnusedProjectionAngle, MipInterpolationFactor, 
-			Layout, Block, CropMin, UncroppedSize, Scratch
+			Layout, BlockId, CropMin, UncroppedSize, Scratch
 		};
 
 		using EPPF = EPixelProcessorFeatures;
@@ -1116,7 +1116,7 @@ void ImageRasterProjectedWrapping( const Mesh* pMesh, Image* pImage,
 	bool bIsRGBFadingEnabled, bool bIsAlphaFadingEnabled,
 	ESamplingMethod SamplingMethod,
 	float FadeStart, float FadeEnd, float MipInterpolationFactor,
-	int32 Layout, int32 Block,
+	int32 Layout, uint64 BlockId,
 	UE::Math::TIntVector2<uint16> CropMin, UE::Math::TIntVector2<uint16> UncroppedSize,
 	FScratchImageProject* Scratch, bool bUseVectorImplementation)
 {
@@ -1173,7 +1173,7 @@ void ImageRasterProjectedWrapping( const Mesh* pMesh, Image* pImage,
 			bIsRGBFadingEnabled, bIsAlphaFadingEnabled,
 			SamplingMethod,
 			FadeStart, FadeEnd, UnusedProjectionAngle, MipInterpolationFactor,
-			Layout, Block, CropMin, UncroppedSize, Scratch
+			Layout, BlockId, CropMin, UncroppedSize, Scratch
 		};
 
 		using EPPF = EPixelProcessorFeatures;
diff --git a/Engine/Plugins/Experimental/Mutable/Source/MutableRuntime/Private/MuR/OpImageRasterMesh.h b/Engine/Plugins/Experimental/Mutable/Source/MutableRuntime/Private/MuR/OpImageRasterMesh.h
--- a/Engine/Plugins/Experimental/Mutable/Source/MutableRuntime/Private/MuR/OpImageRasterMesh.h
+++ b/Engine/Plugins/Experimental/Mutable/Source/MutableRuntime/Private/MuR/OpImageRasterMesh.h
@@ -27,7 +27,7 @@ namespace mu
 	};
 
 
-	inline void ImageRasterMesh( const Mesh* pMesh, Image* pImage, int32 LayoutIndex, int32 BlockId,
+	inline void ImageRasterMesh( const Mesh* pMesh, Image* pImage, int32 LayoutIndex, uint64 BlockId,
 		UE::Math::TIntVector2<uint16> CropMin, UE::Math::TIntVector2<uint16> UncroppedSize )
 	{
 		MUTABLE_CPUPROFILER_SCOPE(ImageRasterMesh)
@@ -80,7 +80,7 @@ namespace mu
 
         UntypedMeshBufferIteratorConst bloIt( pMesh->GetVertexBuffers(), MBS_LAYOUTBLOCK, LayoutIndex );
 
-        if (BlockId <0 || bloIt.GetElementSize()==0 )
+        if (BlockId==Layout::InvalidBlockId || bloIt.GetElementSize()==0 )
 		{
 			// Raster all the faces
             WhitePixelProcessor pixelProc;
@@ -108,30 +108,49 @@ namespace mu
 		{
 			// Raster only the faces in the selected block
 
-			// Get the block per face
-			TArray<uint16> blocks;
-			blocks.SetNumZeroed(vertexCount);
+			// Get the block per vertex
+			TArray<uint64> VertexBlockIds;
+			VertexBlockIds.SetNumZeroed(vertexCount);
 
-			for ( int i=0; i<vertexCount; ++i )
+			if (bloIt.GetFormat() == MBF_UINT16)
 			{
-                uint16 index=0;
-				ConvertData( 0, &index, MBF_UINT16, bloIt.ptr(), bloIt.GetFormat() );
-
-				blocks[i] = index;
-				++bloIt;
+				// Relative blocks.
+				const uint16* SourceIds = reinterpret_cast<const uint16*>(bloIt.ptr());
+				for (int32 i = 0; i < vertexCount; ++i)
+				{
+					uint64 Id = SourceIds[i];
+					Id = Id | (uint64(pMesh->MeshIDPrefix)<<32);
+					VertexBlockIds[i] = Id;
+					++SourceIds;
+				}
+			}
+			else if (bloIt.GetFormat() == MBF_UINT64)
+			{
+				// Absolute blocks.
+				const uint64* SourceIds = reinterpret_cast<const uint64*>(bloIt.ptr());
+				for (int32 i = 0; i < vertexCount; ++i)
+				{
+					uint64 Id = SourceIds[i];
+					VertexBlockIds[i] = Id;
+					++SourceIds;
+				}
+			}
+			else
+			{
+				// Format not supported
+				check(false);
 			}
 
             WhitePixelProcessor pixelProc;
 
 			const TArrayView<uint8> ImageData = pImage->DataStorage.GetLOD(0); 
 
-			//for (int f = 0; f < faceCount; ++f)
 			const auto& ProcessFace = [
-				vertices, indices, blocks, BlockId, ImageData, sizeX, sizeY, pixelProc
+				vertices, indices, VertexBlockIds, BlockId, ImageData, sizeX, sizeY, pixelProc
 			] (int32 f)
 			{
 				// TODO: Select faces outside for loop?
-				if (blocks[indices[f * 3 + 0]] == BlockId)
+				if (VertexBlockIds[indices[f * 3 + 0]] == BlockId)
 				{
 					constexpr int32 NumInterpolators = 1;
 					Triangle<NumInterpolators>(ImageData.GetData(), ImageData.Num(),
diff --git a/Engine/Plugins/Experimental/Mutable/Source/MutableRuntime/Private/MuR/OpLayoutRemoveBlocks.h b/Engine/Plugins/Experimental/Mutable/Source/MutableRuntime/Private/MuR/OpLayoutRemoveBlocks.h
--- a/Engine/Plugins/Experimental/Mutable/Source/MutableRuntime/Private/MuR/OpLayoutRemoveBlocks.h
+++ b/Engine/Plugins/Experimental/Mutable/Source/MutableRuntime/Private/MuR/OpLayoutRemoveBlocks.h
@@ -19,24 +19,86 @@ namespace mu
 		}
 
 		const Layout* Source = InMesh->GetLayout(InLayoutIndex);
-
-		// Create the list of blocks in the mesh
-		TArray<bool> blocksFound;
-		blocksFound.SetNumZeroed(1024);
+		Ptr<Layout> pResult;
 
 		UntypedMeshBufferIteratorConst itBlocks(InMesh->GetVertexBuffers(), MBS_LAYOUTBLOCK, InLayoutIndex);
 		if (itBlocks.GetFormat() == MBF_UINT16)
 		{
+			// Relative blocks.
+
 			const uint16* pBlocks = reinterpret_cast<const uint16*>(itBlocks.ptr());
+
+			uint16 MaxId = 0;
 			for (int32 i = 0; i < InMesh->GetVertexCount(); ++i)
 			{
-				if (pBlocks[i] >= blocksFound.Num())
+				uint16 RelativeId = pBlocks[i];
+				check(RelativeId != 0xffff);
+				MaxId = FMath::Max(MaxId, RelativeId);
+			}
+
+
+			// Create the list of blocks in the mesh.
+			// Array that stores a flag for every ID, possibly wasting space.
+			TBitArray BlocksFound;
+			BlocksFound.Init(false,MaxId+1);
+			for (int32 i = 0; i < InMesh->GetVertexCount(); ++i)
+			{
+				BlocksFound[pBlocks[i]] = true;
+			}
+
+			// Remove blocks that are not in the mesh
+			pResult = Source->Clone();
+			int32 DestBlockIndex = 0;
+			for (int32 BlockIndex = 0; BlockIndex < pResult->Blocks.Num(); ++BlockIndex)
+			{
+				uint64 BlockId = pResult->Blocks[BlockIndex].Id;
+				uint32 BlockIdPrefix = uint32(BlockId >> 32);
+				uint32 RelativeBlockId = uint32(BlockId & 0xffffffff);
+
+				bool bBlockMatchesPrefix = BlockIdPrefix == InMesh->MeshIDPrefix;
+
+				if (bBlockMatchesPrefix && RelativeBlockId < uint32(BlocksFound.Num()) && BlocksFound[RelativeBlockId] )
 				{
-					blocksFound.SetNumZeroed(pBlocks[i] + 1024);
+					// keep this block
+					pResult->Blocks[DestBlockIndex] = pResult->Blocks[BlockIndex];
+					++DestBlockIndex;
 				}
+			}
+			pResult->SetBlockCount(DestBlockIndex);
+
+		}
+		else if (itBlocks.GetFormat() == MBF_UINT64)
+		{
+			// Absolute blocks.
+
+			// Create the list of blocks in the mesh.
+			// TODO: SparseIndexSet? like in the 16-bit case, but per-prefix?
+			TSet<uint64> BlocksFound;
+			BlocksFound.Reserve(64);
+
+			const uint64* BlockIdsPerVertex = reinterpret_cast<const uint64*>(itBlocks.ptr());
+			for (int32 i = 0; i < InMesh->GetVertexCount(); ++i)
+			{
+				BlocksFound.FindOrAdd(BlockIdsPerVertex[i]);
+			}
 
-				blocksFound[pBlocks[i]] = true;
+			// Remove blocks that are not in the mesh
+			pResult = Source->Clone();
+			int32 DestBlockIndex = 0;
+			for (int32 BlockIndex = 0; BlockIndex < pResult->Blocks.Num(); ++BlockIndex)
+			{
+				uint64 BlockId = pResult->Blocks[BlockIndex].Id;
+				bool bBlockFound = BlocksFound.Contains(BlockId);
+
+				if (bBlockFound)
+				{
+					// keep this block
+					pResult->Blocks[DestBlockIndex] = pResult->Blocks[BlockIndex];
+					++DestBlockIndex;
+				}
 			}
+			pResult->SetBlockCount(DestBlockIndex);
+
 		}
 		else if (itBlocks.GetFormat() == MBF_NONE)
 		{
@@ -50,20 +112,6 @@ namespace mu
 			check(false);
 		}
 
-		// Remove blocks that are not in the mesh
-		Ptr<Layout> pResult = Source->Clone();
-		int32 dest = 0;
-		for (int32 b = 0; b < pResult->m_blocks.Num(); ++b)
-		{
-			int blockIndex = pResult->m_blocks[b].m_id;
-			if (blockIndex < (int)blocksFound.Num() && blocksFound[blockIndex])
-			{
-				// keep
-				pResult->m_blocks[dest] = pResult->m_blocks[b];
-				++dest;
-			}
-		}
-		pResult->SetBlockCount(dest);
 
 		return pResult;
 	}
@@ -72,37 +120,21 @@ namespace mu
 	//---------------------------------------------------------------------------------------------
 	inline Ptr<Layout> LayoutRemoveBlocks(const Layout* Source, const Layout* ReferenceLayout)
 	{
-		// Create the list of blocks in the mesh
-		TArray<bool,TInlineAllocator<1024>> BlocksFound;
-		BlocksFound.SetNumZeroed(1024);
-
-		if (ReferenceLayout)
-		{
-			for (const Layout::FBlock& Block: ReferenceLayout->m_blocks)
-			{
-				if (Block.m_id >= BlocksFound.Num())
-				{
-					BlocksFound.SetNumZeroed(Block.m_id + 1024);
-				}
-
-				BlocksFound[Block.m_id] = true;
-			}
-		}
-
 		// Remove blocks that are not in the mesh
 		Ptr<Layout> pResult = Source->Clone();
-		int dest = 0;
-		for (int32 b = 0; b < pResult->m_blocks.Num(); ++b)
+		int32 DestBlockIndex = 0;
+		for (int32 BlockIndexInSource = 0; BlockIndexInSource < pResult->Blocks.Num(); ++BlockIndexInSource)
 		{
-			int blockIndex = pResult->m_blocks[b].m_id;
-			if (blockIndex < (int)BlocksFound.Num() && BlocksFound[blockIndex])
+			uint64 BlockId = pResult->Blocks[BlockIndexInSource].Id;
+			int32 BlockIndexInReference = ReferenceLayout->FindBlock(BlockId);
+			if (BlockIndexInReference >=0 )
 			{
 				// keep
-				pResult->m_blocks[dest] = pResult->m_blocks[b];
-				++dest;
+				pResult->Blocks[DestBlockIndex] = pResult->Blocks[BlockIndexInSource];
+				++DestBlockIndex;
 			}
 		}
-		pResult->SetBlockCount(dest);
+		pResult->SetBlockCount(DestBlockIndex);
 
 		return pResult;
 	}
@@ -113,20 +145,11 @@ namespace mu
 	{
 		Ptr<Layout> pResult = pA->Clone();
 
-		// This is faster but fails in the rare case of a block being in both layouts, which may 
-		// happen if we merge a mesh with itself.
-		//pResult->GetPrivate()->m_blocks.insert
-		//(
-		//	pResult->GetPrivate()->m_blocks.end(),
-		//	pB->GetPrivate()->m_blocks.begin(),
-		//	pB->GetPrivate()->m_blocks.end()
-		//);
-
-		for ( const Layout::FBlock& block: pB->m_blocks )
+		for ( const Layout::FBlock& block: pB->Blocks )
 		{
-			if ( pResult->FindBlock(block.m_id)<0 )
+			if ( pResult->FindBlock(block.Id)<0 )
 			{
-				pResult->m_blocks.Add(block);
+				pResult->Blocks.Add(block);
 			}
 		}
 
diff --git a/Engine/Plugins/Experimental/Mutable/Source/MutableRuntime/Private/MuR/OpMeshApplyLayout.h b/Engine/Plugins/Experimental/Mutable/Source/MutableRuntime/Private/MuR/OpMeshApplyLayout.h
--- a/Engine/Plugins/Experimental/Mutable/Source/MutableRuntime/Private/MuR/OpMeshApplyLayout.h
+++ b/Engine/Plugins/Experimental/Mutable/Source/MutableRuntime/Private/MuR/OpMeshApplyLayout.h
@@ -67,22 +67,41 @@ namespace mu
 
         check( layoutBuffer>=0 && layoutChannel>=0 );
         check( pApplied->GetVertexBuffers().m_buffers[layoutBuffer].m_channels.Num()==1 );
-        check( pApplied->GetVertexBuffers().m_buffers[layoutBuffer].m_channels[layoutChannel].m_format==MBF_UINT16 );
         check( pApplied->GetVertexBuffers().m_buffers[layoutBuffer].m_channels[layoutChannel].m_componentCount==1 );
 
         const uint16* pLayoutData = reinterpret_cast<const uint16*>( pApplied->GetVertexBuffers().GetBufferData( layoutBuffer ) );
+		UntypedMeshBufferIterator ItLayoutData(pApplied->GetVertexBuffers(), MBS_LAYOUTBLOCK, texCoordsSet);
 
-		// In some corner case involcing automatic LODs and remove meshes behaving differently among them we may need to remove vertices that don't 
+		// In some corner case involving automatic LODs and remove meshes behaving differently among them we may need to remove vertices that don't 
 		// have any block in the current layout. Track them here.
 		TArray<int32> VerticesToRemove;
 
         uint8* pVertices = pData;
 		for ( int32 v=0; v<pApplied->GetVertexBuffers().GetElementCount(); ++v )
 		{
-			int32 blockId = pLayoutData[v];
+			uint64 BlockId = 0;
+			if (ItLayoutData.GetFormat() == MBF_UINT16)
+			{
+				// Relative blocks.
+				const uint16* SourceIds = reinterpret_cast<const uint16*>(ItLayoutData.ptr());
+				BlockId = *SourceIds;
+				BlockId = BlockId | (uint64(pApplied->MeshIDPrefix) << 32);
+			}
+			else if (ItLayoutData.GetFormat() == MBF_UINT64)
+			{
+				// Absolute blocks.
+				const uint64* SourceIds = reinterpret_cast<const uint64*>(ItLayoutData.ptr());
+				BlockId = *SourceIds;
+			}
+			else
+			{
+				// Format not supported
+				check(false);
+			}
+			++ItLayoutData;
 
 			// TODO: This could be optimised
-			int32 relBlock = pLayout->FindBlock( blockId );
+			int32 relBlock = pLayout->FindBlock(BlockId);
 
 			// This may still happen with lower LOD and "remove meshes" in a corner case:
 			// Auto LODs with "Remove Meshes" that behave differently across LODs, and leave geometry in a block that has been removed in the higher LOD.
diff --git a/Engine/Plugins/Experimental/Mutable/Source/MutableRuntime/Private/MuR/OpMeshClipWithMesh.cpp b/Engine/Plugins/Experimental/Mutable/Source/MutableRuntime/Private/MuR/OpMeshClipWithMesh.cpp
--- a/Engine/Plugins/Experimental/Mutable/Source/MutableRuntime/Private/MuR/OpMeshClipWithMesh.cpp
+++ b/Engine/Plugins/Experimental/Mutable/Source/MutableRuntime/Private/MuR/OpMeshClipWithMesh.cpp
@@ -572,7 +572,7 @@ namespace mu { namespace
 
 		Result->GetVertexBuffers().SetElementCount(MaskVertexCount);
 		Result->GetVertexBuffers().SetBufferCount(1);
-		Result->VertexIDPrefix = Base->VertexIDPrefix;
+		Result->MeshIDPrefix = Base->MeshIDPrefix;
 
 		EMeshBufferSemantic Semantic = MBS_VERTEXINDEX;
 		int32 SemanticIndex = 0;
diff --git a/Engine/Plugins/Experimental/Mutable/Source/MutableRuntime/Private/MuR/OpMeshDifference.h b/Engine/Plugins/Experimental/Mutable/Source/MutableRuntime/Private/MuR/OpMeshDifference.h
--- a/Engine/Plugins/Experimental/Mutable/Source/MutableRuntime/Private/MuR/OpMeshDifference.h
+++ b/Engine/Plugins/Experimental/Mutable/Source/MutableRuntime/Private/MuR/OpMeshDifference.h
@@ -162,15 +162,7 @@ namespace mu
 			BufferIndex = 1;
 			Result->GetVertexBuffers().SetBuffer(BufferIndex, ElementSize, semantic.Num(), semantic.GetData(), semanticIndex.GetData(), format.GetData(), components.GetData(), offsets.GetData());
 
-
-			if (bGenerateRelativeIds)
-			{
-				Result->VertexIDPrefix = pBase->VertexIDPrefix;
-			}
-			else
-			{
-				Result->VertexIDPrefix = 0;
-			}
+			Result->MeshIDPrefix = pBase->MeshIDPrefix;
 
 			// Source vertex index channel
 			MeshVertexIdIteratorConst IdIterator( pBase );
diff --git a/Engine/Plugins/Experimental/Mutable/Source/MutableRuntime/Private/MuR/OpMeshExtractLayoutBlock.h b/Engine/Plugins/Experimental/Mutable/Source/MutableRuntime/Private/MuR/OpMeshExtractLayoutBlock.h
--- a/Engine/Plugins/Experimental/Mutable/Source/MutableRuntime/Private/MuR/OpMeshExtractLayoutBlock.h
+++ b/Engine/Plugins/Experimental/Mutable/Source/MutableRuntime/Private/MuR/OpMeshExtractLayoutBlock.h
@@ -35,7 +35,7 @@ namespace mu
 
 		// If the vertices are explicit or relative, the above operation will already handle them correctly
 		// Otherwise, create a relative vertex ID buffer if necessary
-		Result->VertexIDPrefix = Source->VertexIDPrefix;
+		Result->MeshIDPrefix = Source->MeshIDPrefix;
 		if (Source->AreVertexIdsImplicit() 
 			&& 
 			// If we extract everything, we can keep the ids implicit.
@@ -116,38 +116,39 @@ namespace mu
 
 
 	//---------------------------------------------------------------------------------------------
-    inline void MeshExtractLayoutBlock(Mesh* Result, const Mesh* pSource,
-                                           uint32 layout,
-                                           uint16 blockCount,
-                                           const uint32* pExtractBlocks, bool& bOutSuccess)
+    inline void MeshExtractLayoutBlock(Mesh* Result, const Mesh* Source,
+                                           uint32 LayoutIndex,
+                                           uint16 BlockCount,
+                                           const uint64* BlockIds, bool& bOutSuccess)
 	{
-		check(pSource);
+		check(Source);
 		bOutSuccess = true;
 		
 		// TODO: Optimise
-		Result->CopyFrom(*pSource);
+		Result->CopyFrom(*Source);
 
-		UntypedMeshBufferIteratorConst itBlocks( pSource->GetVertexBuffers(), MBS_LAYOUTBLOCK, layout );
+		UntypedMeshBufferIteratorConst itBlocks(Source->GetVertexBuffers(), MBS_LAYOUTBLOCK, LayoutIndex);
 
         if (itBlocks.GetFormat()!=MBF_NONE)
         {
-            int resultVertices = 0;
+            int32 ResultVertices = 0;
 			TArray<int32> oldToNew;
-			oldToNew.Init(-1,pSource->GetVertexCount());
+			oldToNew.Init(-1,Source->GetVertexCount());
 			TArray<int32> newToOld;
-            newToOld.Reserve( pSource->GetVertexCount() );
+            newToOld.Reserve( Source->GetVertexCount() );
 
             if ( itBlocks.GetFormat()==MBF_UINT16 )
             {
                 const uint16* pBlocks = reinterpret_cast<const uint16*>( itBlocks.ptr() );
-                for ( int i=0; i<pSource->GetVertexCount(); ++i )
+                for ( int32 i=0; i<Source->GetVertexCount(); ++i )
                 {
-                    uint32_t vertexBlock = pBlocks[i];
+                    uint64 VertexBlockRelative = pBlocks[i];
+					uint64 VertexBlockId = (uint64(Source->MeshIDPrefix) << 32) | VertexBlockRelative;
 
                     bool found = false;
-                    for ( int j=0; j<blockCount; ++j)
+                    for ( int32 j=0; j< BlockCount; ++j)
                     {
-                        if (vertexBlock==pExtractBlocks[j])
+                        if (VertexBlockId == BlockIds[j])
                         {
                             found = true;
                             break;
@@ -156,17 +157,42 @@ namespace mu
 
                     if ( found )
                     {
-                        oldToNew[i] = resultVertices++;
+                        oldToNew[i] = ResultVertices++;
                         newToOld.Add( i );
                     }
                 }
             }
-            else
+            else if (itBlocks.GetFormat() == MBF_UINT64)
+			{
+				const uint64* pBlocks = reinterpret_cast<const uint64*>(itBlocks.ptr());
+				for (int32 i = 0; i < Source->GetVertexCount(); ++i)
+				{
+					uint64 VertexBlockId = pBlocks[i];
+
+					bool found = false;
+					for (int j = 0; j < BlockCount; ++j)
+					{
+						if (VertexBlockId == BlockIds[j])
+						{
+							found = true;
+							break;
+						}
+					}
+
+					if (found)
+					{
+						oldToNew[i] = ResultVertices++;
+						newToOld.Add(i);
+					}
+				}
+			}
+			else
+
             {
                 check( false );
             }
 
-            MeshExtractFromVertices(pSource, Result, oldToNew, newToOld);
+            MeshExtractFromVertices(Source, Result, oldToNew, newToOld);
         }
 
         Result->Surfaces.Empty();
diff --git a/Engine/Plugins/Experimental/Mutable/Source/MutableRuntime/Private/MuR/OpMeshFormat.cpp b/Engine/Plugins/Experimental/Mutable/Source/MutableRuntime/Private/MuR/OpMeshFormat.cpp
--- a/Engine/Plugins/Experimental/Mutable/Source/MutableRuntime/Private/MuR/OpMeshFormat.cpp
+++ b/Engine/Plugins/Experimental/Mutable/Source/MutableRuntime/Private/MuR/OpMeshFormat.cpp
@@ -443,7 +443,7 @@ namespace mu
 		Ptr<const Mesh> Source = PureSource;
 
 		Result->CopyFrom(*Format);
-		Result->VertexIDPrefix = Source->VertexIDPrefix;
+		Result->MeshIDPrefix = Source->MeshIDPrefix;
 
 		// Make sure that the bone indices will fit in this format, or extend it.
 		if (formatVertices)
diff --git a/Engine/Plugins/Experimental/Mutable/Source/MutableRuntime/Private/MuR/OpMeshRemove.cpp b/Engine/Plugins/Experimental/Mutable/Source/MutableRuntime/Private/MuR/OpMeshRemove.cpp
--- a/Engine/Plugins/Experimental/Mutable/Source/MutableRuntime/Private/MuR/OpMeshRemove.cpp
+++ b/Engine/Plugins/Experimental/Mutable/Source/MutableRuntime/Private/MuR/OpMeshRemove.cpp
@@ -23,7 +23,7 @@ namespace mu
     {
 		MeshVertexIdIteratorConst itVI( pSource );
 		FIdInterval current;
-		current.idStart = TNumericLimits<uint64>::Max();
+		current.idStart = Mesh::InvalidVertexId;
 		current.idPosition = 0;
 		current.size = 0;
 		for ( int32 sv=0; sv<pSource->GetVertexBuffers().GetElementCount(); ++sv )
@@ -31,7 +31,7 @@ namespace mu
             uint64 id = itVI.Get();
             ++itVI;
 
-            if (current.idStart==TNumericLimits<uint64>::Max())
+            if (current.idStart== Mesh::InvalidVertexId)
             {
                 current.idStart = id;
                 current.idPosition = sv;
@@ -53,7 +53,7 @@ namespace mu
             }
         }
 
-        if (current.idStart!=TNumericLimits<uint64>::Max())
+        if (current.idStart!= Mesh::InvalidVertexId)
         {
             intervals.Add(current);
         }
@@ -202,7 +202,7 @@ namespace mu
 		// If we had implicit indices, make them explicit or relative to keep them valid
 		if (RemovedIndices && Result->AreVertexIdsImplicit())
 		{
-			Result->MakeVertexIndicesRelative();
+			Result->MakeVertexIdsRelative();
 		}
 
 		// The temp array is necessary because if the vertex buffer is not sorted according to the index buffer we cannot do it in-place
diff --git a/Engine/Plugins/Experimental/Mutable/Source/MutableRuntime/Private/MuR/Operations.cpp b/Engine/Plugins/Experimental/Mutable/Source/MutableRuntime/Private/MuR/Operations.cpp
--- a/Engine/Plugins/Experimental/Mutable/Source/MutableRuntime/Private/MuR/Operations.cpp
+++ b/Engine/Plugins/Experimental/Mutable/Source/MutableRuntime/Private/MuR/Operations.cpp
@@ -910,7 +910,7 @@ namespace mu
         {
             const uint8_t* data = program.GetOpArgsPointer(at);
             mu::OP::ADDRESS source;
-            memcpy( &source, data, sizeof(OP::ADDRESS) );
+			FMemory::Memcpy( &source, data, sizeof(OP::ADDRESS) );
             f(source);
             break;
         }
diff --git a/Engine/Plugins/Experimental/Mutable/Source/MutableRuntime/Public/MuR/Layout.h b/Engine/Plugins/Experimental/Mutable/Source/MutableRuntime/Public/MuR/Layout.h
--- a/Engine/Plugins/Experimental/Mutable/Source/MutableRuntime/Public/MuR/Layout.h
+++ b/Engine/Plugins/Experimental/Mutable/Source/MutableRuntime/Public/MuR/Layout.h
@@ -8,6 +8,7 @@
 #include "Containers/Array.h"
 #include "Misc/AssertionMacros.h"
 #include "Math/IntVector.h"
+#include "Math/NumericLimits.h"
 
 namespace mu
 {
@@ -42,6 +43,10 @@ namespace mu
 	//! \ingroup runtime
     class MUTABLERUNTIME_API Layout : public Resource
 	{
+	public:
+
+		static constexpr uint64 InvalidBlockId = TNumericLimits<uint64>::Max();
+
 	public:
 
 		//-----------------------------------------------------------------------------------------
@@ -51,11 +56,11 @@ namespace mu
 		Layout();
 
 		//! Deep clone this layout.
-		LayoutPtr Clone() const;
+		Ptr<Layout> Clone() const;
 
 		//! Serialisation
 		static void Serialise( const Layout* p, OutputArchive& arch );
-		static LayoutPtr StaticUnserialise( InputArchive& arch );
+		static Ptr<Layout> StaticUnserialise( InputArchive& arch );
 
 		//! Full compare
 		bool operator==( const Layout& other ) const;
@@ -74,18 +79,18 @@ namespace mu
 		//! on each axis.
 		//! \param sizeX width of the grid.
 		//! \param sizeY height of the grid.
-		void SetGridSize( int sizeX, int sizeY );
+		void SetGridSize( int32 SizeX, int32 SizeY );
 
 		//! Get the maximum resolution of the grid where the blocks are defined.
 		//! \param[out] pSizeX The integer pointed by this will be set to the width of the grid.
 		//! \param[out] pSizeY The integer pointed by this will be set to the height of the grid.
-		void GetMaxGridSize(int* pSizeX, int* pSizeY) const;
+		void GetMaxGridSize(int32* SizeX, int32* SizeY) const;
 
 		//! Get the maximum resolution of the grid where the blocks are defined. It must be bigger than 0
 		//! on each axis.
 		//! \param sizeX width of the grid.
 		//! \param sizeY height of the grid.
-		void SetMaxGridSize(int sizeX, int sizeY);
+		void SetMaxGridSize(int32 SizeX, int32 SizeY);
 
 		//! Return the number of blocks in this layout.
 		int32 GetBlockCount() const;
@@ -103,14 +108,14 @@ namespace mu
 		//! (from 0 to the Y size of the grid minus one )
 		//! \param[out] pSizeX will be set to the x size of the block
 		//! \param[out] pSizeY will be set to the y size of the block
-		void GetBlock( int index, uint16* pMinX, uint16* pMinY, uint16* pSizeX, uint16* pSizeY ) const;
+		void GetBlock( int32 Index, uint16* MinX, uint16* MinY, uint16* SizeX, uint16* SizeY ) const;
 
 		//! Returns the reduction priority of a block.
 		//! \param index Block to get priority.
 		//! \param[out] priority reduction priority of the block
 		//! \param[out] bReduceBothAxes reduction method of the block
 		//! \param[out] bReduceByTwo reduction method of the block
-		void GetBlockOptions(int index, int& priority, bool& bReduceBothAxes, bool& bReduceByTwo) const;
+		void GetBlockOptions(int32 Index, int32& Priority, bool& bReduceBothAxes, bool& bReduceByTwo) const;
 
 		//! Set a block of the layout.
 		//! "Position" here means the lower-left corner of the block.
@@ -121,18 +126,18 @@ namespace mu
 		//! (from 0 to the Y size of the grid minus one )
 		//! \param sizeX will be set to the x size of the block
 		//! \param sizeY will be set to the y size of the block
-        void SetBlock( int index, int minX, int minY, int sizeX, int sizeY );
+        void SetBlock( int32 Index, int32 MinX, int32 MinY, int32 SizeX, int32 SizeY );
 
 		//! Set the reduction options of a block
 		//! \param index Block to set the options
 		//! \param priority will be set to the reduction priority of the block. The blocks with the highest values will be the last to be reduced
 		//! \param bReduceBothAxes will be set to reduce the block in both axis at the same time.
 		//! \param bReduceByTwo will reduce by two blocks on a unitary reduction.
-		void SetBlockOptions(int index, int priority, bool bReduceBothAxes, bool bReduceByTwo);
+		void SetBlockOptions(int32 Index, int32 Priority, bool bReduceBothAxes, bool bReduceByTwo);
 
 		//! Set the texture layout packing strategy
 		//! By default the texture layout packing strategy is set to resizable layout
-		void SetLayoutPackingStrategy(EPackStrategy _strategy);
+		void SetLayoutPackingStrategy(EPackStrategy);
 
 		//! Set the texture layout packing strategy
 		EPackStrategy GetLayoutPackingStrategy() const;
@@ -160,24 +165,24 @@ namespace mu
 		//-----------------------------------------------------------------------------------------
 		struct FBlock
 		{
-			FBlock(UE::Math::TIntVector2<uint16> min = UE::Math::TIntVector2<uint16>(), UE::Math::TIntVector2<uint16> size = UE::Math::TIntVector2<uint16>())
+			FBlock(UE::Math::TIntVector2<uint16> InMin = UE::Math::TIntVector2<uint16>(), UE::Math::TIntVector2<uint16> InSize = UE::Math::TIntVector2<uint16>())
 			{
-				m_min = min;
-				m_size = size;
-				m_id = -1;
-				m_priority = 0;
+				Min = InMin;
+				Size = InSize;
+				Id = Layout::InvalidBlockId;
+				Priority = 0;
 				bReduceBothAxes = false;
 				bReduceByTwo = false;
 			}
 
-			UE::Math::TIntVector2<uint16> m_min = UE::Math::TIntVector2<uint16>(0,0);
-			UE::Math::TIntVector2<uint16> m_size = UE::Math::TIntVector2<uint16>(0, 0);
+			UE::Math::TIntVector2<uint16> Min = UE::Math::TIntVector2<uint16>(0,0);
+			UE::Math::TIntVector2<uint16> Size = UE::Math::TIntVector2<uint16>(0, 0);
 
 			//! Absolute id used to control merging of various layouts
-			int32 m_id;
+			uint64 Id;
 
 			//! Priority value to control the shrink texture layout strategy
-			int32 m_priority;
+			int32 Priority;
 
 			//! Value to control the method to reduce the block
 			bool bReduceBothAxes;
@@ -195,10 +200,10 @@ namespace mu
 			//!
 			inline bool operator==(const FBlock& o) const
 			{
-				return (m_min == o.m_min) &&
-					(m_size == o.m_size) &&
-					(m_id == o.m_id) &&
-					(m_priority == o.m_priority) &&
+				return (Min == o.Min) &&
+					(Size == o.Size) &&
+					(Id == o.Id) &&
+					(Priority == o.Priority) &&
 					(bReduceBothAxes == o.bReduceBothAxes) &&
 					(bReduceByTwo == o.bReduceByTwo);
 			}
@@ -206,28 +211,25 @@ namespace mu
 			inline bool IsSimilar(const FBlock& o) const
 			{
 				// All but ids
-				return (m_min == o.m_min) &&
-					(m_size == o.m_size) &&
-					(m_priority == o.m_priority) &&
+				return (Min == o.Min) &&
+					(Size == o.Size) &&
+					(Priority == o.Priority) &&
 					(bReduceBothAxes == o.bReduceBothAxes) &&
 					(bReduceByTwo == o.bReduceByTwo);
 			}
-
-			//!
-			void UnserialiseOldVersion(InputArchive& Archive, const int32 Version);
 		};
 
 
 		//!
-		UE::Math::TIntVector2<uint16> m_size = UE::Math::TIntVector2<uint16>(0, 0);
+		UE::Math::TIntVector2<uint16> Size = UE::Math::TIntVector2<uint16>(0, 0);
 
-		UE::Math::TIntVector2<uint16> m_maxsize = UE::Math::TIntVector2<uint16>(0, 0);
+		UE::Math::TIntVector2<uint16> MaxSize = UE::Math::TIntVector2<uint16>(0, 0);
 
 		//!
-		TArray<FBlock> m_blocks;
+		TArray<FBlock> Blocks;
 
 		//! Packing strategy
-		EPackStrategy m_strategy = EPackStrategy::RESIZABLE_LAYOUT;
+		EPackStrategy Strategy = EPackStrategy::RESIZABLE_LAYOUT;
 		 
 		int32 FirstLODToIgnoreWarnings = 0;
 
@@ -243,9 +245,8 @@ namespace mu
 		//!
 		bool IsSimilar(const Layout& o) const;
 
-		//! Find a block by id. This converts the "absolute" id to a relative index to the layout
-		//! blocks. Return -1 if not found.
-		int32 FindBlock(int32 id) const;
+		/** Find a block by id. This converts the "absolute" id to a relative index to the layout blocks. Return -1 if not found. */
+		int32 FindBlock(uint64 Id) const;
 
 		//! Return true if the layout is a single block filling all area.
 		bool IsSingleBlockAndFull() const;
diff --git a/Engine/Plugins/Experimental/Mutable/Source/MutableRuntime/Public/MuR/Mesh.h b/Engine/Plugins/Experimental/Mutable/Source/MutableRuntime/Public/MuR/Mesh.h
--- a/Engine/Plugins/Experimental/Mutable/Source/MutableRuntime/Public/MuR/Mesh.h
+++ b/Engine/Plugins/Experimental/Mutable/Source/MutableRuntime/Public/MuR/Mesh.h
@@ -140,7 +140,11 @@ namespace mu
     //! \ingroup runtime
     class MUTABLERUNTIME_API Mesh : public Resource
     {
-    public:
+	public:
+
+		static constexpr uint64 InvalidVertexId = TNumericLimits<uint64>::Max();
+
+	public:
 
         //-----------------------------------------------------------------------------------------
         // Life cycle
@@ -202,9 +206,6 @@ namespace mu
 
         //! \}
 
-		/** Return true if the mesh has unique IDs for its vertices. */
-		bool HasVertexIds() const;
-
 		/** Return true if the mesh has unique vertex IDs and they stored in an implicit way. 
 		* This is relevant for some mesh operations that will need to make them explicit so that the result is still correct.
 		*/
@@ -212,8 +213,8 @@ namespace mu
 		bool AreVertexIdsExplicit() const;
 
 		/** Create an explicit vertex buffer for vertex IDs if they are implicit. */
-		void MakeVertexIndicesRelative();
-		void MakeVertexIndicesExplicit();
+		void MakeVertexIdsRelative();
+		void MakeIdsExplicit();
 
         //! \name Texture layouts
         //! \{
@@ -335,10 +336,10 @@ namespace mu
 		//! It should be reset after any operation that modifies the format.
 		mutable uint32 StaticFormatFlags = 0;
 
-		/** Prefix for the unique vertex IDs, if the mesh has them in an implicit, or relative way. 
+		/** Prefix for the unique IDs related to this mesh (vertices and layout blocks). Useful if the mesh stores them in an implicit, or relative way. 
 		* See MeshVertexIdIterator for details.
 		*/
-		uint32 VertexIDPrefix = 0;
+		uint32 MeshIDPrefix = 0;
 
 		//!
 		FMeshBufferSet VertexBuffers;
@@ -406,54 +407,56 @@ namespace mu
         //!
 		inline bool operator==(const Mesh& o) const
 		{
-			bool equal = (IndexBuffers == o.IndexBuffers);
-			if (equal) equal = (VertexBuffers == o.VertexBuffers);
-			if (equal) equal = (Layouts.Num() == o.Layouts.Num());
-			if (equal) equal = (BonePoses.Num() == o.BonePoses.Num());
-			if (equal) equal = (BoneMap.Num() == o.BoneMap.Num());
-			if (equal && Skeleton != o.Skeleton)
+			bool bEqual = true;
+
+			if (bEqual) bEqual = (MeshIDPrefix == o.MeshIDPrefix);
+			if (bEqual) bEqual = (IndexBuffers == o.IndexBuffers);
+			if (bEqual) bEqual = (VertexBuffers == o.VertexBuffers);
+			if (bEqual) bEqual = (Layouts.Num() == o.Layouts.Num());
+			if (bEqual) bEqual = (BonePoses.Num() == o.BonePoses.Num());
+			if (bEqual) bEqual = (BoneMap.Num() == o.BoneMap.Num());
+			if (bEqual && Skeleton != o.Skeleton)
 			{
 				if (Skeleton && o.Skeleton)
 				{
-					equal = (*Skeleton == *o.Skeleton);
+					bEqual = (*Skeleton == *o.Skeleton);
 				}
 				else
 				{
-					equal = false;
+					bEqual = false;
 				}
 			}
-			if (equal) equal = (StreamedResources == o.StreamedResources);
-			if (equal) equal = (Surfaces == o.Surfaces);
-			if (equal) equal = (Tags == o.Tags);
-			if (equal) equal = (SkeletonIDs == o.SkeletonIDs);
-			if (equal) equal = (VertexIDPrefix == o.VertexIDPrefix);
+			if (bEqual) bEqual = (StreamedResources == o.StreamedResources);
+			if (bEqual) bEqual = (Surfaces == o.Surfaces);
+			if (bEqual) bEqual = (Tags == o.Tags);
+			if (bEqual) bEqual = (SkeletonIDs == o.SkeletonIDs);
 
-			for (int32 i = 0; equal && i < Layouts.Num(); ++i)
+			for (int32 i = 0; bEqual && i < Layouts.Num(); ++i)
 			{
-				equal &= (*Layouts[i]) == (*o.Layouts[i]);
+				bEqual &= (*Layouts[i]) == (*o.Layouts[i]);
 			}
 
-			equal &= AdditionalBuffers.Num() == o.AdditionalBuffers.Num();
-			for (int32 i = 0; equal && i < AdditionalBuffers.Num(); ++i)
+			bEqual &= AdditionalBuffers.Num() == o.AdditionalBuffers.Num();
+			for (int32 i = 0; bEqual && i < AdditionalBuffers.Num(); ++i)
 			{
-				equal &= AdditionalBuffers[i] == o.AdditionalBuffers[i];
+				bEqual &= AdditionalBuffers[i] == o.AdditionalBuffers[i];
 			}
 
-			equal &= BonePoses.Num() == o.BonePoses.Num();
-			for (int32 i = 0; equal && i < BonePoses.Num(); ++i)
+			bEqual &= BonePoses.Num() == o.BonePoses.Num();
+			for (int32 i = 0; bEqual && i < BonePoses.Num(); ++i)
 			{
-				equal &= BonePoses[i] == o.BonePoses[i];
+				bEqual &= BonePoses[i] == o.BonePoses[i];
 			}
 
-			if (equal) equal = BoneMap == o.BoneMap;
+			if (bEqual) bEqual = BoneMap == o.BoneMap;
 
-			equal &= AdditionalPhysicsBodies.Num() == o.AdditionalPhysicsBodies.Num();
-			for (int32 i = 0; equal && i < AdditionalPhysicsBodies.Num(); ++i)
+			bEqual &= AdditionalPhysicsBodies.Num() == o.AdditionalPhysicsBodies.Num();
+			for (int32 i = 0; bEqual && i < AdditionalPhysicsBodies.Num(); ++i)
 			{
-				equal &= *AdditionalPhysicsBodies[i] == *o.AdditionalPhysicsBodies[i];
+				bEqual &= *AdditionalPhysicsBodies[i] == *o.AdditionalPhysicsBodies[i];
 			}
 
-			return equal;
+			return bEqual;
 		}
 
 		//! Compare the mesh with another one, but ignore internal data like generated vertex
diff --git a/Engine/Plugins/Experimental/Mutable/Source/MutableRuntime/Public/MuR/MeshBufferSet.h b/Engine/Plugins/Experimental/Mutable/Source/MutableRuntime/Public/MuR/MeshBufferSet.h
--- a/Engine/Plugins/Experimental/Mutable/Source/MutableRuntime/Public/MuR/MeshBufferSet.h
+++ b/Engine/Plugins/Experimental/Mutable/Source/MutableRuntime/Public/MuR/MeshBufferSet.h
@@ -394,6 +394,9 @@ namespace mu
 
 		//!
 		void UpdateOffsets(int32 b);
+		
+		/** Check that all channels of a specific semantic use the provided format. */
+		bool HasAnySemanticWithDifferentFormat(EMeshBufferSemantic Semantic, EMeshBufferFormat ExpectedFormat) const;
 	};	
 
 	
diff --git a/Engine/Plugins/Experimental/Mutable/Source/MutableTools/Internal/MuT/AST.h b/Engine/Plugins/Experimental/Mutable/Source/MutableTools/Internal/MuT/AST.h
--- a/Engine/Plugins/Experimental/Mutable/Source/MutableTools/Internal/MuT/AST.h
+++ b/Engine/Plugins/Experimental/Mutable/Source/MutableTools/Internal/MuT/AST.h
@@ -794,15 +794,15 @@ namespace mu
         virtual FImageDesc GetImageDesc( bool returnBestOption=false, FGetImageDescContext* context=nullptr ) const;
 
         //! Optional cache struct to use int he method below.
-        using FBlockLayoutSizeCache=TMap< const TPair<ASTOp*,int>, TPair<int,int>>;
+        using FBlockLayoutSizeCache=TMap< const TPair<ASTOp*,int32>, TPair<int32,int32>>;
 
         //! Return the size in layout blocks of a particular block given by absolute index
-        virtual void GetBlockLayoutSize( int blockIndex, int* pBlockX, int* pBlockY, FBlockLayoutSizeCache* cache );
-        void GetBlockLayoutSizeCached( int blockIndex, int* pBlockX, int* pBlockY, FBlockLayoutSizeCache* cache );
+        virtual void GetBlockLayoutSize( uint64 BlockId, int32* pBlockX, int32* pBlockY, FBlockLayoutSizeCache* cache );
+        void GetBlockLayoutSizeCached(uint64 BlockId, int32* pBlockX, int32* pBlockY, FBlockLayoutSizeCache* cache );
 
 
         //! Return the size in pixels of the layout grid block for the image operation
-        virtual void GetLayoutBlockSize( int* pBlockX, int* pBlockY );
+        virtual void GetLayoutBlockSize( int32* pBlockX, int32* pBlockY );
 
         virtual bool IsImagePlainConstant(FVector4f& colour ) const;
         virtual bool IsColourConstant(FVector4f& colour ) const;
diff --git a/Engine/Plugins/Experimental/Mutable/Source/MutableTools/Internal/MuT/NodeMeshFragmentPrivate.h b/Engine/Plugins/Experimental/Mutable/Source/MutableTools/Internal/MuT/NodeMeshFragmentPrivate.h
--- a/Engine/Plugins/Experimental/Mutable/Source/MutableTools/Internal/MuT/NodeMeshFragmentPrivate.h
+++ b/Engine/Plugins/Experimental/Mutable/Source/MutableTools/Internal/MuT/NodeMeshFragmentPrivate.h
@@ -17,20 +17,20 @@ namespace mu
 
 		Private()
 		{
-            m_layoutOrGroup = -1;
+            LayoutOrGroup = -1;
 		}
 
 		static FNodeType s_type;
 
-		NodeMeshPtr m_pMesh;
-        int m_layoutOrGroup;
-		TArray<int> m_blocks;
+		Ptr<NodeMesh> m_pMesh;
+        int32 LayoutOrGroup;
+		TArray<uint64> Blocks;
 
         FRAGMENT_TYPE m_fragmentType;
 
 
 		// NodeMesh::Private interface
-        NodeLayoutPtr GetLayout( int index ) const override;
+        Ptr<NodeLayout> GetLayout( int32 index ) const override;
 	};
 
 
diff --git a/Engine/Plugins/Experimental/Mutable/Source/MutableTools/Internal/MuT/NodePatchImage.h b/Engine/Plugins/Experimental/Mutable/Source/MutableTools/Internal/MuT/NodePatchImage.h
--- a/Engine/Plugins/Experimental/Mutable/Source/MutableTools/Internal/MuT/NodePatchImage.h
+++ b/Engine/Plugins/Experimental/Mutable/Source/MutableTools/Internal/MuT/NodePatchImage.h
@@ -14,13 +14,7 @@ namespace mu
 
     // Forward definitions
     class NodeImage;
-    typedef Ptr<NodeImage> NodeImagePtr;
-    typedef Ptr<const NodeImage> NodeImagePtrConst;
-
-    class NodePatchImage;
-    typedef Ptr<NodePatchImage> NodePatchImagePtr;
-    typedef Ptr<const NodePatchImage> NodePatchImagePtrConst;
-
+ 
 
     //! Node that allows to modify an image from an object by blending other images on specific
     //! layout blocks.
@@ -43,32 +37,31 @@ namespace mu
         //-----------------------------------------------------------------------------------------
 
         //! Set the image that will be blended on the destination.
-        void SetImage( NodeImagePtr );
+        void SetImage(Ptr<NodeImage>);
 
         //! Get the image that will be blended on the destination.
-        NodeImagePtr GetImage() const;
+		Ptr<NodeImage> GetImage() const;
 
         //! Set the blending mask image.
-        void SetMask( NodeImagePtr );
+        void SetMask(Ptr<NodeImage>);
 
         //! Get the blending mask image.
-        NodeImagePtr GetMask() const;
+		Ptr<NodeImage> GetMask() const;
 
         //! Set the number of blocks in the layout that will be patched
-        void SetBlockCount( int );
+        void SetBlockCount( int32 );
 
         //! Get the number of blocks in the layout that will be patched
-        int GetBlockCount() const;
+        int32 GetBlockCount() const;
 
         //! Set a block of the layout that will be patched.
         //! \param index is the index in the patching list, from 0 to GetBlockCount()-1
-        //! \param block is the block index in the layout of the image were this patch operation
-        //! will be applied.
-        void SetBlock( int index, int block );
+        //! \param LayoutBlockIndex is the block index in the layout of the image were this patch operation will be applied.
+        void SetBlock( int32 index, int32 LayoutBlockIndex);
 
-        //! Get a block of the layout that will be patched.
+        //! Get a block index of the layout that will be patched.
         //! \param index is the index in the patching list, from 0 to GetBlockCount()-1
-        int GetBlock( int index ) const;
+        int32 GetBlock( int32 index ) const;
 
         //! Set the blending operation to use to combine the pixels
         EBlendType GetBlendType() const;
diff --git a/Engine/Plugins/Experimental/Mutable/Source/MutableTools/Internal/MuT/NodePatchImagePrivate.h b/Engine/Plugins/Experimental/Mutable/Source/MutableTools/Internal/MuT/NodePatchImagePrivate.h
--- a/Engine/Plugins/Experimental/Mutable/Source/MutableTools/Internal/MuT/NodePatchImagePrivate.h
+++ b/Engine/Plugins/Experimental/Mutable/Source/MutableTools/Internal/MuT/NodePatchImagePrivate.h
@@ -20,7 +20,8 @@ namespace mu
         NodeImagePtr m_pImage;
         NodeImagePtr m_pMask;
 
-        TArray<int> m_blocks;
+		// These are the indices of the blocks in the layout being patched (not the Ids)
+        TArray<int32> BlockIndices;
 
         EBlendType m_blendType = EBlendType::BT_BLEND;
 
diff --git a/Engine/Plugins/Experimental/Mutable/Source/MutableTools/Internal/MuT/NodePatchMeshPrivate.h b/Engine/Plugins/Experimental/Mutable/Source/MutableTools/Internal/MuT/NodePatchMeshPrivate.h
--- a/Engine/Plugins/Experimental/Mutable/Source/MutableTools/Internal/MuT/NodePatchMeshPrivate.h
+++ b/Engine/Plugins/Experimental/Mutable/Source/MutableTools/Internal/MuT/NodePatchMeshPrivate.h
@@ -9,15 +9,14 @@
 namespace mu
 {
 
-
 	class NodePatchMesh::Private : public Node::Private
 	{
 	public:
 
 		static FNodeType s_type;
 
-		NodeMeshPtr m_pRemove;
-		NodeMeshPtr m_pAdd;
+		Ptr<NodeMesh> m_pRemove;
+		Ptr<NodeMesh> m_pAdd;
 
 	};
 
diff --git a/Engine/Plugins/Experimental/Mutable/Source/MutableTools/Private/MuT/AST.cpp b/Engine/Plugins/Experimental/Mutable/Source/MutableTools/Private/MuT/AST.cpp
--- a/Engine/Plugins/Experimental/Mutable/Source/MutableTools/Private/MuT/AST.cpp
+++ b/Engine/Plugins/Experimental/Mutable/Source/MutableTools/Private/MuT/AST.cpp
@@ -881,26 +881,25 @@ bool ASTOp::IsColourConstant(FVector4f&) const
 }
 
 
-void ASTOp::GetBlockLayoutSize( int, int*, int*, FBlockLayoutSizeCache* )
+void ASTOp::GetBlockLayoutSize( uint64, int32*, int32*, FBlockLayoutSizeCache* )
 {
     check(false);
 }
 
-void ASTOp::GetBlockLayoutSizeCached( int blockIndex, int* pBlockX, int* pBlockY,
-                               FBlockLayoutSizeCache* cache )
+void ASTOp::GetBlockLayoutSizeCached(uint64 blockIndex, int32* BlockX, int32* BlockY, FBlockLayoutSizeCache* cache)
 {
 	FBlockLayoutSizeCache::KeyType key(this,blockIndex);
 	FBlockLayoutSizeCache::ValueType* ValuePtr = cache->Find( key );
     if (ValuePtr)
     {
-        *pBlockX = ValuePtr->Key;
-        *pBlockY = ValuePtr->Value;
+        *BlockX = ValuePtr->Key;
+        *BlockY = ValuePtr->Value;
         return;
     }
 
-    GetBlockLayoutSize( blockIndex, pBlockX, pBlockY, cache );
+    GetBlockLayoutSize( blockIndex, BlockX, BlockY, cache );
 
-	cache->Add(key, FBlockLayoutSizeCache::ValueType( *pBlockX, *pBlockY) );
+	cache->Add(key, FBlockLayoutSizeCache::ValueType( *BlockX, *BlockY) );
 }
 
 
diff --git a/Engine/Plugins/Experimental/Mutable/Source/MutableTools/Private/MuT/ASTOpConditional.cpp b/Engine/Plugins/Experimental/Mutable/Source/MutableTools/Private/MuT/ASTOpConditional.cpp
--- a/Engine/Plugins/Experimental/Mutable/Source/MutableTools/Private/MuT/ASTOpConditional.cpp
+++ b/Engine/Plugins/Experimental/Mutable/Source/MutableTools/Private/MuT/ASTOpConditional.cpp
@@ -194,7 +194,7 @@ namespace mu
 
 
 	//-------------------------------------------------------------------------------------------------
-	void ASTOpConditional::GetLayoutBlockSize(int* pBlockX, int* pBlockY)
+	void ASTOpConditional::GetLayoutBlockSize(int32* pBlockX, int32* pBlockY)
 	{
 		if (type == OP_TYPE::IM_CONDITIONAL)
 		{
@@ -213,16 +213,15 @@ namespace mu
 
 
 	//-------------------------------------------------------------------------------------------------
-	void ASTOpConditional::GetBlockLayoutSize(int blockIndex, int* pBlockX, int* pBlockY,
-		FBlockLayoutSizeCache* cache)
+	void ASTOpConditional::GetBlockLayoutSize(uint64 BlockId, int32* pBlockX, int32* pBlockY, FBlockLayoutSizeCache* cache)
 	{
 		if (type == OP_TYPE::LA_CONDITIONAL)
 		{
-			yes->GetBlockLayoutSizeCached(blockIndex, pBlockX, pBlockY, cache);
+			yes->GetBlockLayoutSizeCached(BlockId, pBlockX, pBlockY, cache);
 
 			if (*pBlockX == 0)
 			{
-				no->GetBlockLayoutSizeCached(blockIndex, pBlockX, pBlockY, cache);
+				no->GetBlockLayoutSizeCached(BlockId, pBlockX, pBlockY, cache);
 			}
 		}
 		else
diff --git a/Engine/Plugins/Experimental/Mutable/Source/MutableTools/Private/MuT/ASTOpConditional.h b/Engine/Plugins/Experimental/Mutable/Source/MutableTools/Private/MuT/ASTOpConditional.h
--- a/Engine/Plugins/Experimental/Mutable/Source/MutableTools/Private/MuT/ASTOpConditional.h
+++ b/Engine/Plugins/Experimental/Mutable/Source/MutableTools/Private/MuT/ASTOpConditional.h
@@ -45,8 +45,8 @@ struct FProgram;
 		void ForEachChild(const TFunctionRef<void(ASTChild&)> f) override;
 		void Link(FProgram& program, FLinkerOptions* Options) override;
 		FImageDesc GetImageDesc(bool returnBestOption, class FGetImageDescContext* context) const override;
-		void GetLayoutBlockSize(int* pBlockX, int* pBlockY) override;
-		void GetBlockLayoutSize(int blockIndex, int* pBlockX, int* pBlockY, FBlockLayoutSizeCache* cache) override;
+		void GetLayoutBlockSize(int32* pBlockX, int32* pBlockY) override;
+		void GetBlockLayoutSize(uint64 BlockId, int32* pBlockX, int32* pBlockY, FBlockLayoutSizeCache* cache) override;
 		Ptr<ASTOp> OptimiseSemantic(const FModelOptimizationOptions&, int32 Pass) const override;
 		bool GetNonBlackRect(FImageRect& maskUsage) const override;
 		Ptr<ImageSizeExpression> GetImageSizeExpression() const override;
diff --git a/Engine/Plugins/Experimental/Mutable/Source/MutableTools/Private/MuT/ASTOpConstantResource.cpp b/Engine/Plugins/Experimental/Mutable/Source/MutableTools/Private/MuT/ASTOpConstantResource.cpp
--- a/Engine/Plugins/Experimental/Mutable/Source/MutableTools/Private/MuT/ASTOpConstantResource.cpp
+++ b/Engine/Plugins/Experimental/Mutable/Source/MutableTools/Private/MuT/ASTOpConstantResource.cpp
@@ -413,7 +413,7 @@ namespace mu
 	}
 
 	//-------------------------------------------------------------------------------------------------
-	void ASTOpConstantResource::GetBlockLayoutSize(int blockIndex, int* pBlockX, int* pBlockY, FBlockLayoutSizeCache*)
+	void ASTOpConstantResource::GetBlockLayoutSize(uint64 BlockId, int32* BlockX, int32* BlockY, FBlockLayoutSizeCache*)
 	{
 		switch (Type)
 		{
@@ -424,16 +424,16 @@ namespace mu
 
 			if (pLayout)
 			{
-				int relId = pLayout->FindBlock(blockIndex);
+				int relId = pLayout->FindBlock(BlockId);
 				if (relId >= 0)
 				{
-					*pBlockX = pLayout->m_blocks[relId].m_size[0];
-					*pBlockY = pLayout->m_blocks[relId].m_size[1];
+					*BlockX = pLayout->Blocks[relId].Size[0];
+					*BlockY = pLayout->Blocks[relId].Size[1];
 				}
 				else
 				{
-					*pBlockX = 0;
-					*pBlockY = 0;
+					*BlockX = 0;
+					*BlockY = 0;
 				}
 			}
 
@@ -446,7 +446,7 @@ namespace mu
 
 
 	//-------------------------------------------------------------------------------------------------
-	void ASTOpConstantResource::GetLayoutBlockSize(int* pBlockX, int* pBlockY)
+	void ASTOpConstantResource::GetLayoutBlockSize(int32* pBlockX, int32* pBlockY)
 	{
 		switch (Type)
 		{
diff --git a/Engine/Plugins/Experimental/Mutable/Source/MutableTools/Private/MuT/ASTOpConstantResource.h b/Engine/Plugins/Experimental/Mutable/Source/MutableTools/Private/MuT/ASTOpConstantResource.h
--- a/Engine/Plugins/Experimental/Mutable/Source/MutableTools/Private/MuT/ASTOpConstantResource.h
+++ b/Engine/Plugins/Experimental/Mutable/Source/MutableTools/Private/MuT/ASTOpConstantResource.h
@@ -61,9 +61,8 @@ namespace mu
 		uint64 Hash() const override;
 		void Link(FProgram& program, FLinkerOptions*) override;
 		FImageDesc GetImageDesc(bool, class FGetImageDescContext*) const override;
-		void GetBlockLayoutSize(int blockIndex, int* pBlockX, int* pBlockY,
-			FBlockLayoutSizeCache* cache) override;
-		void GetLayoutBlockSize(int* pBlockX, int* pBlockY) override;
+		void GetBlockLayoutSize(uint64 BlockId, int32* pBlockX, int32* pBlockY, FBlockLayoutSizeCache* cache) override;
+		void GetLayoutBlockSize(int32* pBlockX, int32* pBlockY) override;
 		bool GetNonBlackRect(FImageRect& maskUsage) const override;
 		bool IsImagePlainConstant(FVector4f& colour) const override;
 		Ptr<ImageSizeExpression> GetImageSizeExpression() const override;
diff --git a/Engine/Plugins/Experimental/Mutable/Source/MutableTools/Private/MuT/ASTOpImageCompose.cpp b/Engine/Plugins/Experimental/Mutable/Source/MutableTools/Private/MuT/ASTOpImageCompose.cpp
--- a/Engine/Plugins/Experimental/Mutable/Source/MutableTools/Private/MuT/ASTOpImageCompose.cpp
+++ b/Engine/Plugins/Experimental/Mutable/Source/MutableTools/Private/MuT/ASTOpImageCompose.cpp
@@ -50,7 +50,7 @@ bool ASTOpImageCompose::IsEqual(const ASTOp& otherUntyped) const
 			Base ==other->Base &&
 			BlockImage == other->BlockImage &&
 			Mask == other->Mask &&
-			BlockIndex == other->BlockIndex;
+			BlockId == other->BlockId;
     }
     return false;
 }
@@ -64,7 +64,7 @@ uint64 ASTOpImageCompose::Hash() const
 	hash_combine(res, Base.child().get());
 	hash_combine(res, BlockImage.child().get());
 	hash_combine(res, Mask.child().get());
-	hash_combine(res, BlockIndex);
+	hash_combine(res, BlockId);
 	return res;
 }
 
@@ -77,7 +77,7 @@ mu::Ptr<ASTOp> ASTOpImageCompose::Clone(MapChildFuncRef mapChild) const
 	n->Base = mapChild(Base.child());
 	n->BlockImage = mapChild(BlockImage.child());
 	n->Mask = mapChild(Mask.child());
-	n->BlockIndex = BlockIndex;
+	n->BlockId = BlockId;
     return n;
 }
 
@@ -105,7 +105,7 @@ void ASTOpImageCompose::Link( FProgram& program, FLinkerOptions*)
 		if (Base) args.base = Base->linkedAddress;
 		if (BlockImage) args.blockImage = BlockImage->linkedAddress;
 		if (Mask) args.mask = Mask->linkedAddress;
-        args.blockIndex = BlockIndex;
+        args.BlockId = BlockId;
 
         linkedAddress = (OP::ADDRESS)program.m_opAddress.Num();
         program.m_opAddress.Add((uint32_t)program.m_byteCode.Num());
@@ -189,19 +189,17 @@ bool ASTOpImageCompose::IsImagePlainConstant(FVector4f& colour) const
 }
 
 
-void ASTOpImageCompose::GetLayoutBlockSize(int* pBlockX, int* pBlockY)
+void ASTOpImageCompose::GetLayoutBlockSize(int32* pBlockX, int32* pBlockY)
 {
 	// We can find out in two ways: following the block image, or following the base image.
 	// Let's try the block approach: We need the block size and the layout blocks
-	int layoutBlocksX = 0;
-	int layoutBlocksY = 0;
+	int32 layoutBlocksX = 0;
+	int32 layoutBlocksY = 0;
 	if (Layout.child())
 	{
 		MUTABLE_CPUPROFILER_SCOPE(GetLayoutBlockSize_GetBlockLayoutSize);
 		FBlockLayoutSizeCache cache;
-		Layout->GetBlockLayoutSizeCached(BlockIndex,
-			&layoutBlocksX, &layoutBlocksY,
-			&cache);
+		Layout->GetBlockLayoutSizeCached(BlockId, &layoutBlocksX, &layoutBlocksY, &cache);
 	}
 
 	if (layoutBlocksX > 0 && layoutBlocksY > 0 && BlockImage.child())
@@ -223,9 +221,9 @@ mu::Ptr<ASTOp> ASTOpImageCompose::OptimiseSemantic(const FModelOptimizationOptio
 {
 	mu::Ptr<ASTOp> at;
 
-	auto baseAt = Base.child();
-	auto blockAt = BlockImage.child();
-	auto layoutAt = Layout.child();
+	mu::Ptr<ASTOp> baseAt = Base.child();
+	mu::Ptr<ASTOp> blockAt = BlockImage.child();
+	mu::Ptr<ASTOp> layoutAt = Layout.child();
 	if (layoutAt
 		&&
 		layoutAt->GetOpType() == OP_TYPE::LA_CONSTANT
@@ -240,17 +238,17 @@ mu::Ptr<ASTOp> ASTOpImageCompose::OptimiseSemantic(const FModelOptimizationOptio
 		// Constant single-block full layout?
 		if (pLayout->GetBlockCount() == 1
 			&&
-			pLayout->m_blocks[0].m_min == UE::Math::TIntVector2<uint16>(0, 0)
+			pLayout->Blocks[0].Min == UE::Math::TIntVector2<uint16>(0, 0)
 			&&
-			pLayout->m_blocks[0].m_size == pLayout->m_size
+			pLayout->Blocks[0].Size == pLayout->Size
 			&&
-			pLayout->m_blocks[0].m_id == int(BlockIndex)
+			pLayout->Blocks[0].Id == BlockId
 			)
 		{
 			// We could only take the block, but we must make sure it will have the format
 			// and size of the base.
-			auto baseDesc = baseAt->GetImageDesc(true);
-			auto blockDesc = blockAt->GetImageDesc(true);
+			mu::FImageDesc baseDesc = baseAt->GetImageDesc(true);
+			mu::FImageDesc blockDesc = blockAt->GetImageDesc(true);
 
 			at = blockAt;
 
diff --git a/Engine/Plugins/Experimental/Mutable/Source/MutableTools/Private/MuT/ASTOpImageCompose.h b/Engine/Plugins/Experimental/Mutable/Source/MutableTools/Private/MuT/ASTOpImageCompose.h
--- a/Engine/Plugins/Experimental/Mutable/Source/MutableTools/Private/MuT/ASTOpImageCompose.h
+++ b/Engine/Plugins/Experimental/Mutable/Source/MutableTools/Private/MuT/ASTOpImageCompose.h
@@ -21,7 +21,7 @@ namespace mu
 		ASTChild Base;
 		ASTChild BlockImage;
 		ASTChild Mask;
-		uint32 BlockIndex=0;
+		uint64 BlockId=0;
 
 	public:
 
@@ -38,8 +38,8 @@ namespace mu
 		//TODO: Ptr<ASTOp> OptimiseSink(const FModelOptimizationOptions& options, FOptimizeSinkContext& context) const override;
 		Ptr<ASTOp> OptimiseSemantic(const FModelOptimizationOptions& options, int32 Pass) const;
 		FImageDesc GetImageDesc(bool returnBestOption, FGetImageDescContext* context) const override;
-		void GetLayoutBlockSize(int* pBlockX, int* pBlockY) override;
-		bool IsImagePlainConstant(FVector4f& colour) const override;
+		void GetLayoutBlockSize(int32* OutBlockX, int32* OutBlockY) override;
+		bool IsImagePlainConstant(FVector4f& OutColor) const override;
 		Ptr<ImageSizeExpression> GetImageSizeExpression() const override;
 
 	};
diff --git a/Engine/Plugins/Experimental/Mutable/Source/MutableTools/Private/MuT/ASTOpImageRasterMesh.cpp b/Engine/Plugins/Experimental/Mutable/Source/MutableTools/Private/MuT/ASTOpImageRasterMesh.cpp
--- a/Engine/Plugins/Experimental/Mutable/Source/MutableTools/Private/MuT/ASTOpImageRasterMesh.cpp
+++ b/Engine/Plugins/Experimental/Mutable/Source/MutableTools/Private/MuT/ASTOpImageRasterMesh.cpp
@@ -26,7 +26,7 @@ namespace mu
 		, mask(this)
 		, projector(this)
 	{
-		BlockId = -1;
+		BlockId = Layout::InvalidBlockId;
 		LayoutIndex = -1;
 		SizeX = SizeY = 0;
 		SourceSizeX = SourceSizeY = 0;
@@ -139,7 +139,7 @@ namespace mu
 			OP::ImageRasterMeshArgs args;
 			FMemory::Memzero(&args, sizeof(args));
 
-			args.blockId = BlockId;
+			args.BlockId = BlockId;
 			args.LayoutIndex = LayoutIndex;
 			args.sizeX = SizeX;
 			args.sizeY = SizeY;
diff --git a/Engine/Plugins/Experimental/Mutable/Source/MutableTools/Private/MuT/ASTOpImageRasterMesh.h b/Engine/Plugins/Experimental/Mutable/Source/MutableTools/Private/MuT/ASTOpImageRasterMesh.h
--- a/Engine/Plugins/Experimental/Mutable/Source/MutableTools/Private/MuT/ASTOpImageRasterMesh.h
+++ b/Engine/Plugins/Experimental/Mutable/Source/MutableTools/Private/MuT/ASTOpImageRasterMesh.h
@@ -24,7 +24,7 @@ namespace mu
 		ASTChild projector;
 
 		/** If layouts are used, this can indicate a single layout block that we want to raster. */
-		int32 BlockId;
+		uint64 BlockId;
 		int8 LayoutIndex;
 
 		/** Size of the image to generate by rasterization of the mesh. */
diff --git a/Engine/Plugins/Experimental/Mutable/Source/MutableTools/Private/MuT/ASTOpLayoutFromMesh.cpp b/Engine/Plugins/Experimental/Mutable/Source/MutableTools/Private/MuT/ASTOpLayoutFromMesh.cpp
--- a/Engine/Plugins/Experimental/Mutable/Source/MutableTools/Private/MuT/ASTOpLayoutFromMesh.cpp
+++ b/Engine/Plugins/Experimental/Mutable/Source/MutableTools/Private/MuT/ASTOpLayoutFromMesh.cpp
@@ -97,7 +97,7 @@ namespace mu
 	}
 
 
-	void ASTOpLayoutFromMesh::GetBlockLayoutSize(int blockIndex, int* pBlockX, int* pBlockY, FBlockLayoutSizeCache* cache)
+	void ASTOpLayoutFromMesh::GetBlockLayoutSize(uint64 BlockId, int32* pBlockX, int32* pBlockY, FBlockLayoutSizeCache* cache)
 	{
 		// This shouldn't happen for this operation because it is always in a branch of layout operations that is not the main one.
 		check(false);
diff --git a/Engine/Plugins/Experimental/Mutable/Source/MutableTools/Private/MuT/ASTOpLayoutFromMesh.h b/Engine/Plugins/Experimental/Mutable/Source/MutableTools/Private/MuT/ASTOpLayoutFromMesh.h
--- a/Engine/Plugins/Experimental/Mutable/Source/MutableTools/Private/MuT/ASTOpLayoutFromMesh.h
+++ b/Engine/Plugins/Experimental/Mutable/Source/MutableTools/Private/MuT/ASTOpLayoutFromMesh.h
@@ -32,7 +32,7 @@ namespace mu
 		void ForEachChild(const TFunctionRef<void(ASTChild&)>) override;
 		void Link(FProgram& program, FLinkerOptions* Options) override;
 		Ptr<ASTOp> OptimiseSink(const FModelOptimizationOptions& options, FOptimizeSinkContext& context) const override;
-		void GetBlockLayoutSize(int blockIndex, int* pBlockX, int* pBlockY, FBlockLayoutSizeCache* cache) override;
+		void GetBlockLayoutSize(uint64 BlockId, int32* pBlockX, int32* pBlockY, FBlockLayoutSizeCache* cache) override;
 
 	};
 
diff --git a/Engine/Plugins/Experimental/Mutable/Source/MutableTools/Private/MuT/ASTOpLayoutMerge.cpp b/Engine/Plugins/Experimental/Mutable/Source/MutableTools/Private/MuT/ASTOpLayoutMerge.cpp
--- a/Engine/Plugins/Experimental/Mutable/Source/MutableTools/Private/MuT/ASTOpLayoutMerge.cpp
+++ b/Engine/Plugins/Experimental/Mutable/Source/MutableTools/Private/MuT/ASTOpLayoutMerge.cpp
@@ -93,16 +93,16 @@ namespace mu
 	}
 
 
-	void ASTOpLayoutMerge::GetBlockLayoutSize(int blockIndex, int* pBlockX, int* pBlockY, FBlockLayoutSizeCache* cache)
+	void ASTOpLayoutMerge::GetBlockLayoutSize(uint64 BlockId, int32* pBlockX, int32* pBlockY, FBlockLayoutSizeCache* cache)
 	{
 		if (Base)
 		{
-			Base->GetBlockLayoutSize(blockIndex, pBlockX, pBlockY, cache);
+			Base->GetBlockLayoutSize(BlockId, pBlockX, pBlockY, cache);
 		}
 
 		if (!*pBlockX && Added)
 		{
-			Added->GetBlockLayoutSize(blockIndex, pBlockX, pBlockY, cache);
+			Added->GetBlockLayoutSize(BlockId, pBlockX, pBlockY, cache);
 		}
 	}
 
diff --git a/Engine/Plugins/Experimental/Mutable/Source/MutableTools/Private/MuT/ASTOpLayoutMerge.h b/Engine/Plugins/Experimental/Mutable/Source/MutableTools/Private/MuT/ASTOpLayoutMerge.h
--- a/Engine/Plugins/Experimental/Mutable/Source/MutableTools/Private/MuT/ASTOpLayoutMerge.h
+++ b/Engine/Plugins/Experimental/Mutable/Source/MutableTools/Private/MuT/ASTOpLayoutMerge.h
@@ -31,7 +31,7 @@ namespace mu
 		Ptr<ASTOp> Clone(MapChildFuncRef mapChild) const override;
 		void ForEachChild(const TFunctionRef<void(ASTChild&)>) override;
 		void Link(FProgram& program, FLinkerOptions* Options) override;
-		void GetBlockLayoutSize(int blockIndex, int* pBlockX, int* pBlockY, FBlockLayoutSizeCache* cache) override;
+		void GetBlockLayoutSize(uint64 BlockId, int32* pBlockX, int32* pBlockY, FBlockLayoutSizeCache* cache) override;
 
 	};
 
diff --git a/Engine/Plugins/Experimental/Mutable/Source/MutableTools/Private/MuT/ASTOpLayoutPack.cpp b/Engine/Plugins/Experimental/Mutable/Source/MutableTools/Private/MuT/ASTOpLayoutPack.cpp
--- a/Engine/Plugins/Experimental/Mutable/Source/MutableTools/Private/MuT/ASTOpLayoutPack.cpp
+++ b/Engine/Plugins/Experimental/Mutable/Source/MutableTools/Private/MuT/ASTOpLayoutPack.cpp
@@ -88,11 +88,11 @@ namespace mu
 	}
 
 
-	void ASTOpLayoutPack::GetBlockLayoutSize(int blockIndex, int* pBlockX, int* pBlockY, FBlockLayoutSizeCache* cache)
+	void ASTOpLayoutPack::GetBlockLayoutSize(uint64 BlockId, int32* pBlockX, int32* pBlockY, FBlockLayoutSizeCache* cache)
 	{
 		if (Source)
 		{
-			Source->GetBlockLayoutSize(blockIndex, pBlockX, pBlockY, cache);
+			Source->GetBlockLayoutSize(BlockId, pBlockX, pBlockY, cache);
 		}
 	}
 
diff --git a/Engine/Plugins/Experimental/Mutable/Source/MutableTools/Private/MuT/ASTOpLayoutPack.h b/Engine/Plugins/Experimental/Mutable/Source/MutableTools/Private/MuT/ASTOpLayoutPack.h
--- a/Engine/Plugins/Experimental/Mutable/Source/MutableTools/Private/MuT/ASTOpLayoutPack.h
+++ b/Engine/Plugins/Experimental/Mutable/Source/MutableTools/Private/MuT/ASTOpLayoutPack.h
@@ -30,7 +30,7 @@ namespace mu
 		Ptr<ASTOp> Clone(MapChildFuncRef mapChild) const override;
 		void ForEachChild(const TFunctionRef<void(ASTChild&)>) override;
 		void Link(FProgram& program, FLinkerOptions* Options) override;
-		void GetBlockLayoutSize(int blockIndex, int* pBlockX, int* pBlockY, FBlockLayoutSizeCache* cache) override;
+		void GetBlockLayoutSize(uint64 BlockId, int32* pBlockX, int32* pBlockY, FBlockLayoutSizeCache* cache) override;
 
 	};
 
diff --git a/Engine/Plugins/Experimental/Mutable/Source/MutableTools/Private/MuT/ASTOpLayoutRemoveBlocks.cpp b/Engine/Plugins/Experimental/Mutable/Source/MutableTools/Private/MuT/ASTOpLayoutRemoveBlocks.cpp
--- a/Engine/Plugins/Experimental/Mutable/Source/MutableTools/Private/MuT/ASTOpLayoutRemoveBlocks.cpp
+++ b/Engine/Plugins/Experimental/Mutable/Source/MutableTools/Private/MuT/ASTOpLayoutRemoveBlocks.cpp
@@ -93,11 +93,11 @@ namespace mu
 	}
 
 
-	void ASTOpLayoutRemoveBlocks::GetBlockLayoutSize(int blockIndex, int* pBlockX, int* pBlockY, FBlockLayoutSizeCache* cache)
+	void ASTOpLayoutRemoveBlocks::GetBlockLayoutSize(uint64 BlockId, int32* pBlockX, int32* pBlockY, FBlockLayoutSizeCache* cache)
 	{
 		if (Source)
 		{
-			Source->GetBlockLayoutSize(blockIndex, pBlockX, pBlockY, cache);
+			Source->GetBlockLayoutSize(BlockId, pBlockX, pBlockY, cache);
 		}
 	}
 
diff --git a/Engine/Plugins/Experimental/Mutable/Source/MutableTools/Private/MuT/ASTOpLayoutRemoveBlocks.h b/Engine/Plugins/Experimental/Mutable/Source/MutableTools/Private/MuT/ASTOpLayoutRemoveBlocks.h
--- a/Engine/Plugins/Experimental/Mutable/Source/MutableTools/Private/MuT/ASTOpLayoutRemoveBlocks.h
+++ b/Engine/Plugins/Experimental/Mutable/Source/MutableTools/Private/MuT/ASTOpLayoutRemoveBlocks.h
@@ -31,7 +31,7 @@ namespace mu
 		Ptr<ASTOp> Clone(MapChildFuncRef mapChild) const override;
 		void ForEachChild(const TFunctionRef<void(ASTChild&)>) override;
 		void Link(FProgram& program, FLinkerOptions* Options) override;
-		void GetBlockLayoutSize(int blockIndex, int* pBlockX, int* pBlockY, FBlockLayoutSizeCache* cache) override;
+		void GetBlockLayoutSize(uint64 BlockId, int32* pBlockX, int32* pBlockY, FBlockLayoutSizeCache* cache) override;
 
 	};
 
diff --git a/Engine/Plugins/Experimental/Mutable/Source/MutableTools/Private/MuT/ASTOpMeshExtractLayoutBlocks.cpp b/Engine/Plugins/Experimental/Mutable/Source/MutableTools/Private/MuT/ASTOpMeshExtractLayoutBlocks.cpp
--- a/Engine/Plugins/Experimental/Mutable/Source/MutableTools/Private/MuT/ASTOpMeshExtractLayoutBlocks.cpp
+++ b/Engine/Plugins/Experimental/Mutable/Source/MutableTools/Private/MuT/ASTOpMeshExtractLayoutBlocks.cpp
@@ -87,9 +87,9 @@ namespace mu
 			AppendCode(program.m_byteCode, (uint16)Layout);
 			AppendCode(program.m_byteCode, (uint16)Blocks.Num());
 
-			for (auto b : Blocks)
+			for (uint64 Id : Blocks)
 			{
-				AppendCode(program.m_byteCode, (uint32)b);
+				AppendCode(program.m_byteCode, Id);
 			}
 		}
 	}
diff --git a/Engine/Plugins/Experimental/Mutable/Source/MutableTools/Private/MuT/ASTOpMeshExtractLayoutBlocks.h b/Engine/Plugins/Experimental/Mutable/Source/MutableTools/Private/MuT/ASTOpMeshExtractLayoutBlocks.h
--- a/Engine/Plugins/Experimental/Mutable/Source/MutableTools/Private/MuT/ASTOpMeshExtractLayoutBlocks.h
+++ b/Engine/Plugins/Experimental/Mutable/Source/MutableTools/Private/MuT/ASTOpMeshExtractLayoutBlocks.h
@@ -26,8 +26,8 @@ namespace mu
 		/** Layout to use to select the blocks. */
 		uint16 Layout = 0;
 
-		/** Blocks to include in the resulting mesh. */
-		TArray<uint32> Blocks;
+		/** Block Ids to include in the resulting mesh. */
+		TArray<uint64> Blocks;
 
 	public:
 
diff --git a/Engine/Plugins/Experimental/Mutable/Source/MutableTools/Private/MuT/ASTOpReferenceResource.cpp b/Engine/Plugins/Experimental/Mutable/Source/MutableTools/Private/MuT/ASTOpReferenceResource.cpp
--- a/Engine/Plugins/Experimental/Mutable/Source/MutableTools/Private/MuT/ASTOpReferenceResource.cpp
+++ b/Engine/Plugins/Experimental/Mutable/Source/MutableTools/Private/MuT/ASTOpReferenceResource.cpp
@@ -77,7 +77,7 @@ namespace mu
 
 
 	//-------------------------------------------------------------------------------------------------
-	void ASTOpReferenceResource::GetBlockLayoutSize(int blockIndex, int* pBlockX, int* pBlockY, FBlockLayoutSizeCache*)
+	void ASTOpReferenceResource::GetBlockLayoutSize(uint64 BlockId, int32* pBlockX, int32* pBlockY, FBlockLayoutSizeCache*)
 	{
 		switch (type)
 		{
@@ -96,7 +96,7 @@ namespace mu
 
 
 	//-------------------------------------------------------------------------------------------------
-	void ASTOpReferenceResource::GetLayoutBlockSize(int* pBlockX, int* pBlockY)
+	void ASTOpReferenceResource::GetLayoutBlockSize(int32* pBlockX, int32* pBlockY)
 	{
 		switch (type)
 		{
diff --git a/Engine/Plugins/Experimental/Mutable/Source/MutableTools/Private/MuT/ASTOpReferenceResource.h b/Engine/Plugins/Experimental/Mutable/Source/MutableTools/Private/MuT/ASTOpReferenceResource.h
--- a/Engine/Plugins/Experimental/Mutable/Source/MutableTools/Private/MuT/ASTOpReferenceResource.h
+++ b/Engine/Plugins/Experimental/Mutable/Source/MutableTools/Private/MuT/ASTOpReferenceResource.h
@@ -46,8 +46,8 @@ namespace mu
 		uint64 Hash() const override;
 		void Link(FProgram& program, FLinkerOptions*) override;
 		FImageDesc GetImageDesc(bool, class FGetImageDescContext*) const override;
-		void GetBlockLayoutSize(int blockIndex, int* pBlockX, int* pBlockY, FBlockLayoutSizeCache* cache) override;
-		void GetLayoutBlockSize(int* pBlockX, int* pBlockY) override;
+		void GetBlockLayoutSize(uint64 BlockId, int32* pBlockX, int32* pBlockY, FBlockLayoutSizeCache* cache) override;
+		void GetLayoutBlockSize(int32* pBlockX, int32* pBlockY) override;
 		bool GetNonBlackRect(FImageRect& maskUsage) const override;
 		Ptr<ImageSizeExpression> GetImageSizeExpression() const override;
 		Ptr<ASTOp> OptimiseSemantic(const FModelOptimizationOptions& options, int32 Pass) const override;
diff --git a/Engine/Plugins/Experimental/Mutable/Source/MutableTools/Private/MuT/ASTOpSwitch.cpp b/Engine/Plugins/Experimental/Mutable/Source/MutableTools/Private/MuT/ASTOpSwitch.cpp
--- a/Engine/Plugins/Experimental/Mutable/Source/MutableTools/Private/MuT/ASTOpSwitch.cpp
+++ b/Engine/Plugins/Experimental/Mutable/Source/MutableTools/Private/MuT/ASTOpSwitch.cpp
@@ -286,8 +286,7 @@ namespace mu
 
 
 	//-------------------------------------------------------------------------------------------------
-	void ASTOpSwitch::GetBlockLayoutSize(int blockIndex, int* pBlockX, int* pBlockY,
-		FBlockLayoutSizeCache* cache)
+	void ASTOpSwitch::GetBlockLayoutSize(uint64 BlockId, int32* pBlockX, int32* pBlockY, FBlockLayoutSizeCache* cache)
 	{
 		switch (type)
 		{
@@ -301,7 +300,7 @@ namespace mu
 
 			if (child)
 			{
-				child->GetBlockLayoutSizeCached(blockIndex, pBlockX, pBlockY, cache);
+				child->GetBlockLayoutSizeCached(BlockId, pBlockX, pBlockY, cache);
 			}
 			else
 			{
@@ -319,7 +318,7 @@ namespace mu
 
 
 	//-------------------------------------------------------------------------------------------------
-	void ASTOpSwitch::GetLayoutBlockSize(int* pBlockX, int* pBlockY)
+	void ASTOpSwitch::GetLayoutBlockSize(int32* pBlockX, int32* pBlockY)
 	{
 		switch (type)
 		{
diff --git a/Engine/Plugins/Experimental/Mutable/Source/MutableTools/Private/MuT/ASTOpSwitch.h b/Engine/Plugins/Experimental/Mutable/Source/MutableTools/Private/MuT/ASTOpSwitch.h
--- a/Engine/Plugins/Experimental/Mutable/Source/MutableTools/Private/MuT/ASTOpSwitch.h
+++ b/Engine/Plugins/Experimental/Mutable/Source/MutableTools/Private/MuT/ASTOpSwitch.h
@@ -67,8 +67,8 @@ template <class SCALAR> class vec4;
 		void Assert() override;
 		void Link(FProgram& program, FLinkerOptions* Options) override;
 		FImageDesc GetImageDesc(bool returnBestOption, class FGetImageDescContext* context) const  override;
-		void GetBlockLayoutSize(int blockIndex, int* pBlockX, int* pBlockY, FBlockLayoutSizeCache* cache) override;
-		void GetLayoutBlockSize(int* pBlockX, int* pBlockY) override;
+		void GetBlockLayoutSize(uint64 BlockId, int32* pBlockX, int32* pBlockY, FBlockLayoutSizeCache* cache) override;
+		void GetLayoutBlockSize(int32* pBlockX, int32* pBlockY) override;
 		Ptr<ASTOp> OptimiseSemantic(const FModelOptimizationOptions&, int32 Pass) const override;
 		bool GetNonBlackRect(FImageRect& maskUsage) const override;
 		bool IsImagePlainConstant(FVector4f& colour) const override;
diff --git a/Engine/Plugins/Experimental/Mutable/Source/MutableTools/Private/MuT/CodeGenerator.cpp b/Engine/Plugins/Experimental/Mutable/Source/MutableTools/Private/MuT/CodeGenerator.cpp
--- a/Engine/Plugins/Experimental/Mutable/Source/MutableTools/Private/MuT/CodeGenerator.cpp
+++ b/Engine/Plugins/Experimental/Mutable/Source/MutableTools/Private/MuT/CodeGenerator.cpp
@@ -385,10 +385,10 @@ namespace mu
 	}
 
 
-	Ptr<const Layout> CodeGenerator::AddLayout(Ptr<const Layout> SourceLayout)
+	Ptr<const Layout> CodeGenerator::AddLayout(Ptr<const Layout> SourceLayout, uint32 MeshIDPrefix)
 	{
 		// The layout we are adding must be a source layout, without block ids yet.
-		check(SourceLayout->m_blocks.IsEmpty() || SourceLayout->m_blocks[0].m_id == -1);
+		check(SourceLayout->Blocks.IsEmpty() || SourceLayout->Blocks[0].Id == Layout::InvalidBlockId);
 
 		Ptr<const Layout>* it = GeneratedLayouts.Find(SourceLayout.get());
 
@@ -399,14 +399,13 @@ namespace mu
 
 		// Assign unique ids to each layout block
 		Ptr<Layout> ClonedLayout = SourceLayout->Clone();
-		for (int32 b = 0; b < ClonedLayout->m_blocks.Num(); ++b)
+		for (int32 b = 0; b < ClonedLayout->Blocks.Num(); ++b)
 		{
-			// This is a hard limit due to layout block index data being stored in 16 bit.
-			check(m_absoluteLayoutIndex < 65536);
-			ClonedLayout->m_blocks[b].m_id = m_absoluteLayoutIndex++;
+			uint64 Id = uint64(MeshIDPrefix) << 32 | uint64(b);
+			ClonedLayout->Blocks[b].Id = Id;
 		}
-		check(SourceLayout->m_blocks.Num() == ClonedLayout->m_blocks.Num());
-		check(ClonedLayout->m_blocks.IsEmpty() || ClonedLayout->m_blocks[0].m_id != -1);
+		check(SourceLayout->Blocks.Num() == ClonedLayout->Blocks.Num());
+		check(ClonedLayout->Blocks.IsEmpty() || ClonedLayout->Blocks[0].Id != Layout::InvalidBlockId);
 		GeneratedLayouts.Add(SourceLayout.get(), ClonedLayout);
 
 		return ClonedLayout;
@@ -1317,7 +1316,7 @@ namespace mu
 								ImageOptions.RectSize = { 0,0 };
 								ImageOptions.ActiveTags = node.m_tags;
 								ImageOptions.LayoutToApply = pLayout;
-								ImageOptions.LayoutBlockId = pLayout->m_blocks[BlockIndex].m_id;
+								ImageOptions.LayoutBlockId = pLayout->Blocks[BlockIndex].Id;
 								FImageGenerationResult Result;
 								GenerateImage(ImageOptions, Result, pImageNode);
 								Ptr<ASTOp> blockAd = Result.op;
@@ -1357,7 +1356,7 @@ namespace mu
 										if (const NodePatchImage* pPatch = e.node->m_textures[t].m_pPatch.get())
 										{
 											// Is the current block to be patched?
-											if (pPatch->GetPrivate()->m_blocks.Contains(BlockIndex))
+											if (pPatch->GetPrivate()->BlockIndices.Contains(BlockIndex))
 											{
 												blockAd = GenerateImageBlockPatch(blockAd, pPatch, e.condition, ImageOptions);
 											}
@@ -1382,8 +1381,8 @@ namespace mu
 								composeOp->BlockImage = blockAd;
 
 								// Set the absolute block index.
-								check(pLayout->m_blocks[BlockIndex].m_id >= 0);
-								composeOp->BlockIndex = pLayout->m_blocks[BlockIndex].m_id;
+								check(pLayout->Blocks[BlockIndex].Id != Layout::InvalidBlockId);
+								composeOp->BlockId = pLayout->Blocks[BlockIndex].Id;
 
 								imageAd = composeOp;
 							}
@@ -1426,7 +1425,7 @@ namespace mu
 												ImageOptions.ActiveTags = node.m_tags;
 												ImageOptions.RectSize = { 0,0 };
 												ImageOptions.LayoutToApply = pExtendLayout;
-												ImageOptions.LayoutBlockId = pExtendLayout->m_blocks[b].m_id;
+												ImageOptions.LayoutBlockId = pExtendLayout->Blocks[b].Id;
 												FImageGenerationResult ExtendResult;
 												GenerateImage(ImageOptions, ExtendResult, pExtend);
 												Ptr<ASTOp> fragmentAd = ExtendResult.op;
@@ -1464,8 +1463,8 @@ namespace mu
 												composeOp->BlockImage = fragmentAd;
 
 												// Set the absolute block index.
-												check(pExtendLayout->m_blocks[b].m_id >= 0);
-												composeOp->BlockIndex = pExtendLayout->m_blocks[b].m_id;
+												check(pExtendLayout->Blocks[b].Id != Layout::InvalidBlockId);
+												composeOp->BlockId = pExtendLayout->Blocks[b].Id;
 
 												lastBase = composeOp;
 											}
@@ -2177,7 +2176,7 @@ namespace mu
 				// Parameters
 				FImageGenerationOptions ClipOptions;
 				ClipOptions.ImageLayoutStrategy = CompilerOptions::TextureLayoutStrategy::None;
-				ClipOptions.LayoutBlockId = -1;
+				ClipOptions.LayoutBlockId = Layout::InvalidBlockId;
 				ClipOptions.State = Options.State;
 
 				FImageGenerationResult ClipMaskResult;
diff --git a/Engine/Plugins/Experimental/Mutable/Source/MutableTools/Private/MuT/CodeGenerator.h b/Engine/Plugins/Experimental/Mutable/Source/MutableTools/Private/MuT/CodeGenerator.h
--- a/Engine/Plugins/Experimental/Mutable/Source/MutableTools/Private/MuT/CodeGenerator.h
+++ b/Engine/Plugins/Experimental/Mutable/Source/MutableTools/Private/MuT/CodeGenerator.h
@@ -205,8 +205,6 @@ namespace mu
         /** List of image resources for every image formats that have been generated so far as palceholders for missing images. */
         Ptr<Image> m_missingImage[size_t(EImageFormat::IF_COUNT)];
 
-        //! First free index for a layout block
-        int32 m_absoluteLayoutIndex = 0;
 
         //! List of already used vertex ID groups that must be unique.
         TSet<uint32> UniqueVertexIDGroups;
@@ -331,7 +329,7 @@ namespace mu
 			UE::Math::TIntVector2<int32> RectSize = {0, 0};
 
 			/** Layout block that we are trying to generate if any. */
-			int32 LayoutBlockId = -1;
+			uint64 LayoutBlockId = Layout::InvalidBlockId;
 			Ptr<const Layout> LayoutToApply;
 
 			friend FORCEINLINE uint32 GetTypeHash(const FImageGenerationOptions& InKey)
@@ -569,10 +567,11 @@ namespace mu
 			Ptr<Mesh> currentLayoutMesh,
 			int32 currentLayoutChannel,
 			const void* errorContext,
-			const FMeshGenerationOptions& MeshOptions);
+			const FMeshGenerationOptions& MeshOptions,
+			bool bUseAbsoluteBlockIds);
 
 		//!
-		Ptr<const Layout> AddLayout(Ptr<const Layout> SourceLayout);
+		Ptr<const Layout> AddLayout(Ptr<const Layout> SourceLayout, uint32 MeshIDPrefix);
 
 		struct FExtensionDataGenerationResult
 		{
diff --git a/Engine/Plugins/Experimental/Mutable/Source/MutableTools/Private/MuT/CodeGenerator_Image.cpp b/Engine/Plugins/Experimental/Mutable/Source/MutableTools/Private/MuT/CodeGenerator_Image.cpp
--- a/Engine/Plugins/Experimental/Mutable/Source/MutableTools/Private/MuT/CodeGenerator_Image.cpp
+++ b/Engine/Plugins/Experimental/Mutable/Source/MutableTools/Private/MuT/CodeGenerator_Image.cpp
@@ -1389,9 +1389,9 @@ namespace mu
 		int32 LayoutBlockIndex = -1;
 		if (Options.LayoutToApply)
 		{
-			LayoutBlockIndex = Options.LayoutToApply->m_blocks.IndexOfByPredicate([&](const Layout::FBlock& Block) { return Block.m_id == Options.LayoutBlockId; });
+			LayoutBlockIndex = Options.LayoutToApply->Blocks.IndexOfByPredicate([&](const Layout::FBlock& Block) { return Block.Id == Options.LayoutBlockId; });
 		}
-		int32 GeneratedLayoutBlockId = -1;
+		uint64 GeneratedLayoutBlockId = Layout::InvalidBlockId;
 
         // Mesh
         if ( node.m_pMesh )
@@ -1407,18 +1407,16 @@ namespace mu
 			FMeshGenerationResult MeshResult;
 			GenerateMesh( MeshOptions, MeshResult, node.m_pMesh );
 
-			// Match the block id of the block we are generating with the id that resulted in the generated mesh
-			GeneratedLayoutBlockId = -1;
-			
+			// Match the block id of the block we are generating with the id that resulted in the generated mesh			
 			Ptr<const Layout> Layout = MeshResult.GeneratedLayouts.IsValidIndex(node.m_layout) ? MeshResult.GeneratedLayouts[node.m_layout] : nullptr;
-			if (Layout && Layout->m_blocks.IsValidIndex(LayoutBlockIndex))
+			if (Layout && Layout->Blocks.IsValidIndex(LayoutBlockIndex))
 			{
-				GeneratedLayoutBlockId = Layout->m_blocks[LayoutBlockIndex].m_id;
+				GeneratedLayoutBlockId = Layout->Blocks[LayoutBlockIndex].Id;
 			}
-			else if (Layout && Layout->m_blocks.Num() == 1)
+			else if (Layout && Layout->Blocks.Num() == 1)
 			{
 				// Layout management disabled, use the only block available
-				GeneratedLayoutBlockId = Layout->m_blocks[0].m_id;
+				GeneratedLayoutBlockId = Layout->Blocks[0].Id;
 			}
 			else
 			{
@@ -1450,7 +1448,7 @@ namespace mu
             else
             {
                 // Extract the mesh layout block
-                if ( GeneratedLayoutBlockId>=0 )
+                if ( GeneratedLayoutBlockId!= Layout::InvalidBlockId)
                 {
                     Ptr<ASTOpMeshExtractLayoutBlocks> eop = new ASTOpMeshExtractLayoutBlocks();
                     eop->Source = CurrentMeshToProjectOp;
@@ -1543,7 +1541,7 @@ namespace mu
 			FImageGenerationOptions NewOptions = Options;
 			NewOptions.ImageLayoutStrategy = CompilerOptions::TextureLayoutStrategy::None;
 			NewOptions.LayoutToApply = nullptr;
-			NewOptions.LayoutBlockId = -1;
+			NewOptions.LayoutBlockId = Layout::InvalidBlockId;
 			NewOptions.RectSize = { 0,0 };
 
 			FImageGenerationResult ImageResult;
diff --git a/Engine/Plugins/Experimental/Mutable/Source/MutableTools/Private/MuT/CodeGenerator_Mesh.cpp b/Engine/Plugins/Experimental/Mutable/Source/MutableTools/Private/MuT/CodeGenerator_Mesh.cpp
--- a/Engine/Plugins/Experimental/Mutable/Source/MutableTools/Private/MuT/CodeGenerator_Mesh.cpp
+++ b/Engine/Plugins/Experimental/Mutable/Source/MutableTools/Private/MuT/CodeGenerator_Mesh.cpp
@@ -269,7 +269,8 @@ namespace mu
 		Ptr<Mesh> currentLayoutMesh,
 		int32 currentLayoutChannel,
 		const void* errorContext,
-		const FMeshGenerationOptions& MeshOptions
+		const FMeshGenerationOptions& MeshOptions,
+		bool bUseAbsoluteBlockIds
 		)
 	{
 		MUTABLE_CPUPROFILER_SCOPE(LayoutUV_PrepareForLayout);
@@ -280,46 +281,22 @@ namespace mu
 		}
 
 		// The layout must have block ids.
-		check(GeneratedLayout->m_blocks.IsEmpty() || GeneratedLayout->m_blocks[0].m_id != -1);
-
+		check(GeneratedLayout->Blocks.IsEmpty() || GeneratedLayout->Blocks[0].Id != Layout::InvalidBlockId);
 
 		// 
 		Ptr<const Layout> Layout = GeneratedLayout;
 		currentLayoutMesh->AddLayout(Layout);
 
-		// Create the layout block vertex buffer
-		uint16* LayoutData = nullptr;
-		{
-			const int32 LayoutBufferIndex = currentLayoutMesh->GetVertexBuffers().GetBufferCount();
-			currentLayoutMesh->GetVertexBuffers().SetBufferCount(LayoutBufferIndex + 1);
-
-			// TODO
-			check(Layout->GetBlockCount() < MAX_uint16);
-			const EMeshBufferSemantic semantic = MBS_LAYOUTBLOCK;
-			const int32 semanticIndex = int32(currentLayoutChannel);
-			const EMeshBufferFormat format = MBF_UINT16;
-			const int32 components = 1;
-			const int32 offset = 0;
-			currentLayoutMesh->GetVertexBuffers().SetBuffer
-			(
-				LayoutBufferIndex,
-				sizeof(uint16),
-				1,
-				&semantic, &semanticIndex,
-				&format, &components,
-				&offset
-			);
-			LayoutData = (uint16*)currentLayoutMesh->GetVertexBuffers().GetBufferData(LayoutBufferIndex);
-		}
-
 		const int32 NumVertices = currentLayoutMesh->GetVertexCount();
 		const int32 NumBlocks = Layout->GetBlockCount();
 
 		// Clear block data
-		const uint16 NullBlockId = MAX_uint16 - 1;
+		TArray<uint16> LayoutData;
+		LayoutData.SetNumUninitialized(NumVertices);
+		constexpr uint16 NullBlockId = MAX_uint16 - 1;
 		for (int32 VertexIndex = 0; VertexIndex < NumVertices; ++VertexIndex)
 		{
-			*(LayoutData + VertexIndex) = NullBlockId;
+			LayoutData[VertexIndex] = NullBlockId;
 		}
 
 		// Find block ids for each block in the grid
@@ -328,7 +305,7 @@ namespace mu
 		GridBlockBlockId.Init(MAX_uint16, Grid.X * Grid.Y);
 
 		// 
-		TArray<uint16> BlockIds;
+		TArray<uint64> BlockIds;
 		TArray<box<FVector2f>> BlockRects;
 
 		BlockIds.SetNumUninitialized(NumBlocks);
@@ -339,7 +316,7 @@ namespace mu
 		for (int32 BlockIndex = 0; BlockIndex < NumBlocks; ++BlockIndex)
 		{
 			// Get the block id
-			BlockIds[BlockIndex] = Layout->m_blocks[BlockIndex].m_id;
+			BlockIds[BlockIndex] = Layout->Blocks[BlockIndex].Id;
 
 			// Get the block rect
 			uint16 MinX, MinY, SizeX, SizeY;
@@ -485,7 +462,7 @@ namespace mu
 
 			uint16 X, Y;
 
-			uint16& BlockIndexV0 = *(LayoutData + Index0);
+			uint16& BlockIndexV0 = LayoutData[Index0];
 			if (BlockIndexV0 == NullBlockId)
 			{
 				X = (uint16)FMath::Min<uint32>(MaxGridX, FMath::Max<uint32>(0, (uint32)(Grid.X * TempUVs[Index0][0])));
@@ -493,7 +470,7 @@ namespace mu
 				BlockIndexV0 = (X < Grid.X && Y < Grid.Y) ? GridBlockBlockId[Y * Grid.X + X] : 0;
 			}
 
-			uint16& BlockIndexV1 = *(LayoutData + Index1);
+			uint16& BlockIndexV1 = LayoutData[Index1];
 			if (BlockIndexV1 == NullBlockId)
 			{
 				X = (uint16)FMath::Min<uint32>(MaxGridX, FMath::Max<uint32>(0, (uint32)(Grid.X * TempUVs[Index1][0])));
@@ -501,7 +478,7 @@ namespace mu
 				BlockIndexV1 = (X < Grid.X && Y < Grid.Y) ? GridBlockBlockId[Y * Grid.X + X] : 0;
 			}
 
-			uint16& BlockIndexV2 = *(LayoutData + Index2);
+			uint16& BlockIndexV2 = LayoutData[Index2];
 			if (BlockIndexV2 == NullBlockId)
 			{
 				X = (uint16)FMath::Min<uint32>(MaxGridX, FMath::Max<uint32>(0, (uint32)(Grid.X * TempUVs[Index2][0])));
@@ -580,13 +557,13 @@ namespace mu
 			}
 
 			// Get the limits of the predominant block rect
-			const Layout::FBlock& LayoutBlock = Layout->m_blocks[BlockIndex];
+			const Layout::FBlock& LayoutBlock = Layout->Blocks[BlockIndex];
 
 			const float SmallNumber = 0.000001;
-			const float MinX = ((float)LayoutBlock.m_min.X) / (float)Grid.X + SmallNumber;
-			const float MinY = ((float)LayoutBlock.m_min.Y) / (float)Grid.Y + SmallNumber;
-			const float MaxX = (((float)LayoutBlock.m_size.X + LayoutBlock.m_min.X) / (float)Grid.X) - 2 * SmallNumber;
-			const float MaxY = (((float)LayoutBlock.m_size.Y + LayoutBlock.m_min.Y) / (float)Grid.Y) - 2 * SmallNumber;
+			const float MinX = ((float)LayoutBlock.Min.X) / (float)Grid.X + SmallNumber;
+			const float MinY = ((float)LayoutBlock.Min.Y) / (float)Grid.Y + SmallNumber;
+			const float MaxX = (((float)LayoutBlock.Size.X + LayoutBlock.Min.X) / (float)Grid.X) - 2 * SmallNumber;
+			const float MaxY = (((float)LayoutBlock.Size.Y + LayoutBlock.Min.Y) / (float)Grid.Y) - 2 * SmallNumber;
 
 			// Iterate triangles and clamp the UVs
 			FVector2f* TempUVsData = TempUVs.GetData();
@@ -608,7 +585,7 @@ namespace mu
 					FVector2f& UV = TempUVs[UVIndex];
 					UV[0] = FMath::Clamp(UV[0], MinX, MaxX);
 					UV[1] = FMath::Clamp(UV[1], MinY, MaxY);
-					*(LayoutData + UVIndex) = BlockIndex;
+					LayoutData[UVIndex] = BlockIndex;
 				}
 
 				OtherTriangle.bUVsFixed = true;
@@ -643,6 +620,32 @@ namespace mu
 			}
 		}
 
+		// Create the layout block vertex buffer
+		uint8* LayoutBufferPtr = nullptr;
+		{
+			const int32 LayoutBufferIndex = currentLayoutMesh->GetVertexBuffers().GetBufferCount();
+			currentLayoutMesh->GetVertexBuffers().SetBufferCount(LayoutBufferIndex + 1);
+
+			// TODO
+			check(Layout->GetBlockCount() < MAX_uint16);
+			const EMeshBufferSemantic LayoutSemantic = MBS_LAYOUTBLOCK;
+			const int32 LayoutSemanticIndex = int32(currentLayoutChannel);
+			const EMeshBufferFormat LayoutFormat = bUseAbsoluteBlockIds ? MBF_UINT64 : MBF_UINT16;
+			const int32 LayoutComponents = 1;
+			const int32 LayoutOffset = 0;
+			int32 ElementSize = bUseAbsoluteBlockIds ? sizeof(uint64) : sizeof(uint16);
+			currentLayoutMesh->GetVertexBuffers().SetBuffer
+			(
+				LayoutBufferIndex,
+				ElementSize,
+				1,
+				&LayoutSemantic, &LayoutSemanticIndex,
+				&LayoutFormat, &LayoutComponents,
+				&LayoutOffset
+			);
+			LayoutBufferPtr = currentLayoutMesh->GetVertexBuffers().GetBufferData(LayoutBufferIndex);
+		}
+
 		// Format and copy UVs
 		{
 			uint8* pVertices = pData;
@@ -652,19 +655,38 @@ namespace mu
 			{
 				FVector2f* UV = &TempUVs[VertexIndex];
 
-				uint16& LayoutBlockIndex = LayoutData[VertexIndex];
-				if (BlockIds.IsValidIndex(LayoutBlockIndex))
+				uint16 LayoutBlockIndex = LayoutData[VertexIndex];
+
+				// Replace block index by the actual id of the block
+				if (bUseAbsoluteBlockIds)
 				{
-					// Homogenize UVs
-					*UV = BlockRects[LayoutBlockIndex].Homogenize(*UV);
+					uint64* Ptr = reinterpret_cast<uint64*>(LayoutBufferPtr) + VertexIndex;
 
-					// Replace block index by the actual id of the block
-					LayoutBlockIndex = BlockIds[LayoutBlockIndex];
+					if (BlockIds.IsValidIndex(LayoutBlockIndex))
+					{
+						*UV = BlockRects[LayoutBlockIndex].Homogenize(*UV);
+						*Ptr = BlockIds[LayoutBlockIndex];
+					}
+					else
+					{
+						// Map vertices without block to the first block.
+						*Ptr = 0;
+					}
 				}
 				else
 				{
-					// Map vertices without block to the first block.
-					LayoutBlockIndex = 0;
+					uint16* Ptr = reinterpret_cast<uint16*>(LayoutBufferPtr) + VertexIndex;
+
+					if (BlockIds.IsValidIndex(LayoutBlockIndex))
+					{
+						*UV = BlockRects[LayoutBlockIndex].Homogenize(*UV);
+						*Ptr = uint16(BlockIds[LayoutBlockIndex] & 0xffff);
+					}
+					else
+					{
+						// Map vertices without block to the first block.
+						*Ptr = 0;
+					}
 				}
 
 				// Copy UVs
@@ -917,16 +939,16 @@ namespace mu
 
                 op->Source = BaseResult.meshOp;
 
-                if (BaseResult.GeneratedLayouts.Num()>node.m_layoutOrGroup )
+                if (BaseResult.GeneratedLayouts.Num()>node.LayoutOrGroup )
                 {
-                    const Layout* pLayout = BaseResult.GeneratedLayouts[node.m_layoutOrGroup].get();
-                    op->Layout = (uint16)node.m_layoutOrGroup;
+                    const Layout* pLayout = BaseResult.GeneratedLayouts[node.LayoutOrGroup].get();
+                    op->Layout = (uint16)node.LayoutOrGroup;
 
-                    for ( int32 i=0; i<node.m_blocks.Num(); ++i )
+                    for ( int32 i=0; i<node.Blocks.Num(); ++i )
                     {
-                        if (node.m_blocks[i]>=0 && node.m_blocks[i]<pLayout->m_blocks.Num() )
+                        if (node.Blocks[i]>=0 && node.Blocks[i]<pLayout->Blocks.Num() )
                         {
-                            int bid = pLayout->m_blocks[ node.m_blocks[i] ].m_id;
+                            uint64 bid = pLayout->Blocks[ node.Blocks[i] ].Id;
                             op->Blocks.Add(bid);
                         }
                         else
@@ -1264,7 +1286,7 @@ namespace mu
 			// This data is required
 			MeshPtr EmptyMesh = new Mesh();
 			ConstantOp->SetValue(EmptyMesh, m_compilerOptions->OptimisationOptions.DiskCacheContext);
-			EmptyMesh->VertexIDPrefix = ConstantOp->GetValueHash();
+			EmptyMesh->MeshIDPrefix = ConstantOp->GetValueHash();
 
 			FGeneratedConstantMesh MeshEntry;
 			MeshEntry.Mesh = EmptyMesh;
@@ -1363,22 +1385,22 @@ namespace mu
 			ConstantOp->SetValue(Cloned, m_compilerOptions->OptimisationOptions.DiskCacheContext);
 
 			// Add the unique vertex ID prefix in all cases, since it is free memory-wise
-			uint32 VertexIDPrefix = uint32(ConstantOp->GetValueHash());
+			uint32 MeshIDPrefix = uint32(ConstantOp->GetValueHash());
 			{
 				// Ensure the ID group is unique
 				bool bValid = false;
 				do
 				{
 					bool bAlreadyPresent = false;
-					UniqueVertexIDGroups.FindOrAdd(VertexIDPrefix, &bAlreadyPresent);
-					bValid = !bAlreadyPresent && VertexIDPrefix != 0;
+					UniqueVertexIDGroups.FindOrAdd(MeshIDPrefix, &bAlreadyPresent);
+					bValid = !bAlreadyPresent && MeshIDPrefix != 0;
 					if (!bValid)
 					{
-						++VertexIDPrefix;
+						++MeshIDPrefix;
 					}
 				} while (bValid);
 
-				Cloned->VertexIDPrefix = VertexIDPrefix;
+				Cloned->MeshIDPrefix = MeshIDPrefix;
 			}
 
 			// Add the constant data
@@ -1405,9 +1427,11 @@ namespace mu
 						const NodeLayoutBlocks* TypedNode = static_cast<const NodeLayoutBlocks*>(pLayoutNode.get());
 
 						Ptr<const Layout> SourceLayout = TypedNode->GetPrivate()->m_pLayout;
-						Ptr<const Layout> GeneratedLayout = AddLayout( SourceLayout );
+						Ptr<const Layout> GeneratedLayout = AddLayout(SourceLayout, MeshIDPrefix);
 						const void* Context = InOptions.OverrideContext.Get(Node.m_errorContext);
-						PrepareForLayout(GeneratedLayout, Cloned, LayoutIndex, Context, InOptions);
+
+						bool bUseAbsoluteBlockIds = false;
+						PrepareForLayout(GeneratedLayout, Cloned, LayoutIndex, Context, InOptions, bUseAbsoluteBlockIds);
 
 						OutResult.GeneratedLayouts.Add(GeneratedLayout);
 					}
@@ -1419,7 +1443,11 @@ namespace mu
 					{
 						Ptr<const Layout> GeneratedLayout = InOptions.OverrideLayouts[LayoutIndex];
 						const void* Context = InOptions.OverrideContext.Get(Node.m_errorContext);
-						PrepareForLayout(GeneratedLayout, Cloned, LayoutIndex, Context, InOptions);
+
+						// In this case we need the layout block ids to use the ids in the parent layout, and not be prefixed with
+						// the current mesh id prefix. For this reason we need them to be absolute.
+						bool bUseAbsoluteBlockIds = true;
+						PrepareForLayout(GeneratedLayout, Cloned, LayoutIndex, Context, InOptions, bUseAbsoluteBlockIds);
 
 						OutResult.GeneratedLayouts.Add(GeneratedLayout);
 					}
diff --git a/Engine/Plugins/Experimental/Mutable/Source/MutableTools/Private/MuT/CodeOptimiser_Parameters.cpp b/Engine/Plugins/Experimental/Mutable/Source/MutableTools/Private/MuT/CodeOptimiser_Parameters.cpp
--- a/Engine/Plugins/Experimental/Mutable/Source/MutableTools/Private/MuT/CodeOptimiser_Parameters.cpp
+++ b/Engine/Plugins/Experimental/Mutable/Source/MutableTools/Private/MuT/CodeOptimiser_Parameters.cpp
@@ -403,9 +403,7 @@ namespace mu
                         {
 						const ASTOpImageCompose* typedYes = static_cast<const ASTOpImageCompose*>(typedAt->yes.child().get());
 						const ASTOpImageCompose* typedNo = static_cast<const ASTOpImageCompose*>(typedAt->no.child().get());
-                        if ( typedYes->BlockIndex
-                             ==
-                             typedNo->BlockIndex
+                        if ( typedYes->BlockId == typedNo->BlockId
                              &&
                              (
                                  (typedYes->Mask.child().get() != nullptr)
diff --git a/Engine/Plugins/Experimental/Mutable/Source/MutableTools/Private/MuT/NodeMeshFragment.cpp b/Engine/Plugins/Experimental/Mutable/Source/MutableTools/Private/MuT/NodeMeshFragment.cpp
--- a/Engine/Plugins/Experimental/Mutable/Source/MutableTools/Private/MuT/NodeMeshFragment.cpp
+++ b/Engine/Plugins/Experimental/Mutable/Source/MutableTools/Private/MuT/NodeMeshFragment.cpp
@@ -44,60 +44,60 @@ namespace mu
 
 
     //---------------------------------------------------------------------------------------------
-    NodeMeshPtr NodeMeshFragment::GetMesh() const
+    Ptr<NodeMesh> NodeMeshFragment::GetMesh() const
 	{
 		return m_pD->m_pMesh.get();
 	}
 
 
 	//---------------------------------------------------------------------------------------------
-	void NodeMeshFragment::SetMesh( NodeMeshPtr pNode )
+	void NodeMeshFragment::SetMesh(Ptr<NodeMesh> pNode )
 	{
 		m_pD->m_pMesh = pNode;
 	}
 
 
 	//---------------------------------------------------------------------------------------------
-    int NodeMeshFragment::GetLayoutOrGroup() const
+    int32 NodeMeshFragment::GetLayoutOrGroup() const
 	{
-        return m_pD->m_layoutOrGroup;
+        return m_pD->LayoutOrGroup;
 	}
 
 
 	//---------------------------------------------------------------------------------------------
-    void NodeMeshFragment::SetLayoutOrGroup( int l )
+    void NodeMeshFragment::SetLayoutOrGroup( int32 l )
 	{
-        m_pD->m_layoutOrGroup = l;
+        m_pD->LayoutOrGroup = l;
 	}
 
 
 	//---------------------------------------------------------------------------------------------
-	void NodeMeshFragment::SetBlockCount( int t )
+	void NodeMeshFragment::SetBlockCount( int32 t )
 	{
-		m_pD->m_blocks.SetNum(t);
+		m_pD->Blocks.SetNum(t);
 	}
 
 
 	//---------------------------------------------------------------------------------------------
-	int NodeMeshFragment::GetBlock( int t ) const
+	int32 NodeMeshFragment::GetBlock( int32 t ) const
 	{
-		check( t>=0 && t<m_pD->m_blocks.Num() );
-		return m_pD->m_blocks[t];
+		check( t>=0 && t<m_pD->Blocks.Num() );
+		return m_pD->Blocks[t];
 	}
 
 
 	//---------------------------------------------------------------------------------------------
-	void NodeMeshFragment::SetBlock( int t, int b )
+	void NodeMeshFragment::SetBlock( int32 t, int32 b )
 	{
-		check( t>=0 && t<m_pD->m_blocks.Num() );
-		m_pD->m_blocks[t] = b;
+		check( t>=0 && t<m_pD->Blocks.Num() );
+		m_pD->Blocks[t] = b;
 	}
 
 
 	//---------------------------------------------------------------------------------------------
-	NodeLayoutPtr NodeMeshFragment::Private::GetLayout( int index ) const
+	Ptr<NodeLayout> NodeMeshFragment::Private::GetLayout( int32 index ) const
 	{
-		NodeLayoutPtr pResult;
+		Ptr<NodeLayout> pResult;
 
 		if ( m_pMesh )
 		{
diff --git a/Engine/Plugins/Experimental/Mutable/Source/MutableTools/Private/MuT/NodePatchImage.cpp b/Engine/Plugins/Experimental/Mutable/Source/MutableTools/Private/MuT/NodePatchImage.cpp
--- a/Engine/Plugins/Experimental/Mutable/Source/MutableTools/Private/MuT/NodePatchImage.cpp
+++ b/Engine/Plugins/Experimental/Mutable/Source/MutableTools/Private/MuT/NodePatchImage.cpp
@@ -28,60 +28,60 @@ namespace mu
     //---------------------------------------------------------------------------------------------
     // Own Interface
     //---------------------------------------------------------------------------------------------
-    void NodePatchImage::SetImage( NodeImagePtr pImage )
+    void NodePatchImage::SetImage( Ptr<NodeImage> pImage )
     {
         m_pD->m_pImage = pImage;
     }
 
 
     //---------------------------------------------------------------------------------------------
-    NodeImagePtr NodePatchImage::GetImage() const
+	Ptr<NodeImage> NodePatchImage::GetImage() const
     {
         return m_pD->m_pImage.get();
     }
 
 
     //---------------------------------------------------------------------------------------------
-    void NodePatchImage::SetMask( NodeImagePtr pMask )
+    void NodePatchImage::SetMask(Ptr<NodeImage> pMask )
     {
         m_pD->m_pMask = pMask;
     }
 
 
     //---------------------------------------------------------------------------------------------
-    NodeImagePtr NodePatchImage::GetMask() const
+	Ptr<NodeImage> NodePatchImage::GetMask() const
     {
         return m_pD->m_pMask.get();
     }
 
 
     //---------------------------------------------------------------------------------------------
-    void NodePatchImage::SetBlockCount( int c )
+    void NodePatchImage::SetBlockCount( int32 c )
     {
-        m_pD->m_blocks.SetNum( c );
+        m_pD->BlockIndices.SetNum( c );
     }
 
 
     //---------------------------------------------------------------------------------------------
-    int NodePatchImage::GetBlockCount() const
+    int32 NodePatchImage::GetBlockCount() const
     {
-        return (int)m_pD->m_blocks.Num();
+        return (int32)m_pD->BlockIndices.Num();
     }
 
 
     //---------------------------------------------------------------------------------------------
-    void NodePatchImage::SetBlock( int index, int block )
+    void NodePatchImage::SetBlock( int32 index, int32 LayoutBlockIndex )
     {
-        check( index>=0 && index<(int)m_pD->m_blocks.Num() );
-        m_pD->m_blocks[ index ] = block;
+        check( index>=0 && index<m_pD->BlockIndices.Num() );
+        m_pD->BlockIndices[ index ] = LayoutBlockIndex;
     }
 
 
     //---------------------------------------------------------------------------------------------
-    int NodePatchImage::GetBlock( int index ) const
+    int32 NodePatchImage::GetBlock( int32 index ) const
     {
-        check( index>=0 && index<(int)m_pD->m_blocks.Num() );
-        return m_pD->m_blocks[ index ];
+        check( index>=0 && index<m_pD->BlockIndices.Num() );
+        return m_pD->BlockIndices[ index ];
     }
 
 
@@ -93,9 +93,9 @@ namespace mu
 
 
     //---------------------------------------------------------------------------------------------
-    void NodePatchImage::SetBlendType(EBlendType t)
+    void NodePatchImage::SetBlendType(EBlendType InType)
     {
-        m_pD->m_blendType = t;
+        m_pD->m_blendType = InType;
     }
 
 
