[Mutable] Add missing constructor for FMutableParameterIndex.

#rnx

diff --git a/Engine/Plugins/Experimental/Mutable/Source/CustomizableObject/Internal/MuCO/CustomizableObjectPrivate.h b/Engine/Plugins/Experimental/Mutable/Source/CustomizableObject/Internal/MuCO/CustomizableObjectPrivate.h
--- a/Engine/Plugins/Experimental/Mutable/Source/CustomizableObject/Internal/MuCO/CustomizableObjectPrivate.h
+++ b/Engine/Plugins/Experimental/Mutable/Source/CustomizableObject/Internal/MuCO/CustomizableObjectPrivate.h
@@ -559,6 +559,12 @@ template<> struct TCanBulkSerialize<FMorphTargetVertexData> { enum { Value = tru
 
 struct FMutableParameterIndex
 {
+	FMutableParameterIndex(int32 InIndex, int32 InTypedIndex)
+	{
+		Index = InIndex;
+		TypedIndex = InTypedIndex;
+	}
+
 	int32 Index = INDEX_NONE;
 	int32 TypedIndex = INDEX_NONE;
 };
