[Mutable] Fix crash when reference texture is not found.

- An uninitialized image was returned when the reference texture was not found. Now a dummy image is returned.
- Fix crop operations so if the source image is smaller than the crop no uninitialized memory is left.

#rnx

diff --git a/Engine/Plugins/Experimental/Mutable/Source/CustomizableObject/Private/MuCO/UnrealMutableImageProvider.cpp b/Engine/Plugins/Experimental/Mutable/Source/CustomizableObject/Private/MuCO/UnrealMutableImageProvider.cpp
--- a/Engine/Plugins/Experimental/Mutable/Source/CustomizableObject/Private/MuCO/UnrealMutableImageProvider.cpp
+++ b/Engine/Plugins/Experimental/Mutable/Source/CustomizableObject/Private/MuCO/UnrealMutableImageProvider.cpp
@@ -353,7 +353,7 @@ TTuple<UE::Tasks::FTask, TFunction<void()>> FUnrealMutableImageProvider::GetRefe
 		// We just show a warning and move on. This cannot happen during cooks, so it is fine.
 		UE_LOG(LogMutable, Warning, TEXT("Failed to load image [%i]."), Id);
 
-		ResultCallback(Image);
+		ResultCallback(CreateDummy());
 		return Invoke(TrivialReturn);
 	}
 
diff --git a/Engine/Plugins/Experimental/Mutable/Source/MutableRuntime/Private/MuR/OpImageCrop.cpp b/Engine/Plugins/Experimental/Mutable/Source/MutableRuntime/Private/MuR/OpImageCrop.cpp
--- a/Engine/Plugins/Experimental/Mutable/Source/MutableRuntime/Private/MuR/OpImageCrop.cpp
+++ b/Engine/Plugins/Experimental/Mutable/Source/MutableRuntime/Private/MuR/OpImageCrop.cpp
@@ -4,10 +4,10 @@
 #include "MuR/MutableMath.h"
 #include "MuR/Platform.h"
 
+#include "Math/IntRect.h"
 
 namespace mu
 {
-
 	bool FImageOperator::ImageCrop( Image* InCropped, int32 CompressionQuality, const Image* InBase, const box< FIntVector2 >& Rect)
 	{
 		Ptr<const Image> Base = InBase;
@@ -26,6 +26,13 @@ namespace mu
 			Cropped = CreateImage(InCropped->GetSizeX(), InCropped->GetSizeY(), InCropped->GetLODCount(), UncompressedFormat, EInitializationType::NotInitialized);
         }
 
+		// In case the base is smaller than the crop extension, make it as if the base extends indefinitely
+		// with black pixels.
+		if (Rect.min.X + Rect.size.X > InBase->GetSizeX() || Rect.min.Y + Rect.size.Y > InBase->GetSizeY())
+		{
+			Cropped->InitToBlack();			
+		}
+
 		const FImageFormatData& finfo = GetImageFormatData(UncompressedFormat);
 
 		check(Base && Cropped);
@@ -52,20 +59,28 @@ namespace mu
 		checkf(Rect.min[1] % finfo.PixelsPerBlockY == 0, TEXT("Rect must snap to blocks."));
 		checkf(Rect.size[0] % finfo.PixelsPerBlockX == 0, TEXT("Rect must snap to blocks."));
 		checkf(Rect.size[1] % finfo.PixelsPerBlockY == 0, TEXT("Rect must snap to blocks."));
+	
+		const uint32 BytesPerPixel = finfo.BytesPerBlock;
+
+        const uint8* BaseBuf = Base->GetLODData(0);
+        uint8* CropBuf = Cropped->GetLODData(0);
 
-		int32 baseRowSize = finfo.BytesPerBlock * FMath::DivideAndRoundUp(Base->GetSizeX(), uint16(finfo.PixelsPerBlockX));
-		int32 cropRowSize = finfo.BytesPerBlock * FMath::DivideAndRoundUp(Rect.size[0], int32(finfo.PixelsPerBlockX));
+		FImageSize BaseSize = Base->GetSize(); 
+		FImageSize CroppedSize = Cropped->GetSize();
 
-        const uint8* pBaseBuf = Base->GetLODData(0);
-        uint8* pCropBuf = Cropped->GetLODData(0);
+		FIntRect BaseRect = FIntRect(
+				FMath::Min<int32>(BaseSize.X, Rect.min[0]),
+				FMath::Min<int32>(BaseSize.Y, Rect.min[1]),
+				FMath::Min<int32>(BaseSize.X, Rect.min[0] + Rect.size[0]),
+				FMath::Min<int32>(BaseSize.Y, Rect.min[1] + Rect.size[1]));
 
-		int32 SkipPixels = Base->GetSizeX() * Rect.min[1] + Rect.min[0];
-		pBaseBuf += finfo.BytesPerBlock * SkipPixels / finfo.PixelsPerBlockX;
-		for (int32 Y = 0; Y < Rect.size[1]; ++Y)
+		const uint32 NumBytesPerRow = FMath::Max<uint32>(0, (BaseRect.Max.X - BaseRect.Min.X) * BytesPerPixel);
+		for (int32 Y = 0; Y < BaseRect.Max.Y - BaseRect.Min.Y; ++Y)
 		{
-			FMemory::Memcpy(pCropBuf, pBaseBuf, cropRowSize);
-			pCropBuf += cropRowSize;
-			pBaseBuf += baseRowSize;
+			FMemory::Memcpy(
+				CropBuf + CroppedSize.X*Y*BytesPerPixel, 
+				BaseBuf + (BaseSize.X*(Y + BaseRect.Min.Y) + BaseRect.Min.X)*BytesPerPixel, 
+				NumBytesPerRow);
 		}
 
 		if (BaseFormat != UncompressedFormat)
@@ -81,5 +96,4 @@ namespace mu
 
 		return true;
 	}
-
 }
