[Mutable] Import nodes will use their outer CustomizableObject as the default external object.

#rnx

diff --git a/Engine/Plugins/Experimental/Mutable/Source/CustomizableObjectEditor/Private/MuCOE/Nodes/CustomizableObjectNodeExternalPin.cpp b/Engine/Plugins/Experimental/Mutable/Source/CustomizableObjectEditor/Private/MuCOE/Nodes/CustomizableObjectNodeExternalPin.cpp
--- a/Engine/Plugins/Experimental/Mutable/Source/CustomizableObjectEditor/Private/MuCOE/Nodes/CustomizableObjectNodeExternalPin.cpp
+++ b/Engine/Plugins/Experimental/Mutable/Source/CustomizableObjectEditor/Private/MuCOE/Nodes/CustomizableObjectNodeExternalPin.cpp
@@ -46,7 +46,7 @@ void UCustomizableObjectNodeExternalPin::PostBackwardsCompatibleFixup()
 	}
 
 	// Reconstruct the node since the NodeExposePin pin name may have changed while not loaded.
-	ReconstructNode();
+	Super::ReconstructNode();
 }
 
 
@@ -66,7 +66,7 @@ void UCustomizableObjectNodeExternalPin::SetExternalObjectNodeId(FGuid Guid)
 		DestroyNodeDelegateHandle = NodeExposePin->DestroyNodeDelegate.AddUObject(this, &Super::ReconstructNode);
 	}
 
-	ReconstructNode();
+	Super::ReconstructNode();
 }
 
 
@@ -138,6 +138,18 @@ bool UCustomizableObjectNodeExternalPin::CanConnect(const UEdGraphPin* InOwnedIn
 }
 
 
+void UCustomizableObjectNodeExternalPin::ReconstructNode(UCustomizableObjectNodeRemapPins* RemapPinsMode)
+{
+	Super::ReconstructNode(RemapPinsMode);
+
+	if (!ExternalObject)
+	{
+		ExternalObject = Cast<UCustomizableObject>(GetOutermostObject());
+		ExternalObjectNodeId = FGuid();
+	}
+}
+
+
 void UCustomizableObjectNodeExternalPin::BeginPostDuplicate(bool bDuplicateForPIE)
 {
 	Super::BeginPostDuplicate(bDuplicateForPIE);
diff --git a/Engine/Plugins/Experimental/Mutable/Source/CustomizableObjectEditor/Public/MuCOE/Nodes/CustomizableObjectNodeExternalPin.h b/Engine/Plugins/Experimental/Mutable/Source/CustomizableObjectEditor/Public/MuCOE/Nodes/CustomizableObjectNodeExternalPin.h
--- a/Engine/Plugins/Experimental/Mutable/Source/CustomizableObjectEditor/Public/MuCOE/Nodes/CustomizableObjectNodeExternalPin.h
+++ b/Engine/Plugins/Experimental/Mutable/Source/CustomizableObjectEditor/Public/MuCOE/Nodes/CustomizableObjectNodeExternalPin.h
@@ -38,6 +38,8 @@ public:
 	UCustomizableObjectNodeRemapPins* CreateRemapPinsDefault() const override;
 	void BeginPostDuplicate(bool bDuplicateForPIE) override;
 	virtual bool CanConnect(const UEdGraphPin* InOwnedInputPin, const UEdGraphPin* InOutputPin, bool& bOutIsOtherNodeBlocklisted, bool& bOutArePinsCompatible) const override;
+	
+	virtual void ReconstructNode(UCustomizableObjectNodeRemapPins* RemapPinsMode) override;
 
 	void UpdateReferencedNodeId(const FGuid& NewGuid) override;
 	
