[mutable] Undo layout id changes because they break some assets.
#rnx

diff --git a/Engine/Plugins/Experimental/Mutable/Source/CustomizableObject/Internal/MuCO/CustomizableObjectPrivate.h b/Engine/Plugins/Experimental/Mutable/Source/CustomizableObject/Internal/MuCO/CustomizableObjectPrivate.h
--- a/Engine/Plugins/Experimental/Mutable/Source/CustomizableObject/Internal/MuCO/CustomizableObjectPrivate.h
+++ b/Engine/Plugins/Experimental/Mutable/Source/CustomizableObject/Internal/MuCO/CustomizableObjectPrivate.h
@@ -987,8 +987,6 @@ public:
 
 		NumRuntimeReferencedTextures,
 
-		DeterminisiticLayoutBlockIds,
-
 		LastCustomizableObjectVersion
 	};
 	static constexpr int32 CurrentSupportedVersion = ECustomizableObjectVersions::LastCustomizableObjectVersion;
diff --git a/Engine/Plugins/Experimental/Mutable/Source/CustomizableObject/Private/MuCO/CustomizableObjectInstanceDescriptor.cpp b/Engine/Plugins/Experimental/Mutable/Source/CustomizableObject/Private/MuCO/CustomizableObjectInstanceDescriptor.cpp
--- a/Engine/Plugins/Experimental/Mutable/Source/CustomizableObject/Private/MuCO/CustomizableObjectInstanceDescriptor.cpp
+++ b/Engine/Plugins/Experimental/Mutable/Source/CustomizableObject/Private/MuCO/CustomizableObjectInstanceDescriptor.cpp
@@ -1220,7 +1220,7 @@ const FString& FCustomizableObjectInstanceDescriptor::GetIntParameterSelectedOpt
 	const int32 ParameterIndexInObject = CustomizableObject->FindParameter(ParamName);
 	const int32 ParameterIndexInInstance = FindTypedParameterIndex(ParamName, EMutableParameterType::Int);
 
-	if (ParameterIndexInObject >= 0 && IntParameters.IsValidIndex(ParameterIndexInInstance))
+	if (ParameterIndexInObject >= 0 && ParameterIndexInInstance >= 0)
 	{
 		if (RangeIndex == -1)
 		{
diff --git a/Engine/Plugins/Experimental/Mutable/Source/CustomizableObjectEditor/Private/MuCOE/SMutableMeshViewer.cpp b/Engine/Plugins/Experimental/Mutable/Source/CustomizableObjectEditor/Private/MuCOE/SMutableMeshViewer.cpp
--- a/Engine/Plugins/Experimental/Mutable/Source/CustomizableObjectEditor/Private/MuCOE/SMutableMeshViewer.cpp
+++ b/Engine/Plugins/Experimental/Mutable/Source/CustomizableObjectEditor/Private/MuCOE/SMutableMeshViewer.cpp
@@ -273,7 +273,7 @@ TSharedRef<SWidget> SMutableMeshViewer::GenerateDataTableSlates()
 	const FText VerticesCountTitle = LOCTEXT("VerticesCountTitle", "Vertex count : ");
 	const FText FacesCountTitle = LOCTEXT("FacesCountTitle", "Face count : ");
 	const FText BonesCountTitle = LOCTEXT("BonesCountTitle", "Bone count : ");
-	const FText MeshIdPrefixTitle = LOCTEXT("MeshIdPrefixTitle", "Mesh ID prefix : ");
+	const FText VertexIdPrefixTitle = LOCTEXT("VertexIdPrefixTitle", "Vertex ID prefix : ");
 	const FText BuffersTitle = LOCTEXT("BuffersTitle", "Buffers");
 	
 	return SNew(SScrollBox)
@@ -377,14 +377,14 @@ TSharedRef<SWidget> SMutableMeshViewer::GenerateDataTableSlates()
 						.AutoWidth()
 						[
 							SNew(STextBlock).
-								Text(MeshIdPrefixTitle)
+								Text(VertexIdPrefixTitle)
 						]
 
 						+ SHorizontalBox::Slot()
 						.AutoWidth()
 						[
 							SNew(STextBlock).
-								Text(this, &SMutableMeshViewer::GetMeshIdPrefix)
+								Text(this, &SMutableMeshViewer::GetVertexIdPrefix)
 						]
 				]
 			]
@@ -689,9 +689,9 @@ FText SMutableMeshViewer::GetBoneCount() const
 	return FText::AsNumber(MutableMesh && MutableMesh->GetSkeleton() ? MutableMesh->GetSkeleton()->GetBoneCount() : 0);
 }
 
-FText SMutableMeshViewer::GetMeshIdPrefix() const
+FText SMutableMeshViewer::GetVertexIdPrefix() const
 {
-	return FText::AsNumber(MutableMesh ? MutableMesh->MeshIDPrefix : 0);
+	return FText::AsNumber(MutableMesh ? MutableMesh->VertexIDPrefix : 0);
 }
 
 
diff --git a/Engine/Plugins/Experimental/Mutable/Source/CustomizableObjectEditor/Private/MuCOE/SMutableMeshViewer.h b/Engine/Plugins/Experimental/Mutable/Source/CustomizableObjectEditor/Private/MuCOE/SMutableMeshViewer.h
--- a/Engine/Plugins/Experimental/Mutable/Source/CustomizableObjectEditor/Private/MuCOE/SMutableMeshViewer.h
+++ b/Engine/Plugins/Experimental/Mutable/Source/CustomizableObjectEditor/Private/MuCOE/SMutableMeshViewer.h
@@ -123,7 +123,7 @@ private:
 	FText GetVertexCount() const;
 	FText GetFaceCount() const;
 	FText GetBoneCount() const;
-	FText GetMeshIdPrefix() const;
+	FText GetVertexIdPrefix() const;
 
 	/*
 	 * Elements used to feed the buffers list (index and buffer channels as an internal list)
diff --git a/Engine/Plugins/Experimental/Mutable/Source/MutableRuntime/Internal/MuR/MeshPrivate.h b/Engine/Plugins/Experimental/Mutable/Source/MutableRuntime/Internal/MuR/MeshPrivate.h
--- a/Engine/Plugins/Experimental/Mutable/Source/MutableRuntime/Internal/MuR/MeshPrivate.h
+++ b/Engine/Plugins/Experimental/Mutable/Source/MutableRuntime/Internal/MuR/MeshPrivate.h
@@ -717,6 +717,8 @@ namespace mu
 			return 
 				Mesh 
 				&&
+				Mesh->HasVertexIds()
+				&&
 				(CurrentIdIndex < Mesh->GetVertexCount())
 				;
 		}
@@ -726,13 +728,11 @@ namespace mu
 			check(Mesh);
 			check(CurrentIdIndex < Mesh->GetVertexCount());
 
-			uint64 Id=Mesh::InvalidVertexId;
-
 			// Is it implicit?
-			if (!BufferIterator.ptr())
+			if (Mesh->VertexIDPrefix && !BufferIterator.ptr())
 			{
 				// The id is just prefix and index
-				Id = (uint64(Mesh->MeshIDPrefix) << 32) | uint64(CurrentIdIndex);
+				return (uint64(Mesh->VertexIDPrefix) << 32) | uint64(CurrentIdIndex);
 			}
 
 			// Is it relative?
@@ -740,20 +740,19 @@ namespace mu
 			{
 				// There is a buffer storing IDs without prefix because it is the same for all vertices.
 				uint32 RelativeId = BufferIterator.GetAsUINT32();
-				Id = (uint64(Mesh->MeshIDPrefix) << 32) | uint64(RelativeId);
+				return (uint64(Mesh->VertexIDPrefix) << 32) | uint64(RelativeId);
 			}
 
 			// Is it explicit?
 			else if (BufferIterator.GetFormat() == MBF_UINT64)
 			{
-				Id = BufferIterator.GetAsUINT64();
-			}
-			else
-			{
-				check(false);
+				check( Mesh->VertexIDPrefix==0 );
+				uint64 Id = BufferIterator.GetAsUINT64();
+				return Id;
 			}
 
-			return Id;
+			check(false);
+			return 0;
 		}
 
 	};
diff --git a/Engine/Plugins/Experimental/Mutable/Source/MutableRuntime/Internal/MuR/OpImageProject.h b/Engine/Plugins/Experimental/Mutable/Source/MutableRuntime/Internal/MuR/OpImageProject.h
--- a/Engine/Plugins/Experimental/Mutable/Source/MutableRuntime/Internal/MuR/OpImageProject.h
+++ b/Engine/Plugins/Experimental/Mutable/Source/MutableRuntime/Internal/MuR/OpImageProject.h
@@ -24,7 +24,7 @@ namespace mu
 		bool bIsRGBFadingEnabled, bool bIsAlphaFadingEnabled,
 		ESamplingMethod SamplingMethod,
 		float FadeStart, float FadeEnd, float MipInterpolationFactor,
-		int Layout, uint64 BlockId,
+		int Layout, int32 Block,
 		UE::Math::TIntVector2<uint16> CropMin, UE::Math::TIntVector2<uint16> UncroppedSize,
 		FScratchImageProject* scratch, bool bUseVectorImplementation = false);
 
@@ -43,7 +43,7 @@ namespace mu
 		bool bIsRGBFadingEnabled, bool bIsAlphaFadingEnabled,
 		ESamplingMethod SamplingMethod,
 		float FadeStart, float FadeEnd, float MipInterpolationFactor,
-		int32 Layout, uint64 BlockId,
+		int32 Layout, int32 Block,
 		UE::Math::TIntVector2<uint16> CropMin, UE::Math::TIntVector2<uint16> UncroppedSize,
 		FScratchImageProject* Scratch, bool bUseVectorImplementation = false);
 
diff --git a/Engine/Plugins/Experimental/Mutable/Source/MutableRuntime/Internal/MuR/OpMeshMerge.h b/Engine/Plugins/Experimental/Mutable/Source/MutableRuntime/Internal/MuR/OpMeshMerge.h
--- a/Engine/Plugins/Experimental/Mutable/Source/MutableRuntime/Internal/MuR/OpMeshMerge.h
+++ b/Engine/Plugins/Experimental/Mutable/Source/MutableRuntime/Internal/MuR/OpMeshMerge.h
@@ -333,7 +333,7 @@ namespace mu
 			MUTABLE_CPUPROFILER_SCOPE(Layouts);
 
 			Result->Layouts.SetNum(pFirst->Layouts.Num());
-			for (int32 i = 0; i < pFirst->Layouts.Num(); ++i)
+			for (int i = 0; i < pFirst->Layouts.Num(); ++i)
 			{
 				const Layout* pF = pFirst->Layouts[i].get();
 				LayoutPtr pR = pF->Clone();
@@ -342,7 +342,7 @@ namespace mu
 				{
 					const Layout* pS = pSecond->Layouts[i].get();
 
-					pR->Blocks.Append(pS->Blocks);
+					pR->m_blocks.Append(pS->m_blocks);
 				}
 
 				Result->Layouts[i] = pR;
@@ -679,14 +679,6 @@ namespace mu
 			}
 		}
 
-		// This affects both vertex IDs and layout block ids.
-		bool bNeedsExplicitIds = pFirst->MeshIDPrefix != pSecond->MeshIDPrefix;
-		if (!bNeedsExplicitIds)
-		{
-			// This is needed in case a mesh merges with itself.
-			Result->MeshIDPrefix = pFirst->MeshIDPrefix;
-		}
-
 		// Vertices
 		//-----------------
 		{
@@ -696,9 +688,16 @@ namespace mu
 			const int32 SecondCount = pSecond->GetVertexBuffers().GetElementCount();
 
 			// TODO: fast path should be per-buffer
-			// TODO: when formats match, which at runtime should be often.
+			// TODO: when formats match, which at runtime should be always.
 			bool bFastPath = pFirst->GetVertexBuffers().HasSameFormat( pSecond->GetVertexBuffers() );
-			bFastPath = bFastPath && !bNeedsExplicitIds;
+
+			bool bNeedsExplicitVertexIds = pFirst->VertexIDPrefix != pSecond->VertexIDPrefix;
+			if (!bNeedsExplicitVertexIds)
+			{
+				Result->VertexIDPrefix = pFirst->VertexIDPrefix;
+			}
+
+			bFastPath = bFastPath && !bNeedsExplicitVertexIds;
 
 			// Check if the format of the BoneIndex buffer has to change
 			bool bChangeBoneIndicesFormat = false;
@@ -706,10 +705,49 @@ namespace mu
 
 			// Iterate all vertex buffers to check if we need to format bone indices
 			{
-				bChangeBoneIndicesFormat = pFirst->GetVertexBuffers().HasAnySemanticWithDifferentFormat(MBS_BONEINDICES, BoneIndexFormat);
-				if (!bChangeBoneIndicesFormat)
+				// TODO: deduplicate code. move all to cpp.
 				{
-					pSecond->GetVertexBuffers().HasAnySemanticWithDifferentFormat(MBS_BONEINDICES, BoneIndexFormat);
+					// Check if pFirst requires a reformat of the bone index buffers
+					const FMeshBufferSet& VertexBuffers = pFirst->GetVertexBuffers();
+					for (int32 VertexBufferIndex = 0; !bChangeBoneIndicesFormat && VertexBufferIndex < VertexBuffers.m_buffers.Num(); ++VertexBufferIndex)
+					{
+						const FMeshBuffer& Buffer = VertexBuffers.m_buffers[VertexBufferIndex];
+
+						const int32 elemSize = VertexBuffers.GetElementSize(VertexBufferIndex);
+						const int32 firstSize = FirstCount * elemSize;
+
+						const int32 ChannelsCount = VertexBuffers.GetBufferChannelCount(VertexBufferIndex);
+						for (int32 ChannelIndex = 0; ChannelIndex < ChannelsCount; ++ChannelIndex)
+						{
+							if (Buffer.m_channels[ChannelIndex].m_semantic == MBS_BONEINDICES)
+							{
+								bChangeBoneIndicesFormat = Buffer.m_channels[ChannelIndex].m_format != BoneIndexFormat;
+								break;
+							}
+						}
+					}
+				}
+
+				{
+					// Check if pSecond requires a reformat of the bone index buffers
+					const FMeshBufferSet& VertexBuffers = pSecond->GetVertexBuffers();
+					for (int32 VertexBufferIndex = 0; !bChangeBoneIndicesFormat && VertexBufferIndex < VertexBuffers.m_buffers.Num(); ++VertexBufferIndex)
+					{
+						const FMeshBuffer& Buffer = VertexBuffers.m_buffers[VertexBufferIndex];
+
+						const int32 elemSize = VertexBuffers.GetElementSize(VertexBufferIndex);
+						const int32 firstSize = FirstCount * elemSize;
+
+						const int32 ChannelsCount = VertexBuffers.GetBufferChannelCount(VertexBufferIndex);
+						for (int32 ChannelIndex = 0; ChannelIndex < ChannelsCount; ++ChannelIndex)
+						{
+							if (Buffer.m_channels[ChannelIndex].m_semantic == MBS_BONEINDICES)
+							{
+								bChangeBoneIndicesFormat = Buffer.m_channels[ChannelIndex].m_format != BoneIndexFormat;
+								break;
+							}
+						}
+					}
 				}
 
 				bFastPath = bFastPath && !bChangeBoneIndicesFormat;
@@ -737,7 +775,7 @@ namespace mu
 					result.m_elementSize = first.m_elementSize;
 
 					// See if we need to enlarge the components of any of the result channels
-					bool bResetOffsets = false;
+					bool resetOffsets = false;
 					for ( int32 c=0; c<result.m_channels.Num(); ++c )
 					{
 						int32 sb = -1;
@@ -756,16 +794,25 @@ namespace mu
 								 >
 								 result.m_channels[c].m_componentCount )
 							{
-								result.m_channels[c].m_componentCount = second.m_channels[sc].m_componentCount;
-								bResetOffsets = true;
+								result.m_channels[c].m_componentCount =
+										second.m_channels[sc].m_componentCount;
+								resetOffsets = true;
 							}
 						}
 					}
 
 					// Reset the channel offsets if necessary
-					if (bResetOffsets)
+					if (resetOffsets)
 					{
-						Result->GetVertexBuffers().UpdateOffsets(vb);
+						int32 offset = 0;
+						for ( int32 c=0; c<result.m_channels.Num(); ++c )
+						{
+                            result.m_channels[c].m_offset = (uint8_t)offset;
+							offset += result.m_channels[c].m_componentCount
+									*
+									GetMeshFormatData(result.m_channels[c].m_format).SizeInBytes;
+						}
+						result.m_elementSize = offset;
 					}
 				}
 
@@ -811,7 +858,16 @@ namespace mu
 								FMeshBuffer& Buffer = VertexBuffers.m_buffers[foundBuffer];
 								Buffer.m_channels.Insert(chan, foundChannel + 1);
 
-								VertexBuffers.UpdateOffsets(foundBuffer);
+								// Update offsets
+								int32 Offset = Buffer.m_channels[foundChannel].m_offset;
+								for (int32 c = foundChannel; c < Buffer.m_channels.Num(); ++c)
+								{
+									Buffer.m_channels[c].m_offset = (uint8)Offset;
+									Offset += Buffer.m_channels[c].m_componentCount
+										*
+										GetMeshFormatData(Buffer.m_channels[c].m_format).SizeInBytes;
+								}
+								Buffer.m_elementSize = Offset;
 							}
 						}
 					}
@@ -852,41 +908,27 @@ namespace mu
 						}
 					}
 				}
-
-				if (bNeedsExplicitIds)
-				{
-					// Make sure the result format is suitable for the explicit IDs
-
-					// Set to a dummy prefix, to ensure conversion happens
-					Result->MeshIDPrefix = 1;
-					Result->MakeIdsExplicit();
-				}
-
+				
 				// Convert the source meshes to the new format
-                if (!bNeedsExplicitIds && pFirst->GetVertexBuffers().HasSameFormat(Result->GetVertexBuffers()))
+                if (!bNeedsExplicitVertexIds && pFirst->GetVertexBuffers().HasSameFormat(Result->GetVertexBuffers()))
                 {
-					Result->VertexBuffers = pVFirst->VertexBuffers;
-				}
+                    pVFirst = pFirst;
+                }
                 else
                 {
 					check(ScratchMeshes.FirstReformat);
 
 					bool bOutSuccess = false;
+					MeshFormat(ScratchMeshes.FirstReformat.get(), pFirst, Result, false, true, false, false, bOutSuccess);
 
-					if (bNeedsExplicitIds)
+					if (bNeedsExplicitVertexIds)
 					{
-						Ptr<Mesh> TempMesh = pFirst->Clone(EMeshCopyFlags::WithVertexBuffers);
-						TempMesh->MakeIdsExplicit();
-						MeshFormat(ScratchMeshes.FirstReformat.get(), TempMesh.get(), Result, false, true, false, false, bOutSuccess);
-					}
-					else
-					{
-						MeshFormat(ScratchMeshes.FirstReformat.get(), pFirst, Result, false, true, false, false, bOutSuccess);
+						ScratchMeshes.FirstReformat->MakeVertexIndicesExplicit();
 					}
 
 					if (bOutSuccess)
 					{
-						Result->VertexBuffers = ScratchMeshes.FirstReformat->VertexBuffers;
+						pVFirst = ScratchMeshes.FirstReformat;
 					}
 					else
 					{
@@ -896,7 +938,7 @@ namespace mu
                 }
 
 
-                if (!bNeedsExplicitIds && pSecond->GetVertexBuffers().HasSameFormat(Result->GetVertexBuffers()))
+                if (!bNeedsExplicitVertexIds && pSecond->GetVertexBuffers().HasSameFormat(Result->GetVertexBuffers()))
                 {
                     pVSecond = pSecond;
                 }
@@ -905,16 +947,11 @@ namespace mu
 					check(ScratchMeshes.SecondReformat);
 
 					bool bOutSuccess = false;
+                    MeshFormat(ScratchMeshes.SecondReformat.get(), pSecond, Result, false, true, false, false, bOutSuccess);
 
-					if (bNeedsExplicitIds)
-					{
-						Ptr<Mesh> TempMesh = pSecond->Clone(EMeshCopyFlags::WithVertexBuffers);
-						TempMesh->MakeIdsExplicit();
-						MeshFormat(ScratchMeshes.SecondReformat.get(), TempMesh.get(), Result, false, true, false, false, bOutSuccess);
-					}
-					else
+					if (bNeedsExplicitVertexIds)
 					{
-						MeshFormat(ScratchMeshes.SecondReformat.get(), pSecond, Result, false, true, false, false, bOutSuccess);
+						ScratchMeshes.SecondReformat->MakeVertexIndicesExplicit();
 					}
 
 					if (bOutSuccess)
@@ -927,16 +964,17 @@ namespace mu
 					}
 				}
 
-				check(Result->GetVertexBuffers().HasSameFormat(pVSecond->GetVertexBuffers()));
+				check(pVFirst->GetVertexBuffers().HasSameFormat(pVSecond->GetVertexBuffers()));
 			}
 			else
 			{
                 MUTABLE_CPUPROFILER_SCOPE(FastPath);
 
-				Result->VertexBuffers = pFirst->VertexBuffers;
+                pVFirst = pFirst;
 				pVSecond = pSecond;
 			}
 
+			Result->VertexBuffers = pVFirst->VertexBuffers;
 
 			// Allocate additional vertices
 			Result->VertexBuffers.SetElementCount( FirstCount + SecondCount );
diff --git a/Engine/Plugins/Experimental/Mutable/Source/MutableRuntime/Internal/MuR/Operations.h b/Engine/Plugins/Experimental/Mutable/Source/MutableRuntime/Internal/MuR/Operations.h
--- a/Engine/Plugins/Experimental/Mutable/Source/MutableRuntime/Internal/MuR/Operations.h
+++ b/Engine/Plugins/Experimental/Mutable/Source/MutableRuntime/Internal/MuR/Operations.h
@@ -666,7 +666,7 @@ namespace mu
         {
             ADDRESS layout, base, blockImage;
             ADDRESS mask;
-            uint64 BlockId;
+            uint32 blockIndex;
         };
 
         struct ImageInterpolateArgs
@@ -754,9 +754,7 @@ namespace mu
 
         struct ImageRasterMeshArgs
         {
-			uint64 BlockId;
-			
-			ADDRESS mesh;
+            ADDRESS mesh;
 
 			//! These are used in case of projected mesh raster.
 			ADDRESS image;
@@ -769,6 +767,7 @@ namespace mu
 			//! like cylindrical projections.
 			ADDRESS projector;
 			
+			int32 blockId;
 			uint16 sizeX, sizeY;
 			uint16 SourceSizeX, SourceSizeY;
 			uint16 CropMinX, CropMinY;
diff --git a/Engine/Plugins/Experimental/Mutable/Source/MutableRuntime/Private/MuR/CodeRunner.cpp b/Engine/Plugins/Experimental/Mutable/Source/MutableRuntime/Private/MuR/CodeRunner.cpp
--- a/Engine/Plugins/Experimental/Mutable/Source/MutableRuntime/Private/MuR/CodeRunner.cpp
+++ b/Engine/Plugins/Experimental/Mutable/Source/MutableRuntime/Private/MuR/CodeRunner.cpp
@@ -1894,9 +1894,9 @@ namespace mu
                 Ptr<const Mesh> Source = LoadMesh(FCacheAddress(source, item));
 
                 // Access with memcpy necessary for unaligned arm issues.
-                uint64 blocks[512];
-				check(blockCount< 512);
-				FMemory::Memcpy(blocks, data, sizeof(uint64)*FMath::Min(512,int32(blockCount)));
+                uint32 blocks[1024];
+				check(blockCount<1024);
+				FMemory::Memcpy(blocks, data, sizeof(uint32)*FMath::Min(1024,int32(blockCount)));
 
 				if (Source)
 				{
@@ -3719,7 +3719,7 @@ namespace mu
 					Data.Resource = const_cast<Layout*>(ComposeLayout.get());
 					int32 DataPos = m_heapData.Add(Data);
 
-					int32 RelBlockIndex = ComposeLayout->FindBlock(Args.BlockId);
+					int32 RelBlockIndex = ComposeLayout->FindBlock(Args.blockIndex);
 
 					if (RelBlockIndex >= 0)
 					{
@@ -3767,7 +3767,7 @@ namespace mu
 					Data.Resource = const_cast<Layout*>(ComposeLayout.get());
 					int32 DataPos = m_heapData.Add(Data);
 
-					int32 RelBlockIndex = ComposeLayout->FindBlock(Args.BlockId);
+					int32 RelBlockIndex = ComposeLayout->FindBlock(Args.blockIndex);
 					if (RelBlockIndex >= 0)
 					{
 						AddOp(FScheduledOp(item.At, item, 2, DataPos),
@@ -4485,7 +4485,7 @@ namespace mu
 					Ptr<Image> ResultImage = CreateImage(SizeX, SizeY, 1, EImageFormat::IF_L_UBYTE, EInitializationType::Black);
 					if (pMesh)
 					{
-						ImageRasterMesh(pMesh.get(), ResultImage.get(), args.LayoutIndex, args.BlockId, CropMin, UncroppedSize);
+						ImageRasterMesh(pMesh.get(), ResultImage.get(), args.LayoutIndex, args.blockId, CropMin, UncroppedSize);
 						Release(pMesh);
 					}
 
@@ -4687,7 +4687,7 @@ namespace mu
 							args.bIsRGBFadingEnabled, args.bIsAlphaFadingEnabled,
 							SamplingMethod,
 							FadeStartRad, FadeEndRad, FMath::Frac(Data.RasterMesh.MipValue),
-							args.LayoutIndex, args.BlockId,
+							args.LayoutIndex, args.blockId,
 							CropMin, UncroppedSize,
 							&Scratch, bUseProjectionVectorImpl);
 						break;
@@ -4698,7 +4698,7 @@ namespace mu
 							args.bIsRGBFadingEnabled, args.bIsAlphaFadingEnabled,
 							SamplingMethod,
 							FadeStartRad, FadeEndRad, FMath::Frac(Data.RasterMesh.MipValue),
-							args.LayoutIndex, args.BlockId,
+							args.LayoutIndex, args.blockId,
 							CropMin, UncroppedSize,
 							&Scratch, bUseProjectionVectorImpl);
 						break;
diff --git a/Engine/Plugins/Experimental/Mutable/Source/MutableRuntime/Private/MuR/CodeRunnerTasked.cpp b/Engine/Plugins/Experimental/Mutable/Source/MutableRuntime/Private/MuR/CodeRunnerTasked.cpp
--- a/Engine/Plugins/Experimental/Mutable/Source/MutableRuntime/Private/MuR/CodeRunnerTasked.cpp
+++ b/Engine/Plugins/Experimental/Mutable/Source/MutableRuntime/Private/MuR/CodeRunnerTasked.cpp
@@ -1900,7 +1900,7 @@ namespace mu
 
 		Ptr<const Image> Base = Runner->LoadImage({ Args.base, Op.ExecutionIndex, Op.ExecutionOptions });
 
-		int32 RelBlockIndex = Layout->FindBlock(Args.BlockId);
+		int32 RelBlockIndex = Layout->FindBlock(Args.blockIndex);
 
 		// Shortcuts
 		if (RelBlockIndex < 0)
diff --git a/Engine/Plugins/Experimental/Mutable/Source/MutableRuntime/Private/MuR/Layout.cpp b/Engine/Plugins/Experimental/Mutable/Source/MutableRuntime/Private/MuR/Layout.cpp
--- a/Engine/Plugins/Experimental/Mutable/Source/MutableRuntime/Private/MuR/Layout.cpp
+++ b/Engine/Plugins/Experimental/Mutable/Source/MutableRuntime/Private/MuR/Layout.cpp
@@ -39,10 +39,10 @@ namespace mu {
 	{
 		LLM_SCOPE_BYNAME(TEXT("MutableRuntime"));
 		LayoutPtr pResult = new Layout();		
-		pResult->Size = Size;
-		pResult->MaxSize = MaxSize;
-		pResult->Blocks = Blocks;
-		pResult->Strategy = Strategy;
+		pResult->m_size = m_size;
+		pResult->m_maxsize = m_maxsize;
+		pResult->m_blocks = m_blocks;
+		pResult->m_strategy = m_strategy;
 		pResult->FirstLODToIgnoreWarnings = FirstLODToIgnoreWarnings;
 		pResult->ReductionMethod = ReductionMethod;
 		return pResult;
@@ -52,10 +52,10 @@ namespace mu {
 	//---------------------------------------------------------------------------------------------
 	bool Layout::operator==( const Layout& o ) const
 	{
-		return (Size == o.Size) &&
-			(MaxSize == o.MaxSize) &&
-			(Blocks == o.Blocks) &&
-			(Strategy == o.Strategy) &&
+		return  (m_size == o.m_size) &&
+			(m_maxsize == o.m_maxsize) &&
+			(m_blocks == o.m_blocks) &&
+			(m_strategy == o.m_strategy) &&
 			// maybe this is not needed
 			(FirstLODToIgnoreWarnings == o.FirstLODToIgnoreWarnings) &&
 			(ReductionMethod==o.ReductionMethod);
@@ -65,52 +65,52 @@ namespace mu {
 	//---------------------------------------------------------------------------------------------
 	int32 Layout::GetDataSize() const
 	{
-		return sizeof(Layout) + Blocks.GetAllocatedSize();
+		return sizeof(Layout) + m_blocks.GetAllocatedSize();
 	}
 
 
 	//---------------------------------------------------------------------------------------------
 	FIntPoint Layout::GetGridSize() const
 	{
-		return FIntPoint(Size[0], Size[1]);
+		return FIntPoint(m_size[0], m_size[1]);
 	}
 
 
 	//---------------------------------------------------------------------------------------------
-	void Layout::SetGridSize( int32 sizeX, int32 sizeY )
+	void Layout::SetGridSize( int sizeX, int sizeY )
 	{
 		check( sizeX>=0 && sizeY>=0 );
-		Size[0] = (uint16)sizeX;
-		Size[1] = (uint16)sizeY;
+		m_size[0] = (uint16)sizeX;
+		m_size[1] = (uint16)sizeY;
 	}
 
 
 	//---------------------------------------------------------------------------------------------
-	void Layout::GetMaxGridSize(int32* SizeX, int32* SizeY) const
+	void Layout::GetMaxGridSize(int* pSizeX, int* pSizeY) const
 	{
-		check(SizeX && SizeY);
+		check(pSizeX && pSizeY);
 
-		if (SizeX && SizeY)
+		if (pSizeX && pSizeY)
 		{
-			*SizeX = MaxSize[0];
-			*SizeY = MaxSize[1];
+			*pSizeX = m_maxsize[0];
+			*pSizeY = m_maxsize[1];
 		}
 	}
 
 
 	//---------------------------------------------------------------------------------------------
-	void Layout::SetMaxGridSize(int32 sizeX, int32 sizeY)
+	void Layout::SetMaxGridSize(int sizeX, int sizeY)
 	{
 		check(sizeX >= 0 && sizeY >= 0);
-		MaxSize[0] = (uint16)sizeX;
-		MaxSize[1] = (uint16)sizeY;
+		m_maxsize[0] = (uint16)sizeX;
+		m_maxsize[1] = (uint16)sizeY;
 	}
 
 
 	//---------------------------------------------------------------------------------------------
 	int32 Layout::GetBlockCount() const
 	{
-		return Blocks.Num();
+		return m_blocks.Num();
 	}
 
 
@@ -119,22 +119,22 @@ namespace mu {
 	{
 		check( n>=0 );
 		LLM_SCOPE_BYNAME(TEXT("MutableRuntime"));
-		Blocks.SetNum( n );
+		m_blocks.SetNum( n );
 	}
 
 
 	//---------------------------------------------------------------------------------------------
-	void Layout::GetBlock( int32 index, uint16* pMinX, uint16* pMinY, uint16* pSizeX, uint16* pSizeY ) const
+	void Layout::GetBlock( int index, uint16* pMinX, uint16* pMinY, uint16* pSizeX, uint16* pSizeY ) const
 	{
-		check( index >=0 && index < Blocks.Num() );
+		check( index >=0 && index < m_blocks.Num() );
 		check( pMinX && pMinY && pSizeX && pSizeY );
 
 		if (pMinX && pMinY && pSizeX && pSizeY)
 		{
-			*pMinX = Blocks[index].Min[0];
-			*pMinY = Blocks[index].Min[1];
-			*pSizeX = Blocks[index].Size[0];
-			*pSizeY = Blocks[index].Size[1];
+			*pMinX = m_blocks[index].m_min[0];
+			*pMinY = m_blocks[index].m_min[1];
+			*pSizeX = m_blocks[index].m_size[0];
+			*pSizeY = m_blocks[index].m_size[1];
 		}
 	}
 
@@ -142,61 +142,61 @@ namespace mu {
 	//---------------------------------------------------------------------------------------------
 	void Layout::GetBlockOptions(int index, int& pPriority, bool& bReduceBothAxes, bool& bReduceByTwo) const
 	{
-		check(index >= 0 && index < Blocks.Num());
+		check(index >= 0 && index < m_blocks.Num());
 
-		pPriority = Blocks[index].Priority;
-		bReduceBothAxes = Blocks[index].bReduceBothAxes;
-		bReduceByTwo = Blocks[index].bReduceByTwo;
+		pPriority = m_blocks[index].m_priority;
+		bReduceBothAxes = m_blocks[index].bReduceBothAxes;
+		bReduceByTwo = m_blocks[index].bReduceByTwo;
 	}
 
 
 	//---------------------------------------------------------------------------------------------
     void Layout::SetBlock( int index, int minx, int miny, int sizex, int sizey )
 	{
-		check( index >=0 && index < Blocks.Num() );
+		check( index >=0 && index < m_blocks.Num() );
 
 		// Keeps the id
-		Blocks[index].Min = UE::Math::TIntVector2<uint16>((uint16)minx, (uint16)miny);
-		Blocks[index].Size = UE::Math::TIntVector2<uint16>((uint16)sizex, (uint16)sizey);
+		m_blocks[index].m_min = UE::Math::TIntVector2<uint16>((uint16)minx, (uint16)miny);
+		m_blocks[index].m_size = UE::Math::TIntVector2<uint16>((uint16)sizex, (uint16)sizey);
 	}
 
 
 	//---------------------------------------------------------------------------------------------
 	void Layout::SetBlockOptions(int index, int priority, bool bReduceBothAxes, bool bReduceByTwo)
 	{
-		check(index >= 0 && index < Blocks.Num());
+		check(index >= 0 && index < m_blocks.Num());
 
-		Blocks[index].Priority = priority;
-		Blocks[index].bReduceBothAxes = bReduceBothAxes;
-		Blocks[index].bReduceByTwo = bReduceByTwo;
+		m_blocks[index].m_priority = priority;
+		m_blocks[index].bReduceBothAxes = bReduceBothAxes;
+		m_blocks[index].bReduceByTwo = bReduceByTwo;
 	}
 
 
 	//---------------------------------------------------------------------------------------------
 	void Layout::SetLayoutPackingStrategy(EPackStrategy _strategy)
 	{
-		Strategy = _strategy;
+		m_strategy = _strategy;
 	}
 
 
 	//---------------------------------------------------------------------------------------------
 	EPackStrategy Layout::GetLayoutPackingStrategy() const
 	{
-		return Strategy;
+		return m_strategy;
 	}
 
 
 	//---------------------------------------------------------------------------------------------
 	void Layout::Serialise(OutputArchive& arch) const
 	{
-		uint32 ver = 7;
+		uint32 ver = 6;
 		arch << ver;
 
-		arch << Size;
-		arch << Blocks;
+		arch << m_size;
+		arch << m_blocks;
 
-		arch << MaxSize;
-		arch << uint32(Strategy);
+		arch << m_maxsize;
+		arch << uint32(m_strategy);
 		arch << FirstLODToIgnoreWarnings;
 		arch << uint32(ReductionMethod);
 	}
@@ -207,33 +207,56 @@ namespace mu {
 	{
 		uint32 ver;
 		arch >> ver;
-		check(ver <= 7);
+		check(ver <= 6);
 
-		arch >> Size;
-		arch >> Blocks;
-		arch >> MaxSize;
+		arch >> m_size;
+
+		if (ver < 6)
+		{
+			uint32_t Size = 0;
+			arch >> Size;
+			m_blocks.SetNum(Size);
+
+			for (uint32_t BlockIndex = 0; BlockIndex < Size; ++BlockIndex)
+			{
+				// Unserialise taking into account the Layout version
+				m_blocks[BlockIndex].UnserialiseOldVersion(arch, ver);
+			}
+		}
+		else
+		{
+			arch >> m_blocks;
+		}
+
+		arch >> m_maxsize;
 
 		uint32 temp;
 		arch >> temp;
-		Strategy = EPackStrategy(temp);
+		m_strategy = EPackStrategy(temp);
 
-		arch >> FirstLODToIgnoreWarnings;
+		if (ver >= 4)
+		{
+			arch >> FirstLODToIgnoreWarnings;
+		}
 
-		arch >> temp;
-		ReductionMethod = EReductionMethod(temp);
+		if (ver >= 5)
+		{
+			arch >> temp;
+			ReductionMethod = EReductionMethod(temp);
+		}
 	}
 
 
 	//---------------------------------------------------------------------------------------------
 	bool Layout::IsSimilar(const Layout& o) const
 	{
-		if (Size != o.Size || MaxSize != o.MaxSize ||
-			Blocks.Num() != o.Blocks.Num() || Strategy != o.Strategy)
+		if (m_size != o.m_size || m_maxsize != o.m_maxsize ||
+			m_blocks.Num() != o.m_blocks.Num() || m_strategy != o.m_strategy)
 			return false;
 
-		for (int32 i = 0; i < Blocks.Num(); ++i)
+		for (int32 i = 0; i < m_blocks.Num(); ++i)
 		{
-			if (!Blocks[i].IsSimilar(o.Blocks[i])) return false;
+			if (!m_blocks[i].IsSimilar(o.m_blocks[i])) return false;
 		}
 
 		return true;
@@ -242,26 +265,28 @@ namespace mu {
 
 
 	//---------------------------------------------------------------------------------------------
-	int32 Layout::FindBlock(uint64 Id) const
+	int32 Layout::FindBlock(int32_t id) const
 	{
-		for (int32 i = 0; i < Blocks.Num(); ++i)
+		check(id >= 0);
+		int res = -1;
+		for (int32 i = 0; res < 0 && i < m_blocks.Num(); ++i)
 		{
-			if (Blocks[i].Id == Id)
+			if (m_blocks[i].m_id == id)
 			{
-				return i;
+				res = i;
 			}
 		}
 
-		return -1;
+		return res;
 	}
 
 
 	//---------------------------------------------------------------------------------------------
 	bool Layout::IsSingleBlockAndFull() const
 	{
-		if (Blocks.Num() == 1
-			&& Blocks[0].Min == UE::Math::TIntVector2<uint16>(0, 0)
-			&& Blocks[0].Size == Size)
+		if (m_blocks.Num() == 1
+			&& m_blocks[0].m_min == UE::Math::TIntVector2<uint16>(0, 0)
+			&& m_blocks[0].m_size == m_size)
 		{
 			return true;
 		}
@@ -300,10 +325,10 @@ namespace mu {
 	//---------------------------------------------------------------------------------------------
 	void Layout::FBlock::Serialise(OutputArchive& arch) const
 	{
-		arch << Min;
-		arch << Size;
-		arch << Id;
-		arch << Priority;
+		arch << m_min;
+		arch << m_size;
+		arch << m_id;
+		arch << m_priority;
 		arch << bReduceBothAxes;
 		arch << bReduceByTwo;
 	}
@@ -312,13 +337,27 @@ namespace mu {
 	//---------------------------------------------------------------------------------------------
 	void Layout::FBlock::Unserialise(InputArchive& arch)
 	{
-		arch >> Min;
-		arch >> Size;
-		arch >> Id;
-		arch >> Priority;
+		arch >> m_min;
+		arch >> m_size;
+		arch >> m_id;
+		arch >> m_priority;
 		arch >> bReduceBothAxes;
 		arch >> bReduceByTwo;
 	}
 
+	
+	//---------------------------------------------------------------------------------------------
+	void Layout::FBlock::UnserialiseOldVersion(InputArchive& Archive, const int32 Version)
+	{
+		Archive >> m_min;
+		Archive >> m_size;
+		Archive >> m_id;
+		Archive >> m_priority;
+
+		if (Version >= 5)
+		{
+			Archive >> bReduceBothAxes;
+		}
+	}
 }
 
diff --git a/Engine/Plugins/Experimental/Mutable/Source/MutableRuntime/Private/MuR/MeshBufferSet.cpp b/Engine/Plugins/Experimental/Mutable/Source/MutableRuntime/Private/MuR/MeshBufferSet.cpp
--- a/Engine/Plugins/Experimental/Mutable/Source/MutableRuntime/Private/MuR/MeshBufferSet.cpp
+++ b/Engine/Plugins/Experimental/Mutable/Source/MutableRuntime/Private/MuR/MeshBufferSet.cpp
@@ -662,26 +662,6 @@ namespace mu
 			m_buffers[b].m_elementSize = offset;
 	}
 
-
-	bool FMeshBufferSet::HasAnySemanticWithDifferentFormat(EMeshBufferSemantic Semantic, EMeshBufferFormat ExpectedFormat) const
-	{
-		for (const FMeshBuffer& Buffer : m_buffers)
-		{
-			for (int32 ChannelIndex = 0; ChannelIndex < Buffer.m_channels.Num(); ++ChannelIndex)
-			{
-				if (Buffer.m_channels[ChannelIndex].m_semantic == Semantic)
-				{
-					if (Buffer.m_channels[ChannelIndex].m_format != ExpectedFormat)
-					{
-						return true;
-					}
-				}
-			}
-		}
-
-		return false;
-	}
-
 	
 	//-----------------------------------------------------------------------------------------
 	void FMeshBuffer::Serialise(OutputArchive& arch) const
diff --git a/Engine/Plugins/Experimental/Mutable/Source/MutableRuntime/Private/MuR/OpImageProject.cpp b/Engine/Plugins/Experimental/Mutable/Source/MutableRuntime/Private/MuR/OpImageProject.cpp
--- a/Engine/Plugins/Experimental/Mutable/Source/MutableRuntime/Private/MuR/OpImageProject.cpp
+++ b/Engine/Plugins/Experimental/Mutable/Source/MutableRuntime/Private/MuR/OpImageProject.cpp
@@ -902,7 +902,7 @@ void ImageRasterProjectedPlanar(const Mesh* pMesh, Image* pImage,
 	bool bIsRGBFadingEnabled, bool bIsAlphaFadingEnabled,
 	ESamplingMethod SamplingMethod,
 	float FadeStart, float FadeEnd, float MipInterpolationFactor,
-	int32 Layout, uint64 BlockId,
+	int32 Layout, int32 Block,
 	UE::Math::TIntVector2<uint16> CropMin, UE::Math::TIntVector2<uint16> UncroppedSize,
 	FScratchImageProject* Scratch, bool bUseVectorImplementation)
 {
@@ -958,7 +958,7 @@ void ImageRasterProjectedPlanar(const Mesh* pMesh, Image* pImage,
 			bIsRGBFadingEnabled, bIsAlphaFadingEnabled,
 			SamplingMethod,
 			FadeStart, FadeEnd,  UnusedProjectionAngle, MipInterpolationFactor, 
-			Layout, BlockId, CropMin, UncroppedSize, Scratch
+			Layout, Block, CropMin, UncroppedSize, Scratch
 		};
 
 		using EPPF = EPixelProcessorFeatures;
@@ -1116,7 +1116,7 @@ void ImageRasterProjectedWrapping( const Mesh* pMesh, Image* pImage,
 	bool bIsRGBFadingEnabled, bool bIsAlphaFadingEnabled,
 	ESamplingMethod SamplingMethod,
 	float FadeStart, float FadeEnd, float MipInterpolationFactor,
-	int32 Layout, uint64 BlockId,
+	int32 Layout, int32 Block,
 	UE::Math::TIntVector2<uint16> CropMin, UE::Math::TIntVector2<uint16> UncroppedSize,
 	FScratchImageProject* Scratch, bool bUseVectorImplementation)
 {
@@ -1173,7 +1173,7 @@ void ImageRasterProjectedWrapping( const Mesh* pMesh, Image* pImage,
 			bIsRGBFadingEnabled, bIsAlphaFadingEnabled,
 			SamplingMethod,
 			FadeStart, FadeEnd, UnusedProjectionAngle, MipInterpolationFactor,
-			Layout, BlockId, CropMin, UncroppedSize, Scratch
+			Layout, Block, CropMin, UncroppedSize, Scratch
 		};
 
 		using EPPF = EPixelProcessorFeatures;
diff --git a/Engine/Plugins/Experimental/Mutable/Source/MutableRuntime/Private/MuR/OpImageRasterMesh.h b/Engine/Plugins/Experimental/Mutable/Source/MutableRuntime/Private/MuR/OpImageRasterMesh.h
--- a/Engine/Plugins/Experimental/Mutable/Source/MutableRuntime/Private/MuR/OpImageRasterMesh.h
+++ b/Engine/Plugins/Experimental/Mutable/Source/MutableRuntime/Private/MuR/OpImageRasterMesh.h
@@ -27,7 +27,7 @@ namespace mu
 	};
 
 
-	inline void ImageRasterMesh( const Mesh* pMesh, Image* pImage, int32 LayoutIndex, uint64 BlockId,
+	inline void ImageRasterMesh( const Mesh* pMesh, Image* pImage, int32 LayoutIndex, int32 BlockId,
 		UE::Math::TIntVector2<uint16> CropMin, UE::Math::TIntVector2<uint16> UncroppedSize )
 	{
 		MUTABLE_CPUPROFILER_SCOPE(ImageRasterMesh)
@@ -80,7 +80,7 @@ namespace mu
 
         UntypedMeshBufferIteratorConst bloIt( pMesh->GetVertexBuffers(), MBS_LAYOUTBLOCK, LayoutIndex );
 
-        if (BlockId==Layout::InvalidBlockId || bloIt.GetElementSize()==0 )
+        if (BlockId <0 || bloIt.GetElementSize()==0 )
 		{
 			// Raster all the faces
             WhitePixelProcessor pixelProc;
@@ -108,49 +108,30 @@ namespace mu
 		{
 			// Raster only the faces in the selected block
 
-			// Get the block per vertex
-			TArray<uint64> VertexBlockIds;
-			VertexBlockIds.SetNumZeroed(vertexCount);
+			// Get the block per face
+			TArray<uint16> blocks;
+			blocks.SetNumZeroed(vertexCount);
 
-			if (bloIt.GetFormat() == MBF_UINT16)
+			for ( int i=0; i<vertexCount; ++i )
 			{
-				// Relative blocks.
-				const uint16* SourceIds = reinterpret_cast<const uint16*>(bloIt.ptr());
-				for (int32 i = 0; i < vertexCount; ++i)
-				{
-					uint64 Id = SourceIds[i];
-					Id = Id | (uint64(pMesh->MeshIDPrefix)<<32);
-					VertexBlockIds[i] = Id;
-					++SourceIds;
-				}
-			}
-			else if (bloIt.GetFormat() == MBF_UINT64)
-			{
-				// Absolute blocks.
-				const uint64* SourceIds = reinterpret_cast<const uint64*>(bloIt.ptr());
-				for (int32 i = 0; i < vertexCount; ++i)
-				{
-					uint64 Id = SourceIds[i];
-					VertexBlockIds[i] = Id;
-					++SourceIds;
-				}
-			}
-			else
-			{
-				// Format not supported
-				check(false);
+                uint16 index=0;
+				ConvertData( 0, &index, MBF_UINT16, bloIt.ptr(), bloIt.GetFormat() );
+
+				blocks[i] = index;
+				++bloIt;
 			}
 
             WhitePixelProcessor pixelProc;
 
 			const TArrayView<uint8> ImageData = pImage->DataStorage.GetLOD(0); 
 
+			//for (int f = 0; f < faceCount; ++f)
 			const auto& ProcessFace = [
-				vertices, indices, VertexBlockIds, BlockId, ImageData, sizeX, sizeY, pixelProc
+				vertices, indices, blocks, BlockId, ImageData, sizeX, sizeY, pixelProc
 			] (int32 f)
 			{
 				// TODO: Select faces outside for loop?
-				if (VertexBlockIds[indices[f * 3 + 0]] == BlockId)
+				if (blocks[indices[f * 3 + 0]] == BlockId)
 				{
 					constexpr int32 NumInterpolators = 1;
 					Triangle<NumInterpolators>(ImageData.GetData(), ImageData.Num(),
diff --git a/Engine/Plugins/Experimental/Mutable/Source/MutableRuntime/Private/MuR/OpLayoutRemoveBlocks.h b/Engine/Plugins/Experimental/Mutable/Source/MutableRuntime/Private/MuR/OpLayoutRemoveBlocks.h
--- a/Engine/Plugins/Experimental/Mutable/Source/MutableRuntime/Private/MuR/OpLayoutRemoveBlocks.h
+++ b/Engine/Plugins/Experimental/Mutable/Source/MutableRuntime/Private/MuR/OpLayoutRemoveBlocks.h
@@ -19,86 +19,24 @@ namespace mu
 		}
 
 		const Layout* Source = InMesh->GetLayout(InLayoutIndex);
-		Ptr<Layout> pResult;
+
+		// Create the list of blocks in the mesh
+		TArray<bool> blocksFound;
+		blocksFound.SetNumZeroed(1024);
 
 		UntypedMeshBufferIteratorConst itBlocks(InMesh->GetVertexBuffers(), MBS_LAYOUTBLOCK, InLayoutIndex);
 		if (itBlocks.GetFormat() == MBF_UINT16)
 		{
-			// Relative blocks.
-
 			const uint16* pBlocks = reinterpret_cast<const uint16*>(itBlocks.ptr());
-
-			uint16 MaxId = 0;
 			for (int32 i = 0; i < InMesh->GetVertexCount(); ++i)
 			{
-				uint16 RelativeId = pBlocks[i];
-				check(RelativeId != 0xffff);
-				MaxId = FMath::Max(MaxId, RelativeId);
-			}
-
-
-			// Create the list of blocks in the mesh.
-			// Array that stores a flag for every ID, possibly wasting space.
-			TBitArray BlocksFound;
-			BlocksFound.Init(false,MaxId+1);
-			for (int32 i = 0; i < InMesh->GetVertexCount(); ++i)
-			{
-				BlocksFound[pBlocks[i]] = true;
-			}
-
-			// Remove blocks that are not in the mesh
-			pResult = Source->Clone();
-			int32 DestBlockIndex = 0;
-			for (int32 BlockIndex = 0; BlockIndex < pResult->Blocks.Num(); ++BlockIndex)
-			{
-				uint64 BlockId = pResult->Blocks[BlockIndex].Id;
-				uint32 BlockIdPrefix = uint32(BlockId >> 32);
-				uint32 RelativeBlockId = uint32(BlockId & 0xffffffff);
-
-				bool bBlockMatchesPrefix = BlockIdPrefix == InMesh->MeshIDPrefix;
-
-				if (bBlockMatchesPrefix && RelativeBlockId < uint32(BlocksFound.Num()) && BlocksFound[RelativeBlockId] )
+				if (pBlocks[i] >= blocksFound.Num())
 				{
-					// keep this block
-					pResult->Blocks[DestBlockIndex] = pResult->Blocks[BlockIndex];
-					++DestBlockIndex;
+					blocksFound.SetNumZeroed(pBlocks[i] + 1024);
 				}
-			}
-			pResult->SetBlockCount(DestBlockIndex);
-
-		}
-		else if (itBlocks.GetFormat() == MBF_UINT64)
-		{
-			// Absolute blocks.
-
-			// Create the list of blocks in the mesh.
-			// TODO: SparseIndexSet? like in the 16-bit case, but per-prefix?
-			TSet<uint64> BlocksFound;
-			BlocksFound.Reserve(64);
-
-			const uint64* BlockIdsPerVertex = reinterpret_cast<const uint64*>(itBlocks.ptr());
-			for (int32 i = 0; i < InMesh->GetVertexCount(); ++i)
-			{
-				BlocksFound.FindOrAdd(BlockIdsPerVertex[i]);
-			}
 
-			// Remove blocks that are not in the mesh
-			pResult = Source->Clone();
-			int32 DestBlockIndex = 0;
-			for (int32 BlockIndex = 0; BlockIndex < pResult->Blocks.Num(); ++BlockIndex)
-			{
-				uint64 BlockId = pResult->Blocks[BlockIndex].Id;
-				bool bBlockFound = BlocksFound.Contains(BlockId);
-
-				if (bBlockFound)
-				{
-					// keep this block
-					pResult->Blocks[DestBlockIndex] = pResult->Blocks[BlockIndex];
-					++DestBlockIndex;
-				}
+				blocksFound[pBlocks[i]] = true;
 			}
-			pResult->SetBlockCount(DestBlockIndex);
-
 		}
 		else if (itBlocks.GetFormat() == MBF_NONE)
 		{
@@ -112,6 +50,20 @@ namespace mu
 			check(false);
 		}
 
+		// Remove blocks that are not in the mesh
+		Ptr<Layout> pResult = Source->Clone();
+		int32 dest = 0;
+		for (int32 b = 0; b < pResult->m_blocks.Num(); ++b)
+		{
+			int blockIndex = pResult->m_blocks[b].m_id;
+			if (blockIndex < (int)blocksFound.Num() && blocksFound[blockIndex])
+			{
+				// keep
+				pResult->m_blocks[dest] = pResult->m_blocks[b];
+				++dest;
+			}
+		}
+		pResult->SetBlockCount(dest);
 
 		return pResult;
 	}
@@ -120,21 +72,37 @@ namespace mu
 	//---------------------------------------------------------------------------------------------
 	inline Ptr<Layout> LayoutRemoveBlocks(const Layout* Source, const Layout* ReferenceLayout)
 	{
+		// Create the list of blocks in the mesh
+		TArray<bool,TInlineAllocator<1024>> BlocksFound;
+		BlocksFound.SetNumZeroed(1024);
+
+		if (ReferenceLayout)
+		{
+			for (const Layout::FBlock& Block: ReferenceLayout->m_blocks)
+			{
+				if (Block.m_id >= BlocksFound.Num())
+				{
+					BlocksFound.SetNumZeroed(Block.m_id + 1024);
+				}
+
+				BlocksFound[Block.m_id] = true;
+			}
+		}
+
 		// Remove blocks that are not in the mesh
 		Ptr<Layout> pResult = Source->Clone();
-		int32 DestBlockIndex = 0;
-		for (int32 BlockIndexInSource = 0; BlockIndexInSource < pResult->Blocks.Num(); ++BlockIndexInSource)
+		int dest = 0;
+		for (int32 b = 0; b < pResult->m_blocks.Num(); ++b)
 		{
-			uint64 BlockId = pResult->Blocks[BlockIndexInSource].Id;
-			int32 BlockIndexInReference = ReferenceLayout->FindBlock(BlockId);
-			if (BlockIndexInReference >=0 )
+			int blockIndex = pResult->m_blocks[b].m_id;
+			if (blockIndex < (int)BlocksFound.Num() && BlocksFound[blockIndex])
 			{
 				// keep
-				pResult->Blocks[DestBlockIndex] = pResult->Blocks[BlockIndexInSource];
-				++DestBlockIndex;
+				pResult->m_blocks[dest] = pResult->m_blocks[b];
+				++dest;
 			}
 		}
-		pResult->SetBlockCount(DestBlockIndex);
+		pResult->SetBlockCount(dest);
 
 		return pResult;
 	}
@@ -145,11 +113,20 @@ namespace mu
 	{
 		Ptr<Layout> pResult = pA->Clone();
 
-		for ( const Layout::FBlock& block: pB->Blocks )
+		// This is faster but fails in the rare case of a block being in both layouts, which may 
+		// happen if we merge a mesh with itself.
+		//pResult->GetPrivate()->m_blocks.insert
+		//(
+		//	pResult->GetPrivate()->m_blocks.end(),
+		//	pB->GetPrivate()->m_blocks.begin(),
+		//	pB->GetPrivate()->m_blocks.end()
+		//);
+
+		for ( const Layout::FBlock& block: pB->m_blocks )
 		{
-			if ( pResult->FindBlock(block.Id)<0 )
+			if ( pResult->FindBlock(block.m_id)<0 )
 			{
-				pResult->Blocks.Add(block);
+				pResult->m_blocks.Add(block);
 			}
 		}
 
diff --git a/Engine/Plugins/Experimental/Mutable/Source/MutableRuntime/Private/MuR/OpMeshApplyLayout.h b/Engine/Plugins/Experimental/Mutable/Source/MutableRuntime/Private/MuR/OpMeshApplyLayout.h
--- a/Engine/Plugins/Experimental/Mutable/Source/MutableRuntime/Private/MuR/OpMeshApplyLayout.h
+++ b/Engine/Plugins/Experimental/Mutable/Source/MutableRuntime/Private/MuR/OpMeshApplyLayout.h
@@ -67,41 +67,22 @@ namespace mu
 
         check( layoutBuffer>=0 && layoutChannel>=0 );
         check( pApplied->GetVertexBuffers().m_buffers[layoutBuffer].m_channels.Num()==1 );
+        check( pApplied->GetVertexBuffers().m_buffers[layoutBuffer].m_channels[layoutChannel].m_format==MBF_UINT16 );
         check( pApplied->GetVertexBuffers().m_buffers[layoutBuffer].m_channels[layoutChannel].m_componentCount==1 );
 
         const uint16* pLayoutData = reinterpret_cast<const uint16*>( pApplied->GetVertexBuffers().GetBufferData( layoutBuffer ) );
-		UntypedMeshBufferIterator ItLayoutData(pApplied->GetVertexBuffers(), MBS_LAYOUTBLOCK, texCoordsSet);
 
-		// In some corner case involving automatic LODs and remove meshes behaving differently among them we may need to remove vertices that don't 
+		// In some corner case involcing automatic LODs and remove meshes behaving differently among them we may need to remove vertices that don't 
 		// have any block in the current layout. Track them here.
 		TArray<int32> VerticesToRemove;
 
         uint8* pVertices = pData;
 		for ( int32 v=0; v<pApplied->GetVertexBuffers().GetElementCount(); ++v )
 		{
-			uint64 BlockId = 0;
-			if (ItLayoutData.GetFormat() == MBF_UINT16)
-			{
-				// Relative blocks.
-				const uint16* SourceIds = reinterpret_cast<const uint16*>(ItLayoutData.ptr());
-				BlockId = *SourceIds;
-				BlockId = BlockId | (uint64(pApplied->MeshIDPrefix) << 32);
-			}
-			else if (ItLayoutData.GetFormat() == MBF_UINT64)
-			{
-				// Absolute blocks.
-				const uint64* SourceIds = reinterpret_cast<const uint64*>(ItLayoutData.ptr());
-				BlockId = *SourceIds;
-			}
-			else
-			{
-				// Format not supported
-				check(false);
-			}
-			++ItLayoutData;
+			int32 blockId = pLayoutData[v];
 
 			// TODO: This could be optimised
-			int32 relBlock = pLayout->FindBlock(BlockId);
+			int32 relBlock = pLayout->FindBlock( blockId );
 
 			// This may still happen with lower LOD and "remove meshes" in a corner case:
 			// Auto LODs with "Remove Meshes" that behave differently across LODs, and leave geometry in a block that has been removed in the higher LOD.
diff --git a/Engine/Plugins/Experimental/Mutable/Source/MutableRuntime/Private/MuR/OpMeshClipWithMesh.cpp b/Engine/Plugins/Experimental/Mutable/Source/MutableRuntime/Private/MuR/OpMeshClipWithMesh.cpp
--- a/Engine/Plugins/Experimental/Mutable/Source/MutableRuntime/Private/MuR/OpMeshClipWithMesh.cpp
+++ b/Engine/Plugins/Experimental/Mutable/Source/MutableRuntime/Private/MuR/OpMeshClipWithMesh.cpp
@@ -572,7 +572,7 @@ namespace mu { namespace
 
 		Result->GetVertexBuffers().SetElementCount(MaskVertexCount);
 		Result->GetVertexBuffers().SetBufferCount(1);
-		Result->MeshIDPrefix = Base->MeshIDPrefix;
+		Result->VertexIDPrefix = Base->VertexIDPrefix;
 
 		EMeshBufferSemantic Semantic = MBS_VERTEXINDEX;
 		int32 SemanticIndex = 0;
diff --git a/Engine/Plugins/Experimental/Mutable/Source/MutableRuntime/Private/MuR/OpMeshDifference.h b/Engine/Plugins/Experimental/Mutable/Source/MutableRuntime/Private/MuR/OpMeshDifference.h
--- a/Engine/Plugins/Experimental/Mutable/Source/MutableRuntime/Private/MuR/OpMeshDifference.h
+++ b/Engine/Plugins/Experimental/Mutable/Source/MutableRuntime/Private/MuR/OpMeshDifference.h
@@ -162,7 +162,15 @@ namespace mu
 			BufferIndex = 1;
 			Result->GetVertexBuffers().SetBuffer(BufferIndex, ElementSize, semantic.Num(), semantic.GetData(), semanticIndex.GetData(), format.GetData(), components.GetData(), offsets.GetData());
 
-			Result->MeshIDPrefix = pBase->MeshIDPrefix;
+
+			if (bGenerateRelativeIds)
+			{
+				Result->VertexIDPrefix = pBase->VertexIDPrefix;
+			}
+			else
+			{
+				Result->VertexIDPrefix = 0;
+			}
 
 			// Source vertex index channel
 			MeshVertexIdIteratorConst IdIterator( pBase );
diff --git a/Engine/Plugins/Experimental/Mutable/Source/MutableRuntime/Private/MuR/OpMeshExtractLayoutBlock.h b/Engine/Plugins/Experimental/Mutable/Source/MutableRuntime/Private/MuR/OpMeshExtractLayoutBlock.h
--- a/Engine/Plugins/Experimental/Mutable/Source/MutableRuntime/Private/MuR/OpMeshExtractLayoutBlock.h
+++ b/Engine/Plugins/Experimental/Mutable/Source/MutableRuntime/Private/MuR/OpMeshExtractLayoutBlock.h
@@ -35,7 +35,7 @@ namespace mu
 
 		// If the vertices are explicit or relative, the above operation will already handle them correctly
 		// Otherwise, create a relative vertex ID buffer if necessary
-		Result->MeshIDPrefix = Source->MeshIDPrefix;
+		Result->VertexIDPrefix = Source->VertexIDPrefix;
 		if (Source->AreVertexIdsImplicit() 
 			&& 
 			// If we extract everything, we can keep the ids implicit.
@@ -116,39 +116,38 @@ namespace mu
 
 
 	//---------------------------------------------------------------------------------------------
-    inline void MeshExtractLayoutBlock(Mesh* Result, const Mesh* Source,
-                                           uint32 LayoutIndex,
-                                           uint16 BlockCount,
-                                           const uint64* BlockIds, bool& bOutSuccess)
+    inline void MeshExtractLayoutBlock(Mesh* Result, const Mesh* pSource,
+                                           uint32 layout,
+                                           uint16 blockCount,
+                                           const uint32* pExtractBlocks, bool& bOutSuccess)
 	{
-		check(Source);
+		check(pSource);
 		bOutSuccess = true;
 		
 		// TODO: Optimise
-		Result->CopyFrom(*Source);
+		Result->CopyFrom(*pSource);
 
-		UntypedMeshBufferIteratorConst itBlocks(Source->GetVertexBuffers(), MBS_LAYOUTBLOCK, LayoutIndex);
+		UntypedMeshBufferIteratorConst itBlocks( pSource->GetVertexBuffers(), MBS_LAYOUTBLOCK, layout );
 
         if (itBlocks.GetFormat()!=MBF_NONE)
         {
-            int32 ResultVertices = 0;
+            int resultVertices = 0;
 			TArray<int32> oldToNew;
-			oldToNew.Init(-1,Source->GetVertexCount());
+			oldToNew.Init(-1,pSource->GetVertexCount());
 			TArray<int32> newToOld;
-            newToOld.Reserve( Source->GetVertexCount() );
+            newToOld.Reserve( pSource->GetVertexCount() );
 
             if ( itBlocks.GetFormat()==MBF_UINT16 )
             {
                 const uint16* pBlocks = reinterpret_cast<const uint16*>( itBlocks.ptr() );
-                for ( int32 i=0; i<Source->GetVertexCount(); ++i )
+                for ( int i=0; i<pSource->GetVertexCount(); ++i )
                 {
-                    uint64 VertexBlockRelative = pBlocks[i];
-					uint64 VertexBlockId = (uint64(Source->MeshIDPrefix) << 32) | VertexBlockRelative;
+                    uint32_t vertexBlock = pBlocks[i];
 
                     bool found = false;
-                    for ( int32 j=0; j< BlockCount; ++j)
+                    for ( int j=0; j<blockCount; ++j)
                     {
-                        if (VertexBlockId == BlockIds[j])
+                        if (vertexBlock==pExtractBlocks[j])
                         {
                             found = true;
                             break;
@@ -157,42 +156,17 @@ namespace mu
 
                     if ( found )
                     {
-                        oldToNew[i] = ResultVertices++;
+                        oldToNew[i] = resultVertices++;
                         newToOld.Add( i );
                     }
                 }
             }
-            else if (itBlocks.GetFormat() == MBF_UINT64)
-			{
-				const uint64* pBlocks = reinterpret_cast<const uint64*>(itBlocks.ptr());
-				for (int32 i = 0; i < Source->GetVertexCount(); ++i)
-				{
-					uint64 VertexBlockId = pBlocks[i];
-
-					bool found = false;
-					for (int j = 0; j < BlockCount; ++j)
-					{
-						if (VertexBlockId == BlockIds[j])
-						{
-							found = true;
-							break;
-						}
-					}
-
-					if (found)
-					{
-						oldToNew[i] = ResultVertices++;
-						newToOld.Add(i);
-					}
-				}
-			}
-			else
-
+            else
             {
                 check( false );
             }
 
-            MeshExtractFromVertices(Source, Result, oldToNew, newToOld);
+            MeshExtractFromVertices(pSource, Result, oldToNew, newToOld);
         }
 
         Result->Surfaces.Empty();
diff --git a/Engine/Plugins/Experimental/Mutable/Source/MutableRuntime/Private/MuR/OpMeshFormat.cpp b/Engine/Plugins/Experimental/Mutable/Source/MutableRuntime/Private/MuR/OpMeshFormat.cpp
--- a/Engine/Plugins/Experimental/Mutable/Source/MutableRuntime/Private/MuR/OpMeshFormat.cpp
+++ b/Engine/Plugins/Experimental/Mutable/Source/MutableRuntime/Private/MuR/OpMeshFormat.cpp
@@ -443,7 +443,7 @@ namespace mu
 		Ptr<const Mesh> Source = PureSource;
 
 		Result->CopyFrom(*Format);
-		Result->MeshIDPrefix = Source->MeshIDPrefix;
+		Result->VertexIDPrefix = Source->VertexIDPrefix;
 
 		// Make sure that the bone indices will fit in this format, or extend it.
 		if (formatVertices)
diff --git a/Engine/Plugins/Experimental/Mutable/Source/MutableRuntime/Private/MuR/OpMeshRemove.cpp b/Engine/Plugins/Experimental/Mutable/Source/MutableRuntime/Private/MuR/OpMeshRemove.cpp
--- a/Engine/Plugins/Experimental/Mutable/Source/MutableRuntime/Private/MuR/OpMeshRemove.cpp
+++ b/Engine/Plugins/Experimental/Mutable/Source/MutableRuntime/Private/MuR/OpMeshRemove.cpp
@@ -23,7 +23,7 @@ namespace mu
     {
 		MeshVertexIdIteratorConst itVI( pSource );
 		FIdInterval current;
-		current.idStart = Mesh::InvalidVertexId;
+		current.idStart = TNumericLimits<uint64>::Max();
 		current.idPosition = 0;
 		current.size = 0;
 		for ( int32 sv=0; sv<pSource->GetVertexBuffers().GetElementCount(); ++sv )
@@ -31,7 +31,7 @@ namespace mu
             uint64 id = itVI.Get();
             ++itVI;
 
-            if (current.idStart== Mesh::InvalidVertexId)
+            if (current.idStart==TNumericLimits<uint64>::Max())
             {
                 current.idStart = id;
                 current.idPosition = sv;
@@ -53,7 +53,7 @@ namespace mu
             }
         }
 
-        if (current.idStart!= Mesh::InvalidVertexId)
+        if (current.idStart!=TNumericLimits<uint64>::Max())
         {
             intervals.Add(current);
         }
@@ -202,7 +202,7 @@ namespace mu
 		// If we had implicit indices, make them explicit or relative to keep them valid
 		if (RemovedIndices && Result->AreVertexIdsImplicit())
 		{
-			Result->MakeVertexIdsRelative();
+			Result->MakeVertexIndicesRelative();
 		}
 
 		// The temp array is necessary because if the vertex buffer is not sorted according to the index buffer we cannot do it in-place
diff --git a/Engine/Plugins/Experimental/Mutable/Source/MutableRuntime/Private/MuR/Operations.cpp b/Engine/Plugins/Experimental/Mutable/Source/MutableRuntime/Private/MuR/Operations.cpp
--- a/Engine/Plugins/Experimental/Mutable/Source/MutableRuntime/Private/MuR/Operations.cpp
+++ b/Engine/Plugins/Experimental/Mutable/Source/MutableRuntime/Private/MuR/Operations.cpp
@@ -910,7 +910,7 @@ namespace mu
         {
             const uint8_t* data = program.GetOpArgsPointer(at);
             mu::OP::ADDRESS source;
-			FMemory::Memcpy( &source, data, sizeof(OP::ADDRESS) );
+            memcpy( &source, data, sizeof(OP::ADDRESS) );
             f(source);
             break;
         }
diff --git a/Engine/Plugins/Experimental/Mutable/Source/MutableRuntime/Public/MuR/Layout.h b/Engine/Plugins/Experimental/Mutable/Source/MutableRuntime/Public/MuR/Layout.h
--- a/Engine/Plugins/Experimental/Mutable/Source/MutableRuntime/Public/MuR/Layout.h
+++ b/Engine/Plugins/Experimental/Mutable/Source/MutableRuntime/Public/MuR/Layout.h
@@ -8,7 +8,6 @@
 #include "Containers/Array.h"
 #include "Misc/AssertionMacros.h"
 #include "Math/IntVector.h"
-#include "Math/NumericLimits.h"
 
 namespace mu
 {
@@ -43,10 +42,6 @@ namespace mu
 	//! \ingroup runtime
     class MUTABLERUNTIME_API Layout : public Resource
 	{
-	public:
-
-		static constexpr uint64 InvalidBlockId = TNumericLimits<uint64>::Max();
-
 	public:
 
 		//-----------------------------------------------------------------------------------------
@@ -56,11 +51,11 @@ namespace mu
 		Layout();
 
 		//! Deep clone this layout.
-		Ptr<Layout> Clone() const;
+		LayoutPtr Clone() const;
 
 		//! Serialisation
 		static void Serialise( const Layout* p, OutputArchive& arch );
-		static Ptr<Layout> StaticUnserialise( InputArchive& arch );
+		static LayoutPtr StaticUnserialise( InputArchive& arch );
 
 		//! Full compare
 		bool operator==( const Layout& other ) const;
@@ -79,18 +74,18 @@ namespace mu
 		//! on each axis.
 		//! \param sizeX width of the grid.
 		//! \param sizeY height of the grid.
-		void SetGridSize( int32 SizeX, int32 SizeY );
+		void SetGridSize( int sizeX, int sizeY );
 
 		//! Get the maximum resolution of the grid where the blocks are defined.
 		//! \param[out] pSizeX The integer pointed by this will be set to the width of the grid.
 		//! \param[out] pSizeY The integer pointed by this will be set to the height of the grid.
-		void GetMaxGridSize(int32* SizeX, int32* SizeY) const;
+		void GetMaxGridSize(int* pSizeX, int* pSizeY) const;
 
 		//! Get the maximum resolution of the grid where the blocks are defined. It must be bigger than 0
 		//! on each axis.
 		//! \param sizeX width of the grid.
 		//! \param sizeY height of the grid.
-		void SetMaxGridSize(int32 SizeX, int32 SizeY);
+		void SetMaxGridSize(int sizeX, int sizeY);
 
 		//! Return the number of blocks in this layout.
 		int32 GetBlockCount() const;
@@ -108,14 +103,14 @@ namespace mu
 		//! (from 0 to the Y size of the grid minus one )
 		//! \param[out] pSizeX will be set to the x size of the block
 		//! \param[out] pSizeY will be set to the y size of the block
-		void GetBlock( int32 Index, uint16* MinX, uint16* MinY, uint16* SizeX, uint16* SizeY ) const;
+		void GetBlock( int index, uint16* pMinX, uint16* pMinY, uint16* pSizeX, uint16* pSizeY ) const;
 
 		//! Returns the reduction priority of a block.
 		//! \param index Block to get priority.
 		//! \param[out] priority reduction priority of the block
 		//! \param[out] bReduceBothAxes reduction method of the block
 		//! \param[out] bReduceByTwo reduction method of the block
-		void GetBlockOptions(int32 Index, int32& Priority, bool& bReduceBothAxes, bool& bReduceByTwo) const;
+		void GetBlockOptions(int index, int& priority, bool& bReduceBothAxes, bool& bReduceByTwo) const;
 
 		//! Set a block of the layout.
 		//! "Position" here means the lower-left corner of the block.
@@ -126,18 +121,18 @@ namespace mu
 		//! (from 0 to the Y size of the grid minus one )
 		//! \param sizeX will be set to the x size of the block
 		//! \param sizeY will be set to the y size of the block
-        void SetBlock( int32 Index, int32 MinX, int32 MinY, int32 SizeX, int32 SizeY );
+        void SetBlock( int index, int minX, int minY, int sizeX, int sizeY );
 
 		//! Set the reduction options of a block
 		//! \param index Block to set the options
 		//! \param priority will be set to the reduction priority of the block. The blocks with the highest values will be the last to be reduced
 		//! \param bReduceBothAxes will be set to reduce the block in both axis at the same time.
 		//! \param bReduceByTwo will reduce by two blocks on a unitary reduction.
-		void SetBlockOptions(int32 Index, int32 Priority, bool bReduceBothAxes, bool bReduceByTwo);
+		void SetBlockOptions(int index, int priority, bool bReduceBothAxes, bool bReduceByTwo);
 
 		//! Set the texture layout packing strategy
 		//! By default the texture layout packing strategy is set to resizable layout
-		void SetLayoutPackingStrategy(EPackStrategy);
+		void SetLayoutPackingStrategy(EPackStrategy _strategy);
 
 		//! Set the texture layout packing strategy
 		EPackStrategy GetLayoutPackingStrategy() const;
@@ -165,24 +160,24 @@ namespace mu
 		//-----------------------------------------------------------------------------------------
 		struct FBlock
 		{
-			FBlock(UE::Math::TIntVector2<uint16> InMin = UE::Math::TIntVector2<uint16>(), UE::Math::TIntVector2<uint16> InSize = UE::Math::TIntVector2<uint16>())
+			FBlock(UE::Math::TIntVector2<uint16> min = UE::Math::TIntVector2<uint16>(), UE::Math::TIntVector2<uint16> size = UE::Math::TIntVector2<uint16>())
 			{
-				Min = InMin;
-				Size = InSize;
-				Id = Layout::InvalidBlockId;
-				Priority = 0;
+				m_min = min;
+				m_size = size;
+				m_id = -1;
+				m_priority = 0;
 				bReduceBothAxes = false;
 				bReduceByTwo = false;
 			}
 
-			UE::Math::TIntVector2<uint16> Min = UE::Math::TIntVector2<uint16>(0,0);
-			UE::Math::TIntVector2<uint16> Size = UE::Math::TIntVector2<uint16>(0, 0);
+			UE::Math::TIntVector2<uint16> m_min = UE::Math::TIntVector2<uint16>(0,0);
+			UE::Math::TIntVector2<uint16> m_size = UE::Math::TIntVector2<uint16>(0, 0);
 
 			//! Absolute id used to control merging of various layouts
-			uint64 Id;
+			int32 m_id;
 
 			//! Priority value to control the shrink texture layout strategy
-			int32 Priority;
+			int32 m_priority;
 
 			//! Value to control the method to reduce the block
 			bool bReduceBothAxes;
@@ -200,10 +195,10 @@ namespace mu
 			//!
 			inline bool operator==(const FBlock& o) const
 			{
-				return (Min == o.Min) &&
-					(Size == o.Size) &&
-					(Id == o.Id) &&
-					(Priority == o.Priority) &&
+				return (m_min == o.m_min) &&
+					(m_size == o.m_size) &&
+					(m_id == o.m_id) &&
+					(m_priority == o.m_priority) &&
 					(bReduceBothAxes == o.bReduceBothAxes) &&
 					(bReduceByTwo == o.bReduceByTwo);
 			}
@@ -211,25 +206,28 @@ namespace mu
 			inline bool IsSimilar(const FBlock& o) const
 			{
 				// All but ids
-				return (Min == o.Min) &&
-					(Size == o.Size) &&
-					(Priority == o.Priority) &&
+				return (m_min == o.m_min) &&
+					(m_size == o.m_size) &&
+					(m_priority == o.m_priority) &&
 					(bReduceBothAxes == o.bReduceBothAxes) &&
 					(bReduceByTwo == o.bReduceByTwo);
 			}
+
+			//!
+			void UnserialiseOldVersion(InputArchive& Archive, const int32 Version);
 		};
 
 
 		//!
-		UE::Math::TIntVector2<uint16> Size = UE::Math::TIntVector2<uint16>(0, 0);
+		UE::Math::TIntVector2<uint16> m_size = UE::Math::TIntVector2<uint16>(0, 0);
 
-		UE::Math::TIntVector2<uint16> MaxSize = UE::Math::TIntVector2<uint16>(0, 0);
+		UE::Math::TIntVector2<uint16> m_maxsize = UE::Math::TIntVector2<uint16>(0, 0);
 
 		//!
-		TArray<FBlock> Blocks;
+		TArray<FBlock> m_blocks;
 
 		//! Packing strategy
-		EPackStrategy Strategy = EPackStrategy::RESIZABLE_LAYOUT;
+		EPackStrategy m_strategy = EPackStrategy::RESIZABLE_LAYOUT;
 		 
 		int32 FirstLODToIgnoreWarnings = 0;
 
@@ -245,8 +243,9 @@ namespace mu
 		//!
 		bool IsSimilar(const Layout& o) const;
 
-		/** Find a block by id. This converts the "absolute" id to a relative index to the layout blocks. Return -1 if not found. */
-		int32 FindBlock(uint64 Id) const;
+		//! Find a block by id. This converts the "absolute" id to a relative index to the layout
+		//! blocks. Return -1 if not found.
+		int32 FindBlock(int32 id) const;
 
 		//! Return true if the layout is a single block filling all area.
 		bool IsSingleBlockAndFull() const;
diff --git a/Engine/Plugins/Experimental/Mutable/Source/MutableRuntime/Public/MuR/Mesh.h b/Engine/Plugins/Experimental/Mutable/Source/MutableRuntime/Public/MuR/Mesh.h
--- a/Engine/Plugins/Experimental/Mutable/Source/MutableRuntime/Public/MuR/Mesh.h
+++ b/Engine/Plugins/Experimental/Mutable/Source/MutableRuntime/Public/MuR/Mesh.h
@@ -140,11 +140,7 @@ namespace mu
     //! \ingroup runtime
     class MUTABLERUNTIME_API Mesh : public Resource
     {
-	public:
-
-		static constexpr uint64 InvalidVertexId = TNumericLimits<uint64>::Max();
-
-	public:
+    public:
 
         //-----------------------------------------------------------------------------------------
         // Life cycle
@@ -206,6 +202,9 @@ namespace mu
 
         //! \}
 
+		/** Return true if the mesh has unique IDs for its vertices. */
+		bool HasVertexIds() const;
+
 		/** Return true if the mesh has unique vertex IDs and they stored in an implicit way. 
 		* This is relevant for some mesh operations that will need to make them explicit so that the result is still correct.
 		*/
@@ -213,8 +212,8 @@ namespace mu
 		bool AreVertexIdsExplicit() const;
 
 		/** Create an explicit vertex buffer for vertex IDs if they are implicit. */
-		void MakeVertexIdsRelative();
-		void MakeIdsExplicit();
+		void MakeVertexIndicesRelative();
+		void MakeVertexIndicesExplicit();
 
         //! \name Texture layouts
         //! \{
@@ -336,10 +335,10 @@ namespace mu
 		//! It should be reset after any operation that modifies the format.
 		mutable uint32 StaticFormatFlags = 0;
 
-		/** Prefix for the unique IDs related to this mesh (vertices and layout blocks). Useful if the mesh stores them in an implicit, or relative way. 
+		/** Prefix for the unique vertex IDs, if the mesh has them in an implicit, or relative way. 
 		* See MeshVertexIdIterator for details.
 		*/
-		uint32 MeshIDPrefix = 0;
+		uint32 VertexIDPrefix = 0;
 
 		//!
 		FMeshBufferSet VertexBuffers;
@@ -407,56 +406,54 @@ namespace mu
         //!
 		inline bool operator==(const Mesh& o) const
 		{
-			bool bEqual = true;
-
-			if (bEqual) bEqual = (MeshIDPrefix == o.MeshIDPrefix);
-			if (bEqual) bEqual = (IndexBuffers == o.IndexBuffers);
-			if (bEqual) bEqual = (VertexBuffers == o.VertexBuffers);
-			if (bEqual) bEqual = (Layouts.Num() == o.Layouts.Num());
-			if (bEqual) bEqual = (BonePoses.Num() == o.BonePoses.Num());
-			if (bEqual) bEqual = (BoneMap.Num() == o.BoneMap.Num());
-			if (bEqual && Skeleton != o.Skeleton)
+			bool equal = (IndexBuffers == o.IndexBuffers);
+			if (equal) equal = (VertexBuffers == o.VertexBuffers);
+			if (equal) equal = (Layouts.Num() == o.Layouts.Num());
+			if (equal) equal = (BonePoses.Num() == o.BonePoses.Num());
+			if (equal) equal = (BoneMap.Num() == o.BoneMap.Num());
+			if (equal && Skeleton != o.Skeleton)
 			{
 				if (Skeleton && o.Skeleton)
 				{
-					bEqual = (*Skeleton == *o.Skeleton);
+					equal = (*Skeleton == *o.Skeleton);
 				}
 				else
 				{
-					bEqual = false;
+					equal = false;
 				}
 			}
-			if (bEqual) bEqual = (StreamedResources == o.StreamedResources);
-			if (bEqual) bEqual = (Surfaces == o.Surfaces);
-			if (bEqual) bEqual = (Tags == o.Tags);
-			if (bEqual) bEqual = (SkeletonIDs == o.SkeletonIDs);
+			if (equal) equal = (StreamedResources == o.StreamedResources);
+			if (equal) equal = (Surfaces == o.Surfaces);
+			if (equal) equal = (Tags == o.Tags);
+			if (equal) equal = (SkeletonIDs == o.SkeletonIDs);
+			if (equal) equal = (VertexIDPrefix == o.VertexIDPrefix);
 
-			for (int32 i = 0; bEqual && i < Layouts.Num(); ++i)
+			for (int32 i = 0; equal && i < Layouts.Num(); ++i)
 			{
-				bEqual &= (*Layouts[i]) == (*o.Layouts[i]);
+				equal &= (*Layouts[i]) == (*o.Layouts[i]);
 			}
 
-			bEqual &= AdditionalBuffers.Num() == o.AdditionalBuffers.Num();
-			for (int32 i = 0; bEqual && i < AdditionalBuffers.Num(); ++i)
+			equal &= AdditionalBuffers.Num() == o.AdditionalBuffers.Num();
+			for (int32 i = 0; equal && i < AdditionalBuffers.Num(); ++i)
 			{
-				bEqual &= AdditionalBuffers[i] == o.AdditionalBuffers[i];
+				equal &= AdditionalBuffers[i] == o.AdditionalBuffers[i];
 			}
 
-			bEqual &= BonePoses.Num() == o.BonePoses.Num();
-			for (int32 i = 0; bEqual && i < BonePoses.Num(); ++i)
+			equal &= BonePoses.Num() == o.BonePoses.Num();
+			for (int32 i = 0; equal && i < BonePoses.Num(); ++i)
 			{
-				bEqual &= BonePoses[i] == o.BonePoses[i];
+				equal &= BonePoses[i] == o.BonePoses[i];
 			}
 
-			if (bEqual) bEqual = BoneMap == o.BoneMap;
+			if (equal) equal = BoneMap == o.BoneMap;
 
-			bEqual &= AdditionalPhysicsBodies.Num() == o.AdditionalPhysicsBodies.Num();
-			for (int32 i = 0; bEqual && i < AdditionalPhysicsBodies.Num(); ++i)
+			equal &= AdditionalPhysicsBodies.Num() == o.AdditionalPhysicsBodies.Num();
+			for (int32 i = 0; equal && i < AdditionalPhysicsBodies.Num(); ++i)
 			{
-				bEqual &= *AdditionalPhysicsBodies[i] == *o.AdditionalPhysicsBodies[i];
+				equal &= *AdditionalPhysicsBodies[i] == *o.AdditionalPhysicsBodies[i];
 			}
 
-			return bEqual;
+			return equal;
 		}
 
 		//! Compare the mesh with another one, but ignore internal data like generated vertex
diff --git a/Engine/Plugins/Experimental/Mutable/Source/MutableRuntime/Public/MuR/MeshBufferSet.h b/Engine/Plugins/Experimental/Mutable/Source/MutableRuntime/Public/MuR/MeshBufferSet.h
--- a/Engine/Plugins/Experimental/Mutable/Source/MutableRuntime/Public/MuR/MeshBufferSet.h
+++ b/Engine/Plugins/Experimental/Mutable/Source/MutableRuntime/Public/MuR/MeshBufferSet.h
@@ -394,9 +394,6 @@ namespace mu
 
 		//!
 		void UpdateOffsets(int32 b);
-		
-		/** Check that all channels of a specific semantic use the provided format. */
-		bool HasAnySemanticWithDifferentFormat(EMeshBufferSemantic Semantic, EMeshBufferFormat ExpectedFormat) const;
 	};	
 
 	
diff --git a/Engine/Plugins/Experimental/Mutable/Source/MutableTools/Internal/MuT/AST.h b/Engine/Plugins/Experimental/Mutable/Source/MutableTools/Internal/MuT/AST.h
--- a/Engine/Plugins/Experimental/Mutable/Source/MutableTools/Internal/MuT/AST.h
+++ b/Engine/Plugins/Experimental/Mutable/Source/MutableTools/Internal/MuT/AST.h
@@ -794,15 +794,15 @@ namespace mu
         virtual FImageDesc GetImageDesc( bool returnBestOption=false, FGetImageDescContext* context=nullptr ) const;
 
         //! Optional cache struct to use int he method below.
-        using FBlockLayoutSizeCache=TMap< const TPair<ASTOp*,int32>, TPair<int32,int32>>;
+        using FBlockLayoutSizeCache=TMap< const TPair<ASTOp*,int>, TPair<int,int>>;
 
         //! Return the size in layout blocks of a particular block given by absolute index
-        virtual void GetBlockLayoutSize( uint64 BlockId, int32* pBlockX, int32* pBlockY, FBlockLayoutSizeCache* cache );
-        void GetBlockLayoutSizeCached(uint64 BlockId, int32* pBlockX, int32* pBlockY, FBlockLayoutSizeCache* cache );
+        virtual void GetBlockLayoutSize( int blockIndex, int* pBlockX, int* pBlockY, FBlockLayoutSizeCache* cache );
+        void GetBlockLayoutSizeCached( int blockIndex, int* pBlockX, int* pBlockY, FBlockLayoutSizeCache* cache );
 
 
         //! Return the size in pixels of the layout grid block for the image operation
-        virtual void GetLayoutBlockSize( int32* pBlockX, int32* pBlockY );
+        virtual void GetLayoutBlockSize( int* pBlockX, int* pBlockY );
 
         virtual bool IsImagePlainConstant(FVector4f& colour ) const;
         virtual bool IsColourConstant(FVector4f& colour ) const;
diff --git a/Engine/Plugins/Experimental/Mutable/Source/MutableTools/Internal/MuT/NodeMeshFragmentPrivate.h b/Engine/Plugins/Experimental/Mutable/Source/MutableTools/Internal/MuT/NodeMeshFragmentPrivate.h
--- a/Engine/Plugins/Experimental/Mutable/Source/MutableTools/Internal/MuT/NodeMeshFragmentPrivate.h
+++ b/Engine/Plugins/Experimental/Mutable/Source/MutableTools/Internal/MuT/NodeMeshFragmentPrivate.h
@@ -17,20 +17,20 @@ namespace mu
 
 		Private()
 		{
-            LayoutOrGroup = -1;
+            m_layoutOrGroup = -1;
 		}
 
 		static FNodeType s_type;
 
-		Ptr<NodeMesh> m_pMesh;
-        int32 LayoutOrGroup;
-		TArray<uint64> Blocks;
+		NodeMeshPtr m_pMesh;
+        int m_layoutOrGroup;
+		TArray<int> m_blocks;
 
         FRAGMENT_TYPE m_fragmentType;
 
 
 		// NodeMesh::Private interface
-        Ptr<NodeLayout> GetLayout( int32 index ) const override;
+        NodeLayoutPtr GetLayout( int index ) const override;
 	};
 
 
diff --git a/Engine/Plugins/Experimental/Mutable/Source/MutableTools/Internal/MuT/NodePatchImage.h b/Engine/Plugins/Experimental/Mutable/Source/MutableTools/Internal/MuT/NodePatchImage.h
--- a/Engine/Plugins/Experimental/Mutable/Source/MutableTools/Internal/MuT/NodePatchImage.h
+++ b/Engine/Plugins/Experimental/Mutable/Source/MutableTools/Internal/MuT/NodePatchImage.h
@@ -14,7 +14,13 @@ namespace mu
 
     // Forward definitions
     class NodeImage;
- 
+    typedef Ptr<NodeImage> NodeImagePtr;
+    typedef Ptr<const NodeImage> NodeImagePtrConst;
+
+    class NodePatchImage;
+    typedef Ptr<NodePatchImage> NodePatchImagePtr;
+    typedef Ptr<const NodePatchImage> NodePatchImagePtrConst;
+
 
     //! Node that allows to modify an image from an object by blending other images on specific
     //! layout blocks.
@@ -37,31 +43,32 @@ namespace mu
         //-----------------------------------------------------------------------------------------
 
         //! Set the image that will be blended on the destination.
-        void SetImage(Ptr<NodeImage>);
+        void SetImage( NodeImagePtr );
 
         //! Get the image that will be blended on the destination.
-		Ptr<NodeImage> GetImage() const;
+        NodeImagePtr GetImage() const;
 
         //! Set the blending mask image.
-        void SetMask(Ptr<NodeImage>);
+        void SetMask( NodeImagePtr );
 
         //! Get the blending mask image.
-		Ptr<NodeImage> GetMask() const;
+        NodeImagePtr GetMask() const;
 
         //! Set the number of blocks in the layout that will be patched
-        void SetBlockCount( int32 );
+        void SetBlockCount( int );
 
         //! Get the number of blocks in the layout that will be patched
-        int32 GetBlockCount() const;
+        int GetBlockCount() const;
 
         //! Set a block of the layout that will be patched.
         //! \param index is the index in the patching list, from 0 to GetBlockCount()-1
-        //! \param LayoutBlockIndex is the block index in the layout of the image were this patch operation will be applied.
-        void SetBlock( int32 index, int32 LayoutBlockIndex);
+        //! \param block is the block index in the layout of the image were this patch operation
+        //! will be applied.
+        void SetBlock( int index, int block );
 
-        //! Get a block index of the layout that will be patched.
+        //! Get a block of the layout that will be patched.
         //! \param index is the index in the patching list, from 0 to GetBlockCount()-1
-        int32 GetBlock( int32 index ) const;
+        int GetBlock( int index ) const;
 
         //! Set the blending operation to use to combine the pixels
         EBlendType GetBlendType() const;
diff --git a/Engine/Plugins/Experimental/Mutable/Source/MutableTools/Internal/MuT/NodePatchImagePrivate.h b/Engine/Plugins/Experimental/Mutable/Source/MutableTools/Internal/MuT/NodePatchImagePrivate.h
--- a/Engine/Plugins/Experimental/Mutable/Source/MutableTools/Internal/MuT/NodePatchImagePrivate.h
+++ b/Engine/Plugins/Experimental/Mutable/Source/MutableTools/Internal/MuT/NodePatchImagePrivate.h
@@ -20,8 +20,7 @@ namespace mu
         NodeImagePtr m_pImage;
         NodeImagePtr m_pMask;
 
-		// These are the indices of the blocks in the layout being patched (not the Ids)
-        TArray<int32> BlockIndices;
+        TArray<int> m_blocks;
 
         EBlendType m_blendType = EBlendType::BT_BLEND;
 
diff --git a/Engine/Plugins/Experimental/Mutable/Source/MutableTools/Internal/MuT/NodePatchMeshPrivate.h b/Engine/Plugins/Experimental/Mutable/Source/MutableTools/Internal/MuT/NodePatchMeshPrivate.h
--- a/Engine/Plugins/Experimental/Mutable/Source/MutableTools/Internal/MuT/NodePatchMeshPrivate.h
+++ b/Engine/Plugins/Experimental/Mutable/Source/MutableTools/Internal/MuT/NodePatchMeshPrivate.h
@@ -9,14 +9,15 @@
 namespace mu
 {
 
+
 	class NodePatchMesh::Private : public Node::Private
 	{
 	public:
 
 		static FNodeType s_type;
 
-		Ptr<NodeMesh> m_pRemove;
-		Ptr<NodeMesh> m_pAdd;
+		NodeMeshPtr m_pRemove;
+		NodeMeshPtr m_pAdd;
 
 	};
 
diff --git a/Engine/Plugins/Experimental/Mutable/Source/MutableTools/Private/MuT/AST.cpp b/Engine/Plugins/Experimental/Mutable/Source/MutableTools/Private/MuT/AST.cpp
--- a/Engine/Plugins/Experimental/Mutable/Source/MutableTools/Private/MuT/AST.cpp
+++ b/Engine/Plugins/Experimental/Mutable/Source/MutableTools/Private/MuT/AST.cpp
@@ -881,25 +881,26 @@ bool ASTOp::IsColourConstant(FVector4f&) const
 }
 
 
-void ASTOp::GetBlockLayoutSize( uint64, int32*, int32*, FBlockLayoutSizeCache* )
+void ASTOp::GetBlockLayoutSize( int, int*, int*, FBlockLayoutSizeCache* )
 {
     check(false);
 }
 
-void ASTOp::GetBlockLayoutSizeCached(uint64 blockIndex, int32* BlockX, int32* BlockY, FBlockLayoutSizeCache* cache)
+void ASTOp::GetBlockLayoutSizeCached( int blockIndex, int* pBlockX, int* pBlockY,
+                               FBlockLayoutSizeCache* cache )
 {
 	FBlockLayoutSizeCache::KeyType key(this,blockIndex);
 	FBlockLayoutSizeCache::ValueType* ValuePtr = cache->Find( key );
     if (ValuePtr)
     {
-        *BlockX = ValuePtr->Key;
-        *BlockY = ValuePtr->Value;
+        *pBlockX = ValuePtr->Key;
+        *pBlockY = ValuePtr->Value;
         return;
     }
 
-    GetBlockLayoutSize( blockIndex, BlockX, BlockY, cache );
+    GetBlockLayoutSize( blockIndex, pBlockX, pBlockY, cache );
 
-	cache->Add(key, FBlockLayoutSizeCache::ValueType( *BlockX, *BlockY) );
+	cache->Add(key, FBlockLayoutSizeCache::ValueType( *pBlockX, *pBlockY) );
 }
 
 
diff --git a/Engine/Plugins/Experimental/Mutable/Source/MutableTools/Private/MuT/ASTOpConditional.cpp b/Engine/Plugins/Experimental/Mutable/Source/MutableTools/Private/MuT/ASTOpConditional.cpp
--- a/Engine/Plugins/Experimental/Mutable/Source/MutableTools/Private/MuT/ASTOpConditional.cpp
+++ b/Engine/Plugins/Experimental/Mutable/Source/MutableTools/Private/MuT/ASTOpConditional.cpp
@@ -194,7 +194,7 @@ namespace mu
 
 
 	//-------------------------------------------------------------------------------------------------
-	void ASTOpConditional::GetLayoutBlockSize(int32* pBlockX, int32* pBlockY)
+	void ASTOpConditional::GetLayoutBlockSize(int* pBlockX, int* pBlockY)
 	{
 		if (type == OP_TYPE::IM_CONDITIONAL)
 		{
@@ -213,15 +213,16 @@ namespace mu
 
 
 	//-------------------------------------------------------------------------------------------------
-	void ASTOpConditional::GetBlockLayoutSize(uint64 BlockId, int32* pBlockX, int32* pBlockY, FBlockLayoutSizeCache* cache)
+	void ASTOpConditional::GetBlockLayoutSize(int blockIndex, int* pBlockX, int* pBlockY,
+		FBlockLayoutSizeCache* cache)
 	{
 		if (type == OP_TYPE::LA_CONDITIONAL)
 		{
-			yes->GetBlockLayoutSizeCached(BlockId, pBlockX, pBlockY, cache);
+			yes->GetBlockLayoutSizeCached(blockIndex, pBlockX, pBlockY, cache);
 
 			if (*pBlockX == 0)
 			{
-				no->GetBlockLayoutSizeCached(BlockId, pBlockX, pBlockY, cache);
+				no->GetBlockLayoutSizeCached(blockIndex, pBlockX, pBlockY, cache);
 			}
 		}
 		else
diff --git a/Engine/Plugins/Experimental/Mutable/Source/MutableTools/Private/MuT/ASTOpConditional.h b/Engine/Plugins/Experimental/Mutable/Source/MutableTools/Private/MuT/ASTOpConditional.h
--- a/Engine/Plugins/Experimental/Mutable/Source/MutableTools/Private/MuT/ASTOpConditional.h
+++ b/Engine/Plugins/Experimental/Mutable/Source/MutableTools/Private/MuT/ASTOpConditional.h
@@ -45,8 +45,8 @@ struct FProgram;
 		void ForEachChild(const TFunctionRef<void(ASTChild&)> f) override;
 		void Link(FProgram& program, FLinkerOptions* Options) override;
 		FImageDesc GetImageDesc(bool returnBestOption, class FGetImageDescContext* context) const override;
-		void GetLayoutBlockSize(int32* pBlockX, int32* pBlockY) override;
-		void GetBlockLayoutSize(uint64 BlockId, int32* pBlockX, int32* pBlockY, FBlockLayoutSizeCache* cache) override;
+		void GetLayoutBlockSize(int* pBlockX, int* pBlockY) override;
+		void GetBlockLayoutSize(int blockIndex, int* pBlockX, int* pBlockY, FBlockLayoutSizeCache* cache) override;
 		Ptr<ASTOp> OptimiseSemantic(const FModelOptimizationOptions&, int32 Pass) const override;
 		bool GetNonBlackRect(FImageRect& maskUsage) const override;
 		Ptr<ImageSizeExpression> GetImageSizeExpression() const override;
diff --git a/Engine/Plugins/Experimental/Mutable/Source/MutableTools/Private/MuT/ASTOpConstantResource.cpp b/Engine/Plugins/Experimental/Mutable/Source/MutableTools/Private/MuT/ASTOpConstantResource.cpp
--- a/Engine/Plugins/Experimental/Mutable/Source/MutableTools/Private/MuT/ASTOpConstantResource.cpp
+++ b/Engine/Plugins/Experimental/Mutable/Source/MutableTools/Private/MuT/ASTOpConstantResource.cpp
@@ -413,7 +413,7 @@ namespace mu
 	}
 
 	//-------------------------------------------------------------------------------------------------
-	void ASTOpConstantResource::GetBlockLayoutSize(uint64 BlockId, int32* BlockX, int32* BlockY, FBlockLayoutSizeCache*)
+	void ASTOpConstantResource::GetBlockLayoutSize(int blockIndex, int* pBlockX, int* pBlockY, FBlockLayoutSizeCache*)
 	{
 		switch (Type)
 		{
@@ -424,16 +424,16 @@ namespace mu
 
 			if (pLayout)
 			{
-				int relId = pLayout->FindBlock(BlockId);
+				int relId = pLayout->FindBlock(blockIndex);
 				if (relId >= 0)
 				{
-					*BlockX = pLayout->Blocks[relId].Size[0];
-					*BlockY = pLayout->Blocks[relId].Size[1];
+					*pBlockX = pLayout->m_blocks[relId].m_size[0];
+					*pBlockY = pLayout->m_blocks[relId].m_size[1];
 				}
 				else
 				{
-					*BlockX = 0;
-					*BlockY = 0;
+					*pBlockX = 0;
+					*pBlockY = 0;
 				}
 			}
 
@@ -446,7 +446,7 @@ namespace mu
 
 
 	//-------------------------------------------------------------------------------------------------
-	void ASTOpConstantResource::GetLayoutBlockSize(int32* pBlockX, int32* pBlockY)
+	void ASTOpConstantResource::GetLayoutBlockSize(int* pBlockX, int* pBlockY)
 	{
 		switch (Type)
 		{
diff --git a/Engine/Plugins/Experimental/Mutable/Source/MutableTools/Private/MuT/ASTOpConstantResource.h b/Engine/Plugins/Experimental/Mutable/Source/MutableTools/Private/MuT/ASTOpConstantResource.h
--- a/Engine/Plugins/Experimental/Mutable/Source/MutableTools/Private/MuT/ASTOpConstantResource.h
+++ b/Engine/Plugins/Experimental/Mutable/Source/MutableTools/Private/MuT/ASTOpConstantResource.h
@@ -61,8 +61,9 @@ namespace mu
 		uint64 Hash() const override;
 		void Link(FProgram& program, FLinkerOptions*) override;
 		FImageDesc GetImageDesc(bool, class FGetImageDescContext*) const override;
-		void GetBlockLayoutSize(uint64 BlockId, int32* pBlockX, int32* pBlockY, FBlockLayoutSizeCache* cache) override;
-		void GetLayoutBlockSize(int32* pBlockX, int32* pBlockY) override;
+		void GetBlockLayoutSize(int blockIndex, int* pBlockX, int* pBlockY,
+			FBlockLayoutSizeCache* cache) override;
+		void GetLayoutBlockSize(int* pBlockX, int* pBlockY) override;
 		bool GetNonBlackRect(FImageRect& maskUsage) const override;
 		bool IsImagePlainConstant(FVector4f& colour) const override;
 		Ptr<ImageSizeExpression> GetImageSizeExpression() const override;
diff --git a/Engine/Plugins/Experimental/Mutable/Source/MutableTools/Private/MuT/ASTOpImageCompose.cpp b/Engine/Plugins/Experimental/Mutable/Source/MutableTools/Private/MuT/ASTOpImageCompose.cpp
--- a/Engine/Plugins/Experimental/Mutable/Source/MutableTools/Private/MuT/ASTOpImageCompose.cpp
+++ b/Engine/Plugins/Experimental/Mutable/Source/MutableTools/Private/MuT/ASTOpImageCompose.cpp
@@ -50,7 +50,7 @@ bool ASTOpImageCompose::IsEqual(const ASTOp& otherUntyped) const
 			Base ==other->Base &&
 			BlockImage == other->BlockImage &&
 			Mask == other->Mask &&
-			BlockId == other->BlockId;
+			BlockIndex == other->BlockIndex;
     }
     return false;
 }
@@ -64,7 +64,7 @@ uint64 ASTOpImageCompose::Hash() const
 	hash_combine(res, Base.child().get());
 	hash_combine(res, BlockImage.child().get());
 	hash_combine(res, Mask.child().get());
-	hash_combine(res, BlockId);
+	hash_combine(res, BlockIndex);
 	return res;
 }
 
@@ -77,7 +77,7 @@ mu::Ptr<ASTOp> ASTOpImageCompose::Clone(MapChildFuncRef mapChild) const
 	n->Base = mapChild(Base.child());
 	n->BlockImage = mapChild(BlockImage.child());
 	n->Mask = mapChild(Mask.child());
-	n->BlockId = BlockId;
+	n->BlockIndex = BlockIndex;
     return n;
 }
 
@@ -105,7 +105,7 @@ void ASTOpImageCompose::Link( FProgram& program, FLinkerOptions*)
 		if (Base) args.base = Base->linkedAddress;
 		if (BlockImage) args.blockImage = BlockImage->linkedAddress;
 		if (Mask) args.mask = Mask->linkedAddress;
-        args.BlockId = BlockId;
+        args.blockIndex = BlockIndex;
 
         linkedAddress = (OP::ADDRESS)program.m_opAddress.Num();
         program.m_opAddress.Add((uint32_t)program.m_byteCode.Num());
@@ -189,17 +189,19 @@ bool ASTOpImageCompose::IsImagePlainConstant(FVector4f& colour) const
 }
 
 
-void ASTOpImageCompose::GetLayoutBlockSize(int32* pBlockX, int32* pBlockY)
+void ASTOpImageCompose::GetLayoutBlockSize(int* pBlockX, int* pBlockY)
 {
 	// We can find out in two ways: following the block image, or following the base image.
 	// Let's try the block approach: We need the block size and the layout blocks
-	int32 layoutBlocksX = 0;
-	int32 layoutBlocksY = 0;
+	int layoutBlocksX = 0;
+	int layoutBlocksY = 0;
 	if (Layout.child())
 	{
 		MUTABLE_CPUPROFILER_SCOPE(GetLayoutBlockSize_GetBlockLayoutSize);
 		FBlockLayoutSizeCache cache;
-		Layout->GetBlockLayoutSizeCached(BlockId, &layoutBlocksX, &layoutBlocksY, &cache);
+		Layout->GetBlockLayoutSizeCached(BlockIndex,
+			&layoutBlocksX, &layoutBlocksY,
+			&cache);
 	}
 
 	if (layoutBlocksX > 0 && layoutBlocksY > 0 && BlockImage.child())
@@ -221,9 +223,9 @@ mu::Ptr<ASTOp> ASTOpImageCompose::OptimiseSemantic(const FModelOptimizationOptio
 {
 	mu::Ptr<ASTOp> at;
 
-	mu::Ptr<ASTOp> baseAt = Base.child();
-	mu::Ptr<ASTOp> blockAt = BlockImage.child();
-	mu::Ptr<ASTOp> layoutAt = Layout.child();
+	auto baseAt = Base.child();
+	auto blockAt = BlockImage.child();
+	auto layoutAt = Layout.child();
 	if (layoutAt
 		&&
 		layoutAt->GetOpType() == OP_TYPE::LA_CONSTANT
@@ -238,17 +240,17 @@ mu::Ptr<ASTOp> ASTOpImageCompose::OptimiseSemantic(const FModelOptimizationOptio
 		// Constant single-block full layout?
 		if (pLayout->GetBlockCount() == 1
 			&&
-			pLayout->Blocks[0].Min == UE::Math::TIntVector2<uint16>(0, 0)
+			pLayout->m_blocks[0].m_min == UE::Math::TIntVector2<uint16>(0, 0)
 			&&
-			pLayout->Blocks[0].Size == pLayout->Size
+			pLayout->m_blocks[0].m_size == pLayout->m_size
 			&&
-			pLayout->Blocks[0].Id == BlockId
+			pLayout->m_blocks[0].m_id == int(BlockIndex)
 			)
 		{
 			// We could only take the block, but we must make sure it will have the format
 			// and size of the base.
-			mu::FImageDesc baseDesc = baseAt->GetImageDesc(true);
-			mu::FImageDesc blockDesc = blockAt->GetImageDesc(true);
+			auto baseDesc = baseAt->GetImageDesc(true);
+			auto blockDesc = blockAt->GetImageDesc(true);
 
 			at = blockAt;
 
diff --git a/Engine/Plugins/Experimental/Mutable/Source/MutableTools/Private/MuT/ASTOpImageCompose.h b/Engine/Plugins/Experimental/Mutable/Source/MutableTools/Private/MuT/ASTOpImageCompose.h
--- a/Engine/Plugins/Experimental/Mutable/Source/MutableTools/Private/MuT/ASTOpImageCompose.h
+++ b/Engine/Plugins/Experimental/Mutable/Source/MutableTools/Private/MuT/ASTOpImageCompose.h
@@ -21,7 +21,7 @@ namespace mu
 		ASTChild Base;
 		ASTChild BlockImage;
 		ASTChild Mask;
-		uint64 BlockId=0;
+		uint32 BlockIndex=0;
 
 	public:
 
@@ -38,8 +38,8 @@ namespace mu
 		//TODO: Ptr<ASTOp> OptimiseSink(const FModelOptimizationOptions& options, FOptimizeSinkContext& context) const override;
 		Ptr<ASTOp> OptimiseSemantic(const FModelOptimizationOptions& options, int32 Pass) const;
 		FImageDesc GetImageDesc(bool returnBestOption, FGetImageDescContext* context) const override;
-		void GetLayoutBlockSize(int32* OutBlockX, int32* OutBlockY) override;
-		bool IsImagePlainConstant(FVector4f& OutColor) const override;
+		void GetLayoutBlockSize(int* pBlockX, int* pBlockY) override;
+		bool IsImagePlainConstant(FVector4f& colour) const override;
 		Ptr<ImageSizeExpression> GetImageSizeExpression() const override;
 
 	};
diff --git a/Engine/Plugins/Experimental/Mutable/Source/MutableTools/Private/MuT/ASTOpImageRasterMesh.cpp b/Engine/Plugins/Experimental/Mutable/Source/MutableTools/Private/MuT/ASTOpImageRasterMesh.cpp
--- a/Engine/Plugins/Experimental/Mutable/Source/MutableTools/Private/MuT/ASTOpImageRasterMesh.cpp
+++ b/Engine/Plugins/Experimental/Mutable/Source/MutableTools/Private/MuT/ASTOpImageRasterMesh.cpp
@@ -26,7 +26,7 @@ namespace mu
 		, mask(this)
 		, projector(this)
 	{
-		BlockId = Layout::InvalidBlockId;
+		BlockId = -1;
 		LayoutIndex = -1;
 		SizeX = SizeY = 0;
 		SourceSizeX = SourceSizeY = 0;
@@ -139,7 +139,7 @@ namespace mu
 			OP::ImageRasterMeshArgs args;
 			FMemory::Memzero(&args, sizeof(args));
 
-			args.BlockId = BlockId;
+			args.blockId = BlockId;
 			args.LayoutIndex = LayoutIndex;
 			args.sizeX = SizeX;
 			args.sizeY = SizeY;
diff --git a/Engine/Plugins/Experimental/Mutable/Source/MutableTools/Private/MuT/ASTOpImageRasterMesh.h b/Engine/Plugins/Experimental/Mutable/Source/MutableTools/Private/MuT/ASTOpImageRasterMesh.h
--- a/Engine/Plugins/Experimental/Mutable/Source/MutableTools/Private/MuT/ASTOpImageRasterMesh.h
+++ b/Engine/Plugins/Experimental/Mutable/Source/MutableTools/Private/MuT/ASTOpImageRasterMesh.h
@@ -24,7 +24,7 @@ namespace mu
 		ASTChild projector;
 
 		/** If layouts are used, this can indicate a single layout block that we want to raster. */
-		uint64 BlockId;
+		int32 BlockId;
 		int8 LayoutIndex;
 
 		/** Size of the image to generate by rasterization of the mesh. */
diff --git a/Engine/Plugins/Experimental/Mutable/Source/MutableTools/Private/MuT/ASTOpLayoutFromMesh.cpp b/Engine/Plugins/Experimental/Mutable/Source/MutableTools/Private/MuT/ASTOpLayoutFromMesh.cpp
--- a/Engine/Plugins/Experimental/Mutable/Source/MutableTools/Private/MuT/ASTOpLayoutFromMesh.cpp
+++ b/Engine/Plugins/Experimental/Mutable/Source/MutableTools/Private/MuT/ASTOpLayoutFromMesh.cpp
@@ -97,7 +97,7 @@ namespace mu
 	}
 
 
-	void ASTOpLayoutFromMesh::GetBlockLayoutSize(uint64 BlockId, int32* pBlockX, int32* pBlockY, FBlockLayoutSizeCache* cache)
+	void ASTOpLayoutFromMesh::GetBlockLayoutSize(int blockIndex, int* pBlockX, int* pBlockY, FBlockLayoutSizeCache* cache)
 	{
 		// This shouldn't happen for this operation because it is always in a branch of layout operations that is not the main one.
 		check(false);
diff --git a/Engine/Plugins/Experimental/Mutable/Source/MutableTools/Private/MuT/ASTOpLayoutFromMesh.h b/Engine/Plugins/Experimental/Mutable/Source/MutableTools/Private/MuT/ASTOpLayoutFromMesh.h
--- a/Engine/Plugins/Experimental/Mutable/Source/MutableTools/Private/MuT/ASTOpLayoutFromMesh.h
+++ b/Engine/Plugins/Experimental/Mutable/Source/MutableTools/Private/MuT/ASTOpLayoutFromMesh.h
@@ -32,7 +32,7 @@ namespace mu
 		void ForEachChild(const TFunctionRef<void(ASTChild&)>) override;
 		void Link(FProgram& program, FLinkerOptions* Options) override;
 		Ptr<ASTOp> OptimiseSink(const FModelOptimizationOptions& options, FOptimizeSinkContext& context) const override;
-		void GetBlockLayoutSize(uint64 BlockId, int32* pBlockX, int32* pBlockY, FBlockLayoutSizeCache* cache) override;
+		void GetBlockLayoutSize(int blockIndex, int* pBlockX, int* pBlockY, FBlockLayoutSizeCache* cache) override;
 
 	};
 
diff --git a/Engine/Plugins/Experimental/Mutable/Source/MutableTools/Private/MuT/ASTOpLayoutMerge.cpp b/Engine/Plugins/Experimental/Mutable/Source/MutableTools/Private/MuT/ASTOpLayoutMerge.cpp
--- a/Engine/Plugins/Experimental/Mutable/Source/MutableTools/Private/MuT/ASTOpLayoutMerge.cpp
+++ b/Engine/Plugins/Experimental/Mutable/Source/MutableTools/Private/MuT/ASTOpLayoutMerge.cpp
@@ -93,16 +93,16 @@ namespace mu
 	}
 
 
-	void ASTOpLayoutMerge::GetBlockLayoutSize(uint64 BlockId, int32* pBlockX, int32* pBlockY, FBlockLayoutSizeCache* cache)
+	void ASTOpLayoutMerge::GetBlockLayoutSize(int blockIndex, int* pBlockX, int* pBlockY, FBlockLayoutSizeCache* cache)
 	{
 		if (Base)
 		{
-			Base->GetBlockLayoutSize(BlockId, pBlockX, pBlockY, cache);
+			Base->GetBlockLayoutSize(blockIndex, pBlockX, pBlockY, cache);
 		}
 
 		if (!*pBlockX && Added)
 		{
-			Added->GetBlockLayoutSize(BlockId, pBlockX, pBlockY, cache);
+			Added->GetBlockLayoutSize(blockIndex, pBlockX, pBlockY, cache);
 		}
 	}
 
diff --git a/Engine/Plugins/Experimental/Mutable/Source/MutableTools/Private/MuT/ASTOpLayoutMerge.h b/Engine/Plugins/Experimental/Mutable/Source/MutableTools/Private/MuT/ASTOpLayoutMerge.h
--- a/Engine/Plugins/Experimental/Mutable/Source/MutableTools/Private/MuT/ASTOpLayoutMerge.h
+++ b/Engine/Plugins/Experimental/Mutable/Source/MutableTools/Private/MuT/ASTOpLayoutMerge.h
@@ -31,7 +31,7 @@ namespace mu
 		Ptr<ASTOp> Clone(MapChildFuncRef mapChild) const override;
 		void ForEachChild(const TFunctionRef<void(ASTChild&)>) override;
 		void Link(FProgram& program, FLinkerOptions* Options) override;
-		void GetBlockLayoutSize(uint64 BlockId, int32* pBlockX, int32* pBlockY, FBlockLayoutSizeCache* cache) override;
+		void GetBlockLayoutSize(int blockIndex, int* pBlockX, int* pBlockY, FBlockLayoutSizeCache* cache) override;
 
 	};
 
diff --git a/Engine/Plugins/Experimental/Mutable/Source/MutableTools/Private/MuT/ASTOpLayoutPack.cpp b/Engine/Plugins/Experimental/Mutable/Source/MutableTools/Private/MuT/ASTOpLayoutPack.cpp
--- a/Engine/Plugins/Experimental/Mutable/Source/MutableTools/Private/MuT/ASTOpLayoutPack.cpp
+++ b/Engine/Plugins/Experimental/Mutable/Source/MutableTools/Private/MuT/ASTOpLayoutPack.cpp
@@ -88,11 +88,11 @@ namespace mu
 	}
 
 
-	void ASTOpLayoutPack::GetBlockLayoutSize(uint64 BlockId, int32* pBlockX, int32* pBlockY, FBlockLayoutSizeCache* cache)
+	void ASTOpLayoutPack::GetBlockLayoutSize(int blockIndex, int* pBlockX, int* pBlockY, FBlockLayoutSizeCache* cache)
 	{
 		if (Source)
 		{
-			Source->GetBlockLayoutSize(BlockId, pBlockX, pBlockY, cache);
+			Source->GetBlockLayoutSize(blockIndex, pBlockX, pBlockY, cache);
 		}
 	}
 
diff --git a/Engine/Plugins/Experimental/Mutable/Source/MutableTools/Private/MuT/ASTOpLayoutPack.h b/Engine/Plugins/Experimental/Mutable/Source/MutableTools/Private/MuT/ASTOpLayoutPack.h
--- a/Engine/Plugins/Experimental/Mutable/Source/MutableTools/Private/MuT/ASTOpLayoutPack.h
+++ b/Engine/Plugins/Experimental/Mutable/Source/MutableTools/Private/MuT/ASTOpLayoutPack.h
@@ -30,7 +30,7 @@ namespace mu
 		Ptr<ASTOp> Clone(MapChildFuncRef mapChild) const override;
 		void ForEachChild(const TFunctionRef<void(ASTChild&)>) override;
 		void Link(FProgram& program, FLinkerOptions* Options) override;
-		void GetBlockLayoutSize(uint64 BlockId, int32* pBlockX, int32* pBlockY, FBlockLayoutSizeCache* cache) override;
+		void GetBlockLayoutSize(int blockIndex, int* pBlockX, int* pBlockY, FBlockLayoutSizeCache* cache) override;
 
 	};
 
diff --git a/Engine/Plugins/Experimental/Mutable/Source/MutableTools/Private/MuT/ASTOpLayoutRemoveBlocks.cpp b/Engine/Plugins/Experimental/Mutable/Source/MutableTools/Private/MuT/ASTOpLayoutRemoveBlocks.cpp
--- a/Engine/Plugins/Experimental/Mutable/Source/MutableTools/Private/MuT/ASTOpLayoutRemoveBlocks.cpp
+++ b/Engine/Plugins/Experimental/Mutable/Source/MutableTools/Private/MuT/ASTOpLayoutRemoveBlocks.cpp
@@ -93,11 +93,11 @@ namespace mu
 	}
 
 
-	void ASTOpLayoutRemoveBlocks::GetBlockLayoutSize(uint64 BlockId, int32* pBlockX, int32* pBlockY, FBlockLayoutSizeCache* cache)
+	void ASTOpLayoutRemoveBlocks::GetBlockLayoutSize(int blockIndex, int* pBlockX, int* pBlockY, FBlockLayoutSizeCache* cache)
 	{
 		if (Source)
 		{
-			Source->GetBlockLayoutSize(BlockId, pBlockX, pBlockY, cache);
+			Source->GetBlockLayoutSize(blockIndex, pBlockX, pBlockY, cache);
 		}
 	}
 
diff --git a/Engine/Plugins/Experimental/Mutable/Source/MutableTools/Private/MuT/ASTOpLayoutRemoveBlocks.h b/Engine/Plugins/Experimental/Mutable/Source/MutableTools/Private/MuT/ASTOpLayoutRemoveBlocks.h
--- a/Engine/Plugins/Experimental/Mutable/Source/MutableTools/Private/MuT/ASTOpLayoutRemoveBlocks.h
+++ b/Engine/Plugins/Experimental/Mutable/Source/MutableTools/Private/MuT/ASTOpLayoutRemoveBlocks.h
@@ -31,7 +31,7 @@ namespace mu
 		Ptr<ASTOp> Clone(MapChildFuncRef mapChild) const override;
 		void ForEachChild(const TFunctionRef<void(ASTChild&)>) override;
 		void Link(FProgram& program, FLinkerOptions* Options) override;
-		void GetBlockLayoutSize(uint64 BlockId, int32* pBlockX, int32* pBlockY, FBlockLayoutSizeCache* cache) override;
+		void GetBlockLayoutSize(int blockIndex, int* pBlockX, int* pBlockY, FBlockLayoutSizeCache* cache) override;
 
 	};
 
diff --git a/Engine/Plugins/Experimental/Mutable/Source/MutableTools/Private/MuT/ASTOpMeshExtractLayoutBlocks.cpp b/Engine/Plugins/Experimental/Mutable/Source/MutableTools/Private/MuT/ASTOpMeshExtractLayoutBlocks.cpp
--- a/Engine/Plugins/Experimental/Mutable/Source/MutableTools/Private/MuT/ASTOpMeshExtractLayoutBlocks.cpp
+++ b/Engine/Plugins/Experimental/Mutable/Source/MutableTools/Private/MuT/ASTOpMeshExtractLayoutBlocks.cpp
@@ -87,9 +87,9 @@ namespace mu
 			AppendCode(program.m_byteCode, (uint16)Layout);
 			AppendCode(program.m_byteCode, (uint16)Blocks.Num());
 
-			for (uint64 Id : Blocks)
+			for (auto b : Blocks)
 			{
-				AppendCode(program.m_byteCode, Id);
+				AppendCode(program.m_byteCode, (uint32)b);
 			}
 		}
 	}
diff --git a/Engine/Plugins/Experimental/Mutable/Source/MutableTools/Private/MuT/ASTOpMeshExtractLayoutBlocks.h b/Engine/Plugins/Experimental/Mutable/Source/MutableTools/Private/MuT/ASTOpMeshExtractLayoutBlocks.h
--- a/Engine/Plugins/Experimental/Mutable/Source/MutableTools/Private/MuT/ASTOpMeshExtractLayoutBlocks.h
+++ b/Engine/Plugins/Experimental/Mutable/Source/MutableTools/Private/MuT/ASTOpMeshExtractLayoutBlocks.h
@@ -26,8 +26,8 @@ namespace mu
 		/** Layout to use to select the blocks. */
 		uint16 Layout = 0;
 
-		/** Block Ids to include in the resulting mesh. */
-		TArray<uint64> Blocks;
+		/** Blocks to include in the resulting mesh. */
+		TArray<uint32> Blocks;
 
 	public:
 
diff --git a/Engine/Plugins/Experimental/Mutable/Source/MutableTools/Private/MuT/ASTOpReferenceResource.cpp b/Engine/Plugins/Experimental/Mutable/Source/MutableTools/Private/MuT/ASTOpReferenceResource.cpp
--- a/Engine/Plugins/Experimental/Mutable/Source/MutableTools/Private/MuT/ASTOpReferenceResource.cpp
+++ b/Engine/Plugins/Experimental/Mutable/Source/MutableTools/Private/MuT/ASTOpReferenceResource.cpp
@@ -77,7 +77,7 @@ namespace mu
 
 
 	//-------------------------------------------------------------------------------------------------
-	void ASTOpReferenceResource::GetBlockLayoutSize(uint64 BlockId, int32* pBlockX, int32* pBlockY, FBlockLayoutSizeCache*)
+	void ASTOpReferenceResource::GetBlockLayoutSize(int blockIndex, int* pBlockX, int* pBlockY, FBlockLayoutSizeCache*)
 	{
 		switch (type)
 		{
@@ -96,7 +96,7 @@ namespace mu
 
 
 	//-------------------------------------------------------------------------------------------------
-	void ASTOpReferenceResource::GetLayoutBlockSize(int32* pBlockX, int32* pBlockY)
+	void ASTOpReferenceResource::GetLayoutBlockSize(int* pBlockX, int* pBlockY)
 	{
 		switch (type)
 		{
diff --git a/Engine/Plugins/Experimental/Mutable/Source/MutableTools/Private/MuT/ASTOpReferenceResource.h b/Engine/Plugins/Experimental/Mutable/Source/MutableTools/Private/MuT/ASTOpReferenceResource.h
--- a/Engine/Plugins/Experimental/Mutable/Source/MutableTools/Private/MuT/ASTOpReferenceResource.h
+++ b/Engine/Plugins/Experimental/Mutable/Source/MutableTools/Private/MuT/ASTOpReferenceResource.h
@@ -46,8 +46,8 @@ namespace mu
 		uint64 Hash() const override;
 		void Link(FProgram& program, FLinkerOptions*) override;
 		FImageDesc GetImageDesc(bool, class FGetImageDescContext*) const override;
-		void GetBlockLayoutSize(uint64 BlockId, int32* pBlockX, int32* pBlockY, FBlockLayoutSizeCache* cache) override;
-		void GetLayoutBlockSize(int32* pBlockX, int32* pBlockY) override;
+		void GetBlockLayoutSize(int blockIndex, int* pBlockX, int* pBlockY, FBlockLayoutSizeCache* cache) override;
+		void GetLayoutBlockSize(int* pBlockX, int* pBlockY) override;
 		bool GetNonBlackRect(FImageRect& maskUsage) const override;
 		Ptr<ImageSizeExpression> GetImageSizeExpression() const override;
 		Ptr<ASTOp> OptimiseSemantic(const FModelOptimizationOptions& options, int32 Pass) const override;
diff --git a/Engine/Plugins/Experimental/Mutable/Source/MutableTools/Private/MuT/ASTOpSwitch.cpp b/Engine/Plugins/Experimental/Mutable/Source/MutableTools/Private/MuT/ASTOpSwitch.cpp
--- a/Engine/Plugins/Experimental/Mutable/Source/MutableTools/Private/MuT/ASTOpSwitch.cpp
+++ b/Engine/Plugins/Experimental/Mutable/Source/MutableTools/Private/MuT/ASTOpSwitch.cpp
@@ -286,7 +286,8 @@ namespace mu
 
 
 	//-------------------------------------------------------------------------------------------------
-	void ASTOpSwitch::GetBlockLayoutSize(uint64 BlockId, int32* pBlockX, int32* pBlockY, FBlockLayoutSizeCache* cache)
+	void ASTOpSwitch::GetBlockLayoutSize(int blockIndex, int* pBlockX, int* pBlockY,
+		FBlockLayoutSizeCache* cache)
 	{
 		switch (type)
 		{
@@ -300,7 +301,7 @@ namespace mu
 
 			if (child)
 			{
-				child->GetBlockLayoutSizeCached(BlockId, pBlockX, pBlockY, cache);
+				child->GetBlockLayoutSizeCached(blockIndex, pBlockX, pBlockY, cache);
 			}
 			else
 			{
@@ -318,7 +319,7 @@ namespace mu
 
 
 	//-------------------------------------------------------------------------------------------------
-	void ASTOpSwitch::GetLayoutBlockSize(int32* pBlockX, int32* pBlockY)
+	void ASTOpSwitch::GetLayoutBlockSize(int* pBlockX, int* pBlockY)
 	{
 		switch (type)
 		{
diff --git a/Engine/Plugins/Experimental/Mutable/Source/MutableTools/Private/MuT/ASTOpSwitch.h b/Engine/Plugins/Experimental/Mutable/Source/MutableTools/Private/MuT/ASTOpSwitch.h
--- a/Engine/Plugins/Experimental/Mutable/Source/MutableTools/Private/MuT/ASTOpSwitch.h
+++ b/Engine/Plugins/Experimental/Mutable/Source/MutableTools/Private/MuT/ASTOpSwitch.h
@@ -67,8 +67,8 @@ template <class SCALAR> class vec4;
 		void Assert() override;
 		void Link(FProgram& program, FLinkerOptions* Options) override;
 		FImageDesc GetImageDesc(bool returnBestOption, class FGetImageDescContext* context) const  override;
-		void GetBlockLayoutSize(uint64 BlockId, int32* pBlockX, int32* pBlockY, FBlockLayoutSizeCache* cache) override;
-		void GetLayoutBlockSize(int32* pBlockX, int32* pBlockY) override;
+		void GetBlockLayoutSize(int blockIndex, int* pBlockX, int* pBlockY, FBlockLayoutSizeCache* cache) override;
+		void GetLayoutBlockSize(int* pBlockX, int* pBlockY) override;
 		Ptr<ASTOp> OptimiseSemantic(const FModelOptimizationOptions&, int32 Pass) const override;
 		bool GetNonBlackRect(FImageRect& maskUsage) const override;
 		bool IsImagePlainConstant(FVector4f& colour) const override;
diff --git a/Engine/Plugins/Experimental/Mutable/Source/MutableTools/Private/MuT/CodeGenerator.cpp b/Engine/Plugins/Experimental/Mutable/Source/MutableTools/Private/MuT/CodeGenerator.cpp
--- a/Engine/Plugins/Experimental/Mutable/Source/MutableTools/Private/MuT/CodeGenerator.cpp
+++ b/Engine/Plugins/Experimental/Mutable/Source/MutableTools/Private/MuT/CodeGenerator.cpp
@@ -385,10 +385,10 @@ namespace mu
 	}
 
 
-	Ptr<const Layout> CodeGenerator::AddLayout(Ptr<const Layout> SourceLayout, uint32 MeshIDPrefix)
+	Ptr<const Layout> CodeGenerator::AddLayout(Ptr<const Layout> SourceLayout)
 	{
 		// The layout we are adding must be a source layout, without block ids yet.
-		check(SourceLayout->Blocks.IsEmpty() || SourceLayout->Blocks[0].Id == Layout::InvalidBlockId);
+		check(SourceLayout->m_blocks.IsEmpty() || SourceLayout->m_blocks[0].m_id == -1);
 
 		Ptr<const Layout>* it = GeneratedLayouts.Find(SourceLayout.get());
 
@@ -399,13 +399,14 @@ namespace mu
 
 		// Assign unique ids to each layout block
 		Ptr<Layout> ClonedLayout = SourceLayout->Clone();
-		for (int32 b = 0; b < ClonedLayout->Blocks.Num(); ++b)
+		for (int32 b = 0; b < ClonedLayout->m_blocks.Num(); ++b)
 		{
-			uint64 Id = uint64(MeshIDPrefix) << 32 | uint64(b);
-			ClonedLayout->Blocks[b].Id = Id;
+			// This is a hard limit due to layout block index data being stored in 16 bit.
+			check(m_absoluteLayoutIndex < 65536);
+			ClonedLayout->m_blocks[b].m_id = m_absoluteLayoutIndex++;
 		}
-		check(SourceLayout->Blocks.Num() == ClonedLayout->Blocks.Num());
-		check(ClonedLayout->Blocks.IsEmpty() || ClonedLayout->Blocks[0].Id != Layout::InvalidBlockId);
+		check(SourceLayout->m_blocks.Num() == ClonedLayout->m_blocks.Num());
+		check(ClonedLayout->m_blocks.IsEmpty() || ClonedLayout->m_blocks[0].m_id != -1);
 		GeneratedLayouts.Add(SourceLayout.get(), ClonedLayout);
 
 		return ClonedLayout;
@@ -1316,7 +1317,7 @@ namespace mu
 								ImageOptions.RectSize = { 0,0 };
 								ImageOptions.ActiveTags = node.m_tags;
 								ImageOptions.LayoutToApply = pLayout;
-								ImageOptions.LayoutBlockId = pLayout->Blocks[BlockIndex].Id;
+								ImageOptions.LayoutBlockId = pLayout->m_blocks[BlockIndex].m_id;
 								FImageGenerationResult Result;
 								GenerateImage(ImageOptions, Result, pImageNode);
 								Ptr<ASTOp> blockAd = Result.op;
@@ -1356,7 +1357,7 @@ namespace mu
 										if (const NodePatchImage* pPatch = e.node->m_textures[t].m_pPatch.get())
 										{
 											// Is the current block to be patched?
-											if (pPatch->GetPrivate()->BlockIndices.Contains(BlockIndex))
+											if (pPatch->GetPrivate()->m_blocks.Contains(BlockIndex))
 											{
 												blockAd = GenerateImageBlockPatch(blockAd, pPatch, e.condition, ImageOptions);
 											}
@@ -1381,8 +1382,8 @@ namespace mu
 								composeOp->BlockImage = blockAd;
 
 								// Set the absolute block index.
-								check(pLayout->Blocks[BlockIndex].Id != Layout::InvalidBlockId);
-								composeOp->BlockId = pLayout->Blocks[BlockIndex].Id;
+								check(pLayout->m_blocks[BlockIndex].m_id >= 0);
+								composeOp->BlockIndex = pLayout->m_blocks[BlockIndex].m_id;
 
 								imageAd = composeOp;
 							}
@@ -1425,7 +1426,7 @@ namespace mu
 												ImageOptions.ActiveTags = node.m_tags;
 												ImageOptions.RectSize = { 0,0 };
 												ImageOptions.LayoutToApply = pExtendLayout;
-												ImageOptions.LayoutBlockId = pExtendLayout->Blocks[b].Id;
+												ImageOptions.LayoutBlockId = pExtendLayout->m_blocks[b].m_id;
 												FImageGenerationResult ExtendResult;
 												GenerateImage(ImageOptions, ExtendResult, pExtend);
 												Ptr<ASTOp> fragmentAd = ExtendResult.op;
@@ -1463,8 +1464,8 @@ namespace mu
 												composeOp->BlockImage = fragmentAd;
 
 												// Set the absolute block index.
-												check(pExtendLayout->Blocks[b].Id != Layout::InvalidBlockId);
-												composeOp->BlockId = pExtendLayout->Blocks[b].Id;
+												check(pExtendLayout->m_blocks[b].m_id >= 0);
+												composeOp->BlockIndex = pExtendLayout->m_blocks[b].m_id;
 
 												lastBase = composeOp;
 											}
@@ -2176,7 +2177,7 @@ namespace mu
 				// Parameters
 				FImageGenerationOptions ClipOptions;
 				ClipOptions.ImageLayoutStrategy = CompilerOptions::TextureLayoutStrategy::None;
-				ClipOptions.LayoutBlockId = Layout::InvalidBlockId;
+				ClipOptions.LayoutBlockId = -1;
 				ClipOptions.State = Options.State;
 
 				FImageGenerationResult ClipMaskResult;
diff --git a/Engine/Plugins/Experimental/Mutable/Source/MutableTools/Private/MuT/CodeGenerator.h b/Engine/Plugins/Experimental/Mutable/Source/MutableTools/Private/MuT/CodeGenerator.h
--- a/Engine/Plugins/Experimental/Mutable/Source/MutableTools/Private/MuT/CodeGenerator.h
+++ b/Engine/Plugins/Experimental/Mutable/Source/MutableTools/Private/MuT/CodeGenerator.h
@@ -205,6 +205,8 @@ namespace mu
         /** List of image resources for every image formats that have been generated so far as palceholders for missing images. */
         Ptr<Image> m_missingImage[size_t(EImageFormat::IF_COUNT)];
 
+        //! First free index for a layout block
+        int32 m_absoluteLayoutIndex = 0;
 
         //! List of already used vertex ID groups that must be unique.
         TSet<uint32> UniqueVertexIDGroups;
@@ -329,7 +331,7 @@ namespace mu
 			UE::Math::TIntVector2<int32> RectSize = {0, 0};
 
 			/** Layout block that we are trying to generate if any. */
-			uint64 LayoutBlockId = Layout::InvalidBlockId;
+			int32 LayoutBlockId = -1;
 			Ptr<const Layout> LayoutToApply;
 
 			friend FORCEINLINE uint32 GetTypeHash(const FImageGenerationOptions& InKey)
@@ -567,11 +569,10 @@ namespace mu
 			Ptr<Mesh> currentLayoutMesh,
 			int32 currentLayoutChannel,
 			const void* errorContext,
-			const FMeshGenerationOptions& MeshOptions,
-			bool bUseAbsoluteBlockIds);
+			const FMeshGenerationOptions& MeshOptions);
 
 		//!
-		Ptr<const Layout> AddLayout(Ptr<const Layout> SourceLayout, uint32 MeshIDPrefix);
+		Ptr<const Layout> AddLayout(Ptr<const Layout> SourceLayout);
 
 		struct FExtensionDataGenerationResult
 		{
diff --git a/Engine/Plugins/Experimental/Mutable/Source/MutableTools/Private/MuT/CodeGenerator_Image.cpp b/Engine/Plugins/Experimental/Mutable/Source/MutableTools/Private/MuT/CodeGenerator_Image.cpp
--- a/Engine/Plugins/Experimental/Mutable/Source/MutableTools/Private/MuT/CodeGenerator_Image.cpp
+++ b/Engine/Plugins/Experimental/Mutable/Source/MutableTools/Private/MuT/CodeGenerator_Image.cpp
@@ -1389,9 +1389,9 @@ namespace mu
 		int32 LayoutBlockIndex = -1;
 		if (Options.LayoutToApply)
 		{
-			LayoutBlockIndex = Options.LayoutToApply->Blocks.IndexOfByPredicate([&](const Layout::FBlock& Block) { return Block.Id == Options.LayoutBlockId; });
+			LayoutBlockIndex = Options.LayoutToApply->m_blocks.IndexOfByPredicate([&](const Layout::FBlock& Block) { return Block.m_id == Options.LayoutBlockId; });
 		}
-		uint64 GeneratedLayoutBlockId = Layout::InvalidBlockId;
+		int32 GeneratedLayoutBlockId = -1;
 
         // Mesh
         if ( node.m_pMesh )
@@ -1407,16 +1407,18 @@ namespace mu
 			FMeshGenerationResult MeshResult;
 			GenerateMesh( MeshOptions, MeshResult, node.m_pMesh );
 
-			// Match the block id of the block we are generating with the id that resulted in the generated mesh			
+			// Match the block id of the block we are generating with the id that resulted in the generated mesh
+			GeneratedLayoutBlockId = -1;
+			
 			Ptr<const Layout> Layout = MeshResult.GeneratedLayouts.IsValidIndex(node.m_layout) ? MeshResult.GeneratedLayouts[node.m_layout] : nullptr;
-			if (Layout && Layout->Blocks.IsValidIndex(LayoutBlockIndex))
+			if (Layout && Layout->m_blocks.IsValidIndex(LayoutBlockIndex))
 			{
-				GeneratedLayoutBlockId = Layout->Blocks[LayoutBlockIndex].Id;
+				GeneratedLayoutBlockId = Layout->m_blocks[LayoutBlockIndex].m_id;
 			}
-			else if (Layout && Layout->Blocks.Num() == 1)
+			else if (Layout && Layout->m_blocks.Num() == 1)
 			{
 				// Layout management disabled, use the only block available
-				GeneratedLayoutBlockId = Layout->Blocks[0].Id;
+				GeneratedLayoutBlockId = Layout->m_blocks[0].m_id;
 			}
 			else
 			{
@@ -1448,7 +1450,7 @@ namespace mu
             else
             {
                 // Extract the mesh layout block
-                if ( GeneratedLayoutBlockId!= Layout::InvalidBlockId)
+                if ( GeneratedLayoutBlockId>=0 )
                 {
                     Ptr<ASTOpMeshExtractLayoutBlocks> eop = new ASTOpMeshExtractLayoutBlocks();
                     eop->Source = CurrentMeshToProjectOp;
@@ -1541,7 +1543,7 @@ namespace mu
 			FImageGenerationOptions NewOptions = Options;
 			NewOptions.ImageLayoutStrategy = CompilerOptions::TextureLayoutStrategy::None;
 			NewOptions.LayoutToApply = nullptr;
-			NewOptions.LayoutBlockId = Layout::InvalidBlockId;
+			NewOptions.LayoutBlockId = -1;
 			NewOptions.RectSize = { 0,0 };
 
 			FImageGenerationResult ImageResult;
diff --git a/Engine/Plugins/Experimental/Mutable/Source/MutableTools/Private/MuT/CodeGenerator_Mesh.cpp b/Engine/Plugins/Experimental/Mutable/Source/MutableTools/Private/MuT/CodeGenerator_Mesh.cpp
--- a/Engine/Plugins/Experimental/Mutable/Source/MutableTools/Private/MuT/CodeGenerator_Mesh.cpp
+++ b/Engine/Plugins/Experimental/Mutable/Source/MutableTools/Private/MuT/CodeGenerator_Mesh.cpp
@@ -269,8 +269,7 @@ namespace mu
 		Ptr<Mesh> currentLayoutMesh,
 		int32 currentLayoutChannel,
 		const void* errorContext,
-		const FMeshGenerationOptions& MeshOptions,
-		bool bUseAbsoluteBlockIds
+		const FMeshGenerationOptions& MeshOptions
 		)
 	{
 		MUTABLE_CPUPROFILER_SCOPE(LayoutUV_PrepareForLayout);
@@ -281,22 +280,46 @@ namespace mu
 		}
 
 		// The layout must have block ids.
-		check(GeneratedLayout->Blocks.IsEmpty() || GeneratedLayout->Blocks[0].Id != Layout::InvalidBlockId);
+		check(GeneratedLayout->m_blocks.IsEmpty() || GeneratedLayout->m_blocks[0].m_id != -1);
+
 
 		// 
 		Ptr<const Layout> Layout = GeneratedLayout;
 		currentLayoutMesh->AddLayout(Layout);
 
+		// Create the layout block vertex buffer
+		uint16* LayoutData = nullptr;
+		{
+			const int32 LayoutBufferIndex = currentLayoutMesh->GetVertexBuffers().GetBufferCount();
+			currentLayoutMesh->GetVertexBuffers().SetBufferCount(LayoutBufferIndex + 1);
+
+			// TODO
+			check(Layout->GetBlockCount() < MAX_uint16);
+			const EMeshBufferSemantic semantic = MBS_LAYOUTBLOCK;
+			const int32 semanticIndex = int32(currentLayoutChannel);
+			const EMeshBufferFormat format = MBF_UINT16;
+			const int32 components = 1;
+			const int32 offset = 0;
+			currentLayoutMesh->GetVertexBuffers().SetBuffer
+			(
+				LayoutBufferIndex,
+				sizeof(uint16),
+				1,
+				&semantic, &semanticIndex,
+				&format, &components,
+				&offset
+			);
+			LayoutData = (uint16*)currentLayoutMesh->GetVertexBuffers().GetBufferData(LayoutBufferIndex);
+		}
+
 		const int32 NumVertices = currentLayoutMesh->GetVertexCount();
 		const int32 NumBlocks = Layout->GetBlockCount();
 
 		// Clear block data
-		TArray<uint16> LayoutData;
-		LayoutData.SetNumUninitialized(NumVertices);
-		constexpr uint16 NullBlockId = MAX_uint16 - 1;
+		const uint16 NullBlockId = MAX_uint16 - 1;
 		for (int32 VertexIndex = 0; VertexIndex < NumVertices; ++VertexIndex)
 		{
-			LayoutData[VertexIndex] = NullBlockId;
+			*(LayoutData + VertexIndex) = NullBlockId;
 		}
 
 		// Find block ids for each block in the grid
@@ -305,7 +328,7 @@ namespace mu
 		GridBlockBlockId.Init(MAX_uint16, Grid.X * Grid.Y);
 
 		// 
-		TArray<uint64> BlockIds;
+		TArray<uint16> BlockIds;
 		TArray<box<FVector2f>> BlockRects;
 
 		BlockIds.SetNumUninitialized(NumBlocks);
@@ -316,7 +339,7 @@ namespace mu
 		for (int32 BlockIndex = 0; BlockIndex < NumBlocks; ++BlockIndex)
 		{
 			// Get the block id
-			BlockIds[BlockIndex] = Layout->Blocks[BlockIndex].Id;
+			BlockIds[BlockIndex] = Layout->m_blocks[BlockIndex].m_id;
 
 			// Get the block rect
 			uint16 MinX, MinY, SizeX, SizeY;
@@ -462,7 +485,7 @@ namespace mu
 
 			uint16 X, Y;
 
-			uint16& BlockIndexV0 = LayoutData[Index0];
+			uint16& BlockIndexV0 = *(LayoutData + Index0);
 			if (BlockIndexV0 == NullBlockId)
 			{
 				X = (uint16)FMath::Min<uint32>(MaxGridX, FMath::Max<uint32>(0, (uint32)(Grid.X * TempUVs[Index0][0])));
@@ -470,7 +493,7 @@ namespace mu
 				BlockIndexV0 = (X < Grid.X && Y < Grid.Y) ? GridBlockBlockId[Y * Grid.X + X] : 0;
 			}
 
-			uint16& BlockIndexV1 = LayoutData[Index1];
+			uint16& BlockIndexV1 = *(LayoutData + Index1);
 			if (BlockIndexV1 == NullBlockId)
 			{
 				X = (uint16)FMath::Min<uint32>(MaxGridX, FMath::Max<uint32>(0, (uint32)(Grid.X * TempUVs[Index1][0])));
@@ -478,7 +501,7 @@ namespace mu
 				BlockIndexV1 = (X < Grid.X && Y < Grid.Y) ? GridBlockBlockId[Y * Grid.X + X] : 0;
 			}
 
-			uint16& BlockIndexV2 = LayoutData[Index2];
+			uint16& BlockIndexV2 = *(LayoutData + Index2);
 			if (BlockIndexV2 == NullBlockId)
 			{
 				X = (uint16)FMath::Min<uint32>(MaxGridX, FMath::Max<uint32>(0, (uint32)(Grid.X * TempUVs[Index2][0])));
@@ -557,13 +580,13 @@ namespace mu
 			}
 
 			// Get the limits of the predominant block rect
-			const Layout::FBlock& LayoutBlock = Layout->Blocks[BlockIndex];
+			const Layout::FBlock& LayoutBlock = Layout->m_blocks[BlockIndex];
 
 			const float SmallNumber = 0.000001;
-			const float MinX = ((float)LayoutBlock.Min.X) / (float)Grid.X + SmallNumber;
-			const float MinY = ((float)LayoutBlock.Min.Y) / (float)Grid.Y + SmallNumber;
-			const float MaxX = (((float)LayoutBlock.Size.X + LayoutBlock.Min.X) / (float)Grid.X) - 2 * SmallNumber;
-			const float MaxY = (((float)LayoutBlock.Size.Y + LayoutBlock.Min.Y) / (float)Grid.Y) - 2 * SmallNumber;
+			const float MinX = ((float)LayoutBlock.m_min.X) / (float)Grid.X + SmallNumber;
+			const float MinY = ((float)LayoutBlock.m_min.Y) / (float)Grid.Y + SmallNumber;
+			const float MaxX = (((float)LayoutBlock.m_size.X + LayoutBlock.m_min.X) / (float)Grid.X) - 2 * SmallNumber;
+			const float MaxY = (((float)LayoutBlock.m_size.Y + LayoutBlock.m_min.Y) / (float)Grid.Y) - 2 * SmallNumber;
 
 			// Iterate triangles and clamp the UVs
 			FVector2f* TempUVsData = TempUVs.GetData();
@@ -585,7 +608,7 @@ namespace mu
 					FVector2f& UV = TempUVs[UVIndex];
 					UV[0] = FMath::Clamp(UV[0], MinX, MaxX);
 					UV[1] = FMath::Clamp(UV[1], MinY, MaxY);
-					LayoutData[UVIndex] = BlockIndex;
+					*(LayoutData + UVIndex) = BlockIndex;
 				}
 
 				OtherTriangle.bUVsFixed = true;
@@ -620,32 +643,6 @@ namespace mu
 			}
 		}
 
-		// Create the layout block vertex buffer
-		uint8* LayoutBufferPtr = nullptr;
-		{
-			const int32 LayoutBufferIndex = currentLayoutMesh->GetVertexBuffers().GetBufferCount();
-			currentLayoutMesh->GetVertexBuffers().SetBufferCount(LayoutBufferIndex + 1);
-
-			// TODO
-			check(Layout->GetBlockCount() < MAX_uint16);
-			const EMeshBufferSemantic LayoutSemantic = MBS_LAYOUTBLOCK;
-			const int32 LayoutSemanticIndex = int32(currentLayoutChannel);
-			const EMeshBufferFormat LayoutFormat = bUseAbsoluteBlockIds ? MBF_UINT64 : MBF_UINT16;
-			const int32 LayoutComponents = 1;
-			const int32 LayoutOffset = 0;
-			int32 ElementSize = bUseAbsoluteBlockIds ? sizeof(uint64) : sizeof(uint16);
-			currentLayoutMesh->GetVertexBuffers().SetBuffer
-			(
-				LayoutBufferIndex,
-				ElementSize,
-				1,
-				&LayoutSemantic, &LayoutSemanticIndex,
-				&LayoutFormat, &LayoutComponents,
-				&LayoutOffset
-			);
-			LayoutBufferPtr = currentLayoutMesh->GetVertexBuffers().GetBufferData(LayoutBufferIndex);
-		}
-
 		// Format and copy UVs
 		{
 			uint8* pVertices = pData;
@@ -655,38 +652,19 @@ namespace mu
 			{
 				FVector2f* UV = &TempUVs[VertexIndex];
 
-				uint16 LayoutBlockIndex = LayoutData[VertexIndex];
-
-				// Replace block index by the actual id of the block
-				if (bUseAbsoluteBlockIds)
+				uint16& LayoutBlockIndex = LayoutData[VertexIndex];
+				if (BlockIds.IsValidIndex(LayoutBlockIndex))
 				{
-					uint64* Ptr = reinterpret_cast<uint64*>(LayoutBufferPtr) + VertexIndex;
+					// Homogenize UVs
+					*UV = BlockRects[LayoutBlockIndex].Homogenize(*UV);
 
-					if (BlockIds.IsValidIndex(LayoutBlockIndex))
-					{
-						*UV = BlockRects[LayoutBlockIndex].Homogenize(*UV);
-						*Ptr = BlockIds[LayoutBlockIndex];
-					}
-					else
-					{
-						// Map vertices without block to the first block.
-						*Ptr = 0;
-					}
+					// Replace block index by the actual id of the block
+					LayoutBlockIndex = BlockIds[LayoutBlockIndex];
 				}
 				else
 				{
-					uint16* Ptr = reinterpret_cast<uint16*>(LayoutBufferPtr) + VertexIndex;
-
-					if (BlockIds.IsValidIndex(LayoutBlockIndex))
-					{
-						*UV = BlockRects[LayoutBlockIndex].Homogenize(*UV);
-						*Ptr = uint16(BlockIds[LayoutBlockIndex] & 0xffff);
-					}
-					else
-					{
-						// Map vertices without block to the first block.
-						*Ptr = 0;
-					}
+					// Map vertices without block to the first block.
+					LayoutBlockIndex = 0;
 				}
 
 				// Copy UVs
@@ -939,16 +917,16 @@ namespace mu
 
                 op->Source = BaseResult.meshOp;
 
-                if (BaseResult.GeneratedLayouts.Num()>node.LayoutOrGroup )
+                if (BaseResult.GeneratedLayouts.Num()>node.m_layoutOrGroup )
                 {
-                    const Layout* pLayout = BaseResult.GeneratedLayouts[node.LayoutOrGroup].get();
-                    op->Layout = (uint16)node.LayoutOrGroup;
+                    const Layout* pLayout = BaseResult.GeneratedLayouts[node.m_layoutOrGroup].get();
+                    op->Layout = (uint16)node.m_layoutOrGroup;
 
-                    for ( int32 i=0; i<node.Blocks.Num(); ++i )
+                    for ( int32 i=0; i<node.m_blocks.Num(); ++i )
                     {
-                        if (node.Blocks[i]>=0 && node.Blocks[i]<pLayout->Blocks.Num() )
+                        if (node.m_blocks[i]>=0 && node.m_blocks[i]<pLayout->m_blocks.Num() )
                         {
-                            uint64 bid = pLayout->Blocks[ node.Blocks[i] ].Id;
+                            int bid = pLayout->m_blocks[ node.m_blocks[i] ].m_id;
                             op->Blocks.Add(bid);
                         }
                         else
@@ -1286,7 +1264,7 @@ namespace mu
 			// This data is required
 			MeshPtr EmptyMesh = new Mesh();
 			ConstantOp->SetValue(EmptyMesh, m_compilerOptions->OptimisationOptions.DiskCacheContext);
-			EmptyMesh->MeshIDPrefix = ConstantOp->GetValueHash();
+			EmptyMesh->VertexIDPrefix = ConstantOp->GetValueHash();
 
 			FGeneratedConstantMesh MeshEntry;
 			MeshEntry.Mesh = EmptyMesh;
@@ -1385,22 +1363,22 @@ namespace mu
 			ConstantOp->SetValue(Cloned, m_compilerOptions->OptimisationOptions.DiskCacheContext);
 
 			// Add the unique vertex ID prefix in all cases, since it is free memory-wise
-			uint32 MeshIDPrefix = uint32(ConstantOp->GetValueHash());
+			uint32 VertexIDPrefix = uint32(ConstantOp->GetValueHash());
 			{
 				// Ensure the ID group is unique
 				bool bValid = false;
 				do
 				{
 					bool bAlreadyPresent = false;
-					UniqueVertexIDGroups.FindOrAdd(MeshIDPrefix, &bAlreadyPresent);
-					bValid = !bAlreadyPresent && MeshIDPrefix != 0;
+					UniqueVertexIDGroups.FindOrAdd(VertexIDPrefix, &bAlreadyPresent);
+					bValid = !bAlreadyPresent && VertexIDPrefix != 0;
 					if (!bValid)
 					{
-						++MeshIDPrefix;
+						++VertexIDPrefix;
 					}
 				} while (bValid);
 
-				Cloned->MeshIDPrefix = MeshIDPrefix;
+				Cloned->VertexIDPrefix = VertexIDPrefix;
 			}
 
 			// Add the constant data
@@ -1427,11 +1405,9 @@ namespace mu
 						const NodeLayoutBlocks* TypedNode = static_cast<const NodeLayoutBlocks*>(pLayoutNode.get());
 
 						Ptr<const Layout> SourceLayout = TypedNode->GetPrivate()->m_pLayout;
-						Ptr<const Layout> GeneratedLayout = AddLayout(SourceLayout, MeshIDPrefix);
+						Ptr<const Layout> GeneratedLayout = AddLayout( SourceLayout );
 						const void* Context = InOptions.OverrideContext.Get(Node.m_errorContext);
-
-						bool bUseAbsoluteBlockIds = false;
-						PrepareForLayout(GeneratedLayout, Cloned, LayoutIndex, Context, InOptions, bUseAbsoluteBlockIds);
+						PrepareForLayout(GeneratedLayout, Cloned, LayoutIndex, Context, InOptions);
 
 						OutResult.GeneratedLayouts.Add(GeneratedLayout);
 					}
@@ -1443,11 +1419,7 @@ namespace mu
 					{
 						Ptr<const Layout> GeneratedLayout = InOptions.OverrideLayouts[LayoutIndex];
 						const void* Context = InOptions.OverrideContext.Get(Node.m_errorContext);
-
-						// In this case we need the layout block ids to use the ids in the parent layout, and not be prefixed with
-						// the current mesh id prefix. For this reason we need them to be absolute.
-						bool bUseAbsoluteBlockIds = true;
-						PrepareForLayout(GeneratedLayout, Cloned, LayoutIndex, Context, InOptions, bUseAbsoluteBlockIds);
+						PrepareForLayout(GeneratedLayout, Cloned, LayoutIndex, Context, InOptions);
 
 						OutResult.GeneratedLayouts.Add(GeneratedLayout);
 					}
diff --git a/Engine/Plugins/Experimental/Mutable/Source/MutableTools/Private/MuT/CodeOptimiser_Parameters.cpp b/Engine/Plugins/Experimental/Mutable/Source/MutableTools/Private/MuT/CodeOptimiser_Parameters.cpp
--- a/Engine/Plugins/Experimental/Mutable/Source/MutableTools/Private/MuT/CodeOptimiser_Parameters.cpp
+++ b/Engine/Plugins/Experimental/Mutable/Source/MutableTools/Private/MuT/CodeOptimiser_Parameters.cpp
@@ -403,7 +403,9 @@ namespace mu
                         {
 						const ASTOpImageCompose* typedYes = static_cast<const ASTOpImageCompose*>(typedAt->yes.child().get());
 						const ASTOpImageCompose* typedNo = static_cast<const ASTOpImageCompose*>(typedAt->no.child().get());
-                        if ( typedYes->BlockId == typedNo->BlockId
+                        if ( typedYes->BlockIndex
+                             ==
+                             typedNo->BlockIndex
                              &&
                              (
                                  (typedYes->Mask.child().get() != nullptr)
diff --git a/Engine/Plugins/Experimental/Mutable/Source/MutableTools/Private/MuT/NodeMeshFragment.cpp b/Engine/Plugins/Experimental/Mutable/Source/MutableTools/Private/MuT/NodeMeshFragment.cpp
--- a/Engine/Plugins/Experimental/Mutable/Source/MutableTools/Private/MuT/NodeMeshFragment.cpp
+++ b/Engine/Plugins/Experimental/Mutable/Source/MutableTools/Private/MuT/NodeMeshFragment.cpp
@@ -44,60 +44,60 @@ namespace mu
 
 
     //---------------------------------------------------------------------------------------------
-    Ptr<NodeMesh> NodeMeshFragment::GetMesh() const
+    NodeMeshPtr NodeMeshFragment::GetMesh() const
 	{
 		return m_pD->m_pMesh.get();
 	}
 
 
 	//---------------------------------------------------------------------------------------------
-	void NodeMeshFragment::SetMesh(Ptr<NodeMesh> pNode )
+	void NodeMeshFragment::SetMesh( NodeMeshPtr pNode )
 	{
 		m_pD->m_pMesh = pNode;
 	}
 
 
 	//---------------------------------------------------------------------------------------------
-    int32 NodeMeshFragment::GetLayoutOrGroup() const
+    int NodeMeshFragment::GetLayoutOrGroup() const
 	{
-        return m_pD->LayoutOrGroup;
+        return m_pD->m_layoutOrGroup;
 	}
 
 
 	//---------------------------------------------------------------------------------------------
-    void NodeMeshFragment::SetLayoutOrGroup( int32 l )
+    void NodeMeshFragment::SetLayoutOrGroup( int l )
 	{
-        m_pD->LayoutOrGroup = l;
+        m_pD->m_layoutOrGroup = l;
 	}
 
 
 	//---------------------------------------------------------------------------------------------
-	void NodeMeshFragment::SetBlockCount( int32 t )
+	void NodeMeshFragment::SetBlockCount( int t )
 	{
-		m_pD->Blocks.SetNum(t);
+		m_pD->m_blocks.SetNum(t);
 	}
 
 
 	//---------------------------------------------------------------------------------------------
-	int32 NodeMeshFragment::GetBlock( int32 t ) const
+	int NodeMeshFragment::GetBlock( int t ) const
 	{
-		check( t>=0 && t<m_pD->Blocks.Num() );
-		return m_pD->Blocks[t];
+		check( t>=0 && t<m_pD->m_blocks.Num() );
+		return m_pD->m_blocks[t];
 	}
 
 
 	//---------------------------------------------------------------------------------------------
-	void NodeMeshFragment::SetBlock( int32 t, int32 b )
+	void NodeMeshFragment::SetBlock( int t, int b )
 	{
-		check( t>=0 && t<m_pD->Blocks.Num() );
-		m_pD->Blocks[t] = b;
+		check( t>=0 && t<m_pD->m_blocks.Num() );
+		m_pD->m_blocks[t] = b;
 	}
 
 
 	//---------------------------------------------------------------------------------------------
-	Ptr<NodeLayout> NodeMeshFragment::Private::GetLayout( int32 index ) const
+	NodeLayoutPtr NodeMeshFragment::Private::GetLayout( int index ) const
 	{
-		Ptr<NodeLayout> pResult;
+		NodeLayoutPtr pResult;
 
 		if ( m_pMesh )
 		{
diff --git a/Engine/Plugins/Experimental/Mutable/Source/MutableTools/Private/MuT/NodePatchImage.cpp b/Engine/Plugins/Experimental/Mutable/Source/MutableTools/Private/MuT/NodePatchImage.cpp
--- a/Engine/Plugins/Experimental/Mutable/Source/MutableTools/Private/MuT/NodePatchImage.cpp
+++ b/Engine/Plugins/Experimental/Mutable/Source/MutableTools/Private/MuT/NodePatchImage.cpp
@@ -28,60 +28,60 @@ namespace mu
     //---------------------------------------------------------------------------------------------
     // Own Interface
     //---------------------------------------------------------------------------------------------
-    void NodePatchImage::SetImage( Ptr<NodeImage> pImage )
+    void NodePatchImage::SetImage( NodeImagePtr pImage )
     {
         m_pD->m_pImage = pImage;
     }
 
 
     //---------------------------------------------------------------------------------------------
-	Ptr<NodeImage> NodePatchImage::GetImage() const
+    NodeImagePtr NodePatchImage::GetImage() const
     {
         return m_pD->m_pImage.get();
     }
 
 
     //---------------------------------------------------------------------------------------------
-    void NodePatchImage::SetMask(Ptr<NodeImage> pMask )
+    void NodePatchImage::SetMask( NodeImagePtr pMask )
     {
         m_pD->m_pMask = pMask;
     }
 
 
     //---------------------------------------------------------------------------------------------
-	Ptr<NodeImage> NodePatchImage::GetMask() const
+    NodeImagePtr NodePatchImage::GetMask() const
     {
         return m_pD->m_pMask.get();
     }
 
 
     //---------------------------------------------------------------------------------------------
-    void NodePatchImage::SetBlockCount( int32 c )
+    void NodePatchImage::SetBlockCount( int c )
     {
-        m_pD->BlockIndices.SetNum( c );
+        m_pD->m_blocks.SetNum( c );
     }
 
 
     //---------------------------------------------------------------------------------------------
-    int32 NodePatchImage::GetBlockCount() const
+    int NodePatchImage::GetBlockCount() const
     {
-        return (int32)m_pD->BlockIndices.Num();
+        return (int)m_pD->m_blocks.Num();
     }
 
 
     //---------------------------------------------------------------------------------------------
-    void NodePatchImage::SetBlock( int32 index, int32 LayoutBlockIndex )
+    void NodePatchImage::SetBlock( int index, int block )
     {
-        check( index>=0 && index<m_pD->BlockIndices.Num() );
-        m_pD->BlockIndices[ index ] = LayoutBlockIndex;
+        check( index>=0 && index<(int)m_pD->m_blocks.Num() );
+        m_pD->m_blocks[ index ] = block;
     }
 
 
     //---------------------------------------------------------------------------------------------
-    int32 NodePatchImage::GetBlock( int32 index ) const
+    int NodePatchImage::GetBlock( int index ) const
     {
-        check( index>=0 && index<m_pD->BlockIndices.Num() );
-        return m_pD->BlockIndices[ index ];
+        check( index>=0 && index<(int)m_pD->m_blocks.Num() );
+        return m_pD->m_blocks[ index ];
     }
 
 
@@ -93,9 +93,9 @@ namespace mu
 
 
     //---------------------------------------------------------------------------------------------
-    void NodePatchImage::SetBlendType(EBlendType InType)
+    void NodePatchImage::SetBlendType(EBlendType t)
     {
-        m_pD->m_blendType = InType;
+        m_pD->m_blendType = t;
     }
 
 
