[Mutable] Fix is CustomizableObjectInstance::IsParameterDirty().

- Range check was backwards.

#rnx

diff --git a/Engine/Plugins/Experimental/Mutable/Source/CustomizableObject/Private/MuCO/CustomizableObjectInstance.cpp b/Engine/Plugins/Experimental/Mutable/Source/CustomizableObject/Private/MuCO/CustomizableObjectInstance.cpp
--- a/Engine/Plugins/Experimental/Mutable/Source/CustomizableObject/Private/MuCO/CustomizableObjectInstance.cpp
+++ b/Engine/Plugins/Experimental/Mutable/Source/CustomizableObject/Private/MuCO/CustomizableObjectInstance.cpp
@@ -741,7 +741,7 @@ bool UCustomizableObjectInstance::IsParameterDirty(const FString& ParamName, con
 			
 			if (Result && ResultCommited)
 			{
-				if (RangeIndex != INDEX_NONE)
+				if (RangeIndex == INDEX_NONE)
 				{
 					return Result->Value == ResultCommited->Value;					
 				}
@@ -776,7 +776,7 @@ bool UCustomizableObjectInstance::IsParameterDirty(const FString& ParamName, con
 			
 			if (Result && ResultCommited)
 			{
-				if (RangeIndex != INDEX_NONE)
+				if (RangeIndex == INDEX_NONE)
 				{
 					return Result->ParameterValue == ResultCommited->ParameterValue;					
 				}
@@ -812,7 +812,7 @@ bool UCustomizableObjectInstance::IsParameterDirty(const FString& ParamName, con
 			
 			if (Result && ResultCommited)
 			{
-				if (RangeIndex != INDEX_NONE)
+				if (RangeIndex == INDEX_NONE)
 				{
 					return Result->ParameterValue == ResultCommited->ParameterValue;					
 				}
@@ -840,7 +840,7 @@ bool UCustomizableObjectInstance::IsParameterDirty(const FString& ParamName, con
 			
 			if (Result && ResultCommited)
 			{
-				if (RangeIndex != INDEX_NONE)
+				if (RangeIndex == INDEX_NONE)
 				{
 					return Result->ParameterValueName == ResultCommited->ParameterValueName;					
 				}
@@ -876,7 +876,7 @@ bool UCustomizableObjectInstance::IsParameterDirty(const FString& ParamName, con
 			
 			if (Result && ResultCommited)
 			{
-				if (RangeIndex != INDEX_NONE)
+				if (RangeIndex == INDEX_NONE)
 				{
 					return Result->ParameterValue == ResultCommited->ParameterValue;					
 				}
@@ -912,7 +912,7 @@ bool UCustomizableObjectInstance::IsParameterDirty(const FString& ParamName, con
 			
 			if (Result && ResultCommited)
 			{
-				if (RangeIndex != INDEX_NONE)
+				if (RangeIndex == INDEX_NONE)
 				{
 					return Result->ParameterValue == ResultCommited->ParameterValue;					
 				}
