[mutable] Add missing file in previous submit.
#rnx

diff --git a/Engine/Plugins/Experimental/Mutable/Source/MutableRuntime/Private/MuR/Mesh.cpp b/Engine/Plugins/Experimental/Mutable/Source/MutableRuntime/Private/MuR/Mesh.cpp
--- a/Engine/Plugins/Experimental/Mutable/Source/MutableRuntime/Private/MuR/Mesh.cpp
+++ b/Engine/Plugins/Experimental/Mutable/Source/MutableRuntime/Private/MuR/Mesh.cpp
@@ -54,7 +54,7 @@ MeshPtr Mesh::Clone() const
 	pResult->PhysicsBody = PhysicsBody;
 	pResult->Tags = Tags;
 	pResult->StreamedResources = StreamedResources;
-	pResult->MeshIDPrefix = MeshIDPrefix;
+	pResult->VertexIDPrefix = VertexIDPrefix;
 
     // Clone the main buffers
     pResult->VertexBuffers = VertexBuffers;
@@ -178,7 +178,7 @@ void Mesh::CopyFrom(const Mesh& From, EMeshCopyFlags Flags)
 
     InternalId = From.InternalId;
 	StaticFormatFlags = From.StaticFormatFlags;
-	MeshIDPrefix = From.MeshIDPrefix;
+	VertexIDPrefix = From.VertexIDPrefix;
 
 	if (EnumHasAnyFlags(Flags, EMeshCopyFlags::WithSurfaces))
 	{
@@ -282,6 +282,26 @@ const FMeshBufferSet& Mesh::GetVertexBuffers() const
 }
 
 
+bool Mesh::HasVertexIds() const
+{
+	if (VertexIDPrefix != 0) return true;
+
+	// If there is no prefix, there must be a buffer with full indices
+	int32 BufferIndex = -1;
+	int32 ChannelIndex = -1;
+	VertexBuffers.FindChannel(MBS_VERTEXINDEX, 0, &BufferIndex, &ChannelIndex);
+	if (BufferIndex >= 0 && ChannelIndex >= 0)
+	{
+		if (VertexBuffers.m_buffers[BufferIndex].m_channels[ChannelIndex].m_format == MBF_UINT64)
+		{
+			return true;
+		}
+	}
+
+	return false;
+}
+
+
 bool Mesh::AreVertexIdsImplicit() const
 {
 	// Is there a buffer for vertex ids?
@@ -289,7 +309,7 @@ bool Mesh::AreVertexIdsImplicit() const
 	int32 ChannelIndex = -1;
 	VertexBuffers.FindChannel(MBS_VERTEXINDEX, 0, &BufferIndex, &ChannelIndex);
 
-	return (MeshIDPrefix != 0) && (BufferIndex < 0) && (ChannelIndex < 0);
+	return (VertexIDPrefix != 0) && (BufferIndex < 0) && (ChannelIndex < 0);
 }
 
 
@@ -303,13 +323,13 @@ bool Mesh::AreVertexIdsExplicit() const
 	bool bExplicit = (BufferIndex >= 0) && (ChannelIndex >= 0) && (VertexBuffers.m_buffers[BufferIndex].m_channels[ChannelIndex].m_format == MBF_UINT64);
 	if (bExplicit)
 	{
-		check(MeshIDPrefix == 0);
+		check(VertexIDPrefix == 0);
 	}
 	return bExplicit;
 }
 
 
-void Mesh::MakeVertexIdsRelative()
+void Mesh::MakeVertexIndicesRelative()
 {
 	check(AreVertexIdsImplicit());
 
@@ -331,11 +351,11 @@ void Mesh::MakeVertexIdsRelative()
 }
 
 
-void Mesh::MakeIdsExplicit()
+void Mesh::MakeVertexIndicesExplicit()
 {
 	MUTABLE_CPUPROFILER_SCOPE(Mesh_MakeVertexIndicesExplicit);
 
-	if (!MeshIDPrefix)
+	if (!VertexIDPrefix)
 	{
 		// We already have explicit vertex IDs, or we don't have any.
 		return;
@@ -343,121 +363,72 @@ void Mesh::MakeIdsExplicit()
 
 	int32 VertexCount = GetVertexCount();
 
-	// Vertex IDs
-	{
-		MUTABLE_CPUPROFILER_SCOPE(VertexIDs);
+	bool bHasRelativeVertexIndices = false;
 
-		bool bHasRelativeVertexIndices = false;
+	int32 OldBuf = -1;
+	int32 OldChan = -1;
+	VertexBuffers.FindChannel(MBS_VERTEXINDEX, 0, &OldBuf, &OldChan);
+	bool bHasVertexIndices = (OldBuf >= 0 && OldChan >= 0);
+	if (bHasVertexIndices)
+	{
+		check(OldChan == 0 && VertexBuffers.m_buffers[OldBuf].m_channels.Num() == 1);
 
-		int32 OldBuf = -1;
-		int32 OldChan = -1;
-		VertexBuffers.FindChannel(MBS_VERTEXINDEX, 0, &OldBuf, &OldChan);
-		bool bHasVertexIndices = (OldBuf >= 0 && OldChan >= 0);
-		if (bHasVertexIndices)
+		bool bHasExplicitVertexIndices = VertexBuffers.m_buffers[OldBuf].m_channels[0].m_format == MBF_UINT64;
+		if (bHasExplicitVertexIndices)
 		{
-			check(OldChan == 0 && VertexBuffers.m_buffers[OldBuf].m_channels.Num() == 1);
-
-			FMeshBufferChannel& Channel = VertexBuffers.m_buffers[OldBuf].m_channels[0];
-
-			bool bHasExplicitVertexIndices = Channel.m_format == MBF_UINT64;
-			if (bHasExplicitVertexIndices)
-			{
-				// nothing to do
-				return;
-			}
-
-			// The mesh has relative vertex IDs.
-			bHasRelativeVertexIndices = true;
-
-			check(Channel.m_format==MBF_UINT32);
-			const uint32* OldIdData = reinterpret_cast<const uint32*>(VertexBuffers.GetBufferData(OldBuf));
-
-			TMemoryTrackedArray<uint8> NewIds;
-			NewIds.SetNumUninitialized(VertexCount * sizeof(uint64));
-			uint64* NewIdData = reinterpret_cast<uint64*>(NewIds.GetData());
-
-			for (int32 i = 0; i < VertexCount; ++i)
-			{
-				uint32 OldId = *OldIdData++;
-				uint64 Id = (uint64(MeshIDPrefix) << 32) | uint64(OldId);
-				(*NewIdData++) = Id;
-			}
-
-			// 
-			FMeshBuffer& Buffer = VertexBuffers.m_buffers[OldBuf];
-			Swap(Buffer.m_data, NewIds);
-			Buffer.m_channels[0].m_format = MBF_UINT64;
-			Buffer.m_elementSize = sizeof(uint64);
+			// nothing to do
+			return;
 		}
-
-		if (!bHasRelativeVertexIndices)
+				
+		// The mesh has relative vertex IDs.
+		bHasRelativeVertexIndices = true;
+
+		const uint32* OldIdData = reinterpret_cast<const uint32*>(VertexBuffers.GetBufferData(OldBuf));
+
+		// Create a new buffer with explicit ids and remove the old buffer with relative ids
+		int32 NewBuffer = VertexBuffers.GetBufferCount();
+		VertexBuffers.SetBufferCount(NewBuffer + 1);
+		EMeshBufferSemantic Semantic = MBS_VERTEXINDEX;
+		int32 SemanticIndex = 0;
+		EMeshBufferFormat Format = MBF_UINT64;
+		int32 Components = 1;
+		int32 Offset = 0;
+		VertexBuffers.SetBuffer(NewBuffer, sizeof(uint64), 1, &Semantic, &SemanticIndex, &Format, &Components, &Offset);
+		uint64* IdData = reinterpret_cast<uint64*>(VertexBuffers.GetBufferData(NewBuffer));
+
+		for (int32 i = 0; i < VertexCount; ++i)
 		{
-			// The mesh has implicit Ids
-			// Create a new buffer with explicit ids
-			int32 NewBuffer = VertexBuffers.GetBufferCount();
-			VertexBuffers.SetBufferCount(NewBuffer + 1);
-			EMeshBufferSemantic Semantic = MBS_VERTEXINDEX;
-			int32 SemanticIndex = 0;
-			EMeshBufferFormat Format = MBF_UINT64;
-			int32 Components = 1;
-			int32 Offset = 0;
-			VertexBuffers.SetBuffer(NewBuffer, sizeof(uint64), 1, &Semantic, &SemanticIndex, &Format, &Components, &Offset);
-			uint64* IdData = reinterpret_cast<uint64*>(VertexBuffers.GetBufferData(NewBuffer));
-
-			for (int32 VertexIndex = 0; VertexIndex < VertexCount; ++VertexIndex)
-			{
-				uint64 Id = (uint64(MeshIDPrefix) << 32) | uint64(VertexIndex);
-				(*IdData++) = Id;
-			}
+			uint32 OldId = *OldIdData++;
+			uint64 Id = (uint64(VertexIDPrefix) << 32) | uint64(OldId);
+			(*IdData++) = Id;
 		}
+
+		// TODO: Optimize by overwritting the existing buffer.
+		VertexBuffers.RemoveBuffer(OldBuf);
 	}
 
-	// Layout block IDs
+	if (!bHasRelativeVertexIndices)
 	{
-		MUTABLE_CPUPROFILER_SCOPE(LayoutBlockIDs);
-
-		for (FMeshBuffer& Buffer: VertexBuffers.m_buffers)
+		// The mesh has implicit Ids
+		// Create a new buffer with explicit ids
+		int32 NewBuffer = VertexBuffers.GetBufferCount();
+		VertexBuffers.SetBufferCount(NewBuffer + 1);
+		EMeshBufferSemantic Semantic = MBS_VERTEXINDEX;
+		int32 SemanticIndex = 0;
+		EMeshBufferFormat Format = MBF_UINT64;
+		int32 Components = 1;
+		int32 Offset = 0;
+		VertexBuffers.SetBuffer(NewBuffer, sizeof(uint64), 1, &Semantic, &SemanticIndex, &Format, &Components, &Offset);
+		uint64* IdData = reinterpret_cast<uint64*>(VertexBuffers.GetBufferData(NewBuffer));
+
+		for (int32 VertexIndex = 0; VertexIndex < VertexCount; ++VertexIndex)
 		{
-			for (FMeshBufferChannel& Channel : Buffer.m_channels)
-			{
-				if (Channel.m_semantic != MBS_LAYOUTBLOCK)
-				{
-					continue;
-				}
-
-				if (Channel.m_format == MBF_UINT64)
-				{
-					continue;
-				}
-
-				check(Buffer.m_channels.Num() == 1);
-				check(Buffer.m_channels[0].m_offset == 0);
-				check(Buffer.m_elementSize == sizeof(uint16));
-
-				check(Channel.m_format == MBF_UINT16);
-				const uint16* OldIdData = reinterpret_cast<const uint16*>(Buffer.m_data.GetData());
-
-				TMemoryTrackedArray<uint8> NewIds;
-				NewIds.SetNumUninitialized(VertexCount *sizeof(uint64));
-				uint64* NewIdData = reinterpret_cast<uint64*>(NewIds.GetData());
-
-				for (int32 i = 0; i < VertexCount; ++i)
-				{
-					uint16 OldId = *OldIdData++;
-					uint64 Id = (uint64(MeshIDPrefix) << 32) | uint64(OldId);
-					(*NewIdData++) = Id;
-				}
-
-				// 
-				Swap(Buffer.m_data, NewIds);
-				Buffer.m_channels[0].m_format = MBF_UINT64;
-				Buffer.m_elementSize = sizeof(uint64);
-			}
+			uint64 Id = (uint64(VertexIDPrefix) << 32) | uint64(VertexIndex);
+			(*IdData++) = Id;
 		}
 	}
 
-	// Final cleanup
-	MeshIDPrefix = 0;
+	VertexIDPrefix = 0;
 }
 
 
@@ -1132,7 +1103,7 @@ void Mesh::Serialise(OutputArchive& arch) const
 
 	arch << AdditionalPhysicsBodies;
 
-	arch << MeshIDPrefix;
+	arch << VertexIDPrefix;
 }
 
 
@@ -1264,7 +1235,7 @@ void Mesh::Unserialise(InputArchive& arch)
 
 	if (ver >= 19)
 	{
-		arch >> MeshIDPrefix;
+		arch >> VertexIDPrefix;
 	}
 }
 
@@ -1356,11 +1327,11 @@ void Mesh::CheckIntegrity() const
 			EMeshBufferFormat IdFormat = VertexBuffers.m_buffers[BufferIndex].m_channels[ChannelIndex].m_format;
 			if (IdFormat == MBF_UINT64)
 			{
-				check(MeshIDPrefix==0);
+				check(VertexIDPrefix==0);
 			}
 			else if (IdFormat == MBF_UINT32)
 			{
-				check(MeshIDPrefix != 0);
+				check(VertexIDPrefix != 0);
 			}
 			else
 			{
