[Mutable] Fix textures with NoMipmaps generation setttings not generating the correct image size.

- Image with NoMipmaps generation settings will ignore mips to skip in the initial generation.

#rnx

diff --git a/Engine/Plugins/Experimental/Mutable/Source/CustomizableObject/Internal/MuCO/CustomizableObjectSystemPrivate.h b/Engine/Plugins/Experimental/Mutable/Source/CustomizableObject/Internal/MuCO/CustomizableObjectSystemPrivate.h
--- a/Engine/Plugins/Experimental/Mutable/Source/CustomizableObject/Internal/MuCO/CustomizableObjectSystemPrivate.h
+++ b/Engine/Plugins/Experimental/Mutable/Source/CustomizableObject/Internal/MuCO/CustomizableObjectSystemPrivate.h
@@ -320,6 +320,7 @@ struct FInstanceUpdateData
 		TWeakObjectPtr<UTexture2D> Cached;
 
 		bool bIsPassThrough = false;
+		bool bIsNonProgressive = false;
 	};
 
 	struct FVector
diff --git a/Engine/Plugins/Experimental/Mutable/Source/CustomizableObject/Private/MuCO/BusyWaits_Deprecated.cpp b/Engine/Plugins/Experimental/Mutable/Source/CustomizableObject/Private/MuCO/BusyWaits_Deprecated.cpp
--- a/Engine/Plugins/Experimental/Mutable/Source/CustomizableObject/Private/MuCO/BusyWaits_Deprecated.cpp
+++ b/Engine/Plugins/Experimental/Mutable/Source/CustomizableObject/Private/MuCO/BusyWaits_Deprecated.cpp
@@ -69,8 +69,9 @@ namespace CustomizableObjectSystem::ImplDeprecated
 				const int32 MinMipsInImage = FMath::Min(FullLODCount, UTexture::GetStaticMinTextureResidentMipCount());
 				const int32 MaxMipsToSkip = FullLODCount - MinMipsInImage;
 				int32 MipsToSkip = FMath::Min(MaxMipsToSkip, OperationData->MipsToSkip);
-
-				if (!FMath::IsPowerOfTwo(Image.FullImageSizeX) || !FMath::IsPowerOfTwo(Image.FullImageSizeY))
+				
+				// This has beed replicated to the non deperecated code path.
+				if (Image.bIsNonProgressive || !FMath::IsPowerOfTwo(Image.FullImageSizeX) || !FMath::IsPowerOfTwo(Image.FullImageSizeY))
 				{
 					// It doesn't make sense to skip mips as non-power-of-two size textures cannot be streamed anyway
 					MipsToSkip = 0;
@@ -307,6 +308,9 @@ namespace CustomizableObjectSystem::ImplDeprecated
 							if (ImageKey >= 0 && ImageKey < ModelResources.ImageProperties.Num())
 							{
 								const FMutableModelImageProperties& Props = ModelResources.ImageProperties[ImageKey];
+								
+								// This has beed replicated to the non deperecated code path.
+								Image.bIsNonProgressive = Props.MipGenSettings == TMGS_NoMipmaps;
 
 								if (Props.IsPassThrough)
 								{
diff --git a/Engine/Plugins/Experimental/Mutable/Source/CustomizableObject/Private/MuCO/CustomizableObjectInstance.cpp b/Engine/Plugins/Experimental/Mutable/Source/CustomizableObject/Private/MuCO/CustomizableObjectInstance.cpp
--- a/Engine/Plugins/Experimental/Mutable/Source/CustomizableObject/Private/MuCO/CustomizableObjectInstance.cpp
+++ b/Engine/Plugins/Experimental/Mutable/Source/CustomizableObject/Private/MuCO/CustomizableObjectInstance.cpp
@@ -3012,15 +3012,13 @@ void SetTexturePropertiesFromMutableImageProps(UTexture2D* Texture, const FMutab
 	Texture->Filter = Props.Filter;
 	Texture->LODBias = Props.LODBias;
 
-	// TODO: MTBL-1957
-	//if (Props.MipGenSettings == TextureMipGenSettings::TMGS_NoMipmaps)
-	//{
-		//Texture->NeverStream = true;
-	//}
+	if (Props.MipGenSettings == TextureMipGenSettings::TMGS_NoMipmaps)
+	{
+		Texture->NeverStream = true;
+	}
 
 #if WITH_EDITORONLY_DATA
-	// TODO: MTBL-1957
-	//Texture->MipGenSettings = Props.MipGenSettings;
+	Texture->MipGenSettings = Props.MipGenSettings;
 
 	Texture->bFlipGreenChannel = Props.FlipGreenChannel;
 #endif
diff --git a/Engine/Plugins/Experimental/Mutable/Source/CustomizableObject/Private/MuCO/CustomizableObjectSystem.cpp b/Engine/Plugins/Experimental/Mutable/Source/CustomizableObject/Private/MuCO/CustomizableObjectSystem.cpp
--- a/Engine/Plugins/Experimental/Mutable/Source/CustomizableObject/Private/MuCO/CustomizableObjectSystem.cpp
+++ b/Engine/Plugins/Experimental/Mutable/Source/CustomizableObject/Private/MuCO/CustomizableObjectSystem.cpp
@@ -2260,6 +2260,8 @@ namespace impl
 							{
 								const FMutableModelImageProperties& Props = ModelResources.ImageProperties[ImageKey];
 
+								Image.bIsNonProgressive = Props.MipGenSettings == TMGS_NoMipmaps;
+
 								if (Props.IsPassThrough)
 								{
 									Image.bIsPassThrough = true;
@@ -2591,7 +2593,7 @@ namespace impl
 		const int32 MaxMipsToSkip = FullLODCount - MinMipsInImage;
 		int32 MipsToSkip = FMath::Min(MaxMipsToSkip, OperationData->MipsToSkip);
 
-		if (!FMath::IsPowerOfTwo(Image.FullImageSizeX) || !FMath::IsPowerOfTwo(Image.FullImageSizeY))
+		if (Image.bIsNonProgressive || !FMath::IsPowerOfTwo(Image.FullImageSizeX) || !FMath::IsPowerOfTwo(Image.FullImageSizeY))
 		{
 			// It doesn't make sense to skip mips as non-power-of-two size textures cannot be streamed anyway
 			MipsToSkip = 0;
