[Mutable] Avoid crashing when pasting the wrong descriptor.

#rnx

diff --git a/Engine/Plugins/Experimental/Mutable/Source/CustomizableObject/Internal/MuCO/CustomizableObjectInstanceDescriptor.h b/Engine/Plugins/Experimental/Mutable/Source/CustomizableObject/Internal/MuCO/CustomizableObjectInstanceDescriptor.h
--- a/Engine/Plugins/Experimental/Mutable/Source/CustomizableObject/Internal/MuCO/CustomizableObjectInstanceDescriptor.h
+++ b/Engine/Plugins/Experimental/Mutable/Source/CustomizableObject/Internal/MuCO/CustomizableObjectInstanceDescriptor.h
@@ -60,7 +60,7 @@ struct CUSTOMIZABLEOBJECT_API FCustomizableObjectInstanceDescriptor
 
 	UCustomizableObject* GetCustomizableObject() const;
 
-	void SetCustomizableObject(UCustomizableObject& InCustomizableObject);
+	void SetCustomizableObject(UCustomizableObject* InCustomizableObject);
 	
 	bool GetBuildParameterRelevancy() const;
 	
diff --git a/Engine/Plugins/Experimental/Mutable/Source/CustomizableObject/Private/MuCO/CustomizableObjectInstance.cpp b/Engine/Plugins/Experimental/Mutable/Source/CustomizableObject/Private/MuCO/CustomizableObjectInstance.cpp
--- a/Engine/Plugins/Experimental/Mutable/Source/CustomizableObject/Private/MuCO/CustomizableObjectInstance.cpp
+++ b/Engine/Plugins/Experimental/Mutable/Source/CustomizableObject/Private/MuCO/CustomizableObjectInstance.cpp
@@ -650,7 +650,7 @@ void UCustomizableObjectInstance::SetObject(UCustomizableObject* InObject)
 	PrivateData->BindObjectDelegates(GetCustomizableObject(), InObject);
 #endif
 
-	Descriptor.SetCustomizableObject(*InObject);
+	Descriptor.SetCustomizableObject(InObject);
 	PrivateData->InitCustomizableObjectData(InObject);
 }
 
diff --git a/Engine/Plugins/Experimental/Mutable/Source/CustomizableObject/Private/MuCO/CustomizableObjectInstanceDescriptor.cpp b/Engine/Plugins/Experimental/Mutable/Source/CustomizableObject/Private/MuCO/CustomizableObjectInstanceDescriptor.cpp
--- a/Engine/Plugins/Experimental/Mutable/Source/CustomizableObject/Private/MuCO/CustomizableObjectInstanceDescriptor.cpp
+++ b/Engine/Plugins/Experimental/Mutable/Source/CustomizableObject/Private/MuCO/CustomizableObjectInstanceDescriptor.cpp
@@ -36,9 +36,9 @@ FString GetAvailableOptionsString(const UCustomizableObject& CustomizableObject,
 }
 
 
-FCustomizableObjectInstanceDescriptor::FCustomizableObjectInstanceDescriptor(UCustomizableObject &Object)
+FCustomizableObjectInstanceDescriptor::FCustomizableObjectInstanceDescriptor(UCustomizableObject& Object)
 {
-	SetCustomizableObject(Object);
+	SetCustomizableObject(&Object);
 }
 
 
@@ -423,9 +423,9 @@ UCustomizableObject* FCustomizableObjectInstanceDescriptor::GetCustomizableObjec
 }
 
 
-void FCustomizableObjectInstanceDescriptor::SetCustomizableObject(UCustomizableObject& InCustomizableObject)
+void FCustomizableObjectInstanceDescriptor::SetCustomizableObject(UCustomizableObject* InCustomizableObject)
 {
-	CustomizableObject = &InCustomizableObject;
+	CustomizableObject = InCustomizableObject;
 	ReloadParameters();
 }
 
diff --git a/Engine/Plugins/Experimental/Mutable/Source/CustomizableObjectEditor/Private/MuCOE/CustomizableInstanceDetails.cpp b/Engine/Plugins/Experimental/Mutable/Source/CustomizableObjectEditor/Private/MuCOE/CustomizableInstanceDetails.cpp
--- a/Engine/Plugins/Experimental/Mutable/Source/CustomizableObjectEditor/Private/MuCOE/CustomizableInstanceDetails.cpp
+++ b/Engine/Plugins/Experimental/Mutable/Source/CustomizableObjectEditor/Private/MuCOE/CustomizableInstanceDetails.cpp
@@ -1978,6 +1978,8 @@ FReply FCustomizableInstanceDetails::OnPasteAllParameters()
 	FCustomizableObjectInstanceDescriptor& Descriptor = CustomInstance->GetPrivate()->GetDescriptor();
 	const UScriptStruct* Struct = Descriptor.StaticStruct();
 
+	UCustomizableObject* CustomizableObject = Descriptor.GetCustomizableObject();
+	
 	const int32 MinLOD = Descriptor.GetMinLod();
 	const TArray<uint16> RequestedLODLevels = Descriptor.GetRequestedLODLevels();
 	
@@ -1986,6 +1988,8 @@ FReply FCustomizableInstanceDetails::OnPasteAllParameters()
 
 	if (Struct->ImportText(*ClipText, &Descriptor, nullptr, 0, GLog, GetPathNameSafe(Struct)))
 	{
+		Descriptor.SetCustomizableObject(CustomizableObject);
+		
 		// Keep current LOD
 		Descriptor.SetMinLod(MinLOD);
 		Descriptor.SetRequestedLODLevels(RequestedLODLevels);
