[Mutable] Fix clothing physics mesh not pointing to the correct data when multiple clothing assets have the same GUID.

#rnx

diff --git a/Engine/Plugins/Experimental/Mutable/Source/CustomizableObjectEditor/Private/MuCOE/GenerateMutableSource/GenerateMutableSourceMesh.cpp b/Engine/Plugins/Experimental/Mutable/Source/CustomizableObjectEditor/Private/MuCOE/GenerateMutableSource/GenerateMutableSourceMesh.cpp
--- a/Engine/Plugins/Experimental/Mutable/Source/CustomizableObjectEditor/Private/MuCOE/GenerateMutableSource/GenerateMutableSourceMesh.cpp
+++ b/Engine/Plugins/Experimental/Mutable/Source/CustomizableObjectEditor/Private/MuCOE/GenerateMutableSource/GenerateMutableSourceMesh.cpp
@@ -1405,7 +1405,78 @@ mu::MeshPtr ConvertSkeletalMeshToMutable(const USkeletalMesh* InSkeletalMesh, co
 			}
 
 			int32 FoundIndex = GenerationContext.ClothingAssetsData.IndexOfByPredicate(
-				[AssetGuid = ClothingAssetBase->GetAssetGuid()](const FCustomizableObjectClothingAssetData& Asset){ return Asset.OriginalAssetGuid == AssetGuid; });
+				[SourceAsset = Asset](const FCustomizableObjectClothingAssetData& Asset)
+				{ 
+					bool bIsSameAsset = 
+							Asset.OriginalAssetGuid == SourceAsset->GetAssetGuid() && 
+						   	Asset.Name == SourceAsset->GetFName() && 
+						   	Asset.ReferenceBoneIndex == SourceAsset->ReferenceBoneIndex &&
+						   	Asset.UsedBoneIndices == SourceAsset->UsedBoneIndices &&
+						   	Asset.UsedBoneNames == SourceAsset->UsedBoneNames &&
+						   	Asset.LodMap == SourceAsset->LodMap &&
+							Asset.LodData.Num() == SourceAsset->LodData.Num();
+					
+					const int32 LodDataNum = Asset.LodData.Num(); 
+					for (int32 LodDataIndex = LodDataNum - 1; LodDataIndex >= 0 && bIsSameAsset; --LodDataIndex)
+					{
+						const FClothPhysicalMeshData& DataA = Asset.LodData[LodDataIndex].PhysicalMeshData; 
+						const FClothPhysicalMeshData& DataB = SourceAsset->LodData[LodDataIndex].PhysicalMeshData; 
+					
+						bIsSameAsset = bIsSameAsset && DataA.NumFixedVerts == DataB.NumFixedVerts;
+						bIsSameAsset = bIsSameAsset && DataA.MaxBoneWeights == DataB.MaxBoneWeights;
+						bIsSameAsset = bIsSameAsset && DataA.Vertices == DataB.Vertices;
+						bIsSameAsset = bIsSameAsset && DataA.Normals == DataB.Normals;
+						bIsSameAsset = bIsSameAsset && DataA.Indices == DataB.Indices;
+						bIsSameAsset = bIsSameAsset && DataA.InverseMasses == DataB.InverseMasses;
+						bIsSameAsset = bIsSameAsset && DataA.EuclideanTethers.Tethers == DataB.EuclideanTethers.Tethers;
+						bIsSameAsset = bIsSameAsset && DataA.GeodesicTethers.Tethers == DataB.GeodesicTethers.Tethers;
+						
+						bIsSameAsset = bIsSameAsset && DataA.WeightMaps.Num() == DataB.WeightMaps.Num();
+						bIsSameAsset = bIsSameAsset && DataA.SelfCollisionVertexSet.Num() == DataB.SelfCollisionVertexSet.Num();
+						bIsSameAsset = bIsSameAsset && DataA.BoneData.Num() == DataB.BoneData.Num();
+						
+						static_assert(std::is_trivially_copyable_v<FClothVertBoneData>);
+						// Assume the FClothVertBoneData does not have any padding. In case there was padding, same assets should
+						// have the same unset memory so false negatives can only happen with different assets that have the 
+						// same data. This reasoning relays on the fact that the data buffers have been copied byte for byte using 
+						// Memcpy or similar.
+						bIsSameAsset = bIsSameAsset && FMemory::Memcmp(
+								DataA.BoneData.GetData(), DataB.BoneData.GetData(), DataA.BoneData.Num()*sizeof(FClothVertBoneData)) == 0;
+
+						if (bIsSameAsset)
+						{
+							for (const TPair<uint32, FPointWeightMap>& WeightMap : DataA.WeightMaps)
+							{
+								const FPointWeightMap* FoundWeightMap = DataB.WeightMaps.Find(WeightMap.Key);
+								
+								if (!FoundWeightMap)
+								{
+									bIsSameAsset = false;
+									break;
+								}
+
+								if (FoundWeightMap->Values != WeightMap.Value.Values)
+								{
+									bIsSameAsset = false;
+									break;
+								}
+							}
+						}
+
+						if (bIsSameAsset)
+						{
+							TArray<int32> DataASelfCollisionVertexArray = DataA.SelfCollisionVertexSet.Array();
+							Algo::Sort(DataASelfCollisionVertexArray);
+
+							TArray<int32> DataBSelfCollisionVertexArray = DataB.SelfCollisionVertexSet.Array();
+							Algo::Sort(DataBSelfCollisionVertexArray);
+
+							bIsSameAsset = DataASelfCollisionVertexArray == DataBSelfCollisionVertexArray;
+						}
+					}
+
+					return bIsSameAsset;
+				});
 			
 			if (FoundIndex != INDEX_NONE)
 			{
