[mutable] Remove leftover code to support older engine versions.
Remove ensure that may be caused by user data, but there is an early-out anyway.
#rnx

diff --git a/Engine/Plugins/Experimental/Mutable/Source/CustomizableObjectEditor/Private/MuCOE/Nodes/CustomizableObjectNodeMaterial.cpp b/Engine/Plugins/Experimental/Mutable/Source/CustomizableObjectEditor/Private/MuCOE/Nodes/CustomizableObjectNodeMaterial.cpp
--- a/Engine/Plugins/Experimental/Mutable/Source/CustomizableObjectEditor/Private/MuCOE/Nodes/CustomizableObjectNodeMaterial.cpp
+++ b/Engine/Plugins/Experimental/Mutable/Source/CustomizableObjectEditor/Private/MuCOE/Nodes/CustomizableObjectNodeMaterial.cpp
@@ -907,11 +907,7 @@ int32 UCustomizableObjectNodeMaterial::GetNumParameters(const EMaterialParameter
 {
 	if (Material)
 	{
-#if ENGINE_MAJOR_VERSION==5 && ENGINE_MINOR_VERSION>=1
 		return Material->GetCachedExpressionData().GetParameterTypeEntry(Type).ParameterInfoSet.Num();
-#else
-		return Material->GetCachedExpressionData().Parameters.GetNumParameters(Type);
-#endif
 	}
 	else
 	{
@@ -924,7 +920,6 @@ FGuid UCustomizableObjectNodeMaterial::GetParameterId(const EMaterialParameterTy
 {
 	const FMaterialCachedExpressionData& Data = Material->GetCachedExpressionData();
 
-#if ENGINE_MAJOR_VERSION==5 && ENGINE_MINOR_VERSION>=1
 	if (Data.EditorOnlyData)
 	{
 		if (Data.EditorOnlyData->EditorEntries[(int32)Type].EditorInfo.Num() != 0)
@@ -932,10 +927,6 @@ FGuid UCustomizableObjectNodeMaterial::GetParameterId(const EMaterialParameterTy
 			return Data.EditorOnlyData->EditorEntries[(int32)Type].EditorInfo[ParameterIndex].ExpressionGuid;
 		}
 	}
-#else
-	const FMaterialCachedParameterEntry& Entry = Data.Parameters.GetParameterTypeEntry(Type);
-	return Entry.EditorInfo[ParameterIndex].ExpressionGuid;
-#endif
 	
 	return FGuid();
 }
@@ -948,8 +939,6 @@ FName UCustomizableObjectNodeMaterial::GetParameterName(const EMaterialParameter
 		return FName();
 	}
 
-#if ENGINE_MAJOR_VERSION==5 && ENGINE_MINOR_VERSION>=1
-
 	const FMaterialCachedParameterEntry& Entry = Material->GetCachedExpressionData().GetParameterTypeEntry(Type);
 
 	for (TSet<FMaterialParameterInfo>::TConstIterator It(Entry.ParameterInfoSet); It; ++It)
@@ -965,16 +954,6 @@ FName UCustomizableObjectNodeMaterial::GetParameterName(const EMaterialParameter
 	// The parameter should exist
 	check(false);
 
-#else
-
-	TArray<FGuid> ParameterIds;
-	TArray<FMaterialParameterInfo> ParameterInfo;
-	Material->GetAllParameterInfoOfType(Type, ParameterInfo, ParameterIds);
-
-	return ParameterInfo.IsValidIndex(ParameterIndex) ? ParameterInfo[ParameterIndex].Name : FName();
-
-#endif
-
 	return FName();
 }
 
@@ -986,8 +965,6 @@ int32 UCustomizableObjectNodeMaterial::GetParameterLayerIndex(const EMaterialPar
 		return -1;
 	}
 
-#if ENGINE_MAJOR_VERSION==5 && ENGINE_MINOR_VERSION>=1
-
 	const FMaterialCachedParameterEntry& Entry = Material->GetCachedExpressionData().GetParameterTypeEntry(Type);
 
 	for (TSet<FMaterialParameterInfo>::TConstIterator It(Entry.ParameterInfoSet); It; ++It)
@@ -1003,16 +980,6 @@ int32 UCustomizableObjectNodeMaterial::GetParameterLayerIndex(const EMaterialPar
 	// The parameter should exist
 	check(false);
 
-#else
-
-	TArray<FGuid> ParameterIds;
-	TArray<FMaterialParameterInfo> ParameterInfo;
-	Material->GetAllParameterInfoOfType(Type, ParameterInfo, ParameterIds);
-
-	return ParameterInfo.IsValidIndex(ParameterIndex) ? ParameterInfo[ParameterIndex].Index : -1;
-
-#endif
-
 	return -1;
 }
 
@@ -1030,28 +997,18 @@ FText UCustomizableObjectNodeMaterial::GetParameterLayerName(const EMaterialPara
 	FMaterialLayersFunctions LayersValue;
 	Material->GetMaterialLayers(LayersValue);
 
-#if ENGINE_MAJOR_VERSION==5 && ENGINE_MINOR_VERSION>=1
 
 	return LayersValue.EditorOnly.LayerNames.IsValidIndex(LayerIndex) ? LayersValue.EditorOnly.LayerNames[LayerIndex] : FText();
-
-#else
-
-	return LayersValue.LayerNames.IsValidIndex(LayerIndex) ? LayersValue.LayerNames[LayerIndex] : FText();
-
-#endif
 }
 
 
 bool UCustomizableObjectNodeMaterial::HasParameter(const FGuid& ParameterId) const
 {
-	ensure(Material);
 	if (!Material)
 	{
 		return false;
 	}
 
-#if ENGINE_MAJOR_VERSION==5 && ENGINE_MINOR_VERSION>=1
-
 	for (const EMaterialParameterType Type : ParameterTypes)
 	{
 		const FMaterialCachedExpressionData& Data = Material->GetCachedExpressionData();
@@ -1074,32 +1031,6 @@ bool UCustomizableObjectNodeMaterial::HasParameter(const FGuid& ParameterId) con
 		}
 	}
 
-#else
-
-	const FMaterialCachedExpressionData& Data = Material->GetCachedExpressionData();
-	for (const FMaterialCachedParameterEntry& Entry : Data.Parameters.RuntimeEntries)
-	{
-		for (const FMaterialCachedParameterEditorInfo& Info : Entry.EditorInfo)
-		{
-			if (Info.ExpressionGuid == ParameterId)
-			{
-				return true;
-			}
-		}
-	}
-	for (const FMaterialCachedParameterEntry& Entry : Data.Parameters.EditorOnlyEntries)
-	{
-		for (const FMaterialCachedParameterEditorInfo& Info : Entry.EditorInfo)
-		{
-			if (Info.ExpressionGuid == ParameterId)
-			{
-				return true;
-			}
-		}
-	}
-
-#endif
-
 	return false;
 }
 
