[Mutable] Lowered texture conversion performance penalty warnings to info.

#rnx

diff --git a/Engine/Plugins/Experimental/Mutable/Source/CustomizableObject/Private/MuCO/CustomizableObjectInstance.cpp b/Engine/Plugins/Experimental/Mutable/Source/CustomizableObject/Private/MuCO/CustomizableObjectInstance.cpp
--- a/Engine/Plugins/Experimental/Mutable/Source/CustomizableObject/Private/MuCO/CustomizableObjectInstance.cpp
+++ b/Engine/Plugins/Experimental/Mutable/Source/CustomizableObject/Private/MuCO/CustomizableObjectInstance.cpp
@@ -3167,13 +3167,13 @@ FTexturePlatformData* MutableCreateImagePlatformData(mu::Ptr<const mu::Image> Mu
 	case mu::EImageFormat::IF_RGB_UBYTE:
 		// performance penalty. can happen in states that remove compression.
 		PlatformFormat = PF_R8G8B8A8;	
-		UE_LOG(LogMutable, Warning, TEXT("Building instance: a texture was generated in a format not supported by the hardware (RGB), this results in an additional conversion, so a performance penalty."));
+		UE_LOG(LogMutable, Display, TEXT("Building instance: a texture was generated in a format not supported by the hardware (RGB), this results in an additional conversion, so a performance penalty."));
 		break; 
 
 	case mu::EImageFormat::IF_BGRA_UBYTE:			
 		// performance penalty. can happen with texture parameter images.
 		PlatformFormat = PF_R8G8B8A8;	
-		UE_LOG(LogMutable, Warning, TEXT("Building instance: a texture was generated in a format not supported by the hardware (BGRA), this results in an additional conversion, so a performance penalty."));
+		UE_LOG(LogMutable, Display, TEXT("Building instance: a texture was generated in a format not supported by the hardware (BGRA), this results in an additional conversion, so a performance penalty."));
 		break;
 
 	// Good cases:
