[mutable] Fixed crashes when accessing a CustomizableObjectInstace without setting its parent CustomizableObject first. This situation could easily be replicated when using the Mutable API from BP.
#rnx
#tests PIE

diff --git a/Engine/Plugins/Experimental/Mutable/Source/CustomizableObject/Internal/MuCO/CustomizableObjectInstanceDescriptor.h b/Engine/Plugins/Experimental/Mutable/Source/CustomizableObject/Internal/MuCO/CustomizableObjectInstanceDescriptor.h
--- a/Engine/Plugins/Experimental/Mutable/Source/CustomizableObject/Internal/MuCO/CustomizableObjectInstanceDescriptor.h
+++ b/Engine/Plugins/Experimental/Mutable/Source/CustomizableObject/Internal/MuCO/CustomizableObjectInstanceDescriptor.h
@@ -58,6 +58,7 @@ struct CUSTOMIZABLEOBJECT_API FCustomizableObjectInstanceDescriptor
 	 * Multilayer Projectors not supported */
 	void LoadDescriptor(FArchive &Ar);
 
+	// Could return nullptr in some rare situations, so check first
 	UCustomizableObject* GetCustomizableObject() const;
 
 	void SetCustomizableObject(UCustomizableObject* InCustomizableObject);
diff --git a/Engine/Plugins/Experimental/Mutable/Source/CustomizableObject/Private/MuCO/CustomizableObjectInstance.cpp b/Engine/Plugins/Experimental/Mutable/Source/CustomizableObject/Private/MuCO/CustomizableObjectInstance.cpp
--- a/Engine/Plugins/Experimental/Mutable/Source/CustomizableObject/Private/MuCO/CustomizableObjectInstance.cpp
+++ b/Engine/Plugins/Experimental/Mutable/Source/CustomizableObject/Private/MuCO/CustomizableObjectInstance.cpp
@@ -714,8 +714,15 @@ bool UCustomizableObjectInstance::IsParameterRelevant(int32 ParameterIndex) cons
 
 bool UCustomizableObjectInstance::IsParameterRelevant(const FString& ParamName) const
 {
+	UCustomizableObject* CustomizableObject = GetCustomizableObject();
+
+	if (!CustomizableObject)
+	{
+		return false;
+	}
+
 	// This should have been precalculated in the last update if the appropriate flag in the instance was set.
-	int32 ParameterIndexInObject = GetCustomizableObject()->FindParameter(ParamName);
+	int32 ParameterIndexInObject = CustomizableObject->FindParameter(ParamName);
 	return GetPrivate()->RelevantParameters.Contains(ParameterIndexInObject);
 }
 
@@ -1860,7 +1867,11 @@ bool UCustomizableInstancePrivate::DoComponentsNeedUpdate(UCustomizableObjectIns
 	MUTABLE_CPUPROFILER_SCOPE(UCustomizableInstancePrivate::DoComponentsNeedUpdate);
 
 	UCustomizableObject* CustomizableObject = Public->GetCustomizableObject();
-	check(CustomizableObject);
+	
+	if (!CustomizableObject)
+	{
+		return false;
+	}
 
 	const int32 NumComponents = CustomizableObject->GetComponentCount();
 
@@ -1963,13 +1974,24 @@ bool UCustomizableInstancePrivate::UpdateSkeletalMesh_PostBeginUpdate0(UCustomiz
 
 	bool bUpdateMeshes = DoComponentsNeedUpdate(Public, OperationData, bHasInvalidMesh);
 
+	UCustomizableObject* CustomizableObject = Public->GetCustomizableObject();
+
+	if (!CustomizableObject)
+	{
+		UE_LOG(LogMutable, Warning, TEXT("Failed to generate SkeletalMesh for CO Instance %s. It does not have a CO."), *Public->GetName());
+		
+		InvalidateGeneratedData();
+		OperationData->UpdateResult = EUpdateResult::Error;
+
+		return false;
+	}
+
 	// We can not handle empty meshes, clear any generated mesh and return
 	if (bHasInvalidMesh)
 	{
-		UE_LOG(LogMutable, Warning, TEXT("Failed to generate SkeletalMesh for CO Instance %s. CO [%s]"), *Public->GetName(), *GetNameSafe(Public->GetCustomizableObject()));
+		UE_LOG(LogMutable, Warning, TEXT("Failed to generate SkeletalMesh for CO Instance %s. CO [%s]"), *Public->GetName(), *GetNameSafe(CustomizableObject));
 
 		InvalidateGeneratedData();
-
 		OperationData->UpdateResult = EUpdateResult::Error;
 
 		return false;
@@ -1987,9 +2009,6 @@ bool UCustomizableInstancePrivate::UpdateSkeletalMesh_PostBeginUpdate0(UCustomiz
 
 	bool bSuccess = true;
 
-	UCustomizableObject* CustomizableObject = Public->GetCustomizableObject();
-	check(CustomizableObject);
-
 	const FModelResources& ModelResources = CustomizableObject->GetPrivate()->GetModelResources();
 
 	// Initialize the maximum number of SkeletalMeshes we could possibly have. 
@@ -2056,7 +2075,7 @@ bool UCustomizableInstancePrivate::UpdateSkeletalMesh_PostBeginUpdate0(UCustomiz
 
 		if (OperationData->bStreamMeshLODs)
 		{
-			SkeletalMeshes[Component.Id] = UCustomizableObjectSkeletalMesh::CreateSkeletalMesh(OperationData, *Public, Component.Id);
+			SkeletalMeshes[Component.Id] = UCustomizableObjectSkeletalMesh::CreateSkeletalMesh(OperationData, *Public, *CustomizableObject, Component.Id);
 		}
 		else
 		{
@@ -3857,13 +3876,11 @@ void UCustomizableInstancePrivate::BuildOrCopyClothingData(const TSharedRef<FUpd
 {
 	MUTABLE_CPUPROFILER_SCOPE(UCustomizableInstancePrivate::BuildOrCopyClothingData);
 
-	if (!SkeletalMesh)
-	{
-		return;
-	}
-
+	UCustomizableObject* CustomizableObject = CustomizableObjectInstance->GetCustomizableObject();
+	// It must be not null as it's checked in the calling function
+	check(CustomizableObject);
 
-	const FModelResources& ModelResources = CustomizableObjectInstance->GetCustomizableObject()->GetPrivate()->GetModelResources();
+	const FModelResources& ModelResources = CustomizableObject->GetPrivate()->GetModelResources();
 	const TArray<FCustomizableObjectClothingAssetData>& ClothingAssetsData = ModelResources.ClothingAssetsData;
 	const TArray<FCustomizableObjectClothConfigData>& ClothSharedConfigsData = ModelResources.ClothSharedConfigsData;
 
@@ -4827,7 +4844,10 @@ bool UCustomizableInstancePrivate::BuildOrCopyRenderData(const TSharedRef<FUpdat
 	check(RenderData);
 
 	UCustomizableObject* CustomizableObject = Public->GetCustomizableObject();
+
+	// It must be not null as it's checked in the calling function
 	check(CustomizableObject);
+
 	const FModelResources& ModelResources = CustomizableObject->GetPrivate()->GetModelResources();
 
 	const int32 FirstGeneratedLOD = FMath::Max((int32)OperationData->GetRequestedLODs()[ComponentIndex], OperationData->GetMinLOD());
@@ -6665,8 +6685,15 @@ void UCustomizableObjectInstance::SetRequestedLODs(int32 InMinLOD, int32 , const
 		return;
 	}
 
+	UCustomizableObject* CustomizableObject = GetCustomizableObject();
+
+	if (!CustomizableObject)
+	{
+		return;
+	}
+
 	if (CVarPreserveUserLODsOnFirstGeneration.GetValueOnGameThread() &&
-		GetCustomizableObject()->bPreserveUserLODsOnFirstGeneration &&
+		CustomizableObject->bPreserveUserLODsOnFirstGeneration &&
 		GetPrivate()->SkeletalMeshStatus != ESkeletalMeshStatus::Success)
 	{
 		return;
@@ -6683,7 +6710,7 @@ void UCustomizableObjectInstance::SetRequestedLODs(int32 InMinLOD, int32 , const
 	FMutableUpdateCandidate MutableUpdateCandidate(this);
 
 	// Clamp Min LOD
-	const int32 MinLODIdx = GetCustomizableObject()->GetPrivate()->GetMinLODIndex();
+	const int32 MinLODIdx = CustomizableObject->GetPrivate()->GetMinLODIndex();
 	const int32 MaxLODIdx = PrivateData->NumLODsAvailable - 1;
 	InMinLOD = FMath::Clamp(InMinLOD, MinLODIdx, MaxLODIdx);
 
diff --git a/Engine/Plugins/Experimental/Mutable/Source/CustomizableObject/Private/MuCO/CustomizableObjectInstanceDescriptor.cpp b/Engine/Plugins/Experimental/Mutable/Source/CustomizableObject/Private/MuCO/CustomizableObjectInstanceDescriptor.cpp
--- a/Engine/Plugins/Experimental/Mutable/Source/CustomizableObject/Private/MuCO/CustomizableObjectInstanceDescriptor.cpp
+++ b/Engine/Plugins/Experimental/Mutable/Source/CustomizableObject/Private/MuCO/CustomizableObjectInstanceDescriptor.cpp
@@ -17,6 +17,12 @@
 #include UE_INLINE_GENERATED_CPP_BY_NAME(CustomizableObjectInstanceDescriptor)
 
 
+void CustomizableObjectNullErrorMessage()
+{
+	UE_LOG(LogMutable, Error,
+		TEXT("Tried to perform actions on a CustomizableObjectInstance with no CustomizableObject set. Please set the CustomizableObject of the Instance before doing anything with it."));
+}
+
 FString GetAvailableOptionsString(const UCustomizableObject& CustomizableObject, const int32 ParameterIndexInObject)
 {
 	FString OptionsString;
@@ -44,7 +50,11 @@ FCustomizableObjectInstanceDescriptor::FCustomizableObjectInstanceDescriptor(UCu
 
 void FCustomizableObjectInstanceDescriptor::SaveDescriptor(FArchive& Ar, bool bUseCompactDescriptor)
 {
-	check(CustomizableObject);
+	if (!CustomizableObject)
+	{
+		CustomizableObjectNullErrorMessage();
+		return;
+	}
 
 	Ar << bUseCompactDescriptor;
 
@@ -226,7 +236,11 @@ void FCustomizableObjectInstanceDescriptor::SaveDescriptor(FArchive& Ar, bool bU
 
 void FCustomizableObjectInstanceDescriptor::LoadDescriptor(FArchive& Ar)
 {
-	check(CustomizableObject);
+	if (!CustomizableObject)
+	{
+		CustomizableObjectNullErrorMessage();
+		return;
+	}
 
 	bool bUseCompactDescriptor;
 	Ar << bUseCompactDescriptor;
@@ -450,6 +464,7 @@ mu::ParametersPtr FCustomizableObjectInstanceDescriptor::GetParameters() const
 {
 	if (!CustomizableObject)
 	{
+		CustomizableObjectNullErrorMessage();
 		return nullptr;
 	}
 
@@ -673,7 +688,13 @@ void FCustomizableObjectInstanceDescriptor::ReloadParameters()
 		return;
 	}
 
-	if (!CustomizableObject || !CustomizableObject->IsCompiled())
+	if (!CustomizableObject)
+	{
+		CustomizableObjectNullErrorMessage();
+		return;
+	}
+
+	if (!CustomizableObject->IsCompiled())
 	{
 		return;
 	}
@@ -1215,7 +1236,11 @@ void LogParameterNotFoundWarning(const FString& ParameterName, const int32 Objec
 
 const FString& FCustomizableObjectInstanceDescriptor::GetIntParameterSelectedOption(const FString& ParamName, const int32 RangeIndex) const
 {
-	check(CustomizableObject);
+	if (!CustomizableObject)
+	{
+		CustomizableObjectNullErrorMessage();
+		return FCustomizableObjectBoolParameterValue::DEFAULT_PARAMETER_VALUE_NAME;
+	}
 	
 	const int32 ParameterIndexInObject = CustomizableObject->FindParameter(ParamName);
 	const int32 ParameterIndexInInstance = FindTypedParameterIndex(ParamName, EMutableParameterType::Int);
@@ -1246,7 +1271,12 @@ const FString& FCustomizableObjectInstanceDescriptor::GetIntParameterSelectedOpt
 
 void FCustomizableObjectInstanceDescriptor::SetIntParameterSelectedOption(const int32 ParameterIndexInInstance, const FString& SelectedOption, const int32 RangeIndex)
 {
-	check(CustomizableObject);
+	if (!CustomizableObject)
+	{
+		CustomizableObjectNullErrorMessage();
+		return;
+	}
+
 	RETURN_ON_UNCOMPILED_CO(CustomizableObject, TEXT("Error: Cannot set Int parameter "));
 
 	const int32 ParameterIndexInObject = IntParameters.IsValidIndex(ParameterIndexInInstance) ? CustomizableObject->FindParameter(IntParameters[ParameterIndexInInstance].ParameterName) : INDEX_NONE; //-V781
@@ -1300,7 +1330,12 @@ void FCustomizableObjectInstanceDescriptor::SetIntParameterSelectedOption(const
 
 void FCustomizableObjectInstanceDescriptor::SetIntParameterSelectedOption(const FString& ParamName, const FString& SelectedOptionName, const int32 RangeIndex)
 {
-	check(CustomizableObject);
+	if (!CustomizableObject)
+	{
+		CustomizableObjectNullErrorMessage();
+		return;
+	}
+
 	RETURN_ON_UNCOMPILED_CO(CustomizableObject, TEXT("Error: Cannot set Int parameter "));
 
 	const int32 ParamIndexInInstance = FindTypedParameterIndex(ParamName, EMutableParameterType::Int);
@@ -1317,7 +1352,11 @@ void FCustomizableObjectInstanceDescriptor::SetIntParameterSelectedOption(const
 
 float FCustomizableObjectInstanceDescriptor::GetFloatParameterSelectedOption(const FString& FloatParamName, const int32 RangeIndex) const
 {
-	check(CustomizableObject);
+	if (!CustomizableObject)
+	{
+		CustomizableObjectNullErrorMessage();
+		return FCustomizableObjectFloatParameterValue::DEFAULT_PARAMETER_VALUE;
+	}
 
 	const int32 ParameterIndexInObject = CustomizableObject->FindParameter(FloatParamName);
 	const int32 FloatParamIndex = FindTypedParameterIndex(FloatParamName, EMutableParameterType::Float);
@@ -1342,13 +1381,18 @@ float FCustomizableObjectInstanceDescriptor::GetFloatParameterSelectedOption(con
 
 	LogParameterNotFoundWarning(FloatParamName, ParameterIndexInObject, FloatParamIndex, CustomizableObject, __FUNCTION__);
 	
-	return 	FCustomizableObjectFloatParameterValue::DEFAULT_PARAMETER_VALUE; 
+	return FCustomizableObjectFloatParameterValue::DEFAULT_PARAMETER_VALUE; 
 }
 
 
 void FCustomizableObjectInstanceDescriptor::SetFloatParameterSelectedOption(const FString& FloatParamName, const float FloatValue, const int32 RangeIndex)
 {
-	check(CustomizableObject);
+	if (!CustomizableObject)
+	{
+		CustomizableObjectNullErrorMessage();
+		return;
+	}
+
 	RETURN_ON_UNCOMPILED_CO(CustomizableObject, TEXT("Error: Cannot set Int parameter "));
 
 	const int32 ParameterIndexInObject = CustomizableObject->FindParameter(FloatParamName);
@@ -1386,7 +1430,11 @@ void FCustomizableObjectInstanceDescriptor::SetFloatParameterSelectedOption(cons
 
 FName FCustomizableObjectInstanceDescriptor::GetTextureParameterSelectedOption(const FString& TextureParamName, const int32 RangeIndex) const
 {
-	check(CustomizableObject);
+	if (!CustomizableObject)
+	{
+		CustomizableObjectNullErrorMessage();
+		return FName();
+	}
 
 	const int32 ParameterIndexInObject = CustomizableObject->FindParameter(TextureParamName);
 	const int32 ParameterIndexInInstance = FindTypedParameterIndex(TextureParamName, EMutableParameterType::Texture);
@@ -1417,7 +1465,12 @@ FName FCustomizableObjectInstanceDescriptor::GetTextureParameterSelectedOption(c
 
 void FCustomizableObjectInstanceDescriptor::SetTextureParameterSelectedOption(const FString& TextureParamName, const FString& TextureValue, const int32 RangeIndex)
 {
-	check(CustomizableObject);
+	if (!CustomizableObject)
+	{
+		CustomizableObjectNullErrorMessage();
+		return;
+	}
+
 	RETURN_ON_UNCOMPILED_CO(CustomizableObject, TEXT("Error: Cannot set Int parameter "));
 
 	const int32 ParameterIndexInObject = CustomizableObject->FindParameter(TextureParamName);
@@ -1454,7 +1507,11 @@ void FCustomizableObjectInstanceDescriptor::SetTextureParameterSelectedOption(co
 
 FLinearColor FCustomizableObjectInstanceDescriptor::GetColorParameterSelectedOption(const FString& ColorParamName) const
 {
-	check(CustomizableObject);
+	if (!CustomizableObject)
+	{
+		CustomizableObjectNullErrorMessage();
+		return FCustomizableObjectVectorParameterValue::DEFAULT_PARAMETER_VALUE;
+	}
 
 	const int32 ParameterIndexInObject = CustomizableObject->FindParameter(ColorParamName);
 	const int32 ColorParamIndex = FindTypedParameterIndex(ColorParamName, EMutableParameterType::Color);
@@ -1471,7 +1528,12 @@ FLinearColor FCustomizableObjectInstanceDescriptor::GetColorParameterSelectedOpt
 
 void FCustomizableObjectInstanceDescriptor::SetColorParameterSelectedOption(const FString& ColorParamName, const FLinearColor& ColorValue)
 {
-	check(CustomizableObject);
+	if (!CustomizableObject)
+	{
+		CustomizableObjectNullErrorMessage();
+		return;
+	}
+
 	RETURN_ON_UNCOMPILED_CO(CustomizableObject, TEXT("Error: Cannot set Int parameter "));
 
 	SetVectorParameterSelectedOption(ColorParamName, ColorValue);
@@ -1480,7 +1542,11 @@ void FCustomizableObjectInstanceDescriptor::SetColorParameterSelectedOption(cons
 
 bool FCustomizableObjectInstanceDescriptor::GetBoolParameterSelectedOption(const FString& BoolParamName) const
 {
-	check(CustomizableObject);
+	if (!CustomizableObject)
+	{
+		CustomizableObjectNullErrorMessage();
+		return FCustomizableObjectBoolParameterValue::DEFAULT_PARAMETER_VALUE;
+	}
 
 	const int32 ParameterIndexInObject = CustomizableObject->FindParameter(BoolParamName);
 	const int32 BoolParamIndex = FindTypedParameterIndex(BoolParamName, EMutableParameterType::Bool);
@@ -1497,7 +1563,12 @@ bool FCustomizableObjectInstanceDescriptor::GetBoolParameterSelectedOption(const
 
 void FCustomizableObjectInstanceDescriptor::SetBoolParameterSelectedOption(const FString& BoolParamName, const bool BoolValue)
 {
-	check(CustomizableObject);
+	if (!CustomizableObject)
+	{
+		CustomizableObjectNullErrorMessage();
+		return;
+	}
+
 	RETURN_ON_UNCOMPILED_CO(CustomizableObject, TEXT("Error: Cannot set Int parameter "));
 
 	const int32 ParameterIndexInObject = CustomizableObject->FindParameter(BoolParamName);
@@ -1516,7 +1587,12 @@ void FCustomizableObjectInstanceDescriptor::SetBoolParameterSelectedOption(const
 
 void FCustomizableObjectInstanceDescriptor::SetVectorParameterSelectedOption(const FString& VectorParamName, const FLinearColor& VectorValue)
 {
-	check(CustomizableObject);
+	if (!CustomizableObject)
+	{
+		CustomizableObjectNullErrorMessage();
+		return;
+	}
+
 	RETURN_ON_UNCOMPILED_CO(CustomizableObject, TEXT("Error: Cannot set Int parameter "));
 
 	const int32 ParameterIndexInObject = CustomizableObject->FindParameter(VectorParamName);
@@ -1538,7 +1614,12 @@ void FCustomizableObjectInstanceDescriptor::SetProjectorValue(const FString& Pro
 	const float Angle,
 	const int32 RangeIndex)
 {
-	check(CustomizableObject);
+	if (!CustomizableObject)
+	{
+		CustomizableObjectNullErrorMessage();
+		return;
+	}
+
 	RETURN_ON_UNCOMPILED_CO(CustomizableObject, TEXT("Error: Cannot set Projector parameter "))
 
 	const int32 ParameterIndexInObject = CustomizableObject->FindParameter(ProjectorParamName);
@@ -1660,7 +1741,11 @@ void FCustomizableObjectInstanceDescriptor::GetProjectorValueF(const FString& Pr
 	float& OutAngle, ECustomizableObjectProjectorType& OutType,
 	const int32 RangeIndex) const
 {
-	check(CustomizableObject);
+	if (!CustomizableObject)
+	{
+		CustomizableObjectNullErrorMessage();
+		return;
+	}
 
 	const int32 ParameterIndexInObject = CustomizableObject->FindParameter(ProjectorParamName);
 	const int32 ParameterIndexInInstance = FindTypedParameterIndex(ProjectorParamName, EMutableParameterType::Projector);
@@ -1703,7 +1788,11 @@ void FCustomizableObjectInstanceDescriptor::GetProjectorValueF(const FString& Pr
 
 FVector FCustomizableObjectInstanceDescriptor::GetProjectorPosition(const FString& ParamName, const int32 RangeIndex) const
 {
-	check(CustomizableObject);
+	if (!CustomizableObject)
+	{
+		CustomizableObjectNullErrorMessage();
+		return FVector(-0.0, -0.0, -0.0);
+	}
 
 	const int32 ParameterIndexInObject = CustomizableObject->FindParameter(ParamName);
 	const int32 ParameterIndexInInstance = FindTypedParameterIndex(ParamName, EMutableParameterType::Projector);
@@ -1728,7 +1817,11 @@ FVector FCustomizableObjectInstanceDescriptor::GetProjectorPosition(const FStrin
 
 FVector FCustomizableObjectInstanceDescriptor::GetProjectorDirection(const FString& ParamName, const int32 RangeIndex) const
 {
-	check(CustomizableObject);
+	if (!CustomizableObject)
+	{
+		CustomizableObjectNullErrorMessage();
+		return FVector(-0.0, -0.0, -0.0);
+	}
 
 	const int32 ParameterIndexInObject = CustomizableObject->FindParameter(ParamName);
 	const int32 ParameterIndexInInstance = FindTypedParameterIndex(ParamName, EMutableParameterType::Projector);
@@ -1753,7 +1846,11 @@ FVector FCustomizableObjectInstanceDescriptor::GetProjectorDirection(const FStri
 
 FVector FCustomizableObjectInstanceDescriptor::GetProjectorUp(const FString& ParamName, const int32 RangeIndex) const
 {
-	check(CustomizableObject);
+	if (!CustomizableObject)
+	{
+		CustomizableObjectNullErrorMessage();
+		return FVector(-0.0, -0.0, -0.0);
+	}
 
 	const int32 ParameterIndexInObject = CustomizableObject->FindParameter(ParamName);
 	const int32 ParameterIndexInInstance = FindTypedParameterIndex(ParamName, EMutableParameterType::Projector);
@@ -1778,7 +1875,11 @@ FVector FCustomizableObjectInstanceDescriptor::GetProjectorUp(const FString& Par
 
 FVector FCustomizableObjectInstanceDescriptor::GetProjectorScale(const FString& ParamName, const int32 RangeIndex) const
 {
-	check(CustomizableObject);
+	if (!CustomizableObject)
+	{
+		CustomizableObjectNullErrorMessage();
+		return FVector(-0.0, -0.0, -0.0);
+	}
 
 	const int32 ParameterIndexInObject = CustomizableObject->FindParameter(ParamName);
 	const int32 ParameterIndexInInstance = FindTypedParameterIndex(ParamName, EMutableParameterType::Projector);
@@ -1803,7 +1904,11 @@ FVector FCustomizableObjectInstanceDescriptor::GetProjectorScale(const FString&
 
 float FCustomizableObjectInstanceDescriptor::GetProjectorAngle(const FString& ParamName, const int32 RangeIndex) const
 {
-	check(CustomizableObject);
+	if (!CustomizableObject)
+	{
+		CustomizableObjectNullErrorMessage();
+		return 0.0;
+	}
 
 	const int32 ParameterIndexInObject = CustomizableObject->FindParameter(ParamName);
 	const int32 ParameterIndexInInstance = FindTypedParameterIndex(ParamName, EMutableParameterType::Projector);
@@ -1828,7 +1933,11 @@ float FCustomizableObjectInstanceDescriptor::GetProjectorAngle(const FString& Pa
 
 ECustomizableObjectProjectorType FCustomizableObjectInstanceDescriptor::GetProjectorParameterType(const FString& ParamName, const int32 RangeIndex) const
 {
-	check(CustomizableObject);
+	if (!CustomizableObject)
+	{
+		CustomizableObjectNullErrorMessage();
+		return ECustomizableObjectProjectorType::Planar;
+	}
 
 	const int32 ParameterIndexInObject = CustomizableObject->FindParameter(ParamName);
 	const int32 ParameterIndexInInstance = FindTypedParameterIndex(ParamName, EMutableParameterType::Projector);
@@ -1881,7 +1990,11 @@ int32 FCustomizableObjectInstanceDescriptor::FindTypedParameterIndex(const FStri
 
 int32 FCustomizableObjectInstanceDescriptor::GetProjectorValueRange(const FString& ParamName) const
 {
-	check(CustomizableObject);
+	if (!CustomizableObject)
+	{
+		CustomizableObjectNullErrorMessage();
+		return -1;
+	}
 
 	const int32 ProjectorParamIndex = FindTypedParameterIndex(ParamName, EMutableParameterType::Projector);
 	if (ProjectorParamIndex < 0)
@@ -1894,7 +2007,11 @@ int32 FCustomizableObjectInstanceDescriptor::GetProjectorValueRange(const FStrin
 
 int32 FCustomizableObjectInstanceDescriptor::GetIntValueRange(const FString& ParamName) const
 {
-	check(CustomizableObject);
+	if (!CustomizableObject)
+	{
+		CustomizableObjectNullErrorMessage();
+		return -1;
+	}
 
 	const int32 IntParamIndex = FindTypedParameterIndex(ParamName, EMutableParameterType::Int);
 	if (IntParamIndex < 0)
@@ -1908,7 +2025,11 @@ int32 FCustomizableObjectInstanceDescriptor::GetIntValueRange(const FString& Par
 
 int32 FCustomizableObjectInstanceDescriptor::GetFloatValueRange(const FString& ParamName) const
 {
-	check(CustomizableObject);
+	if (!CustomizableObject)
+	{
+		CustomizableObjectNullErrorMessage();
+		return -1;
+	}
 
 	const int32 FloatParamIndex = FindTypedParameterIndex(ParamName, EMutableParameterType::Float);
 	if (FloatParamIndex < 0)
@@ -1922,7 +2043,11 @@ int32 FCustomizableObjectInstanceDescriptor::GetFloatValueRange(const FString& P
 
 int32 FCustomizableObjectInstanceDescriptor::GetTextureValueRange(const FString& ParamName) const
 {
-	check(CustomizableObject);
+	if (!CustomizableObject)
+	{
+		CustomizableObjectNullErrorMessage();
+		return -1;
+	}
 
 	const int32 TextureParamIndex = FindTypedParameterIndex(ParamName, EMutableParameterType::Texture);
 	if (TextureParamIndex < 0)
@@ -1936,7 +2061,11 @@ int32 FCustomizableObjectInstanceDescriptor::GetTextureValueRange(const FString&
 
 int32 FCustomizableObjectInstanceDescriptor::AddValueToIntRange(const FString& ParamName)
 {
-	check(CustomizableObject);
+	if (!CustomizableObject)
+	{
+		CustomizableObjectNullErrorMessage();
+		return -1;
+	}
 
 	const int32 IntParameterIndex = FindTypedParameterIndex(ParamName, EMutableParameterType::Int);
 	if (IntParameterIndex != INDEX_NONE)
@@ -1979,7 +2108,11 @@ int32 FCustomizableObjectInstanceDescriptor::AddValueToTextureRange(const FStrin
 
 int32 FCustomizableObjectInstanceDescriptor::AddValueToProjectorRange(const FString& ParamName)
 {
-	check(CustomizableObject);
+	if (!CustomizableObject)
+	{
+		CustomizableObjectNullErrorMessage();
+		return -1;
+	}
 
 	const int32 ProjectorParameterIndex = FindTypedParameterIndex(ParamName, EMutableParameterType::Projector);
 	if (ProjectorParameterIndex != INDEX_NONE)
@@ -2129,7 +2262,11 @@ int32 FCustomizableObjectInstanceDescriptor::GetState() const
 
 FString FCustomizableObjectInstanceDescriptor::GetCurrentState() const
 {
-	check(CustomizableObject);
+	if (!CustomizableObject)
+	{
+		CustomizableObjectNullErrorMessage();
+		return FString();
+	}
 
 	return CustomizableObject->GetStateName(GetState());
 }
@@ -2143,7 +2280,12 @@ void FCustomizableObjectInstanceDescriptor::SetState(const int32 InState)
 
 void FCustomizableObjectInstanceDescriptor::SetCurrentState(const FString& StateName)
 {
-	check(CustomizableObject);
+	if (!CustomizableObject)
+	{
+		CustomizableObjectNullErrorMessage();
+		return;
+	}
+
 	RETURN_ON_UNCOMPILED_CO(CustomizableObject, TEXT("Error: Cannot set state"))
 
 	const int32 Result = CustomizableObject->FindState(StateName);
@@ -2172,10 +2314,9 @@ void FCustomizableObjectInstanceDescriptor::SetRandomValues()
 
 void FCustomizableObjectInstanceDescriptor::SetRandomValuesFromStream(const FRandomStream& InStream)
 {
-	check(CustomizableObject);
-
 	if (!CustomizableObject)
 	{
+		CustomizableObjectNullErrorMessage();
 		return;
 	}
 	
@@ -2256,7 +2397,11 @@ void FCustomizableObjectInstanceDescriptor::SetRandomValuesFromStream(const FRan
 
 void FCustomizableObjectInstanceDescriptor::SetDefaultValues()
 {
-	check(CustomizableObject);
+	if (!CustomizableObject)
+	{
+		CustomizableObjectNullErrorMessage();
+		return;
+	}
 
 	for (FCustomizableObjectBoolParameterValue& Value : BoolParameters)
 	{	
diff --git a/Engine/Plugins/Experimental/Mutable/Source/CustomizableObject/Private/MuCO/CustomizableObjectSkeletalMesh.cpp b/Engine/Plugins/Experimental/Mutable/Source/CustomizableObject/Private/MuCO/CustomizableObjectSkeletalMesh.cpp
--- a/Engine/Plugins/Experimental/Mutable/Source/CustomizableObject/Private/MuCO/CustomizableObjectSkeletalMesh.cpp
+++ b/Engine/Plugins/Experimental/Mutable/Source/CustomizableObject/Private/MuCO/CustomizableObjectSkeletalMesh.cpp
@@ -19,21 +19,19 @@ FName GenerateUniqueNameFromCOInstance(const UCustomizableObjectInstance& Instan
 }
 
 
-UCustomizableObjectSkeletalMesh* UCustomizableObjectSkeletalMesh::CreateSkeletalMesh(const TSharedRef<FUpdateContextPrivate>& OperationData, const UCustomizableObjectInstance& Instance, const int32 InComponentIndex)
+UCustomizableObjectSkeletalMesh* UCustomizableObjectSkeletalMesh::CreateSkeletalMesh(const TSharedRef<FUpdateContextPrivate>& OperationData, 
+	const UCustomizableObjectInstance& Instance, const UCustomizableObject& CustomizableObject, const int32 InComponentIndex)
 {
-	UCustomizableObject* CustomizableObject = Instance.GetCustomizableObject();
-	check(CustomizableObject);
-
 	FName SkeletalMeshName = GenerateUniqueNameFromCOInstance(Instance);
 	UCustomizableObjectSkeletalMesh* OutSkeletalMesh = NewObject<UCustomizableObjectSkeletalMesh>(GetTransientPackage(), SkeletalMeshName, RF_Transient);
 	
 	// Debug info
-	OutSkeletalMesh->CustomizableObjectPathName = GetNameSafe(CustomizableObject);
+	OutSkeletalMesh->CustomizableObjectPathName = GetNameSafe(&CustomizableObject);
 	OutSkeletalMesh->InstancePathName = Instance.GetName();
 
 
 	// Init properties
-	OutSkeletalMesh->Model = CustomizableObject->GetPrivate()->GetModel();
+	OutSkeletalMesh->Model = CustomizableObject.GetPrivate()->GetModel();
 
 	OutSkeletalMesh->Parameters = OperationData->Parameters;
 	OutSkeletalMesh->State = OperationData->GetCapturedDescriptor().GetState();
diff --git a/Engine/Plugins/Experimental/Mutable/Source/CustomizableObject/Private/MuCO/MultilayerProjector.cpp b/Engine/Plugins/Experimental/Mutable/Source/CustomizableObject/Private/MuCO/MultilayerProjector.cpp
--- a/Engine/Plugins/Experimental/Mutable/Source/CustomizableObject/Private/MuCO/MultilayerProjector.cpp
+++ b/Engine/Plugins/Experimental/Mutable/Source/CustomizableObject/Private/MuCO/MultilayerProjector.cpp
@@ -135,7 +135,11 @@ void FMultilayerProjector::CreateLayer(FCustomizableObjectInstanceDescriptor& De
 	}
 
 	const UCustomizableObject* Object = Descriptor.GetCustomizableObject(); 
-	check(Object);
+	
+	if (!Object)
+	{
+		return;
+	}
 	
 	// Num Layers.
 	{
diff --git a/Engine/Plugins/Experimental/Mutable/Source/CustomizableObjectEditor/Private/MuCOE/CustomizableObjectEditorViewportClient.cpp b/Engine/Plugins/Experimental/Mutable/Source/CustomizableObjectEditor/Private/MuCOE/CustomizableObjectEditorViewportClient.cpp
--- a/Engine/Plugins/Experimental/Mutable/Source/CustomizableObjectEditor/Private/MuCOE/CustomizableObjectEditorViewportClient.cpp
+++ b/Engine/Plugins/Experimental/Mutable/Source/CustomizableObjectEditor/Private/MuCOE/CustomizableObjectEditorViewportClient.cpp
@@ -1637,7 +1637,12 @@ void FCustomizableObjectEditorViewportClient::OnInstanceForBakingUpdate(const FU
 	}
 	
 	const UCustomizableObject* CO = BakeTempInstance->GetCustomizableObject();
-	check (CO);
+
+	if (!CO)
+	{
+		UE_LOG(LogMutable, Error, TEXT("The Customizable Object instance provided for the baking does not have a Customizable Object."));
+		return;
+	}
 	
 	// Let the user set some configurations at the editor level
 	const FText DefaultFileName = FText::Format(LOCTEXT("DefaultFileNameForBakeInstance", "{0}"), FText::AsCultureInvariant(CO->GetName()));
diff --git a/Engine/Plugins/Experimental/Mutable/Source/CustomizableObjectEditor/Private/MuCOE/CustomizableObjectInstanceBakingUtils.cpp b/Engine/Plugins/Experimental/Mutable/Source/CustomizableObjectEditor/Private/MuCOE/CustomizableObjectInstanceBakingUtils.cpp
--- a/Engine/Plugins/Experimental/Mutable/Source/CustomizableObjectEditor/Private/MuCOE/CustomizableObjectInstanceBakingUtils.cpp
+++ b/Engine/Plugins/Experimental/Mutable/Source/CustomizableObjectEditor/Private/MuCOE/CustomizableObjectInstanceBakingUtils.cpp
@@ -373,7 +373,6 @@ bool BakeCustomizableObjectInstance(
 	
 	// Ensure that the state of the COI provided is valid --------------------------------------------------------------------------------------------
 	UCustomizableObject* InstanceCO = InInstance.GetCustomizableObject();
-	check (InstanceCO);
 
 	// Ensure the CO of the COI is accessible 
 	if (!InstanceCO || InstanceCO->GetPrivate()->IsLocked())
diff --git a/Engine/Plugins/Experimental/Mutable/Source/CustomizableObjectEditor/Private/MuCOE/CustomizableObjectInstanceEditor.cpp b/Engine/Plugins/Experimental/Mutable/Source/CustomizableObjectEditor/Private/MuCOE/CustomizableObjectInstanceEditor.cpp
--- a/Engine/Plugins/Experimental/Mutable/Source/CustomizableObjectEditor/Private/MuCOE/CustomizableObjectInstanceEditor.cpp
+++ b/Engine/Plugins/Experimental/Mutable/Source/CustomizableObjectEditor/Private/MuCOE/CustomizableObjectInstanceEditor.cpp
@@ -718,9 +718,12 @@ void FCustomizableObjectInstanceEditor::OnCustomizableObjectStatusChanged(FCusto
 	case FCustomizableObjectStatus::EState::NoModel:
 		{
 			UCustomizableObject* CustomizableObject = CustomizableObjectInstance->GetCustomizableObject();
-			check(CustomizableObject)
-		
-			CustomizableObject->ConditionalAutoCompile();	
+			
+			if (CustomizableObject)
+			{
+				CustomizableObject->ConditionalAutoCompile();
+			}
+
 			break;			
 		}
 
@@ -1004,7 +1007,11 @@ void FCustomizableObjectInstanceEditor::HideGizmoProjectorParameter(const TShare
 void FCustomizableObjectInstanceEditor::CreatePreviewComponents()
 {
 	const UCustomizableObject* Object = CustomizableObjectInstance->GetCustomizableObject();
-	check(Object);
+	
+	if (!Object)
+	{
+		return;
+	}
 	
 	for (const TObjectPtr<UCustomizableSkeletalComponent>& Component : PreviewCustomizableSkeletalComponents)
 	{
diff --git a/Engine/Plugins/Experimental/Mutable/Source/CustomizableObjectEditor/Private/MuCOE/SCustomizableInstanceProperties.cpp b/Engine/Plugins/Experimental/Mutable/Source/CustomizableObjectEditor/Private/MuCOE/SCustomizableInstanceProperties.cpp
--- a/Engine/Plugins/Experimental/Mutable/Source/CustomizableObjectEditor/Private/MuCOE/SCustomizableInstanceProperties.cpp
+++ b/Engine/Plugins/Experimental/Mutable/Source/CustomizableObjectEditor/Private/MuCOE/SCustomizableInstanceProperties.cpp
@@ -516,6 +516,12 @@ void SCustomizableInstanceProperties::ResetParamBox()
 void SCustomizableInstanceProperties::RecursivelyAddParamAndChildren(int32 ParamIndexInObject)
 {
 	const UCustomizableObject* CustomizableObject = CustomInstance->GetCustomizableObject();
+
+	if (!CustomizableObject)
+	{
+		return;
+	}
+
 	const FString ParamName = CustomizableObject->GetParameterName(ParamIndexInObject);
 
 	AddParameter(ParamIndexInObject);
@@ -534,6 +540,12 @@ void SCustomizableInstanceProperties::RecursivelyAddParamAndChildren(int32 Param
 void SCustomizableInstanceProperties::FillChildrenMap(int32 ParamIndexInObject)
 {
 	const UCustomizableObject* CustomizableObject = CustomInstance->GetCustomizableObject();
+
+	if (!CustomizableObject)
+	{
+		return;
+	}
+
 	const FString ParamName = CustomizableObject->GetParameterName(ParamIndexInObject);
 	FMutableParamUIMetadata UIMetadata = CustomizableObject->GetParameterUIMetadata(ParamName);
 
@@ -548,7 +560,11 @@ void SCustomizableInstanceProperties::FillChildrenMap(int32 ParamIndexInObject)
 void SCustomizableInstanceProperties::AddParameter(int32 ParamIndexInObject)
 {
 	UCustomizableObject* CustomizableObject = CustomInstance->GetCustomizableObject();
-	check(CustomizableObject);
+	
+	if (!CustomizableObject)
+	{
+		return;
+	}
 
 	TSharedPtr<ICustomizableObjectInstanceEditor> Editor = GetEditorChecked();
 	
@@ -741,7 +757,7 @@ void SCustomizableInstanceProperties::AddParameter(int32 ParamIndexInObject)
 	case EMutableParameterType::Int:
 	{
 		int numValues = CustomizableObject->GetIntParameterNumOptions(ParamIndexInObject);
-		bool bIsParamMultidimensional = CustomInstance->GetCustomizableObject()->IsParameterMultidimensional((ParamIndexInObject));
+		bool bIsParamMultidimensional = CustomizableObject->IsParameterMultidimensional((ParamIndexInObject));
 
 		if (!bIsParamMultidimensional && numValues)
 		{
@@ -915,7 +931,7 @@ void SCustomizableInstanceProperties::AddParameter(int32 ParamIndexInObject)
 
 	case EMutableParameterType::Projector:
 	{
-		bool bIsParamMultidimensional = CustomInstance->GetCustomizableObject()->IsParameterMultidimensional(ParamIndexInObject);
+		bool bIsParamMultidimensional = CustomizableObject->IsParameterMultidimensional(ParamIndexInObject);
 
 		if (!bIsParamMultidimensional)
 		{
@@ -1451,6 +1467,12 @@ FReply SCustomizableInstanceProperties::OnResetAllParameters()
 	TArray<FCustomizableObjectIntParameterValue>& IntParameters = CustomInstance->GetPrivate()->GetDescriptor().GetIntParameters();
 	
 	const UCustomizableObject* CustomObject = CustomInstance->GetCustomizableObject();
+
+	if (!CustomObject)
+	{
+		return FReply::Handled();
+	}
+
 	const int32 NumObjectParameter = CustomObject->GetParameterCount();
 
 	for (int32 ParameterIndex = 0; ParameterIndex < NumObjectParameter; ++ParameterIndex)
@@ -1600,6 +1622,12 @@ void SCustomizableInstanceProperties::OnIntParameterComboBoxChanged(TSharedPtr<F
 	TArray<FCustomizableObjectIntParameterValue>& IntParameters = CustomInstance->GetPrivate()->GetDescriptor().GetIntParameters();
 	
 	const UCustomizableObject* CustomObject = CustomInstance->GetCustomizableObject();
+
+	if (!CustomObject)
+	{
+		return;
+	}
+
 	for (int32 i = 0; i < IntParameters.Num(); ++i)
 	{
 		if (IntParameters[i].ParameterName == ParamName)
@@ -1813,6 +1841,12 @@ void SCustomizableInstanceProperties::OnProjectorTextureParameterComboBoxChanged
 	TArray<FCustomizableObjectIntParameterValue>& IntParameters = CustomInstance->GetPrivate()->GetDescriptor().GetIntParameters();
 
 	const UCustomizableObject* CustomObject = CustomInstance->GetCustomizableObject();
+
+	if (!CustomObject)
+	{
+		return;
+	}
+
 	for (int32 i = 0; i < IntParameters.Num(); ++i)
 	{
 		if (IntParameters[i].ParameterName == ParamName)
@@ -1882,6 +1916,13 @@ void SCustomizableInstanceProperties::OnProjectorFloatParameterSliderEnd(float V
 
 FReply SCustomizableInstanceProperties::OnProjectorCopyTransform(const FString ParamName, const int32 RangeIndex) const
 {
+	UCustomizableObject* CustomizableObject = CustomInstance->GetCustomizableObject();
+
+	if (!CustomizableObject)
+	{
+		FReply::Handled();
+	}
+
 	const int32 ParameterIndexInObject = CustomInstance->GetCustomizableObject()->FindParameter(ParamName);
 	const int32 ProjectorParamIndex = CustomInstance->FindProjectorParameterNameIndex(ParamName);
 
@@ -1946,8 +1987,15 @@ FReply SCustomizableInstanceProperties::OnProjectorResetTransform(const FString
 {
 	FScopedTransaction Transaction(LOCTEXT("ResetTransform", "Reset Transform"));
 	CustomInstance->Modify();
+
+	UCustomizableObject* CustomizableObject = CustomInstance->GetCustomizableObject();
+
+	if (!CustomizableObject)
+	{
+		return  FReply::Handled();;
+	}
 	
-	const FCustomizableObjectProjector DefaultValue = CustomInstance->GetCustomizableObject()->GetProjectorParameterDefaultValue(ParamName);
+	const FCustomizableObjectProjector DefaultValue = CustomizableObject->GetProjectorParameterDefaultValue(ParamName);
 
 	CustomInstance->SetProjectorValue(ParamName,
 		static_cast<FVector>(DefaultValue.Position),
@@ -2044,7 +2092,12 @@ FReply SCustomizableInstanceProperties::CreateParameterProfileWindow()
 
 FReply SCustomizableInstanceProperties::RemoveParameterProfile()
 {
-	UCustomizableObject* CustomizableObject = CustomInstance->GetCustomizableObject(); 
+	UCustomizableObject* CustomizableObject = CustomInstance->GetCustomizableObject();
+
+	if (!CustomizableObject)
+	{
+		return FReply::Handled();
+	}
 
 	const int32 ProfileIdx = CustomInstance->GetPrivate()->SelectedProfileIndex;
 	if (ProfileIdx == INDEX_NONE)
@@ -2090,7 +2143,14 @@ void SCustomizableInstanceProperties::OnProfileSelectedChanged(TSharedPtr<FStrin
 	}	
 
 	//Set selected profile
-	TArray<FProfileParameterDat>& Profiles = CustomInstance->GetCustomizableObject()->GetPrivate()->GetInstancePropertiesProfiles();
+	UCustomizableObject* CustomizableObject = CustomInstance->GetCustomizableObject();
+
+	if (!CustomizableObject)
+	{
+		return;
+	}
+
+	TArray<FProfileParameterDat>& Profiles = CustomizableObject->GetPrivate()->GetInstancePropertiesProfiles();
 	for (int32 Idx = 0; Idx < Profiles.Num(); ++Idx)
 	{
 		if (Profiles[Idx].ProfileName == *Selection)
@@ -2122,7 +2182,14 @@ void SCustomizableInstanceProperties::SetParameterProfileNamesOnEditor()
 {
 	ParameterProfileNames.Empty();
 
-	for (FProfileParameterDat& Profile : CustomInstance->GetCustomizableObject()->GetPrivate()->GetInstancePropertiesProfiles())
+	UCustomizableObject* CustomizableObject = CustomInstance->GetCustomizableObject();
+
+	if (!CustomizableObject)
+	{
+		return;
+	}
+
+	for (FProfileParameterDat& Profile : CustomizableObject->GetPrivate()->GetInstancePropertiesProfiles())
 	{
 		ParameterProfileNames.Emplace(MakeShared<FString>(Profile.ProfileName));
 	}
@@ -2141,6 +2208,12 @@ bool SCustomizableInstanceProperties::HasAnyParameters() const
 bool SCustomizableInstanceProperties::SetParameterValueToDefault(int32 ParameterIndex)
 {
 	UCustomizableObject* CustomObject = CustomInstance->GetCustomizableObject();
+
+	if (!CustomObject)
+	{
+		return false;
+	}
+
 	FString ParameterName = CustomObject->GetParameterName(ParameterIndex);
 	EMutableParameterType ParameterType = CustomObject->GetParameterType(ParameterIndex);
 
@@ -2365,6 +2438,12 @@ FReply SCreateProfileParameters::OnButtonClick(EAppReturnType::Type ButtonID)
 		RequestDestroyWindow();
 
 		UCustomizableObject* CustomizableObject = CustomInstance->GetCustomizableObject(); 
+
+		if (!CustomizableObject)
+		{
+			return  FReply::Handled();;
+		}
+
 		CustomizableObject->GetPrivate()->AddNewParameterProfile(GetFileName(), *CustomInstance.Get());
 
 		if (CustomInstance->GetPrivate()->bSelectedProfileDirty && CustomInstance->GetPrivate()->SelectedProfileIndex != INDEX_NONE)
diff --git a/Engine/Plugins/Experimental/Mutable/Source/CustomizableObjectEditor/Private/MuCOE/SCustomizableObjectEditorTagExplorer.cpp b/Engine/Plugins/Experimental/Mutable/Source/CustomizableObjectEditor/Private/MuCOE/SCustomizableObjectEditorTagExplorer.cpp
--- a/Engine/Plugins/Experimental/Mutable/Source/CustomizableObjectEditor/Private/MuCOE/SCustomizableObjectEditorTagExplorer.cpp
+++ b/Engine/Plugins/Experimental/Mutable/Source/CustomizableObjectEditor/Private/MuCOE/SCustomizableObjectEditorTagExplorer.cpp
@@ -195,6 +195,11 @@ TSharedRef<SWidget> SCustomizableObjectEditorTagExplorer::OnGetTagsMenuContent()
 
 	UCustomizableObject* CustomizableObject = CustomizableObjectEditorPtr->GetCustomizableObject();
 
+	if (!CustomizableObject)
+	{
+		return SNullWidget::NullWidget;
+	}
+
 	UCustomizableObjectNodeObject* Root = GetRootNode(CustomizableObject, bMultipleBaseObjectsFound);
 	TArray<UCustomizableObject*> CustomizableObjectFamily;
 	TArray<FString> Tags;
diff --git a/Engine/Plugins/Experimental/Mutable/Source/CustomizableObjectEditor/Private/MuCOE/StatePerformanceTesting.cpp b/Engine/Plugins/Experimental/Mutable/Source/CustomizableObjectEditor/Private/MuCOE/StatePerformanceTesting.cpp
--- a/Engine/Plugins/Experimental/Mutable/Source/CustomizableObjectEditor/Private/MuCOE/StatePerformanceTesting.cpp
+++ b/Engine/Plugins/Experimental/Mutable/Source/CustomizableObjectEditor/Private/MuCOE/StatePerformanceTesting.cpp
@@ -325,6 +325,11 @@ void FRuntimeTest::LoadStateData(int32 StateIndex)
 {
 	UCustomizableObject* CustomizableObject = Instance->GetCustomizableObject(); 
 
+	if (!CustomizableObject)
+	{
+		return;
+	}
+
 	StateWholeInformation StateData;
 	StateData.State = StateIndex;
 	StateData.StateName = CustomizableObject->GetStateName(StateIndex);
@@ -499,6 +504,13 @@ void FRuntimeTest::InitTest()
 	MapStateUpdateData.Empty();
 	MapStateRuntimeData.Empty();
 
+	UCustomizableObject* CustomizableObject = Instance->GetCustomizableObject();
+
+	if (!CustomizableObject)
+	{
+		return;
+	}
+
 	// Build a batch of changes for each state
 	const int32 NumStates = Instance->GetCustomizableObject()->GetStateCount();
 
diff --git a/Engine/Plugins/Experimental/Mutable/Source/CustomizableObjectEditor/Private/MuCOE/Thumbnails/CustomizableObjectInstanceThumbnailRenderer.cpp b/Engine/Plugins/Experimental/Mutable/Source/CustomizableObjectEditor/Private/MuCOE/Thumbnails/CustomizableObjectInstanceThumbnailRenderer.cpp
--- a/Engine/Plugins/Experimental/Mutable/Source/CustomizableObjectEditor/Private/MuCOE/Thumbnails/CustomizableObjectInstanceThumbnailRenderer.cpp
+++ b/Engine/Plugins/Experimental/Mutable/Source/CustomizableObjectEditor/Private/MuCOE/Thumbnails/CustomizableObjectInstanceThumbnailRenderer.cpp
@@ -45,7 +45,12 @@ void UCustomizableObjectInstanceThumbnailRenderer::Draw(UObject* Object, int32 X
 	
 	if (CustomizableObjectInstance)
 	{
-		UCustomizableObject* CustomizableObject = CustomizableObjectInstance->GetCustomizableObject(); 
+		UCustomizableObject* CustomizableObject = CustomizableObjectInstance->GetCustomizableObject();
+
+		if (!CustomizableObject)
+		{
+			return;
+		}
 
 		if ( ThumbnailScene == nullptr )
 		{
diff --git a/Engine/Plugins/Experimental/Mutable/Source/CustomizableObjectPopulation/Private/MuCOP/CustomizableObjectPopulationGenerator.cpp b/Engine/Plugins/Experimental/Mutable/Source/CustomizableObjectPopulation/Private/MuCOP/CustomizableObjectPopulationGenerator.cpp
--- a/Engine/Plugins/Experimental/Mutable/Source/CustomizableObjectPopulation/Private/MuCOP/CustomizableObjectPopulationGenerator.cpp
+++ b/Engine/Plugins/Experimental/Mutable/Source/CustomizableObjectPopulation/Private/MuCOP/CustomizableObjectPopulationGenerator.cpp
@@ -22,7 +22,12 @@ void FCustomizableObjectPopulationClassGenerator::GenerateParameters(
 
 		if (ConstraintSampler.IsValidSampler())
 		{
-			UCustomizableObject* CustomizableObject = Instance->GetCustomizableObject(); 
+			UCustomizableObject* CustomizableObject = Instance->GetCustomizableObject();
+
+			if (!CustomizableObject)
+			{
+				return;
+			}
 
 			const FConstraintIndex& SamplerIndex = ConstraintSampler.GetSamplerID(ConstraintSampler.Sample(Rand));
 			switch (SamplerIndex.SamplerType)
