[mutable] Added missing file (function signature change) to the CL 33546212.
#tests Editor
#rnx

diff --git a/Engine/Plugins/Experimental/Mutable/Source/CustomizableObject/Public/MuCO/CustomizableObjectSkeletalMesh.h b/Engine/Plugins/Experimental/Mutable/Source/CustomizableObject/Public/MuCO/CustomizableObjectSkeletalMesh.h
--- a/Engine/Plugins/Experimental/Mutable/Source/CustomizableObject/Public/MuCO/CustomizableObjectSkeletalMesh.h
+++ b/Engine/Plugins/Experimental/Mutable/Source/CustomizableObject/Public/MuCO/CustomizableObjectSkeletalMesh.h
@@ -17,6 +17,7 @@ namespace mu
 }
 
 class UCustomizableObjectInstance;
+class UCustomizableObject;
 class FUpdateContextPrivate;
 
 class FCustomizableObjectMeshStreamIn;
@@ -36,7 +37,8 @@ class CUSTOMIZABLEOBJECT_API UCustomizableObjectSkeletalMesh : public USkeletalM
 
 public: // Public Methods
 
-	static UCustomizableObjectSkeletalMesh* CreateSkeletalMesh(const TSharedRef<FUpdateContextPrivate>& InOperationData, const UCustomizableObjectInstance& InInstance, const int32 ComponentIndex);
+	static UCustomizableObjectSkeletalMesh* CreateSkeletalMesh(const TSharedRef<FUpdateContextPrivate>& InOperationData, 
+		const UCustomizableObjectInstance& InInstance, const UCustomizableObject& CustomizableObject, const int32 ComponentIndex);
 
 private: // Private Methods
 
