[Mutable] Improved wireframe visibility.

- Now uses the same background color as the default viewport.
- Removed unused code.

#rnx

diff --git a/Engine/Plugins/Experimental/Mutable/Source/CustomizableObjectEditor/Private/MuCOE/CustomizableObjectEditorViewportClient.cpp b/Engine/Plugins/Experimental/Mutable/Source/CustomizableObjectEditor/Private/MuCOE/CustomizableObjectEditorViewportClient.cpp
--- a/Engine/Plugins/Experimental/Mutable/Source/CustomizableObjectEditor/Private/MuCOE/CustomizableObjectEditorViewportClient.cpp
+++ b/Engine/Plugins/Experimental/Mutable/Source/CustomizableObjectEditor/Private/MuCOE/CustomizableObjectEditorViewportClient.cpp
@@ -768,12 +768,6 @@ void FCustomizableObjectEditorViewportClient::HideGizmoLight()
 }
 
 
-void FCustomizableObjectEditorViewportClient::SetVisibilityForWireframeMode(bool bIsWireframeMode)
-{
-	EngineShowFlags.SetDirectLighting(!bIsWireframeMode && bDrawSky);
-}
-
-
 FSphere FCustomizableObjectEditorViewportClient::GetCameraTarget()
 {
 	bool bFoundTarget = false;
@@ -811,14 +805,6 @@ FSphere FCustomizableObjectEditorViewportClient::GetCameraTarget()
 }
 
 
-FLinearColor FCustomizableObjectEditorViewportClient::GetBackgroundColor() const
-{
-	FLinearColor BackgroundColor = FColor(55, 55, 55);
-
-	return BackgroundColor;
-}
-
-
 void FCustomizableObjectEditorViewportClient::SetPreviewComponent(UStaticMeshComponent* InStaticMeshComponent)
 {
 	StaticMeshComponent = InStaticMeshComponent;
diff --git a/Engine/Plugins/Experimental/Mutable/Source/CustomizableObjectEditor/Private/MuCOE/SCustomizableObjectEditorViewport.cpp b/Engine/Plugins/Experimental/Mutable/Source/CustomizableObjectEditor/Private/MuCOE/SCustomizableObjectEditorViewport.cpp
--- a/Engine/Plugins/Experimental/Mutable/Source/CustomizableObjectEditor/Private/MuCOE/SCustomizableObjectEditorViewport.cpp
+++ b/Engine/Plugins/Experimental/Mutable/Source/CustomizableObjectEditor/Private/MuCOE/SCustomizableObjectEditorViewport.cpp
@@ -1081,10 +1081,4 @@ void SCustomizableObjectEditorViewportTabBody::SetCustomizableObject(UCustomizab
 }
 
 
-FLinearColor SCustomizableObjectEditorViewportTabBody::GetViewportBackgroundColor() const
-{
-	return LevelViewportClient->GetBackgroundColor();
-}
-
-
 #undef LOCTEXT_NAMESPACE
\ No newline at end of file
diff --git a/Engine/Plugins/Experimental/Mutable/Source/CustomizableObjectEditor/Private/MuCOE/SCustomizableObjectEditorViewport.h b/Engine/Plugins/Experimental/Mutable/Source/CustomizableObjectEditor/Private/MuCOE/SCustomizableObjectEditorViewport.h
--- a/Engine/Plugins/Experimental/Mutable/Source/CustomizableObjectEditor/Private/MuCOE/SCustomizableObjectEditorViewport.h
+++ b/Engine/Plugins/Experimental/Mutable/Source/CustomizableObjectEditor/Private/MuCOE/SCustomizableObjectEditorViewport.h
@@ -160,9 +160,6 @@ public:
 	
 	void SetPreviewComponents(const TArray<UDebugSkelMeshComponent*>& InSkeletalMeshComponents);
 	
-	/** Function to get viewport's current background color */
-	FLinearColor GetViewportBackgroundColor() const;
-
 	/** Function to get the number of LOD models associated with the preview skeletal mesh*/
 	int32 GetLODModelCount() const;
 
diff --git a/Engine/Plugins/Experimental/Mutable/Source/CustomizableObjectEditor/Private/MuCOE/SCustomizableObjectEditorViewportToolBar.cpp b/Engine/Plugins/Experimental/Mutable/Source/CustomizableObjectEditor/Private/MuCOE/SCustomizableObjectEditorViewportToolBar.cpp
--- a/Engine/Plugins/Experimental/Mutable/Source/CustomizableObjectEditor/Private/MuCOE/SCustomizableObjectEditorViewportToolBar.cpp
+++ b/Engine/Plugins/Experimental/Mutable/Source/CustomizableObjectEditor/Private/MuCOE/SCustomizableObjectEditorViewportToolBar.cpp
@@ -367,37 +367,6 @@ TSharedRef<SWidget> SCustomizableObjectEditorViewportToolBar::GenerateViewportOp
 }
 
 
-FSlateColor SCustomizableObjectEditorViewportToolBar::GetFontColor() const
-{
-	const UAssetViewerSettings* Settings = UAssetViewerSettings::Get();
-	const UEditorPerProjectUserSettings* PerProjectUserSettings = GetDefault<UEditorPerProjectUserSettings>();
-	const int32 ProfileIndex = Settings->Profiles.IsValidIndex(PerProjectUserSettings->AssetViewerProfileIndex) ? PerProjectUserSettings->AssetViewerProfileIndex : 0;
-
-	ensureMsgf(Settings->Profiles.IsValidIndex(PerProjectUserSettings->AssetViewerProfileIndex), TEXT("Invalid default settings pointer or current profile index"));
-
-	FLinearColor FontColor;
-	if (Settings->Profiles[ProfileIndex].bShowEnvironment)
-	{
-		FontColor = FLinearColor::White;
-	}
-	else
-	{
-		FLinearColor BackgroundColorInHSV = Viewport.Pin()->GetViewportBackgroundColor().LinearRGBToHSV();
-
-		// see if it's dark, if V is less than 0.2
-		if (BackgroundColorInHSV.B < 0.3f)
-		{
-			FontColor = FLinearColor::White;
-		}
-		else
-		{
-			FontColor = FLinearColor::Black;
-		}
-	}
-
-	return FontColor;
-}
-
 FText SCustomizableObjectEditorViewportToolBar::GetPlaybackMenuLabel() const
 {
 	return LOCTEXT("PlaybackError", "Error");
diff --git a/Engine/Plugins/Experimental/Mutable/Source/CustomizableObjectEditor/Public/MuCOE/CustomizableObjectEditorViewportClient.h b/Engine/Plugins/Experimental/Mutable/Source/CustomizableObjectEditor/Public/MuCOE/CustomizableObjectEditorViewportClient.h
--- a/Engine/Plugins/Experimental/Mutable/Source/CustomizableObjectEditor/Public/MuCOE/CustomizableObjectEditorViewportClient.h
+++ b/Engine/Plugins/Experimental/Mutable/Source/CustomizableObjectEditor/Public/MuCOE/CustomizableObjectEditorViewportClient.h
@@ -75,8 +75,6 @@ public:
 	class UPersonaOptions* ConfigOption;
 
 	// FEditorViewportClient interface
-	virtual FLinearColor GetBackgroundColor() const override;
-
 	virtual void Tick(float DeltaSeconds) override;
 	virtual void Draw(const FSceneView* View,FPrimitiveDrawInterface* PDI) override;
 	virtual void Draw(FViewport* Viewport,FCanvas* Canvas) override;
@@ -212,8 +210,6 @@ public:
 	/** Do not call directly. Use ICustomizableObjectEditor functions instead. */
 	void HideGizmoLight();
 	
-	void SetVisibilityForWireframeMode(bool bIsWireframeMode);
-
 	void SetAnimation(UAnimationAsset* Animation, EAnimationMode::Type AnimationType);
 
 	/** Set again the animation saved in AnimationBeingPlayed (if any) */
diff --git a/Engine/Plugins/Experimental/Mutable/Source/CustomizableObjectEditor/Public/MuCOE/SCustomizableObjectEditorViewportToolBar.h b/Engine/Plugins/Experimental/Mutable/Source/CustomizableObjectEditor/Public/MuCOE/SCustomizableObjectEditorViewportToolBar.h
--- a/Engine/Plugins/Experimental/Mutable/Source/CustomizableObjectEditor/Public/MuCOE/SCustomizableObjectEditorViewportToolBar.h
+++ b/Engine/Plugins/Experimental/Mutable/Source/CustomizableObjectEditor/Public/MuCOE/SCustomizableObjectEditorViewportToolBar.h
@@ -74,11 +74,6 @@ private:
 	*/
 	TSharedRef<SWidget> GenerateViewportOptionsMenu() const;
 
-	/**
-	* Generate color of the text on the top
-	*/
-	FSlateColor GetFontColor() const;
-
 	/**
 	* Returns the label for the Playback tool bar menu, which changes depending on the current playback speed
 	*
