[mutable] Added tooltips to layout strategy combobox options.

#rnx

diff --git a/Engine/Plugins/Experimental/Mutable/Source/CustomizableObjectEditor/Private/MuCOE/Nodes/CustomizableObjectNodeLayoutBlocksDetails.cpp b/Engine/Plugins/Experimental/Mutable/Source/CustomizableObjectEditor/Private/MuCOE/Nodes/CustomizableObjectNodeLayoutBlocksDetails.cpp
--- a/Engine/Plugins/Experimental/Mutable/Source/CustomizableObjectEditor/Private/MuCOE/Nodes/CustomizableObjectNodeLayoutBlocksDetails.cpp
+++ b/Engine/Plugins/Experimental/Mutable/Source/CustomizableObjectEditor/Private/MuCOE/Nodes/CustomizableObjectNodeLayoutBlocksDetails.cpp
@@ -11,6 +11,7 @@
 #include "MuCOE/Nodes/CustomizableObjectNodeLayoutBlocks.h"
 #include "MuCOE/SCustomizableObjectNodeLayoutBlocksEditor.h"
 #include "PropertyCustomizationHelpers.h"
+#include "SSearchableComboBox.h"
 #include "Styling/SlateTypes.h"
 #include "Widgets/Input/SCheckBox.h"
 #include "Widgets/Input/SSpinBox.h"
@@ -67,6 +68,25 @@ void FCustomizableObjectNodeLayoutBlocksDetails::CustomizeDetails( IDetailLayout
 			.Font(DetailBuilder.GetDetailFont())
 		];
 
+
+		// Layout Strategy options. Hardcoded: we should get names and tooltips from the enum property
+		{
+			LayoutPackingStrategies.Empty();
+			LayoutPackingStrategiesTooltips.Empty();
+
+			LayoutPackingStrategies.Add(MakeShareable(new FString("Resizable")));
+			LayoutPackingStrategiesTooltips.Add(LOCTEXT("LayoutDetails_ResizableStrategyTooltip", "In a layout merge, Layout size will increase if blocks don't fit inside."));
+
+			LayoutPackingStrategies.Add(MakeShareable(new FString("Fixed")));
+			LayoutPackingStrategiesTooltips.Add(LOCTEXT("LayoutDetails_FixedStrategyTooltip", "In a layout merge, the layout will increase its size until the maximum layout grid size"
+				"\nBlock sizes will be reduced if they don't fit inside the layout."
+				"\nSet the reduction priority of each block to control which blocks are reduced first and how they are reduced."));
+
+			LayoutPackingStrategies.Add(MakeShareable(new FString("Overlay")));
+			LayoutPackingStrategiesTooltips.Add(LOCTEXT("LayoutDetails_OverlayStrategyTooltip", "In a layout merge, the layout will not be modified and blocks will be ignored."
+				"\nExtend material nodes just add their layouts on top of the base one"));
+		}
+
 		// Layout strategy selector group widget
 		IDetailGroup* LayoutStrategyOptionsGroup = &CustomizableObjectCategory.AddGroup(TEXT("LayoutStrategyOptionsGroup"), LOCTEXT("LayoutStrategyGroup", "Layout Strategy Group"), false, true);
 		LayoutStrategyOptionsGroup->HeaderRow()
@@ -74,18 +94,25 @@ void FCustomizableObjectNodeLayoutBlocksDetails::CustomizeDetails( IDetailLayout
 		[
 			SNew(STextBlock)
 			.Text(LOCTEXT("LayoutStrategy_Text", "Layout Strategy:"))
-			.ToolTipText(LOCTEXT("LayoutStrategyTooltup", "Selects the packing strategy: Resizable Layout or Fixed Layout"))
+			.ToolTipText(LOCTEXT("LayoutStrategyTooltup", "Selects the packing strategy in case of a layout merge."))
 			.Font(DetailBuilder.GetDetailFont())
 		]
 		.ValueContent()
 		[
-			SNew(STextComboBox)
+			SNew(SSearchableComboBox)
 			.InitiallySelectedItem(CurrentStrategy)
 			.OptionsSource(&LayoutPackingStrategies)
 			.OnSelectionChanged(this, &FCustomizableObjectNodeLayoutBlocksDetails::OnLayoutPackingStrategyChanged)
-			.Font(DetailBuilder.GetDetailFont())
+			.OnGenerateWidget(this, &FCustomizableObjectNodeLayoutBlocksDetails::OnGenerateStrategyComboBox)
+			.ToolTipText(this, &FCustomizableObjectNodeLayoutBlocksDetails::GetSelectedLayoutStrategyTooltip)
+			[
+				SNew(STextBlock)
+				.Text(this, &FCustomizableObjectNodeLayoutBlocksDetails::GetSelectedLayoutStrategyName)
+				.Font(IDetailLayoutBuilder::GetDetailFont())
+			]
 		];
 
+
 		// Max layout size selector widget
 		LayoutStrategyOptionsGroup->AddWidgetRow()
 		.Visibility(TAttribute<EVisibility>(this, &FCustomizableObjectNodeLayoutBlocksDetails::FixedStrategyOptionsVisibility))
@@ -104,6 +131,17 @@ void FCustomizableObjectNodeLayoutBlocksDetails::CustomizeDetails( IDetailLayout
 			.Font(DetailBuilder.GetDetailFont())
 		];
 
+
+		// Block reduction methods options
+		{
+			BlockReductionMethods.Empty();
+			BlockReductionMethods.Add(MakeShareable(new FString("Halve")));
+			BlockReductionMethodsTooltips.Add(LOCTEXT("LayoutDetails_HalveRedMethodTooltip", "Blocks will be reduced by half each time."));
+
+			BlockReductionMethods.Add(MakeShareable(new FString("Unitary")));
+			BlockReductionMethodsTooltips.Add(LOCTEXT("LayoutDetails_UnitaryRedMethodTooltip", "Blocks will be reduced by one unit each time."));
+		}
+
 		// Reduction method selector widget
 		LayoutStrategyOptionsGroup->AddWidgetRow()
 		.Visibility(TAttribute<EVisibility>(this, &FCustomizableObjectNodeLayoutBlocksDetails::FixedStrategyOptionsVisibility))
@@ -111,18 +149,22 @@ void FCustomizableObjectNodeLayoutBlocksDetails::CustomizeDetails( IDetailLayout
 		[
 			SNew(STextBlock)
 			.Text(LOCTEXT("ReductionMethod_Text", "Reduction Method:"))
-			.ToolTipText(LOCTEXT("Reduction_Method_Tooltip", "Select how blocks will be reduced in case that they do not fit in the layout:"
-				"\n Halve: blocks will be reduced by half each time."
-				"\n Unit: blocks will be reduced by one unit each time."))
+			.ToolTipText(LOCTEXT("Reduction_Method_Tooltip", "Select how blocks will be reduced in case that they do not fit in the layout."))
 			.Font(DetailBuilder.GetDetailFont())
 		]
 		.ValueContent()
 		[
-			SNew(STextComboBox)
+			SNew(SSearchableComboBox)
 			.InitiallySelectedItem(CurrentReductionMethod)
 			.OptionsSource(&BlockReductionMethods)
 			.OnSelectionChanged(this, &FCustomizableObjectNodeLayoutBlocksDetails::OnReductionMethodChanged)
-			.Font(DetailBuilder.GetDetailFont())
+			.OnGenerateWidget(this, &FCustomizableObjectNodeLayoutBlocksDetails::OnGenerateReductionMethodComboBox)
+			.ToolTipText(this, &FCustomizableObjectNodeLayoutBlocksDetails::GetSelectedLayoutReductionMethodTooltip)
+			[
+				SNew(STextBlock)
+					.Text(this, &FCustomizableObjectNodeLayoutBlocksDetails::GetSelectedLayoutReductionMethodName)
+					.Font(IDetailLayoutBuilder::GetDetailFont())
+			]
 		];
 
 		// Warning selector group widget
@@ -334,4 +376,100 @@ void FCustomizableObjectNodeLayoutBlocksDetails::OnReductionMethodChanged(TShare
 	}
 }
 
+
+TSharedRef<SWidget> FCustomizableObjectNodeLayoutBlocksDetails::OnGenerateStrategyComboBox(TSharedPtr<FString> InItem) const
+{
+	FText Tooltip;
+
+	if (InItem.IsValid())
+	{
+		int32 TooltipIndex = LayoutPackingStrategies.IndexOfByKey(InItem);
+
+		if (LayoutPackingStrategiesTooltips.IsValidIndex(TooltipIndex))
+		{
+			//A list of tool tips should have been populated in a 1 to 1 correspondance
+			check(LayoutPackingStrategies.Num() == LayoutPackingStrategiesTooltips.Num());
+			Tooltip = LayoutPackingStrategiesTooltips[TooltipIndex];
+		}
+	}
+
+	return SNew(STextBlock)
+		.Text(FText::FromString(*InItem.Get()))
+		.Font(IDetailLayoutBuilder::GetDetailFont())
+		.ToolTipText(Tooltip);
+}
+
+
+TSharedRef<SWidget> FCustomizableObjectNodeLayoutBlocksDetails::OnGenerateReductionMethodComboBox(TSharedPtr<FString> InItem) const
+{
+	FText Tooltip;
+
+	if (InItem.IsValid())
+	{
+		int32 TooltipIndex = BlockReductionMethods.IndexOfByKey(InItem);
+
+		if (BlockReductionMethodsTooltips.IsValidIndex(TooltipIndex))
+		{
+			//A list of tool tips should have been populated in a 1 to 1 correspondance
+			check(BlockReductionMethods.Num() == BlockReductionMethodsTooltips.Num());
+			Tooltip = BlockReductionMethodsTooltips[TooltipIndex];
+		}
+	}
+
+	return SNew(STextBlock)
+		.Text(FText::FromString(*InItem.Get()))
+		.Font(IDetailLayoutBuilder::GetDetailFont())
+		.ToolTipText(Tooltip);
+}
+
+
+FText FCustomizableObjectNodeLayoutBlocksDetails::GetSelectedLayoutStrategyName() const
+{
+	if (Node->Layout)
+	{
+		return FText::FromString(*LayoutPackingStrategies[(uint32)Node->Layout->GetPackingStrategy()]);
+	}
+
+	return FText();
+}
+
+
+FText FCustomizableObjectNodeLayoutBlocksDetails::GetSelectedLayoutReductionMethodName() const
+{
+	if (Node->Layout)
+	{
+		return FText::FromString(*BlockReductionMethods[(uint32)Node->Layout->GetBlockReductionMethod()]);
+	}
+
+	return FText();
+}
+
+
+FText FCustomizableObjectNodeLayoutBlocksDetails::GetSelectedLayoutStrategyTooltip() const
+{
+	if (Node->Layout)
+	{
+		//A list of tool tips should have been populated in a 1 to 1 correspondance
+		check(LayoutPackingStrategies.Num() == LayoutPackingStrategiesTooltips.Num());
+
+		return LayoutPackingStrategiesTooltips[(uint32)Node->Layout->GetPackingStrategy()];
+	}
+
+	return FText();
+}
+
+
+FText FCustomizableObjectNodeLayoutBlocksDetails::GetSelectedLayoutReductionMethodTooltip() const
+{
+	if (Node->Layout)
+	{
+		//A list of tool tips should have been populated in a 1 to 1 correspondance
+		check(BlockReductionMethods.Num() == BlockReductionMethodsTooltips.Num());
+
+		return BlockReductionMethodsTooltips[(uint32)Node->Layout->GetBlockReductionMethod()];
+	}
+
+	return FText();
+}
+
 #undef LOCTEXT_NAMESPACE
diff --git a/Engine/Plugins/Experimental/Mutable/Source/CustomizableObjectEditor/Private/MuCOE/Nodes/CustomizableObjectNodeLayoutBlocksDetails.h b/Engine/Plugins/Experimental/Mutable/Source/CustomizableObjectEditor/Private/MuCOE/Nodes/CustomizableObjectNodeLayoutBlocksDetails.h
--- a/Engine/Plugins/Experimental/Mutable/Source/CustomizableObjectEditor/Private/MuCOE/Nodes/CustomizableObjectNodeLayoutBlocksDetails.h
+++ b/Engine/Plugins/Experimental/Mutable/Source/CustomizableObjectEditor/Private/MuCOE/Nodes/CustomizableObjectNodeLayoutBlocksDetails.h
@@ -41,6 +41,12 @@ private:
 	void OnLayoutPackingStrategyChanged(TSharedPtr<FString> NewSelection, ESelectInfo::Type SelectInfo);
 	void OnIgnoreErrorsCheckStateChanged(ECheckBoxState State);
 	void OnLODBoxValueChanged(int32 Value);
+	TSharedRef<SWidget> OnGenerateStrategyComboBox(TSharedPtr<FString> InItem) const;
+	TSharedRef<SWidget> OnGenerateReductionMethodComboBox(TSharedPtr<FString> InItem) const;
+	FText GetSelectedLayoutStrategyName() const;
+	FText GetSelectedLayoutReductionMethodName() const;
+	FText GetSelectedLayoutStrategyTooltip() const;
+	FText GetSelectedLayoutReductionMethodTooltip() const;
 
 private:
 
@@ -59,8 +65,10 @@ private:
 
 	/** List of available layout packing strategies. */
 	TArray< TSharedPtr< FString > > LayoutPackingStrategies;
+	TArray<FText> LayoutPackingStrategiesTooltips;
 
 	/** List of available block reduction methods. */
 	TArray< TSharedPtr< FString > > BlockReductionMethods;
+	TArray<FText> BlockReductionMethodsTooltips;
 
 };
diff --git a/Engine/Plugins/Experimental/Mutable/Source/CustomizableObjectEditor/Private/MuCOE/Nodes/CustomizableObjectNodeTableDetails.cpp b/Engine/Plugins/Experimental/Mutable/Source/CustomizableObjectEditor/Private/MuCOE/Nodes/CustomizableObjectNodeTableDetails.cpp
--- a/Engine/Plugins/Experimental/Mutable/Source/CustomizableObjectEditor/Private/MuCOE/Nodes/CustomizableObjectNodeTableDetails.cpp
+++ b/Engine/Plugins/Experimental/Mutable/Source/CustomizableObjectEditor/Private/MuCOE/Nodes/CustomizableObjectNodeTableDetails.cpp
@@ -15,6 +15,7 @@
 #include "MuCOE/Nodes/CustomizableObjectNodeTable.h"
 #include "MuCOE/SCustomizableObjectNodeLayoutBlocksEditor.h"
 #include "MuCOE/UnrealEditorPortabilityHelpers.h"
+#include "SSearchableComboBox.h"
 #include "Styling/SlateColor.h"
 #include "Widgets/Input/SButton.h"
 #include "Widgets/Input/SCheckBox.h"
@@ -195,7 +196,6 @@ void FCustomizableObjectNodeTableDetails::CustomizeDetails(const TSharedPtr<IDet
 
 		//	Layout Selector
 		SelectedLayout = nullptr;
-		LayoutBlocksEditor = SNew(SCustomizableObjectNodeLayoutBlocksEditor);
 
 		// Mesh selector of the layout editor
 		LayoutCategory.AddCustomRow(LOCTEXT("TableLayoutEditor_MeshSelector", "Mesh Selector"))
@@ -239,6 +239,14 @@ void FCustomizableObjectNodeTableDetails::CustomizeDetails(const TSharedPtr<IDet
 		]
 		.OverrideResetToDefault(FResetToDefaultOverride::Create(FSimpleDelegate::CreateSP(this, &FCustomizableObjectNodeTableDetails::OnLayoutMeshCustomRowResetButtonClicked)));
 
+
+		// Layout grid size options
+		int32 MaxGridSize = 32;
+		for (int32 Size = 1; Size <= MaxGridSize; Size *= 2)
+		{
+			LayoutGridSizes.Add(MakeShareable(new FString(FString::Printf(TEXT("%d x %d"), Size, Size))));
+		}
+
 		// Layout size selector widget
 		LayoutCategory.AddCustomRow(LOCTEXT("TableBlocksDetails_SizeSelector", "SizeSelector"))
 		.Visibility(TAttribute<EVisibility>(this, &FCustomizableObjectNodeTableDetails::LayoutOptionsVisibility))
@@ -257,6 +265,24 @@ void FCustomizableObjectNodeTableDetails::CustomizeDetails(const TSharedPtr<IDet
 		];
 
 
+		// Layout Strategy options. Hardcoded: we should get names and tooltips from the enum property
+		{
+			LayoutPackingStrategies.Empty();
+			LayoutPackingStrategiesTooltips.Empty();
+
+			LayoutPackingStrategies.Add(MakeShareable(new FString("Resizable")));
+			LayoutPackingStrategiesTooltips.Add(LOCTEXT("TableDetails_ResizableStrategyTooltip", "In a layout merge, Layout size will increase if blocks don't fit inside."));
+
+			LayoutPackingStrategies.Add(MakeShareable(new FString("Fixed")));
+			LayoutPackingStrategiesTooltips.Add(LOCTEXT("TableDetails_FixedStrategyTooltip", "In a layout merge, the layout will increase its size until the maximum layout grid size"
+				"\nBlock sizes will be reduced if they don't fit inside the layout."
+				"\nSet the reduction priority of each block to control which blocks are reduced first and how they are reduced."));
+
+			LayoutPackingStrategies.Add(MakeShareable(new FString("Overlay")));
+			LayoutPackingStrategiesTooltips.Add(LOCTEXT("TableDetails_OverlayStrategyTooltip", "In a layout merge, the layout will not be modified and blocks will be ignored."
+				"\nExtend material nodes just add their layouts on top of the base one"));
+		}
+
 		// Layout strategy selector group widget
 		IDetailGroup* LayoutStrategyOptionsGroup = &LayoutCategory.AddGroup(TEXT("TableLayoutStrategyOptionsGroup"), LOCTEXT("TableLayoutStrategyGroup", "Table Layout Strategy Group"), false, true);
 		LayoutStrategyOptionsGroup->HeaderRow()
@@ -265,15 +291,21 @@ void FCustomizableObjectNodeTableDetails::CustomizeDetails(const TSharedPtr<IDet
 		[
 			SNew(STextBlock)
 			.Text(LOCTEXT("TableLayoutStrategy_Text", "Layout Strategy:"))
-			.ToolTipText(LOCTEXT("TableLayoutStrategyTooltip", "Selects the packing strategy: Resizable Layout or Fixed Layout"))
+			.ToolTipText(LOCTEXT("TableLayoutStrategyTooltip", "Selects the packing strategy in case of a layout merge."))
 			.Font(IDetailLayoutBuilder::GetDetailFont())
 		]
 		.ValueContent()
 		[
-			SAssignNew(StrategyComboBox, STextComboBox)
+			SAssignNew(StrategyComboBox, SSearchableComboBox)
 			.OptionsSource(&LayoutPackingStrategies)
 			.OnSelectionChanged(this, &FCustomizableObjectNodeTableDetails::OnLayoutPackingStrategyChanged)
-			.Font(IDetailLayoutBuilder::GetDetailFont())
+			.OnGenerateWidget(this, &FCustomizableObjectNodeTableDetails::OnGenerateStrategyComboBox)
+			.ToolTipText(this, &FCustomizableObjectNodeTableDetails::GetSelectedLayoutStrategyTooltip)
+			[
+				SNew(STextBlock)
+				.Text(this, &FCustomizableObjectNodeTableDetails::GetSelectedLayoutStrategyName)
+				.Font(IDetailLayoutBuilder::GetDetailFont())
+			]
 		];
 
 
@@ -295,6 +327,16 @@ void FCustomizableObjectNodeTableDetails::CustomizeDetails(const TSharedPtr<IDet
 		];
 
 
+		// Block reduction methods options
+		{
+			BlockReductionMethods.Empty();
+			BlockReductionMethods.Add(MakeShareable(new FString("Halve")));
+			BlockReductionMethodsTooltips.Add(LOCTEXT("TableDetails_HalveRedMethodTooltip", "Blocks will be reduced by half each time."));
+
+			BlockReductionMethods.Add(MakeShareable(new FString("Unitary")));
+			BlockReductionMethodsTooltips.Add(LOCTEXT("TableDetails_UnitaryRedMethodTooltip", "Blocks will be reduced by one unit each time."));
+		}
+
 		// Reduction method selector widget
 		LayoutStrategyOptionsGroup->AddWidgetRow()
 		.Visibility(TAttribute<EVisibility>(this, &FCustomizableObjectNodeTableDetails::FixedStrategyOptionsVisibility))
@@ -302,19 +344,25 @@ void FCustomizableObjectNodeTableDetails::CustomizeDetails(const TSharedPtr<IDet
 		[
 			SNew(STextBlock)
 			.Text(LOCTEXT("TableReductionMethod_Text", "Reduction Method:"))
-			.ToolTipText(LOCTEXT("TableReduction_Method_Tooltip", "Select how blocks will be reduced in case that they do not fit in the layout:"
-				"\n Halve: blocks will be reduced by half each time."
-				"\n Unit: blocks will be reduced by one unit each time."))
+			.ToolTipText(LOCTEXT("TableReduction_Method_Tooltip", "Select how blocks will be reduced in case that they do not fit in the layout:"))
 			.Font(IDetailLayoutBuilder::GetDetailFont())
 		]
 		.ValueContent()
 		[
-			SAssignNew(ReductionMethodComboBox, STextComboBox)
+			SAssignNew(ReductionMethodComboBox, SSearchableComboBox)
 			.OptionsSource(&BlockReductionMethods)
 			.OnSelectionChanged(this, &FCustomizableObjectNodeTableDetails::OnReductionMethodChanged)
-			.Font(IDetailLayoutBuilder::GetDetailFont())
+			.OnGenerateWidget(this, &FCustomizableObjectNodeTableDetails::OnGenerateReductionMethodComboBox)
+			.ToolTipText(this, &FCustomizableObjectNodeTableDetails::GetSelectedLayoutReductionMethodTooltip)
+			[
+				SNew(STextBlock)
+				.Text(this, &FCustomizableObjectNodeTableDetails::GetSelectedLayoutReductionMethodName)
+				.Font(IDetailLayoutBuilder::GetDetailFont())
+			]
 		];
 
+		// Layout blocks editor widget
+		LayoutBlocksEditor = SNew(SCustomizableObjectNodeLayoutBlocksEditor);
 
 		// Block editor Widget
 		LayoutCategory.AddCustomRow(LOCTEXT("TableLayoutEditor", "Layout Editor"))
@@ -708,39 +756,95 @@ void FCustomizableObjectNodeTableDetails::FillLayoutComboBoxOptions()
 	if (SelectedLayout.IsValid() && GridSizeComboBox.IsValid() && StrategyComboBox.IsValid()
 		&& MaxGridSizeComboBox.IsValid() && ReductionMethodComboBox.IsValid())
 	{
-		// Static const variable?
-		int32 MaxGridSize = 32;
-		LayoutGridSizes.Empty();
-
-		for (int32 Size = 1; Size <= MaxGridSize; Size *= 2)
+		for (int32 OptionIndex = 0; OptionIndex < LayoutGridSizes.Num(); ++OptionIndex)
 		{
-			LayoutGridSizes.Add(MakeShareable(new FString(FString::Printf(TEXT("%d x %d"), Size, Size))));
+			int32 Size = 1 << OptionIndex;
 
 			if (SelectedLayout->GetGridSize() == FIntPoint(Size))
 			{
-				GridSizeComboBox->SetSelectedItem(LayoutGridSizes.Last());
+				GridSizeComboBox->SetSelectedItem(LayoutGridSizes[OptionIndex]);
 			}
 
 			if (SelectedLayout->GetMaxGridSize() == FIntPoint(Size))
 			{
-				MaxGridSizeComboBox->SetSelectedItem(LayoutGridSizes.Last());
+				MaxGridSizeComboBox->SetSelectedItem(LayoutGridSizes[OptionIndex]);
 			}
 		}
 
-		LayoutPackingStrategies.Empty();
-		LayoutPackingStrategies.Add(MakeShareable(new FString("Resizable")));
-		LayoutPackingStrategies.Add(MakeShareable(new FString("Fixed")));
-		LayoutPackingStrategies.Add(MakeShareable(new FString("Overlay")));
 		StrategyComboBox->SetSelectedItem(LayoutPackingStrategies[(uint32)SelectedLayout->GetPackingStrategy()]);
-
-		BlockReductionMethods.Empty();
-		BlockReductionMethods.Add(MakeShareable(new FString("Halve")));
-		BlockReductionMethods.Add(MakeShareable(new FString("Unitary")));
 		ReductionMethodComboBox->SetSelectedItem(BlockReductionMethods[(uint32)SelectedLayout->GetBlockReductionMethod()]);
 	}
 }
 
 
+TSharedRef<SWidget> FCustomizableObjectNodeTableDetails::OnGenerateStrategyComboBox(TSharedPtr<FString> InItem) const
+{
+	FText Tooltip;
+
+	if (InItem.IsValid())
+	{
+		int32 TooltipIndex = LayoutPackingStrategies.IndexOfByKey(InItem);
+
+		if (LayoutPackingStrategiesTooltips.IsValidIndex(TooltipIndex))
+		{
+			//A list of tool tips should have been populated in a 1 to 1 correspondance
+			check(LayoutPackingStrategies.Num() == LayoutPackingStrategiesTooltips.Num());
+			Tooltip = LayoutPackingStrategiesTooltips[TooltipIndex];
+		}
+	}
+
+	return SNew(STextBlock)
+		.Text(FText::FromString(*InItem.Get()))
+		.Font(IDetailLayoutBuilder::GetDetailFont())
+		.ToolTipText(Tooltip);
+}
+
+
+TSharedRef<SWidget> FCustomizableObjectNodeTableDetails::OnGenerateReductionMethodComboBox(TSharedPtr<FString> InItem) const
+{
+	FText Tooltip;
+
+	if (InItem.IsValid())
+	{
+		int32 TooltipIndex = BlockReductionMethods.IndexOfByKey(InItem);
+
+		if (BlockReductionMethodsTooltips.IsValidIndex(TooltipIndex))
+		{
+			//A list of tool tips should have been populated in a 1 to 1 correspondance
+			check(BlockReductionMethods.Num() == BlockReductionMethodsTooltips.Num());
+			Tooltip = BlockReductionMethodsTooltips[TooltipIndex];
+		}
+	}
+
+	return SNew(STextBlock)
+		.Text(FText::FromString(*InItem.Get()))
+		.Font(IDetailLayoutBuilder::GetDetailFont())
+		.ToolTipText(Tooltip);
+}
+
+
+FText FCustomizableObjectNodeTableDetails::GetSelectedLayoutStrategyName() const
+{
+	if (SelectedLayout.IsValid())
+	{
+		return FText::FromString(*LayoutPackingStrategies[(uint32)SelectedLayout->GetPackingStrategy()]);
+	}
+
+	return FText();
+}
+
+
+FText FCustomizableObjectNodeTableDetails::GetSelectedLayoutReductionMethodName() const
+{
+	if (SelectedLayout.IsValid())
+	{
+		return FText::FromString(*BlockReductionMethods[(uint32)SelectedLayout->GetBlockReductionMethod()]);
+	}
+
+	return FText();
+}
+
+
 void FCustomizableObjectNodeTableDetails::OnGridSizeChanged(TSharedPtr<FString> NewSelection, ESelectInfo::Type SelectInfo)
 {
 	if (SelectedLayout.IsValid())
@@ -811,6 +915,34 @@ void FCustomizableObjectNodeTableDetails::OnReductionMethodChanged(TSharedPtr<FS
 }
 
 
+FText FCustomizableObjectNodeTableDetails::GetSelectedLayoutStrategyTooltip() const
+{
+	if (SelectedLayout.IsValid())
+	{
+		//A list of tool tips should have been populated in a 1 to 1 correspondance
+		check(LayoutPackingStrategies.Num() == LayoutPackingStrategiesTooltips.Num());
+
+		return LayoutPackingStrategiesTooltips[(uint32)SelectedLayout->GetPackingStrategy()];
+	}
+
+	return FText();
+}
+
+
+FText FCustomizableObjectNodeTableDetails::GetSelectedLayoutReductionMethodTooltip() const
+{
+	if (SelectedLayout.IsValid())
+	{
+		//A list of tool tips should have been populated in a 1 to 1 correspondance
+		check(BlockReductionMethods.Num() == BlockReductionMethodsTooltips.Num());
+
+		return BlockReductionMethodsTooltips[(uint32)SelectedLayout->GetBlockReductionMethod()];
+	}
+
+	return FText();
+}
+
+
 // Metadata Category --------------------------------------------------------------------------------
 
 TSharedPtr<FString> FCustomizableObjectNodeTableDetails::GenerateMutableMetaDataColumnComboBoxOptions()
diff --git a/Engine/Plugins/Experimental/Mutable/Source/CustomizableObjectEditor/Private/MuCOE/Nodes/CustomizableObjectNodeTableDetails.h b/Engine/Plugins/Experimental/Mutable/Source/CustomizableObjectEditor/Private/MuCOE/Nodes/CustomizableObjectNodeTableDetails.h
--- a/Engine/Plugins/Experimental/Mutable/Source/CustomizableObjectEditor/Private/MuCOE/Nodes/CustomizableObjectNodeTableDetails.h
+++ b/Engine/Plugins/Experimental/Mutable/Source/CustomizableObjectEditor/Private/MuCOE/Nodes/CustomizableObjectNodeTableDetails.h
@@ -10,11 +10,14 @@ namespace ESelectInfo { enum Type : int; }
 
 class FString;
 class IDetailLayoutBuilder;
+class ISinglePropertyView;
 class FReply;
 class SButton;
 class SCustomizableObjectNodeLayoutBlocksEditor;
+class SSearchableComboBox;
 class STextBlock;
 class STextComboBox;
+class SWidget;
 class UCustomizableObjectLayout;
 class UCustomizableObjectNodeTable;
 struct EVisibility;
@@ -91,6 +94,24 @@ private:
 	/** Returns the visibility of the Fixed layout widgets */
 	EVisibility FixedStrategyOptionsVisibility() const;
 
+	/** Generates the Layout Packing Strategy text widgets of the ComboBox */
+	TSharedRef<SWidget> OnGenerateStrategyComboBox(TSharedPtr<FString> InItem) const;
+	
+	/** Generates the Layout Redutction Methods text widgets of the ComboBox*/
+	TSharedRef<SWidget> OnGenerateReductionMethodComboBox(TSharedPtr<FString> InItem) const;
+
+	/** Returns the selected Packing Strategy Name as a FText */
+	FText GetSelectedLayoutStrategyName() const;
+
+	/** Returns the selected Reduction Method Name as a FText */
+	FText GetSelectedLayoutReductionMethodName() const;
+
+	/** Returns the selected Packing Strategy Tooltip as a FText */
+	FText GetSelectedLayoutStrategyTooltip() const;
+
+	/** Returns the selected Reduction Method Tooltip as a FText */
+	FText GetSelectedLayoutReductionMethodTooltip() const;
+
 	/** Fills the combo box arrays sources */
 	void FillLayoutComboBoxOptions();
 
@@ -191,18 +212,20 @@ private:
 
 	/** List of available layout packing strategies. */
 	TArray< TSharedPtr< FString > > LayoutPackingStrategies;
+	TArray<FText> LayoutPackingStrategiesTooltips;
 
 	/** List of available block reduction methods. */
 	TArray< TSharedPtr< FString > > BlockReductionMethods;
+	TArray<FText> BlockReductionMethodsTooltips;
 
 	// ComboBox widget to select a Grid Size from the Selected Layout
 	TSharedPtr<STextComboBox> GridSizeComboBox;
-	// ComboBox widget to select a Strategy from the Selected Layout
-	TSharedPtr<STextComboBox> StrategyComboBox;
+	// ComboBox widget to select a Strategy from the Selected Layout. SSearchableComboBox allows us to set a custom tooltip per option.
+	TSharedPtr<SSearchableComboBox> StrategyComboBox;
 	// ComboBox widget to select a Max Grid Size from the Selected Layout
 	TSharedPtr<STextComboBox> MaxGridSizeComboBox;
 	// ComboBox widget to select a Reduction Method from the Selected Layout
-	TSharedPtr<STextComboBox> ReductionMethodComboBox;
+	TSharedPtr<SSearchableComboBox> ReductionMethodComboBox;
 
 
 	// Mutable UI Metadata -------------
