[Mutable] Fix FGCObject being destroyed outside the GameThread.

- Task_Mutable_ReleaseInstance was holding a reference to the context, which was destorying the FUpdateContextPrivate (FGCObject) in a BackgroundThread.
- Changed FUpdateContextPrivate FGCObject to TStrongObjecPtr (only because they are safer).
- Added additional IsInGameThread checks to FUpdateContextPrivate.

#rnx

diff --git a/Engine/Plugins/Experimental/Mutable/Source/CustomizableObject/Internal/MuCO/CustomizableObjectSystemPrivate.h b/Engine/Plugins/Experimental/Mutable/Source/CustomizableObject/Internal/MuCO/CustomizableObjectSystemPrivate.h
--- a/Engine/Plugins/Experimental/Mutable/Source/CustomizableObject/Internal/MuCO/CustomizableObjectSystemPrivate.h
+++ b/Engine/Plugins/Experimental/Mutable/Source/CustomizableObject/Internal/MuCO/CustomizableObjectSystemPrivate.h
@@ -26,6 +26,7 @@
 #include "AssetRegistry/AssetData.h"
 #include "ContentStreaming.h"
 #include "Animation/MorphTarget.h"
+#include "UObject/StrongObjectPtr.h"
 
 #include "Tasks/Task.h"
 
@@ -470,16 +471,13 @@ struct FInstanceUpdateData
 /** Update Context.
  *
  * Alive from the start to the end of the update (both API and LOD update). */
-class FUpdateContextPrivate : public FGCObject
+class FUpdateContextPrivate
 {
 public:
 	FUpdateContextPrivate(UCustomizableObjectInstance& InInstance, const FCustomizableObjectInstanceDescriptor& Descriptor);
 
 	FUpdateContextPrivate(UCustomizableObjectInstance& InInstance);
-	virtual ~FUpdateContextPrivate() override;
-
-	virtual FString GetReferencerName() const override;
-	virtual void AddReferencedObjects(FReferenceCollector& Collector) override;
+	~FUpdateContextPrivate();
 
 	int32 GetMinLOD() const;
 
@@ -588,7 +586,7 @@ public:
 	uint32 MutableRuntimeCycles = 0;
 
 	/** Hard references to objects. Avoids GC to collect them. */
-	TArray<TObjectPtr<const UObject>> Objects;
+	TArray<TStrongObjectPtr<const UObject>> Objects;
 };
 
 
diff --git a/Engine/Plugins/Experimental/Mutable/Source/CustomizableObject/Private/MuCO/CustomizableObjectSystem.cpp b/Engine/Plugins/Experimental/Mutable/Source/CustomizableObject/Private/MuCO/CustomizableObjectSystem.cpp
--- a/Engine/Plugins/Experimental/Mutable/Source/CustomizableObject/Private/MuCO/CustomizableObjectSystem.cpp
+++ b/Engine/Plugins/Experimental/Mutable/Source/CustomizableObject/Private/MuCO/CustomizableObjectSystem.cpp
@@ -149,6 +149,7 @@ static FAutoConsoleVariableSink CVarMutableSink(FConsoleCommandDelegate::CreateS
 
 FUpdateContextPrivate::FUpdateContextPrivate(UCustomizableObjectInstance& InInstance, const FCustomizableObjectInstanceDescriptor& Descriptor)
 {
+	check(IsInGameThread());
 	check(InInstance.GetPrivate());
 	check(InInstance.GetCustomizableObject());
 
@@ -178,6 +179,8 @@ FUpdateContextPrivate::FUpdateContextPrivate(UCustomizableObjectInstance& InInst
 
 FUpdateContextPrivate::~FUpdateContextPrivate()
 {
+	check(IsInGameThread());
+
 	if (UCustomizableObjectSystem::IsCreated())
 	{
 		UCustomizableObjectSystem* System = UCustomizableObjectSystem::GetInstance();
@@ -186,18 +189,6 @@ FUpdateContextPrivate::~FUpdateContextPrivate()
 }
 
 
-FString FUpdateContextPrivate::GetReferencerName() const
-{
-	return TEXT("FUpdateContextPrivate");
-}
-
-
-void FUpdateContextPrivate::AddReferencedObjects(FReferenceCollector& Collector)
-{
-	Collector.AddReferencedObjects(Objects);
-}
-
-
 int32 FUpdateContextPrivate::GetMinLOD() const
 {
 	return CapturedDescriptor.GetMinLod();
@@ -2680,23 +2671,19 @@ namespace impl
 
 
 	// This runs in a worker thread.
-	void Task_Mutable_ReleaseInstance(const TSharedRef<FUpdateContextPrivate>& OperationData, mu::Ptr<mu::System> MutableSystem)
+	void Task_Mutable_ReleaseInstance(mu::Instance::ID InstanceID, mu::Ptr<mu::System> MutableSystem, bool bLiveUpdateMode)
 	{
 		MUTABLE_CPUPROFILER_SCOPE(Task_Mutable_ReleaseInstance)
 
 		check(MutableSystem);
 
-		if (OperationData->InstanceID > 0)
+		if (InstanceID > 0)
 		{
-			MUTABLE_CPUPROFILER_SCOPE(EndUpdate);
-			MutableSystem->EndUpdate(OperationData->InstanceID);
-			OperationData->InstanceUpdateData.Clear();
+			MutableSystem->EndUpdate(InstanceID);
 
-			if (!OperationData->bLiveUpdateMode)
+			if (!bLiveUpdateMode)
 			{
-				MutableSystem->ReleaseInstance(OperationData->InstanceID);
-				OperationData->InstanceID = 0;
-				OperationData->MutableInstance = nullptr;
+				MutableSystem->ReleaseInstance(InstanceID);
 			}
 		}
 
@@ -2867,7 +2854,7 @@ namespace impl
 			const mu::Ptr<mu::System> MutableSystem = CustomizableObjectSystemPrivateData->MutableSystem;
 			CustomizableObjectSystemPrivateData->MutableTaskGraph.AddMutableThreadTask(
 				TEXT("Task_Mutable_ReleaseInstance"),
-				[OperationData, MutableSystem]() {Task_Mutable_ReleaseInstance(OperationData, MutableSystem); });
+				[OperationData, MutableSystem]() {Task_Mutable_ReleaseInstance(OperationData->InstanceID, MutableSystem, OperationData->bLiveUpdateMode); });
 		}
 
 
@@ -3088,7 +3075,7 @@ namespace impl
 		{
 			if (USkeletalMesh* CachedMesh = CustomizableObject->GetPrivate()->MeshCache.Get(MeshId))
 			{
-				Operation->Objects.Add(CachedMesh);
+				Operation->Objects.Emplace(CachedMesh);
 			}
 		}
 
