[Mutable] Documented CacheRuntimeReferencedImages.

#rnx

diff --git a/Engine/Plugins/Experimental/Mutable/Source/CustomizableObject/Private/MuCO/UnrealMutableImageProvider.cpp b/Engine/Plugins/Experimental/Mutable/Source/CustomizableObject/Private/MuCO/UnrealMutableImageProvider.cpp
--- a/Engine/Plugins/Experimental/Mutable/Source/CustomizableObject/Private/MuCO/UnrealMutableImageProvider.cpp
+++ b/Engine/Plugins/Experimental/Mutable/Source/CustomizableObject/Private/MuCO/UnrealMutableImageProvider.cpp
@@ -704,7 +704,7 @@ void FUnrealMutableImageProvider::CacheRuntimeReferencedImages(const TSharedRef<
 			continue;
 		}
 
-		ModelImages.SourceTextures.Emplace(*Texture);
+		ModelImages.SourceTextures.Emplace(*Texture); // Perform a CopyTornOff. Once done, we no longer need the texture loaded.
 	}
 }
 #endif
