[Mutable] Delete CO compiled data files not accessed in the last MaxAccessTimespan days from the Saved/MutableStreamedDataEditor folder.

#rnx

diff --git a/Engine/Plugins/Experimental/Mutable/Source/CustomizableObject/Internal/MuCO/CustomizableObjectPrivate.h b/Engine/Plugins/Experimental/Mutable/Source/CustomizableObject/Internal/MuCO/CustomizableObjectPrivate.h
--- a/Engine/Plugins/Experimental/Mutable/Source/CustomizableObject/Internal/MuCO/CustomizableObjectPrivate.h
+++ b/Engine/Plugins/Experimental/Mutable/Source/CustomizableObject/Internal/MuCO/CustomizableObjectPrivate.h
@@ -857,7 +857,7 @@ public:
 	FReply AddNewParameterProfile(FString Name, class UCustomizableObjectInstance& CustomInstance);
 
 	// Compose folder name where the data is stored
-	FString GetCompiledDataFolderPath() const;
+	static FString GetCompiledDataFolderPath();
 
 	/** Generic Save/Load methods to write/read compiled data */
 	void SaveCompiledData(FArchive& Ar, bool bSkipEditorOnlyData = false);
diff --git a/Engine/Plugins/Experimental/Mutable/Source/CustomizableObject/Private/MuCO/CustomizableObject.cpp b/Engine/Plugins/Experimental/Mutable/Source/CustomizableObject/Private/MuCO/CustomizableObject.cpp
--- a/Engine/Plugins/Experimental/Mutable/Source/CustomizableObject/Private/MuCO/CustomizableObject.cpp
+++ b/Engine/Plugins/Experimental/Mutable/Source/CustomizableObject/Private/MuCO/CustomizableObject.cpp
@@ -995,7 +995,7 @@ FReply UCustomizableObjectPrivate::AddNewParameterProfile(FString Name, UCustomi
 }
 
 
-FString UCustomizableObjectPrivate::GetCompiledDataFolderPath() const
+FString UCustomizableObjectPrivate::GetCompiledDataFolderPath()
 {	
 	return FPaths::ConvertRelativePathToFull(FPaths::ProjectSavedDir() + TEXT("MutableStreamedDataEditor/"));
 }
diff --git a/Engine/Plugins/Experimental/Mutable/Source/CustomizableObjectEditor/Private/MuCOE/CustomizableObjectEditorModule.cpp b/Engine/Plugins/Experimental/Mutable/Source/CustomizableObjectEditor/Private/MuCOE/CustomizableObjectEditorModule.cpp
--- a/Engine/Plugins/Experimental/Mutable/Source/CustomizableObjectEditor/Private/MuCOE/CustomizableObjectEditorModule.cpp
+++ b/Engine/Plugins/Experimental/Mutable/Source/CustomizableObjectEditor/Private/MuCOE/CustomizableObjectEditorModule.cpp
@@ -8,6 +8,9 @@
 #include "ISettingsSection.h"
 #include "MessageLogModule.h"
 #include "Misc/MessageDialog.h"
+#include "Misc/DateTime.h"
+#include "Misc/Timespan.h"
+#include "HAL/FileManager.h"
 #include "MuCO/CustomizableObjectSystem.h"		// For defines related to memory function replacements.
 #include "MuCO/CustomizableObjectInstanceUsage.h"
 #include "MuCO/CustomizableSkeletalMeshActor.h"
@@ -84,6 +87,8 @@ const FName CustomizableObjectDebuggerAppIdentifier = FName(TEXT("CustomizableOb
 
 #define LOCTEXT_NAMESPACE "MutableSettings"
 
+/** Max timespan in days before a Saved/MutableStreamedDataEditor file is deleted. */
+constexpr int32 MaxAccessTimespan = 30;
 
 constexpr float ShowOnScreenCompileWarningsTickerTime = 1.0f;
 
@@ -173,6 +178,35 @@ void ShowOnScreenCompileWarnings()
 	}
 }
 
+void DeleteUnusedMutableStreamedDataEditorFiles()
+{
+	const FDateTime CurrentTime = FDateTime::Now();
+
+	const FString CompiledDataFolder = UCustomizableObjectPrivate::GetCompiledDataFolderPath();
+	const FString FileExtension = TEXT(".mut");
+
+	TArray<FString> Files;
+	IFileManager& FileManager = IFileManager::Get();
+	FileManager.FindFiles(Files, *CompiledDataFolder, *FileExtension);
+	
+	for (const FString& File : Files)
+	{
+		const FString FullFilePath = CompiledDataFolder + File;
+		const FDateTime AccessTimeStamp = FileManager.GetAccessTimeStamp(*(FullFilePath));
+		if (AccessTimeStamp == FDateTime::MinValue())
+		{
+			continue;
+		}
+
+		// Delete files that remain unused for more than MaxAccessTimespan
+		const FTimespan TimeSpan = CurrentTime - AccessTimeStamp;
+		if (TimeSpan.GetDays() > MaxAccessTimespan)
+		{
+			FileManager.Delete(*FullFilePath);
+		}
+	}
+}
+
 
 IMPLEMENT_MODULE( FCustomizableObjectEditorModule, CustomizableObjectEditor );
 
@@ -191,6 +225,9 @@ static void CustomFree(void* mem)
 
 void FCustomizableObjectEditorModule::StartupModule()
 {
+	// Delete unused local compiled data
+	DeleteUnusedMutableStreamedDataEditorFiles();
+
 	// Register the thumbnail renderers
 	//UThumbnailManager::Get().RegisterCustomRenderer(UCustomizableObject::StaticClass(), UCustomizableObjectThumbnailRenderer::StaticClass());
 	//UThumbnailManager::Get().RegisterCustomRenderer(UCustomizableObjectInstance::StaticClass(), UCustomizableObjectInstanceThumbnailRenderer::StaticClass());
