[mutable] Invoke GC at key points during the execution of the testing commandlets.

diff --git a/Engine/Plugins/Experimental/Mutable/Source/CustomizableObjectEditor/Private/MuCOE/CustomizableObjectCompiler.h b/Engine/Plugins/Experimental/Mutable/Source/CustomizableObjectEditor/Private/MuCOE/CustomizableObjectCompiler.h
--- a/Engine/Plugins/Experimental/Mutable/Source/CustomizableObjectEditor/Private/MuCOE/CustomizableObjectCompiler.h
+++ b/Engine/Plugins/Experimental/Mutable/Source/CustomizableObjectEditor/Private/MuCOE/CustomizableObjectCompiler.h
@@ -27,7 +27,7 @@ class FCustomizableObjectCompiler : public FTickableEditorObject, public FTickab
 public:
 
 	FCustomizableObjectCompiler();
-	virtual ~FCustomizableObjectCompiler() {}
+	virtual ~FCustomizableObjectCompiler() override {}
 
 	/** Check for pending compilation process. Returns true if an object has been updated. */
 	bool Tick(bool bBlocking = false);
diff --git a/Engine/Plugins/Experimental/Mutable/Source/MutableValidation/Private/MuV/COIBulkUpdateTestCommandlet.cpp b/Engine/Plugins/Experimental/Mutable/Source/MutableValidation/Private/MuV/COIBulkUpdateTestCommandlet.cpp
--- a/Engine/Plugins/Experimental/Mutable/Source/MutableValidation/Private/MuV/COIBulkUpdateTestCommandlet.cpp
+++ b/Engine/Plugins/Experimental/Mutable/Source/MutableValidation/Private/MuV/COIBulkUpdateTestCommandlet.cpp
@@ -63,7 +63,7 @@ int32 UCOIBulkUpdateTestCommandlet::Main(const FString& Params)
 	}
 
 	// Get all the CO's that need compilation before proceeding
-	TMap<UCustomizableObject*, TSet<UCustomizableObjectInstance*>>  MutableResources;
+	TMap<UCustomizableObject*, TArray<TStrongObjectPtr<UCustomizableObjectInstance>>>  MutableResources;
 	for (FAssetData& Data : FoundAssetData)
 	{
 		UObject* LoadedAsset =  Data.GetAsset();
@@ -80,8 +80,8 @@ int32 UCOIBulkUpdateTestCommandlet::Main(const FString& Params)
 		}
 		
 		// Add/update an entry with the new instance for a given CO. 
-		TSet<UCustomizableObjectInstance*>& Entry = MutableResources.FindOrAdd(InstanceCO);
-		Entry.Add(LoadedInstance);
+		TArray<TStrongObjectPtr<UCustomizableObjectInstance>>& Entry = MutableResources.FindOrAdd(InstanceCO);
+		Entry.Emplace(LoadedInstance);
 	}
 
 	// At this point it is safe to assume that all keys are valid COs and all instances are also valid.
@@ -97,7 +97,7 @@ int32 UCOIBulkUpdateTestCommandlet::Main(const FString& Params)
 	UE_LOG(LogMutable, Log, TEXT("(int) customizable_object_instances_count : %u "), TotalAmountOfInstances);
 	
 	// Report the amount of instances for each of the COs found as parents of the instances in the target path
-	for (TTuple<UCustomizableObject*, TSet<UCustomizableObjectInstance*>>& MutableResourceTuple : MutableResources)
+	for (TTuple<UCustomizableObject*, TArray<TStrongObjectPtr<UCustomizableObjectInstance>>>& MutableResourceTuple : MutableResources)
 	{
 		const UCustomizableObject* CustomizableObjectToCompile = MutableResourceTuple.Key;
 		check (CustomizableObjectToCompile);
@@ -109,7 +109,7 @@ int32 UCOIBulkUpdateTestCommandlet::Main(const FString& Params)
 		
 		// print the name of the instances+
 		uint32 InstanceIndex = 0;
-		for (const UCustomizableObjectInstance* Instance : MutableResourceTuple.Value)
+		for (const TStrongObjectPtr<UCustomizableObjectInstance>& Instance : MutableResourceTuple.Value)
 		{
 			check(Instance);
 			const FString CustomizableObjectInstanceName = Instance->GetName();
@@ -131,38 +131,48 @@ int32 UCOIBulkUpdateTestCommandlet::Main(const FString& Params)
 	
 	// Compile all found COs one by one
 	uint32 CurrentInstanceIndex = 1;
-	for (TTuple<UCustomizableObject*, TSet<UCustomizableObjectInstance*>>& MutableResourceTuple : MutableResources)
+	
+	for ( TMap<UCustomizableObject*, TArray<TStrongObjectPtr<UCustomizableObjectInstance>>>::TIterator ResourcesIterator = MutableResources.CreateIterator(); ResourcesIterator; ++ResourcesIterator)
 	{
-		UCustomizableObject* CustomizableObjectToCompile = MutableResourceTuple.Key;
-		check (CustomizableObjectToCompile);
+		UCustomizableObject*& CustomizableObject = ResourcesIterator->Key;
+		check (CustomizableObject);
 
-		const FString CustomizableObjectName = CustomizableObjectToCompile->GetName();
+		const FString CustomizableObjectName = CustomizableObject->GetName();
 		
 		// Set the compilation platform based on what the system is currently running on
-		FCompilationOptions CompilationOptions = CustomizableObjectToCompile->CompileOptions;
+		FCompilationOptions CompilationOptions = CustomizableObject->CompileOptions;
 		CompilationOptions.TargetPlatform = TargetCompilationPlatform;
 		CompilationOptions.bUseDiskCompilation = false;
-		
-		// Compile the current CO object
+		CompilationOptions.OptimizationLevel = UE_MUTABLE_MAX_OPTIMIZATION;
+		CompilationOptions.bSilentCompilation = false;
 
-		if (!CompilationUtility->CompileCustomizableObject(CustomizableObjectToCompile, false, &CompilationOptions))	// Do not log mutable data since mongoDB will not be able to handle it correctly 
+		// Compile the current CO object
+		if (!CompilationUtility->CompileCustomizableObject(CustomizableObject, false, &CompilationOptions))	// Do not log mutable data since mongoDB will not be able to handle it correctly 
 		{
 			UE_LOG(LogMutable,Error,TEXT("The CO %s could not be compiled succesfully. Skipping the update of all COIs that use it."), *CustomizableObjectName )
 			continue;
 		}
 
 		// Now that CO has been compiled proceed with the update of the instances that use it
-		const uint32 CustomizableObjectInstanceCount = MutableResourceTuple.Value.Num();
+		const uint32 CustomizableObjectInstanceCount = ResourcesIterator->Value.Num();
 		UE_LOG(LogMutable,Display,TEXT("Starting update of the \"%u\" instances with CO : \"%s\"."), CustomizableObjectInstanceCount, *CustomizableObjectName);
-		
+
 		// Iterate over all the COIs of the CO and update them
-		for (UCustomizableObjectInstance* Instance : MutableResourceTuple.Value)
+		TArray<TStrongObjectPtr<UCustomizableObjectInstance>>& Instances = ResourcesIterator->Value;
+		while (!Instances.IsEmpty())
 		{
+			TStrongObjectPtr<UCustomizableObjectInstance> Instance = Instances[0];
+			Instances.RemoveAt(0);
 			UE_LOG(LogMutable,Display,TEXT("\t( %u / %u ) Processing instance : \"%s\" ."),CurrentInstanceIndex++, TotalAmountOfInstances ,*Instance->GetName());
+
+			CollectGarbage(RF_NoFlags, true);
 			
 			// Update each one of the instances and notify if the update failed in any manner
-			InstanceUpdatingUtility->UpdateInstance(Instance);
+			InstanceUpdatingUtility->UpdateInstance(Instance.Get());
 		}
+
+		ResourcesIterator.RemoveCurrent();
+		CollectGarbage(RF_NoFlags, true);
 	}
 
 	UE_LOG(LogMutable,Display,TEXT("Mutable commandlet finished."));
diff --git a/Engine/Plugins/Experimental/Mutable/Source/MutableValidation/Private/MuV/CustomizableObjectCompilationUtility.cpp b/Engine/Plugins/Experimental/Mutable/Source/MutableValidation/Private/MuV/CustomizableObjectCompilationUtility.cpp
--- a/Engine/Plugins/Experimental/Mutable/Source/MutableValidation/Private/MuV/CustomizableObjectCompilationUtility.cpp
+++ b/Engine/Plugins/Experimental/Mutable/Source/MutableValidation/Private/MuV/CustomizableObjectCompilationUtility.cpp
@@ -15,6 +15,8 @@
 
 bool FCustomizableObjectCompilationUtility::CompileCustomizableObject(UCustomizableObject* InCustomizableObject, const bool bShouldLogMutableLogs /* = true */, const FCompilationOptions* InCompilationOptionsOverride  /* nullptr */)
 {
+	LLM_SCOPE_BYNAME(TEXT("FCustomizableObjectCompilationUtility/Compile"));
+	
 	check(InCustomizableObject);
 	
 	CustomizableObject = TStrongObjectPtr(InCustomizableObject);
@@ -93,6 +95,8 @@ bool FCustomizableObjectCompilationUtility::CompileCustomizableObject(UCustomiza
 		const double CompilationStartSeconds = FPlatformTime::Seconds();
 		while (CompileRequest->GetCompilationState() != ECompilationStatePrivate::Completed)
 		{
+			LLM_SCOPE_BYNAME(TEXT("FCustomizableObjectCompilationUtility/CompileLoop"));
+			
 			// Tick the engine
 			CommandletHelpers::TickEngine();
 
diff --git a/Engine/Plugins/Experimental/Mutable/Source/MutableValidation/Private/MuV/CustomizableObjectInstanceUpdateUtility.cpp b/Engine/Plugins/Experimental/Mutable/Source/MutableValidation/Private/MuV/CustomizableObjectInstanceUpdateUtility.cpp
--- a/Engine/Plugins/Experimental/Mutable/Source/MutableValidation/Private/MuV/CustomizableObjectInstanceUpdateUtility.cpp
+++ b/Engine/Plugins/Experimental/Mutable/Source/MutableValidation/Private/MuV/CustomizableObjectInstanceUpdateUtility.cpp
@@ -12,6 +12,8 @@
 
 bool FCustomizableObjectInstanceUpdateUtility::UpdateInstance(UCustomizableObjectInstance* InInstance)
 {
+	LLM_SCOPE_BYNAME(TEXT("FCustomizableObjectInstanceUpdateUtility/UpdateInstance"));
+	
 	check (InInstance);
 	check (ComponentsBeingUpdated.IsEmpty());
 	
@@ -20,23 +22,30 @@ bool FCustomizableObjectInstanceUpdateUtility::UpdateInstance(UCustomizableObjec
 	
 	// Schedule the update of the COI
 	{
-		UE_LOG(LogMutable, Display, TEXT("Invoking update for %s instance."), *InInstance->GetName());
+		UE_LOG(LogMutable, Display, TEXT("Invoking update for %s instance."), *Instance->GetName());
 
 		// Instance update delegate	
 		FInstanceUpdateNativeDelegate InstanceUpdateNativeDelegate;
 		InstanceUpdateNativeDelegate.AddSP(this, &FCustomizableObjectInstanceUpdateUtility::OnInstanceUpdateResult);
 
 		bIsInstanceBeingUpdated = true;
-		InInstance->UpdateSkeletalMeshAsyncResult(InstanceUpdateNativeDelegate, true, true);
+		Instance->UpdateSkeletalMeshAsyncResult(InstanceUpdateNativeDelegate, true, true);
 	}
-	
+
+	// Debug
+	// FLowLevelMemTracker& LowLevelMemoryTracker = FLowLevelMemTracker::Get();
 	
 	// Wait until the update has been completed and the mips streamed
 	while (bIsInstanceBeingUpdated)
 	{
+		LLM_SCOPE_BYNAME(TEXT("FCustomizableObjectInstanceUpdateUtility/UpdateLoop"));
+		
 		// Tick the engine
 		CommandletHelpers::TickEngine();
 
+		// Debug
+		// LowLevelMemoryTracker.UpdateStatsPerFrame();
+		
 		// Stop if exit was requested
 		if (IsEngineExitRequested())
 		{
@@ -66,7 +75,7 @@ bool FCustomizableObjectInstanceUpdateUtility::UpdateInstance(UCustomizableObjec
 
 			if (bFullyStreamed)
 			{
-				UE_LOG(LogMutable, Display, TEXT("Instance %s finished streaming all MIPs."), *InInstance->GetName());
+				UE_LOG(LogMutable, Display, TEXT("Instance %s finished streaming all MIPs."), *Instance->GetName());
 				ComponentsBeingUpdated.Reset();
 				
 				bIsInstanceBeingUpdated = false;		// Exit the while loop
@@ -81,6 +90,8 @@ bool FCustomizableObjectInstanceUpdateUtility::UpdateInstance(UCustomizableObjec
 
 void FCustomizableObjectInstanceUpdateUtility::OnInstanceUpdateResult(const FUpdateContext& Result)
 {
+	LLM_SCOPE_BYNAME(TEXT("FCustomizableObjectInstanceUpdateUtility/OnInstanceUpdated"));
+	
 	const FString InstanceName = Instance->GetName();
 	
 	if (UCustomizableObjectSystem::IsUpdateResultValid(Result.UpdateResult))
@@ -122,4 +133,5 @@ void FCustomizableObjectInstanceUpdateUtility::OnInstanceUpdateResult(const FUpd
 		// Tell the system the instance finished it's update so we can continue the execution without waitting for the mips to stream in
 		bIsInstanceBeingUpdated = false;
 	}	
-}
\ No newline at end of file
+}
+
diff --git a/Engine/Plugins/Experimental/Mutable/Source/MutableValidation/Private/MuV/CustomizableObjectValidationCommandlet.cpp b/Engine/Plugins/Experimental/Mutable/Source/MutableValidation/Private/MuV/CustomizableObjectValidationCommandlet.cpp
--- a/Engine/Plugins/Experimental/Mutable/Source/MutableValidation/Private/MuV/CustomizableObjectValidationCommandlet.cpp
+++ b/Engine/Plugins/Experimental/Mutable/Source/MutableValidation/Private/MuV/CustomizableObjectValidationCommandlet.cpp
@@ -7,6 +7,7 @@
 #include "ValidationUtils.h"
 #include "Components/SkeletalMeshComponent.h"
 #include "Containers/Array.h"
+#include "Containers/SpscQueue.h"
 #include "MuCO/CustomizableObject.h"
 #include "MuCO/CustomizableObjectInstancePrivate.h"
 #include "MuCO/CustomizableObjectPrivate.h"
@@ -20,6 +21,8 @@
 
 int32 UCustomizableObjectValidationCommandlet::Main(const FString& Params)
 {
+	LLM_SCOPE_BYNAME(TEXT("CustomizableObjectValidationCommandlet"));
+	
 	// Execution arguments for commandlet from IDE
 	// -run=CustomizableObjectValidation -CustomizableObject=(PathToCO)
 
@@ -76,10 +79,12 @@ int32 UCustomizableObjectValidationCommandlet::Main(const FString& Params)
 	// Compile the Customizable Object ------------------------------------------------------------------------------ //
 	bool bWasCoCompilationSuccessful = false;
 	{
+		LLM_SCOPE_BYNAME(TEXT("CustomizableObjectValidationCommandlet/Compile"));
+		
 		// Override some configurations that may have been changed by the user
 		FCompilationOptions CompilationOptions = ToTestCustomizableObject->CompileOptions;
 		CompilationOptions.bSilentCompilation = false;
-		CompilationOptions.OptimizationLevel = 2;			// Set the optimization level to the max
+		CompilationOptions.OptimizationLevel = UE_MUTABLE_MAX_OPTIMIZATION;
 		CompilationOptions.TextureCompression = ECustomizableObjectTextureCompression::Fast;
 
 		// Set the target compilation platform based on what the caller wants
@@ -150,7 +155,10 @@ int32 UCustomizableObjectValidationCommandlet::Main(const FString& Params)
 
 	// Generate target random instances to be tested ------------------------------------------------------------ //
 	bool bWasInstancesCreationSuccessful = true;
+	uint32 GeneratedInstances = 0;
 	{
+		LLM_SCOPE_BYNAME(TEXT("CustomizableObjectValidationCommandlet/GenerateInstances"));
+		
 		// Test this parameter configuration in all the states of the CO
 		const uint32 StateCount = ToTestCustomizableObject->GetStateCount();
 		check(StateCount >= 1);
@@ -183,7 +191,8 @@ int32 UCustomizableObjectValidationCommandlet::Main(const FString& Params)
 				{
 					// Set the state for the instance and store it for later update.
 					GeneratedInstance->GetPrivate()->SetState(State);
-                	InstancesToProcess.Push(GeneratedInstance->Clone());
+                	InstancesToProcess.Enqueue(GeneratedInstance->Clone());
+					GeneratedInstances++;
 				}
 			}
 			else
@@ -195,17 +204,23 @@ int32 UCustomizableObjectValidationCommandlet::Main(const FString& Params)
 	}
 	// ---------------------------------------------------------------------------------------------------------- //
 
-	UE_LOG(LogMutable, Log,TEXT("(int) generated_instances_count : %u "), InstancesToProcess.Num());
+	UE_LOG(LogMutable, Log,TEXT("(int) generated_instances_count : %u "), GeneratedInstances);
 	
 	// Update the instances generated --------------------------------------------------------------------------- //
 	UE_LOG(LogMutable,Display,TEXT("Updating generated instances..."));
 	bool bInstanceFailedUpdate = false;
 	const double InstancesUpdateStartSeconds = FPlatformTime::Seconds();
 	{
+		LLM_SCOPE_BYNAME(TEXT("CustomizableObjectValidationCommandlet/Update"));
+		
 		TSharedRef<FCustomizableObjectInstanceUpdateUtility> InstanceUpdatingUtility = MakeShared<FCustomizableObjectInstanceUpdateUtility>();
-		for (UCustomizableObjectInstance* InstanceToUpdate : InstancesToProcess)
+
+		TStrongObjectPtr<UCustomizableObjectInstance> InstanceToUpdate;
+		while (InstancesToProcess.Dequeue(InstanceToUpdate))
 		{
-			if (!InstanceUpdatingUtility->UpdateInstance(InstanceToUpdate))
+			CollectGarbage(RF_NoFlags, true);
+			
+			if (!InstanceUpdatingUtility->UpdateInstance(InstanceToUpdate.Get()))
 			{
 				bInstanceFailedUpdate = true;
 			}
@@ -217,8 +232,8 @@ int32 UCustomizableObjectValidationCommandlet::Main(const FString& Params)
 	const double CombinedInstanceUpdateSeconds = InstancesUpdateEndSeconds - InstancesUpdateStartSeconds;
 	UE_LOG(LogMutable, Log,TEXT("(double) combined_update_time_ms : %f "), CombinedInstanceUpdateSeconds * 1000);
 
-	check(InstancesToProcess.Num() > 0);
-	const double AverageInstanceUpdateSeconds = CombinedInstanceUpdateSeconds / InstancesToProcess.Num();
+	check(GeneratedInstances > 0);
+	const double AverageInstanceUpdateSeconds = CombinedInstanceUpdateSeconds / GeneratedInstances;
 	UE_LOG(LogMutable, Log,TEXT("(double) avg_update_time_ms : %f "), AverageInstanceUpdateSeconds * 1000);
 
 	UE_LOG(LogMutable,Display,TEXT("Generation of Customizable object instances took %f seconds (%f seconds avg)."), CombinedInstanceUpdateSeconds, AverageInstanceUpdateSeconds);
diff --git a/Engine/Plugins/Experimental/Mutable/Source/MutableValidation/Public/MuV/CustomizableObjectValidationCommandlet.h b/Engine/Plugins/Experimental/Mutable/Source/MutableValidation/Public/MuV/CustomizableObjectValidationCommandlet.h
--- a/Engine/Plugins/Experimental/Mutable/Source/MutableValidation/Public/MuV/CustomizableObjectValidationCommandlet.h
+++ b/Engine/Plugins/Experimental/Mutable/Source/MutableValidation/Public/MuV/CustomizableObjectValidationCommandlet.h
@@ -3,7 +3,7 @@
 #pragma once
 
 #include "Commandlets/Commandlet.h"
-#include "Containers/Array.h"
+#include "Containers/SpscQueue.h"
 #include "MuV/CustomizableObjectCompilationUtility.h"
 #include "MuV/CustomizableObjectInstanceUpdateUtility.h"
 
@@ -37,6 +37,5 @@ private:
 	TObjectPtr<UCustomizableObject> ToTestCustomizableObject = nullptr;
 
 	/** Array of COI to be generated with randomized parameter values */
-	UPROPERTY()
-	TArray<TObjectPtr<UCustomizableObjectInstance>> InstancesToProcess;
+	TSpscQueue<TStrongObjectPtr<UCustomizableObjectInstance>> InstancesToProcess;
 };
\ No newline at end of file
