[mutable] Initial support for pass-through meshes, referenced meshes and mesh parameters at the core level.

- meshes now have the flags to check if they are a reference.
- pass-through meshes should now be possible
- stubs have been created for the callbacks at UE level that would be need to be implemented for mesh parameters and referenced meshes.
- removed support for serialization of old versions in Meshes.

#rnx

diff --git a/Engine/Plugins/Experimental/Mutable/Source/CustomizableObject/Internal/MuCO/CustomizableObjectInstancePrivate.h b/Engine/Plugins/Experimental/Mutable/Source/CustomizableObject/Internal/MuCO/CustomizableObjectInstancePrivate.h
--- a/Engine/Plugins/Experimental/Mutable/Source/CustomizableObject/Internal/MuCO/CustomizableObjectInstancePrivate.h
+++ b/Engine/Plugins/Experimental/Mutable/Source/CustomizableObject/Internal/MuCO/CustomizableObjectInstancePrivate.h
@@ -350,8 +350,9 @@ public:
 	UPROPERTY(Transient, Category = Animation, editfixedsize, VisibleAnywhere)
 	TMap<TSubclassOf<UAnimInstance>, FAnimBpGeneratedPhysicsAssets> AnimBpPhysicsAssets;
 
-	// The pass-through textures that will be loaded during an update
+	// The pass-through assets that will be loaded during an update
 	TArray<TSoftObjectPtr<UTexture>> PassThroughTexturesToLoad;
+	TArray<TSoftObjectPtr<USkeletalMesh>> PassThroughMeshesToLoad;
 
 	// Used during an update to prevent the pass-through textures loaded by LoadAdditionalAssetsAsync() from being unloaded by GC
 	// between AdditionalAssetsAsyncLoaded() and their setting into the generated materials in BuildMaterials()
diff --git a/Engine/Plugins/Experimental/Mutable/Source/CustomizableObject/Internal/MuCO/CustomizableObjectPrivate.h b/Engine/Plugins/Experimental/Mutable/Source/CustomizableObject/Internal/MuCO/CustomizableObjectPrivate.h
--- a/Engine/Plugins/Experimental/Mutable/Source/CustomizableObject/Internal/MuCO/CustomizableObjectPrivate.h
+++ b/Engine/Plugins/Experimental/Mutable/Source/CustomizableObject/Internal/MuCO/CustomizableObjectPrivate.h
@@ -39,7 +39,7 @@ FGuid CUSTOMIZABLEOBJECT_API GenerateIdentifier(const UCustomizableObject& Custo
 // Warning! MutableCompiledDataHeader must be the first data serialized in a stream
 struct MutableCompiledDataStreamHeader
 {
-	int32 InternalVersion;
+	int32 InternalVersion=0;
 	FGuid VersionId;
 
 	MutableCompiledDataStreamHeader() { }
@@ -667,6 +667,10 @@ struct FModelResources
 	UPROPERTY()
 	TArray<TSoftObjectPtr<UTexture>> PassThroughTextures;
 
+	/** PassThrough meshes used by the mu::Model. */
+	UPROPERTY()
+	TArray<TSoftObjectPtr<USkeletalMesh>> PassThroughMeshes;
+
 #if WITH_EDITORONLY_DATA
 	/** Runtime referenced textures used by the mu::Model. */
 	UPROPERTY()
@@ -993,6 +997,8 @@ public:
 
 		FixWrappingProjectorLayoutBlockId,
 
+		MeshReferenceSupport,
+
 		LastCustomizableObjectVersion
 	};
 	static constexpr int32 CurrentSupportedVersion = ECustomizableObjectVersions::LastCustomizableObjectVersion;
diff --git a/Engine/Plugins/Experimental/Mutable/Source/CustomizableObject/Internal/MuCO/CustomizableObjectSystemPrivate.h b/Engine/Plugins/Experimental/Mutable/Source/CustomizableObject/Internal/MuCO/CustomizableObjectSystemPrivate.h
--- a/Engine/Plugins/Experimental/Mutable/Source/CustomizableObject/Internal/MuCO/CustomizableObjectSystemPrivate.h
+++ b/Engine/Plugins/Experimental/Mutable/Source/CustomizableObject/Internal/MuCO/CustomizableObjectSystemPrivate.h
@@ -631,10 +631,10 @@ public:
 
 	TSharedPtr<class FUnrealExtensionDataStreamer> ExtensionDataStreamer;
 
-	// This object is responsible for providing custom images to mutable models (for image parameters)
+	// This object is responsible for providing custom images and meshes to mutable (for image parameters, etc.)
 	// This object is called from the mutable thread, and it should only access data already safely submitted from
 	// the game thread and stored in FUnrealMutableImageProvider::GlobalExternalImages.
-	TSharedPtr<class FUnrealMutableImageProvider> ImageProvider;
+	TSharedPtr<class FUnrealMutableResourceProvider> ResourceProvider;
 
 	// Cache of weak references to generated resources to see if they can be reused.
 	TArray<FMutableResourceCache> ModelResourcesCache;
@@ -737,7 +737,7 @@ public:
 
 	bool IsMutableAnimInfoDebuggingEnabled() const;
 
-	FUnrealMutableImageProvider* GetImageProviderChecked() const;
+	FUnrealMutableResourceProvider* GetResourceProviderChecked() const;
 
 	/** Start the actual work of Update Skeletal Mesh process (Update Skeletal Mesh without the queue). */
 	void StartUpdateSkeletalMesh(const TSharedRef<FUpdateContextPrivate>& Context);
diff --git a/Engine/Plugins/Experimental/Mutable/Source/CustomizableObject/Private/MuCO/BusyWaits_Deprecated.cpp b/Engine/Plugins/Experimental/Mutable/Source/CustomizableObject/Private/MuCO/BusyWaits_Deprecated.cpp
--- a/Engine/Plugins/Experimental/Mutable/Source/CustomizableObject/Private/MuCO/BusyWaits_Deprecated.cpp
+++ b/Engine/Plugins/Experimental/Mutable/Source/CustomizableObject/Private/MuCO/BusyWaits_Deprecated.cpp
@@ -236,6 +236,22 @@ namespace CustomizableObjectSystem::ImplDeprecated
 					if (bGenerateLOD)
 					{
 						Component.Mesh = System->GetMeshInline(OperationData->InstanceID, Component.MeshID);
+
+						if (Component.Mesh->IsReference())
+						{
+							uint32 ReferenceID = Component.Mesh->GetReferencedMesh();
+
+							if (ModelResources.PassThroughMeshes.IsValidIndex(ReferenceID))
+							{
+								TSoftObjectPtr<USkeletalMesh> Ref = ModelResources.PassThroughMeshes[ReferenceID];
+								CustomizableObjectInstancePrivateData->PassThroughMeshesToLoad.Add(Ref);
+							}
+							else
+							{
+								// internal error.
+								UE_LOG(LogMutable, Error, TEXT("Referenced mesh [%d] was not stored in the resource array."), ReferenceID);
+							}
+						}
 					}
 				}
 
diff --git a/Engine/Plugins/Experimental/Mutable/Source/CustomizableObject/Private/MuCO/CustomizableObject.cpp b/Engine/Plugins/Experimental/Mutable/Source/CustomizableObject/Private/MuCO/CustomizableObject.cpp
--- a/Engine/Plugins/Experimental/Mutable/Source/CustomizableObject/Private/MuCO/CustomizableObject.cpp
+++ b/Engine/Plugins/Experimental/Mutable/Source/CustomizableObject/Private/MuCO/CustomizableObject.cpp
@@ -541,6 +541,15 @@ void UCustomizableObjectPrivate::SaveCompiledData(FArchive& MemoryWriter, bool b
 		MemoryWriter << StringRef;
 	}
 
+	int32 NumPassthroughMeshes = LocalModelResources.PassThroughMeshes.Num();
+	MemoryWriter << NumPassthroughMeshes;
+
+	for (const TSoftObjectPtr<USkeletalMesh>& PassthroughMesh : LocalModelResources.PassThroughMeshes)
+	{
+		FString StringRef = PassthroughMesh.ToString();
+		MemoryWriter << StringRef;
+	}
+
 #if WITH_EDITORONLY_DATA
 	int32 NumRuntimeReferencedTextures = LocalModelResources.RuntimeReferencedTextures.Num();
 	MemoryWriter << NumRuntimeReferencedTextures;
@@ -669,6 +678,17 @@ void UCustomizableObjectPrivate::LoadCompiledData(FArchive& MemoryReader, const
 			LocalModelResource.PassThroughTextures.Add(TSoftObjectPtr<UTexture>(FSoftObjectPath(StringRef)));
 		}
 
+		int32 NumPassthroughMeshes = 0;
+		MemoryReader << NumPassthroughMeshes;
+
+		for (int32 Index = 0; Index < NumPassthroughMeshes; ++Index)
+		{
+			FString StringRef;
+			MemoryReader << StringRef;
+
+			LocalModelResource.PassThroughMeshes.Add(TSoftObjectPtr<USkeletalMesh>(FSoftObjectPath(StringRef)));
+		}
+
 #if WITH_EDITORONLY_DATA
 		int32 NumRuntimeReferencedTextures = 0;
 		MemoryReader << NumRuntimeReferencedTextures;
diff --git a/Engine/Plugins/Experimental/Mutable/Source/CustomizableObject/Private/MuCO/CustomizableObjectInstance.cpp b/Engine/Plugins/Experimental/Mutable/Source/CustomizableObject/Private/MuCO/CustomizableObjectInstance.cpp
--- a/Engine/Plugins/Experimental/Mutable/Source/CustomizableObject/Private/MuCO/CustomizableObjectInstance.cpp
+++ b/Engine/Plugins/Experimental/Mutable/Source/CustomizableObject/Private/MuCO/CustomizableObjectInstance.cpp
@@ -1897,13 +1897,18 @@ bool UCustomizableInstancePrivate::DoComponentsNeedUpdate(UCustomizableObjectIns
 				return false;
 			}
 
-			if (!Component.bGenerated || !Component.Mesh || Component.SurfaceCount == 0) // else
+			if (!Component.bGenerated || !Component.Mesh ) 
+			{
+				continue;
+			}
+
+			if (Component.SurfaceCount == 0 && !Component.Mesh->IsReference()) 
 			{
 				continue;
 			}
 
 			// Unreal does not support empty sections.
-			if (Component.Mesh->GetVertexCount() == 0) // else
+			if (!Component.Mesh->IsReference() && Component.Mesh->GetVertexCount() == 0) // else
 			{
 				UE_LOG(LogMutable, Error, TEXT("Failed to generate SkeletalMesh for CO Instance [%s]. CO [%s] has invalid geometry for LOD [%d] Component [%d]."),
 					*Public->GetName(), *CustomizableObject->GetName(),
@@ -2057,11 +2062,27 @@ bool UCustomizableInstancePrivate::UpdateSkeletalMesh_PostBeginUpdate0(UCustomiz
 		// Reset last mesh IDs.
 		ComponentsData[Component.Id].LastMeshIdPerLOD.Init(MAX_uint64, MAX_MESH_LOD_COUNT);
 
-		if (!Component.bGenerated || !Component.Mesh || Component.SurfaceCount == 0)
+		if (!Component.bGenerated || !Component.Mesh)
 		{
 			continue;
 		}
 
+		if (Component.SurfaceCount == 0 && !Component.Mesh->IsReference())
+		{
+			continue;
+		}
+
+
+		if (Component.Mesh && Component.Mesh->IsReference())
+		{
+			// TODO: This shouldn't happen here synchronosuly. It should have been requested as an async load.
+			int32 ReferenceID = Component.Mesh->GetReferencedMesh();
+			TSoftObjectPtr<USkeletalMesh> Ref = ModelResources.PassThroughMeshes[ReferenceID];
+
+			SkeletalMeshes[Component.Id] = Ref.LoadSynchronous();
+			break;
+		}		
+
 		if (!ModelResources.ReferenceSkeletalMeshesData.IsValidIndex(Component.Id))
 		{
 			bSuccess = false;
@@ -2264,6 +2285,12 @@ bool UCustomizableInstancePrivate::UpdateSkeletalMesh_PostBeginUpdate0(UCustomiz
 		const int32 NumComponents = SkeletalMeshes.Num();
 		for (int32 ComponentIndex = 0; bSuccess && ComponentIndex < NumComponents; ++ComponentIndex)
 		{
+			const FInstanceUpdateData::FComponent& Component = OperationData->InstanceUpdateData.Components[ComponentIndex];
+			if (Component.Mesh && Component.Mesh->IsReference())
+			{
+				continue;
+			}
+
 			if (OperationData->bUseMeshCache)
 			{
 				const TArray<mu::FResourceID>& MeshId = OperationData->MeshDescriptors[ComponentIndex];
diff --git a/Engine/Plugins/Experimental/Mutable/Source/CustomizableObject/Private/MuCO/CustomizableObjectInstanceDescriptor.cpp b/Engine/Plugins/Experimental/Mutable/Source/CustomizableObject/Private/MuCO/CustomizableObjectInstanceDescriptor.cpp
--- a/Engine/Plugins/Experimental/Mutable/Source/CustomizableObject/Private/MuCO/CustomizableObjectInstanceDescriptor.cpp
+++ b/Engine/Plugins/Experimental/Mutable/Source/CustomizableObject/Private/MuCO/CustomizableObjectInstanceDescriptor.cpp
@@ -2367,7 +2367,7 @@ void FCustomizableObjectInstanceDescriptor::SetRandomValuesFromStream(const FRan
 
 	// Get all possible values
 	TArray<FCustomizableObjectExternalTexture> ProviderValues;
-	for (const TWeakObjectPtr<UCustomizableSystemImageProvider>& Provider : SystemPrivate->ImageProvider->ImageProviders)
+	for (const TWeakObjectPtr<UCustomizableSystemImageProvider>& Provider : SystemPrivate->ResourceProvider->ImageProviders)
 	{
 		ProviderValues.Reset();
 		Provider->GetTextureParameterValues(ProviderValues);
diff --git a/Engine/Plugins/Experimental/Mutable/Source/CustomizableObject/Private/MuCO/CustomizableObjectMipDataProvider.cpp b/Engine/Plugins/Experimental/Mutable/Source/CustomizableObject/Private/MuCO/CustomizableObjectMipDataProvider.cpp
--- a/Engine/Plugins/Experimental/Mutable/Source/CustomizableObject/Private/MuCO/CustomizableObjectMipDataProvider.cpp
+++ b/Engine/Plugins/Experimental/Mutable/Source/CustomizableObject/Private/MuCO/CustomizableObjectMipDataProvider.cpp
@@ -40,7 +40,7 @@ FMutableUpdateContext::FMutableUpdateContext(
 	if (Parameters)
 	{
 		const UCustomizableObjectSystemPrivate* Private = UCustomizableObjectSystem::GetInstance()->GetPrivate();
-		Private->GetImageProviderChecked()->CacheImages(*Parameters);
+		Private->GetResourceProviderChecked()->CacheImages(*Parameters);
 	}
 }
 
@@ -51,7 +51,7 @@ FMutableUpdateContext::~FMutableUpdateContext()
 		UCustomizableObjectSystem::IsCreated())
 	{
 		const UCustomizableObjectSystemPrivate* Private = UCustomizableObjectSystem::GetInstance()->GetPrivate();
-		Private->GetImageProviderChecked()->UnCacheImages(*Parameters);
+		Private->GetResourceProviderChecked()->UnCacheImages(*Parameters);
 	}
 }
 
diff --git a/Engine/Plugins/Experimental/Mutable/Source/CustomizableObject/Private/MuCO/CustomizableObjectSystem.cpp b/Engine/Plugins/Experimental/Mutable/Source/CustomizableObject/Private/MuCO/CustomizableObjectSystem.cpp
--- a/Engine/Plugins/Experimental/Mutable/Source/CustomizableObject/Private/MuCO/CustomizableObjectSystem.cpp
+++ b/Engine/Plugins/Experimental/Mutable/Source/CustomizableObject/Private/MuCO/CustomizableObjectSystem.cpp
@@ -548,10 +548,10 @@ void UCustomizableObjectSystem::InitSystem()
 	Private->MutableSystem->SetStreamingInterface(Private->Streamer);
 
 	// Set up the external image provider, for image parameters.
-	TSharedPtr<FUnrealMutableImageProvider> Provider = MakeShared<FUnrealMutableImageProvider>();
+	TSharedPtr<FUnrealMutableResourceProvider> Provider = MakeShared<FUnrealMutableResourceProvider>();
 	check(Provider != nullptr);
-	Private->ImageProvider = Provider;
-	Private->MutableSystem->SetImageParameterGenerator(Provider);
+	Private->ResourceProvider = Provider;
+	Private->MutableSystem->SetExternalResourceProvider(Provider);
 
 #if WITH_EDITORONLY_DATA
 	Private->EditorImageProvider = NewObject<UEditorImageProvider>();
@@ -1736,7 +1736,7 @@ namespace impl
 		for (FInstanceUpdateData::FComponent& Component : OperationData->InstanceUpdateData.Components)
 		{
 			mu::MeshPtrConst Mesh = Component.Mesh;
-			if (!Mesh)
+			if (!Mesh || Mesh->IsReference())
 			{
 				continue;
 			}
@@ -3278,7 +3278,7 @@ namespace impl
 		Operation->Model = CustomizableObject->GetPrivate()->GetModel().ToSharedRef();
 
 #if WITH_EDITOR
-		SystemPrivateData->GetImageProviderChecked()->CacheRuntimeReferencedImages(Operation->Model.ToSharedRef(), CustomizableObject->GetPrivate()->GetModelResources().RuntimeReferencedTextures);
+		SystemPrivateData->GetResourceProviderChecked()->CacheRuntimeReferencedImages(Operation->Model.ToSharedRef(), CustomizableObject->GetPrivate()->GetModelResources().RuntimeReferencedTextures);
 #endif
 		
 		// Task: Mutable Update and GetMesh
@@ -3792,7 +3792,7 @@ TArray<FCustomizableObjectExternalTexture> UCustomizableObjectSystem::GetTexture
 {
 	TArray<FCustomizableObjectExternalTexture> Result;
 
-	for (const TWeakObjectPtr<UCustomizableSystemImageProvider> Provider : GetPrivate()->GetImageProviderChecked()->ImageProviders)
+	for (const TWeakObjectPtr<UCustomizableSystemImageProvider> Provider : GetPrivate()->GetResourceProviderChecked()->ImageProviders)
 	{
 		if (Provider.IsValid())
 		{
@@ -3806,13 +3806,13 @@ TArray<FCustomizableObjectExternalTexture> UCustomizableObjectSystem::GetTexture
 
 void UCustomizableObjectSystem::RegisterImageProvider(UCustomizableSystemImageProvider* Provider)
 {
-	GetPrivate()->GetImageProviderChecked()->ImageProviders.Add(Provider);
+	GetPrivate()->GetResourceProviderChecked()->ImageProviders.Add(Provider);
 }
 
 
 void UCustomizableObjectSystem::UnregisterImageProvider(UCustomizableSystemImageProvider* Provider)
 {
-	GetPrivate()->GetImageProviderChecked()->ImageProviders.Remove(Provider);
+	GetPrivate()->GetResourceProviderChecked()->ImageProviders.Remove(Provider);
 }
 
 
@@ -3820,11 +3820,11 @@ void UCustomizableObjectSystemPrivate::CacheTextureParameters(const TArray<FCust
 {
 	for (const FCustomizableObjectTextureParameterValue& TextureParameter : TextureParameters)
 	{
-		ImageProvider->CacheImage(TextureParameter.ParameterValue, false);
+		ResourceProvider->CacheImage(TextureParameter.ParameterValue, false);
 
 		for (const FName& RangeValue : TextureParameter.ParameterRangeValues)
 		{
-			ImageProvider->CacheImage(RangeValue, false);
+			ResourceProvider->CacheImage(RangeValue, false);
 		}
 	}
 }
@@ -3834,11 +3834,11 @@ void UCustomizableObjectSystemPrivate::UnCacheTextureParameters(const TArray<FCu
 {
 	for (const FCustomizableObjectTextureParameterValue& TextureParameter : TextureParameters)
 	{
-		ImageProvider->UnCacheImage(TextureParameter.ParameterValue, false);
+		ResourceProvider->UnCacheImage(TextureParameter.ParameterValue, false);
 
 		for (const FName& RangeValue : TextureParameter.ParameterRangeValues)
 		{
-			ImageProvider->UnCacheImage(RangeValue, false);
+			ResourceProvider->UnCacheImage(RangeValue, false);
 		}
 	}
 }
@@ -4055,19 +4055,19 @@ uint64 UCustomizableObjectSystem::GetMaxChunkSizeForPlatform(const ITargetPlatfo
 
 void UCustomizableObjectSystem::CacheImage(FName ImageId)
 {
-	GetPrivate()->GetImageProviderChecked()->CacheImage(ImageId, true);
+	GetPrivate()->GetResourceProviderChecked()->CacheImage(ImageId, true);
 }
 
 
 void UCustomizableObjectSystem::UnCacheImage(FName ImageId)
 {
-	GetPrivate()->GetImageProviderChecked()->UnCacheImage(ImageId, true);
+	GetPrivate()->GetResourceProviderChecked()->UnCacheImage(ImageId, true);
 }
 
 
 void UCustomizableObjectSystem::ClearImageCache()
 {
-	GetPrivate()->GetImageProviderChecked()->ClearCache(true);
+	GetPrivate()->GetResourceProviderChecked()->ClearCache(true);
 }
 
 
@@ -4081,10 +4081,10 @@ bool UCustomizableObjectSystemPrivate::IsMutableAnimInfoDebuggingEnabled() const
 }
 
 
-FUnrealMutableImageProvider* UCustomizableObjectSystemPrivate::GetImageProviderChecked() const
+FUnrealMutableResourceProvider* UCustomizableObjectSystemPrivate::GetResourceProviderChecked() const
 {
-	check(ImageProvider)
-	return ImageProvider.Get();
+	check(ResourceProvider)
+	return ResourceProvider.Get();
 }
 
 
diff --git a/Engine/Plugins/Experimental/Mutable/Source/CustomizableObject/Private/MuCO/UnrealMutableImageProvider.cpp b/Engine/Plugins/Experimental/Mutable/Source/CustomizableObject/Private/MuCO/UnrealMutableImageProvider.cpp
--- a/Engine/Plugins/Experimental/Mutable/Source/CustomizableObject/Private/MuCO/UnrealMutableImageProvider.cpp
+++ b/Engine/Plugins/Experimental/Mutable/Source/CustomizableObject/Private/MuCO/UnrealMutableImageProvider.cpp
@@ -11,6 +11,7 @@
 #include "MuR/Parameters.h"
 #include "MuR/ImageTypes.h"
 #include "MuR/Model.h"
+#include "MuR/Mesh.h"
 
 
 //-------------------------------------------------------------------------------------------------
@@ -80,7 +81,7 @@ mu::EImageFormat GetMutablePixelFormat(EPixelFormat InTextureFormat)
 
 
 //-------------------------------------------------------------------------------------------------
-TTuple<UE::Tasks::FTask, TFunction<void()>> FUnrealMutableImageProvider::GetImageAsync(FName Id, uint8 MipmapsToSkip, TFunction<void(mu::Ptr<mu::Image>)>& ResultCallback)
+TTuple<UE::Tasks::FTask, TFunction<void()>> FUnrealMutableResourceProvider::GetImageAsync(FName Id, uint8 MipmapsToSkip, TFunction<void(mu::Ptr<mu::Image>)>& ResultCallback)
 {
 	// Thread: worker
 	MUTABLE_CPUPROFILER_SCOPE(FUnrealMutableImageProvider::GetImageAsync);
@@ -330,7 +331,7 @@ TTuple<UE::Tasks::FTask, TFunction<void()>> FUnrealMutableImageProvider::GetImag
 
 
 //-------------------------------------------------------------------------------------------------
-TTuple<UE::Tasks::FTask, TFunction<void()>> FUnrealMutableImageProvider::GetReferencedImageAsync(const void* ModelPtr, int32 Id, uint8 MipmapsToSkip, TFunction<void(mu::Ptr<mu::Image>)>& ResultCallback)
+TTuple<UE::Tasks::FTask, TFunction<void()>> FUnrealMutableResourceProvider::GetReferencedImageAsync(const void* ModelPtr, int32 Id, uint8 MipmapsToSkip, TFunction<void(mu::Ptr<mu::Image>)>& ResultCallback)
 {
 	// Thread: worker
 	MUTABLE_CPUPROFILER_SCOPE(FUnrealMutableImageProvider::GetReferencedImageAsync);
@@ -386,7 +387,7 @@ TTuple<UE::Tasks::FTask, TFunction<void()>> FUnrealMutableImageProvider::GetRefe
 
 
 // This should mantain parity with the descriptor of the images generated by GetImageAsync 
-mu::FImageDesc FUnrealMutableImageProvider::GetImageDesc(FName Id, uint8 MipmapsToSkip)
+mu::FImageDesc FUnrealMutableResourceProvider::GetImageDesc(FName Id, uint8 MipmapsToSkip)
 {
 	MUTABLE_CPUPROFILER_SCOPE(FUnrealMutableImageProvider::GetImageDesc);
 
@@ -469,7 +470,52 @@ mu::FImageDesc FUnrealMutableImageProvider::GetImageDesc(FName Id, uint8 Mipmaps
 }
 
 
-void FUnrealMutableImageProvider::CacheImage(FName Id, bool bUser)
+TTuple<UE::Tasks::FTask, TFunction<void()>> FUnrealMutableResourceProvider::GetMeshAsync(FName Id, TFunction<void(mu::Ptr<mu::Mesh>)>& ResultCallback)
+{
+	// TODO: This is a stub function where Mesh Parameters would be supported.
+
+	// Thread: worker
+	MUTABLE_CPUPROFILER_SCOPE(FUnrealMutableImageProvider::GetMeshAsync);
+
+	auto TrivialReturn = []() -> TTuple<UE::Tasks::FTask, TFunction<void()>>
+		{
+			return MakeTuple(UE::Tasks::MakeCompletedTask<void>(), []() -> void {});
+		};
+
+	// No USkeletalMesh was provided, cannot do anything, just provide a dummy mesh
+	UE_LOG(LogMutable, Warning, TEXT("No USkeletalMesh was provided for an application-specific mesh parameter."));
+
+	mu::Ptr<mu::Mesh> Dummy = new mu::Mesh;
+	ResultCallback( Dummy );
+	return Invoke(TrivialReturn);
+}
+
+
+TTuple<UE::Tasks::FTask, TFunction<void()>> FUnrealMutableResourceProvider::GetReferencedMeshAsync(const void* ModelPtr, int32 Id, TFunction<void(mu::Ptr<mu::Mesh>)>& ResultCallback)
+{
+	// TODO: This is a stub function where referenced Meshes would be supported.
+
+	// Thread: worker
+	MUTABLE_CPUPROFILER_SCOPE(FUnrealMutableImageProvider::GetReferencedMeshAsync);
+
+	auto TrivialReturn = []() -> TTuple<UE::Tasks::FTask, TFunction<void()>>
+		{
+			return MakeTuple(UE::Tasks::MakeCompletedTask<void>(), []() -> void {});
+		};
+
+	mu::Ptr<mu::Mesh> Mesh = new mu::Mesh();
+
+	// Not supported outside editor yet.
+	UE_LOG(LogMutable, Warning, TEXT("Failed to get referenced mesh. Not implemented yet."));
+
+	ResultCallback(Mesh);
+	return Invoke(TrivialReturn);
+
+}
+
+
+
+void FUnrealMutableResourceProvider::CacheImage(FName Id, bool bUser)
 {
 	if (Id == NAME_None)
 	{
@@ -586,7 +632,7 @@ void FUnrealMutableImageProvider::CacheImage(FName Id, bool bUser)
 }
 
 
-void FUnrealMutableImageProvider::UnCacheImage(FName Id, bool bUser)
+void FUnrealMutableResourceProvider::UnCacheImage(FName Id, bool bUser)
 {
 	if (Id == NAME_None)
 	{
@@ -618,7 +664,7 @@ void FUnrealMutableImageProvider::UnCacheImage(FName Id, bool bUser)
 }
 
 
-void FUnrealMutableImageProvider::ClearCache(bool bUser)
+void FUnrealMutableResourceProvider::ClearCache(bool bUser)
 {
 	FScopeLock Lock(&ExternalImagesLock);
 	for (TTuple<FName, FUnrealMutableImageInfo> Tuple : GlobalExternalImages)
@@ -628,7 +674,7 @@ void FUnrealMutableImageProvider::ClearCache(bool bUser)
 }
 
 
-void FUnrealMutableImageProvider::CacheImages(const mu::Parameters& Parameters)
+void FUnrealMutableResourceProvider::CacheImages(const mu::Parameters& Parameters)
 {
 	const int32 NumParams = Parameters.GetCount();
 	for (int32 ParamIndex = 0; ParamIndex < NumParams; ++ParamIndex)
@@ -655,7 +701,7 @@ void FUnrealMutableImageProvider::CacheImages(const mu::Parameters& Parameters)
 }
 
 
-void FUnrealMutableImageProvider::UnCacheImages(const mu::Parameters& Parameters)
+void FUnrealMutableResourceProvider::UnCacheImages(const mu::Parameters& Parameters)
 {
 	const int32 NumParams = Parameters.GetCount();
 	for (int32 ParamIndex = 0; ParamIndex < NumParams; ++ParamIndex)
@@ -683,7 +729,7 @@ void FUnrealMutableImageProvider::UnCacheImages(const mu::Parameters& Parameters
 
 
 #if WITH_EDITOR
-void FUnrealMutableImageProvider::CacheRuntimeReferencedImages(const TSharedRef<const mu::Model>& Model, const TArray<TSoftObjectPtr<UTexture2D>>& RuntimeReferencedTextures)
+void FUnrealMutableResourceProvider::CacheRuntimeReferencedImages(const TSharedRef<const mu::Model>& Model, const TArray<TSoftObjectPtr<UTexture2D>>& RuntimeReferencedTextures)
 {
 	check(IsInGameThread());
 	
@@ -710,7 +756,7 @@ void FUnrealMutableImageProvider::CacheRuntimeReferencedImages(const TSharedRef<
 #endif
 
 
-mu::ImagePtr FUnrealMutableImageProvider::CreateDummy()
+mu::ImagePtr FUnrealMutableResourceProvider::CreateDummy()
 {
 	// Create a dummy image
 	const int32 Size = DUMMY_IMAGE_DESC.m_size[0];
@@ -745,7 +791,7 @@ mu::ImagePtr FUnrealMutableImageProvider::CreateDummy()
 }
 
 
-mu::FImageDesc FUnrealMutableImageProvider::CreateDummyDesc()
+mu::FImageDesc FUnrealMutableResourceProvider::CreateDummyDesc()
 {
 	return DUMMY_IMAGE_DESC;
 }
@@ -758,7 +804,7 @@ TAutoConsoleVariable<bool> CVarMutableLockExternalImagesDuringGC(
 	ECVF_Default);
 
 
-void FUnrealMutableImageProvider::AddReferencedObjects(FReferenceCollector& Collector)
+void FUnrealMutableResourceProvider::AddReferencedObjects(FReferenceCollector& Collector)
 {
 #if WITH_EDITOR
 	{
@@ -796,7 +842,7 @@ void FUnrealMutableImageProvider::AddReferencedObjects(FReferenceCollector& Coll
 }
 
 
-FUnrealMutableImageProvider::FUnrealMutableImageInfo::FUnrealMutableImageInfo(const mu::ImagePtr& InImage)
+FUnrealMutableResourceProvider::FUnrealMutableImageInfo::FUnrealMutableImageInfo(const mu::ImagePtr& InImage)
 {
 	check(IsInGameThread())
 	
@@ -804,7 +850,7 @@ FUnrealMutableImageProvider::FUnrealMutableImageInfo::FUnrealMutableImageInfo(co
 }
 
 
-FUnrealMutableImageProvider::FUnrealMutableImageInfo::FUnrealMutableImageInfo(UTexture2D& Texture)
+FUnrealMutableResourceProvider::FUnrealMutableImageInfo::FUnrealMutableImageInfo(UTexture2D& Texture)
 {
 	check(IsInGameThread())
 
diff --git a/Engine/Plugins/Experimental/Mutable/Source/CustomizableObject/Private/MuCO/UnrealMutableImageProvider.h b/Engine/Plugins/Experimental/Mutable/Source/CustomizableObject/Private/MuCO/UnrealMutableImageProvider.h
--- a/Engine/Plugins/Experimental/Mutable/Source/CustomizableObject/Private/MuCO/UnrealMutableImageProvider.h
+++ b/Engine/Plugins/Experimental/Mutable/Source/CustomizableObject/Private/MuCO/UnrealMutableImageProvider.h
@@ -18,16 +18,17 @@ class UTexture2D;
 
 
 /** Implementation of a mutable core provider for image parameters that are application-specific. */
-class FUnrealMutableImageProvider : public mu::ImageParameterGenerator, public FGCObject
+class FUnrealMutableResourceProvider : public mu::ExternalResourceProvider, public FGCObject
 {
 
 public:
-	// mu::ImageParameterGenerator interface
+	// mu::ExternalResourceProvider interface
 	// Thread: worker
 	virtual TTuple<UE::Tasks::FTask, TFunction<void()>> GetImageAsync(FName Id, uint8 MipmapsToSkip, TFunction<void(mu::Ptr<mu::Image>)>& ResultCallback) override;
 	virtual TTuple<UE::Tasks::FTask, TFunction<void()>> GetReferencedImageAsync(const void* ModelPtr, int32 Id, uint8 MipmapsToSkip, TFunction<void(mu::Ptr<mu::Image>)>& ResultCallback) override;
-
 	virtual mu::FImageDesc GetImageDesc(FName Id, uint8 MipmapsToSkip) override;
+	virtual TTuple<UE::Tasks::FTask, TFunction<void()>> GetMeshAsync(FName Id, TFunction<void(mu::Ptr<mu::Mesh>)>& ResultCallback) override;
+	virtual TTuple<UE::Tasks::FTask, TFunction<void()>> GetReferencedMeshAsync(const void* ModelPtr, int32 Id, TFunction<void(mu::Ptr<mu::Mesh>)>& ResultCallback) override;
 
 	
 	// Own interface	
diff --git a/Engine/Plugins/Experimental/Mutable/Source/CustomizableObjectEditor/Private/MuCOE/CustomizableObjectCompiler.cpp b/Engine/Plugins/Experimental/Mutable/Source/CustomizableObjectEditor/Private/MuCOE/CustomizableObjectCompiler.cpp
--- a/Engine/Plugins/Experimental/Mutable/Source/CustomizableObjectEditor/Private/MuCOE/CustomizableObjectCompiler.cpp
+++ b/Engine/Plugins/Experimental/Mutable/Source/CustomizableObjectEditor/Private/MuCOE/CustomizableObjectCompiler.cpp
@@ -923,6 +923,12 @@ void FCustomizableObjectCompiler::CompileInternal(bool bAsync)
 			ModelResources.Materials.Emplace(Material);
 		}
 
+		for (const TPair<TSoftObjectPtr<USkeletalMesh>, FMutableGraphGenerationContext::FGeneratedReferencedMesh>& Pair : GenerationContext.PassthroughMeshMap)
+		{
+			check(Pair.Value.ID == ModelResources.PassThroughMeshes.Num());
+			ModelResources.PassThroughMeshes.Add(Pair.Key);
+		}
+
 		for (const TPair<TSoftObjectPtr<UTexture>, FMutableGraphGenerationContext::FGeneratedReferencedTexture>& Pair : GenerationContext.PassthroughTextureMap)
 		{
 			check(Pair.Value.ID == ModelResources.PassThroughTextures.Num());
diff --git a/Engine/Plugins/Experimental/Mutable/Source/CustomizableObjectEditor/Private/MuCOE/EdGraphSchema_CustomizableObject.cpp b/Engine/Plugins/Experimental/Mutable/Source/CustomizableObjectEditor/Private/MuCOE/EdGraphSchema_CustomizableObject.cpp
--- a/Engine/Plugins/Experimental/Mutable/Source/CustomizableObjectEditor/Private/MuCOE/EdGraphSchema_CustomizableObject.cpp
+++ b/Engine/Plugins/Experimental/Mutable/Source/CustomizableObjectEditor/Private/MuCOE/EdGraphSchema_CustomizableObject.cpp
@@ -45,6 +45,7 @@
 #include "MuCOE/Nodes/CustomizableObjectNodeMeshReshape.h"
 #include "MuCOE/Nodes/CustomizableObjectNodeMeshSwitch.h"
 #include "MuCOE/Nodes/CustomizableObjectNodeMeshVariation.h"
+#include "MuCOE/Nodes/CustomizableObjectNodePassThroughMesh.h"
 #include "MuCOE/Nodes/CustomizableObjectNodeMorphMaterial.h"
 #include "MuCOE/Nodes/CustomizableObjectNodeObjectChild.h"
 #include "MuCOE/Nodes/CustomizableObjectNodeObjectGroup.h"
@@ -218,6 +219,7 @@ UEdGraphNode* FCustomizableObjectSchemaAction_Paste::PerformAction(class UEdGrap
 const FName UEdGraphSchema_CustomizableObject::PC_Object("object");
 const FName UEdGraphSchema_CustomizableObject::PC_Material("material");
 const FName UEdGraphSchema_CustomizableObject::PC_Mesh("mesh");
+const FName UEdGraphSchema_CustomizableObject::PC_PassThroughMesh("passThroughMesh");
 const FName UEdGraphSchema_CustomizableObject::PC_Layout("layout");
 const FName UEdGraphSchema_CustomizableObject::PC_Image("image");
 const FName UEdGraphSchema_CustomizableObject::PC_PassThroughImage("passThroughImage");
@@ -394,6 +396,7 @@ void UEdGraphSchema_CustomizableObject::GetGraphContextActions(FGraphContextMenu
 		{
 			ContextMenuBuilder.CreateTemplateNode<UCustomizableObjectNodeSkeletalMesh>(),
 			ContextMenuBuilder.CreateTemplateNode<UCustomizableObjectNodeStaticMesh>(),
+			ContextMenuBuilder.CreateTemplateNode<UCustomizableObjectNodePassThroughMesh>(),
 			ContextMenuBuilder.CreateTemplateNode<UCustomizableObjectNodeLayoutBlocks>(),
 			ContextMenuBuilder.CreateTemplateNode<UCustomizableObjectNodeMeshMorph>(),
 			ContextMenuBuilder.CreateTemplateNode<UCustomizableObjectNodeMeshSwitch>(),
diff --git a/Engine/Plugins/Experimental/Mutable/Source/CustomizableObjectEditor/Private/MuCOE/GenerateMutableSource/GenerateMutableSource.cpp b/Engine/Plugins/Experimental/Mutable/Source/CustomizableObjectEditor/Private/MuCOE/GenerateMutableSource/GenerateMutableSource.cpp
--- a/Engine/Plugins/Experimental/Mutable/Source/CustomizableObjectEditor/Private/MuCOE/GenerateMutableSource/GenerateMutableSource.cpp
+++ b/Engine/Plugins/Experimental/Mutable/Source/CustomizableObjectEditor/Private/MuCOE/GenerateMutableSource/GenerateMutableSource.cpp
@@ -1954,6 +1954,68 @@ mu::Ptr<mu::Image> GenerateImageConstant(UTexture* Texture, FMutableGraphGenerat
 }
 
 
+mu::Ptr<mu::Mesh> GenerateMeshConstant(const USkeletalMesh* Mesh, FMutableGraphGenerationContext& GenerationContext, bool bIsReference)
+{
+	if (!Mesh)
+	{
+		return nullptr;
+	}
+
+	bool bForceLoad = false;
+	bool bIsCompileTime = false;
+	if (!bIsReference)
+	{
+		bForceLoad = true;
+		if (GenerationContext.Options.OptimizationLevel == 0)
+		{
+			bIsCompileTime = false;
+		}
+		else
+		{
+			bIsCompileTime = true;
+		}
+	}
+
+
+	FMutableGraphGenerationContext::FGeneratedReferencedMesh InvalidEntry;
+	InvalidEntry.ID = TNumericLimits<uint32>::Max();
+
+	FMutableGraphGenerationContext::FGeneratedReferencedMesh* Entry = nullptr;
+	int32 Num = 0;
+
+	if (bIsReference)
+	{
+		Entry = &GenerationContext.PassthroughMeshMap.FindOrAdd(Mesh, InvalidEntry);
+		Num = GenerationContext.PassthroughMeshMap.Num();
+	}
+	// Not supported yet: this will be implemented when mesh conversion happens on-demand in the compilation process
+	// like it happens with images. For now we only implement support for pure pass-through.
+	//else if (bIsCompileTime)
+	//{
+	//	Entry = &GenerationContext.CompileTimeTextureMap.FindOrAdd(Texture, InvalidEntry);
+	//	Num = GenerationContext.CompileTimeTextureMap.Num();
+	//}
+	//else
+	//{
+	//	Entry = &GenerationContext.RuntimeReferencedTextureMap.FindOrAdd(Texture, InvalidEntry);
+	//	Num = GenerationContext.RuntimeReferencedTextureMap.Num();
+	//}
+	else
+	{
+		check(false);
+		return nullptr;
+	}
+
+	if (Entry->ID == TNumericLimits<uint32>::Max())
+	{
+		Entry->ID = Num - 1;
+	}
+
+	mu::Ptr<mu::Mesh> Result = mu::Mesh::CreateAsReference(Entry->ID, bForceLoad);
+	return Result;
+}
+
+
 void AddSocketTagsToMesh(const USkeletalMesh* SourceMesh, mu::Ptr<mu::Mesh> MutableMesh, FMutableGraphGenerationContext& GenerationContext)
 {
 	for (int32 SocketIndex = 0; SocketIndex < SourceMesh->NumSockets(); ++SocketIndex)
diff --git a/Engine/Plugins/Experimental/Mutable/Source/CustomizableObjectEditor/Private/MuCOE/GenerateMutableSource/GenerateMutableSource.h b/Engine/Plugins/Experimental/Mutable/Source/CustomizableObjectEditor/Private/MuCOE/GenerateMutableSource/GenerateMutableSource.h
--- a/Engine/Plugins/Experimental/Mutable/Source/CustomizableObjectEditor/Private/MuCOE/GenerateMutableSource/GenerateMutableSource.h
+++ b/Engine/Plugins/Experimental/Mutable/Source/CustomizableObjectEditor/Private/MuCOE/GenerateMutableSource/GenerateMutableSource.h
@@ -582,10 +582,20 @@ struct FMutableGraphGenerationContext
 		//mu::FImageDesc ImageDesc;
 	};
 
+	/** Data stored per-generated passthrough mesh. */
+	struct FGeneratedReferencedMesh
+	{
+		uint32 ID;
+	};
+
+	// Cache of runtime pass-through meshes and their IDs used in the core to indentify them.
+	// These meshes will remain as external references even in optimized models.
+	TMap<TSoftObjectPtr<USkeletalMesh>, FGeneratedReferencedMesh> PassthroughMeshMap;
+
 	// Cache of runtime pass-through images and their IDs used in the core to indentify them.
 	// These textures will remain as external references even in optimized models.
 	TMap<TSoftObjectPtr<UTexture>, FGeneratedReferencedTexture> PassthroughTextureMap;
-	
+
 	// Cache of runtime images and their IDs used in the core to indentify them.
 	// These textures will remain as external references even in optimized models.
 	TMap<TSoftObjectPtr<UTexture2D>, FGeneratedReferencedTexture> RuntimeReferencedTextureMap;
@@ -905,7 +915,8 @@ int32 GetMaxTextureSize(const UTexture2D& ReferenceTexture, const UTextureLODSet
 // Max texture size of the texture with per platform MaxTextureSize and LODBias applied.
 int32 GetTextureSizeInGame(const UTexture2D& Texture, const UTextureLODSettings& LODSettings, uint8 SurfaceLODBias = 0);
 
-mu::Ptr<mu::Image> GenerateImageConstant( UTexture*, FMutableGraphGenerationContext&, bool bIsReference);
+mu::Ptr<mu::Image> GenerateImageConstant(UTexture*, FMutableGraphGenerationContext&, bool bIsReference);
+mu::Ptr<mu::Mesh> GenerateMeshConstant(const USkeletalMesh*, FMutableGraphGenerationContext&, bool bIsReference);
 
 /** Generates a mutable image descriptor from an unreal engine texture */
 mu::FImageDesc GenerateImageDescriptor(UTexture* Texture);
diff --git a/Engine/Plugins/Experimental/Mutable/Source/CustomizableObjectEditor/Private/MuCOE/GenerateMutableSource/GenerateMutableSourceMesh.cpp b/Engine/Plugins/Experimental/Mutable/Source/CustomizableObjectEditor/Private/MuCOE/GenerateMutableSource/GenerateMutableSourceMesh.cpp
--- a/Engine/Plugins/Experimental/Mutable/Source/CustomizableObjectEditor/Private/MuCOE/GenerateMutableSource/GenerateMutableSourceMesh.cpp
+++ b/Engine/Plugins/Experimental/Mutable/Source/CustomizableObjectEditor/Private/MuCOE/GenerateMutableSource/GenerateMutableSourceMesh.cpp
@@ -36,6 +36,7 @@
 #include "MuCOE/Nodes/CustomizableObjectNodeSkeletalMesh.h"
 #include "MuCOE/Nodes/CustomizableObjectNodeStaticMesh.h"
 #include "MuCOE/Nodes/CustomizableObjectNodeTable.h"
+#include "MuCOE/Nodes/CustomizableObjectNodePassThroughMesh.h"
 #include "MuCOE/UnrealEditorPortabilityHelpers.h"
 #include "MuT/NodeMeshConstant.h"
 #include "MuT/NodeMeshGeometryOperation.h"
@@ -2122,8 +2123,8 @@ mu::MeshPtr ConvertStaticMeshToMutable(const UStaticMesh* StaticMesh, int32 LODI
 	mu::MeshPtr MutableMesh = new mu::Mesh();
 
 	// Vertices
-	int VertexStart = StaticMesh->GetRenderData()->LODResources[LODIndex].Sections[SectionIndex].MinVertexIndex;
-	int VertexCount = StaticMesh->GetRenderData()->LODResources[LODIndex].Sections[SectionIndex].MaxVertexIndex - VertexStart + 1;
+	int32 VertexStart = StaticMesh->GetRenderData()->LODResources[LODIndex].Sections[SectionIndex].MinVertexIndex;
+	int32 VertexCount = StaticMesh->GetRenderData()->LODResources[LODIndex].Sections[SectionIndex].MaxVertexIndex - VertexStart + 1;
 
 	MutableMesh->GetVertexBuffers().SetElementCount(VertexCount);
 	{
@@ -2135,13 +2136,13 @@ mu::MeshPtr ConvertStaticMeshToMutable(const UStaticMesh* StaticMesh, int32 LODI
 		{
 			const FPositionVertexBuffer& VertexBuffer = StaticMesh->GetRenderData()->LODResources[LODIndex].VertexBuffers.PositionVertexBuffer;
 
-			const int ElementSize = 12;
-			const int ChannelCount = 1;
+			const int32 ElementSize = 12;
+			const int32 ChannelCount = 1;
 			const EMeshBufferSemantic Semantics[ChannelCount] = { MBS_POSITION };
-			const int SemanticIndices[ChannelCount] = { 0 };
+			const int32 SemanticIndices[ChannelCount] = { 0 };
 			const EMeshBufferFormat Formats[ChannelCount] = { MBF_FLOAT32 };
-			const int Components[ChannelCount] = { 3 };
-			const int Offsets[ChannelCount] = { 0 };
+			const int32 Components[ChannelCount] = { 3 };
+			const int32 Offsets[ChannelCount] = { 0 };
 
 			MutableMesh->GetVertexBuffers().SetBuffer(MUTABLE_VERTEXBUFFER_POSITION, ElementSize, ChannelCount, Semantics, SemanticIndices, Formats, Components, Offsets);
 			FMemory::Memcpy(
@@ -2155,13 +2156,13 @@ mu::MeshPtr ConvertStaticMeshToMutable(const UStaticMesh* StaticMesh, int32 LODI
 			const FStaticMeshVertexBuffer& VertexBuffer = StaticMesh->GetRenderData()->LODResources[LODIndex].VertexBuffers.StaticMeshVertexBuffer;
 
 			EMeshBufferSemantic Semantics[2];
-			int SemanticIndices[2];
+			int32 SemanticIndices[2];
 			EMeshBufferFormat Formats[2];
-			int Components[2];
-			int Offsets[2];
+			int32 Components[2];
+			int32 Offsets[2];
 
-			int currentChannel = 0;
-			int currentOffset = 0;
+			int32 currentChannel = 0;
+			int32 currentOffset = 0;
 
 			Semantics[currentChannel] = MBS_TANGENT;
 			SemanticIndices[currentChannel] = 0;
@@ -2193,19 +2194,19 @@ mu::MeshPtr ConvertStaticMeshToMutable(const UStaticMesh* StaticMesh, int32 LODI
 		{
 			const FStaticMeshVertexBuffer& VertexBuffer = StaticMesh->GetRenderData()->LODResources[LODIndex].VertexBuffers.StaticMeshVertexBuffer;
 
-			int texChannels = VertexBuffer.GetNumTexCoords();
-			int ChannelCount = texChannels;
+			int32 texChannels = VertexBuffer.GetNumTexCoords();
+			int32 ChannelCount = texChannels;
 
 			EMeshBufferSemantic* Semantics = new EMeshBufferSemantic[ChannelCount];
-			int* SemanticIndices = new int[ChannelCount];
+			int32* SemanticIndices = new int[ChannelCount];
 			EMeshBufferFormat* Formats = new EMeshBufferFormat[ChannelCount];
-			int* Components = new int[ChannelCount];
-			int* Offsets = new int[ChannelCount];
+			int32* Components = new int[ChannelCount];
+			int32* Offsets = new int[ChannelCount];
 
-			int currentChannel = 0;
-			int currentOffset = 0;
+			int32 currentChannel = 0;
+			int32 currentOffset = 0;
 
-			int texChannelSize;
+			int32 texChannelSize;
 			EMeshBufferFormat texChannelFormat;
 			if (VertexBuffer.GetUseFullPrecisionUVs())
 			{
@@ -2218,7 +2219,7 @@ mu::MeshPtr ConvertStaticMeshToMutable(const UStaticMesh* StaticMesh, int32 LODI
 				texChannelFormat = MBF_FLOAT16;
 			}
 
-			for (int c = 0; c < texChannels; ++c)
+			for (int32 c = 0; c < texChannels; ++c)
 			{
 				Semantics[currentChannel] = MBS_TEXCOORDS;
 				SemanticIndices[currentChannel] = c;
@@ -2282,7 +2283,7 @@ mu::MeshPtr ConvertStaticMeshToMutable(const UStaticMesh* StaticMesh, int32 LODI
 			//pSource += IndexStart;
 			uint16* pDest = reinterpret_cast<uint16*>(MutableMesh->GetIndexBuffers().GetBufferData(0));
 
-			for (int i = 0; i < IndexCount; ++i)
+			for (int32 i = 0; i < IndexCount; ++i)
 			{
 				*pDest = Source[IndexStart + i] - VertexStart;
 				++pDest;
@@ -2295,8 +2296,8 @@ mu::MeshPtr ConvertStaticMeshToMutable(const UStaticMesh* StaticMesh, int32 LODI
 
 
 // Convert a Mesh constant to a mutable format. UniqueTags are the tags that make this Mesh unique that cannot be merged in the cache 
-//  with the exact same Mesh with other tags
-mu::MeshPtr GenerateMutableMesh(const UObject * Mesh, const TSoftClassPtr<UAnimInstance>& AnimInstance, int32 LODIndexConnected, int32 SectionIndexConnected, int32 LODIndex, int32 SectionIndex, const FString& UniqueTags, FMutableGraphGenerationContext & GenerationContext, const UCustomizableObjectNode* CurrentNode)
+// with the exact same Mesh with other tags
+mu::Ptr<mu::Mesh> GenerateMutableMesh(const UObject * Mesh, const TSoftClassPtr<UAnimInstance>& AnimInstance, int32 LODIndexConnected, int32 SectionIndexConnected, int32 LODIndex, int32 SectionIndex, const FString& UniqueTags, FMutableGraphGenerationContext & GenerationContext, const UCustomizableObjectNode* CurrentNode, bool bIsReference)
 {
 	// Get the mesh generation flags to use
 	EMutableMeshConversionFlags CurrentFlags = GenerationContext.MeshGenerationFlags.Last();
@@ -2308,20 +2309,28 @@ mu::MeshPtr GenerateMutableMesh(const UObject * Mesh, const TSoftClassPtr<UAnimI
 	{
 		if (const USkeletalMesh* SkeletalMesh = Cast<USkeletalMesh>(Mesh))
 		{
-			MutableMesh = ConvertSkeletalMeshToMutable(SkeletalMesh, AnimInstance, LODIndexConnected, SectionIndexConnected, LODIndex, SectionIndex, GenerationContext, CurrentNode);
-
-			FSkeletalMeshModel* ImportedModel = SkeletalMesh->GetImportedModel();
-			
-			if (MutableMesh &&
-				ImportedModel->LODModels.IsValidIndex(LODIndex) &&
-				ImportedModel->LODModels[LODIndex].Sections.IsValidIndex(SectionIndex))
+			// At some point we will want all meshes to be references at compile-time. For now, just create the actual pass-through meshes.
+			if (bIsReference)
 			{
-				FMeshData MeshData;
-				MeshData.Mesh = Mesh;
-				MeshData.LOD = LODIndex;
-				MeshData.MaterialIndex = SectionIndex;
-				MeshData.Node = CurrentNode;
-				GenerationContext.PinData.GetCurrent().MeshesData.Add(MeshData); // Set::Emplace only supports single element constructors
+				MutableMesh = GenerateMeshConstant(SkeletalMesh,GenerationContext, bIsReference);
+			}
+			else
+			{
+				MutableMesh = ConvertSkeletalMeshToMutable(SkeletalMesh, AnimInstance, LODIndexConnected, SectionIndexConnected, LODIndex, SectionIndex, GenerationContext, CurrentNode);
+
+				FSkeletalMeshModel* ImportedModel = SkeletalMesh->GetImportedModel();
+
+				if (MutableMesh &&
+					ImportedModel->LODModels.IsValidIndex(LODIndex) &&
+					ImportedModel->LODModels[LODIndex].Sections.IsValidIndex(SectionIndex))
+				{
+					FMeshData MeshData;
+					MeshData.Mesh = Mesh;
+					MeshData.LOD = LODIndex;
+					MeshData.MaterialIndex = SectionIndex;
+					MeshData.Node = CurrentNode;
+					GenerationContext.PinData.GetCurrent().MeshesData.Add(MeshData); // Set::Emplace only supports single element constructors
+				}
 			}
 		}
 		else if (const UStaticMesh* StaticMesh = Cast<UStaticMesh>(Mesh))
@@ -2401,7 +2410,8 @@ mu::MeshPtr BuildMorphedMutableMesh(const UEdGraphPin* BaseSourcePin, const FStr
 	{
 		GetLODAndSectionForAutomaticLODs(GenerationContext, *Node, *SkeletalMesh, LODIndexConnected, SectionIndexConnected, LODIndex, SectionIndex, bOnlyConnectedLOD);
 		// Get the base mesh
-		mu::MeshPtr BaseSourceMesh = GenerateMutableMesh(SkeletalMesh, TSoftClassPtr<UAnimInstance>(), LODIndexConnected, SectionIndexConnected, LODIndex, SectionIndex, FString(), GenerationContext, Node);
+		constexpr bool bIsReference = false;
+		mu::Ptr<mu::Mesh> BaseSourceMesh = GenerateMutableMesh(SkeletalMesh, TSoftClassPtr<UAnimInstance>(), LODIndexConnected, SectionIndexConnected, LODIndex, SectionIndex, FString(), GenerationContext, Node, bIsReference);
 		if (BaseSourceMesh)
 		{
 			// Clone it (it will probably be shared)
@@ -3250,7 +3260,8 @@ mu::NodeMeshPtr GenerateMutableSourceMesh(const UEdGraphPin* Pin,
 
 			FSkeletalMeshModel* ImportedModel = TypedNodeSkel->SkeletalMesh->GetImportedModel();
 			
-			mu::MeshPtr MutableMesh = GenerateMutableMesh(TypedNodeSkel->SkeletalMesh, TypedNodeSkel->AnimInstance, LODIndexConnected, SectionIndexConnected, LODIndex, SectionIndex, MeshUniqueTags, GenerationContext, TypedNodeSkel);
+			constexpr bool bIsReference = false;
+			mu::Ptr<mu::Mesh> MutableMesh = GenerateMutableMesh(TypedNodeSkel->SkeletalMesh, TypedNodeSkel->AnimInstance, LODIndexConnected, SectionIndexConnected, LODIndex, SectionIndex, MeshUniqueTags, GenerationContext, TypedNodeSkel, bIsReference);
 			if (MutableMesh)
 			{
 				MeshNode->SetValue(MutableMesh);
@@ -3436,7 +3447,8 @@ mu::NodeMeshPtr GenerateMutableSourceMesh(const UEdGraphPin* Pin,
 			
 			check(SectionIndex < TypedNodeStatic->LODs[LODIndex].Materials.Num());
 
-			mu::MeshPtr MutableMesh = GenerateMutableMesh(TypedNodeStatic->StaticMesh, TSoftClassPtr<UAnimInstance>(), LODIndex, SectionIndex, LODIndex, SectionIndex, FString(), GenerationContext, TypedNodeStatic);
+			constexpr bool bIsReference = false;
+			mu::MeshPtr MutableMesh = GenerateMutableMesh(TypedNodeStatic->StaticMesh, TSoftClassPtr<UAnimInstance>(), LODIndex, SectionIndex, LODIndex, SectionIndex, FString(), GenerationContext, TypedNodeStatic, bIsReference);
 			if (MutableMesh)
 			{
 				MeshNode->SetValue(MutableMesh);
@@ -3463,6 +3475,27 @@ mu::NodeMeshPtr GenerateMutableSourceMesh(const UEdGraphPin* Pin,
 		}
 	}
 
+	else if (const UCustomizableObjectNodePassThroughMesh* TypedNodePassThroughMesh = Cast<UCustomizableObjectNodePassThroughMesh>(Node))
+	{
+		USkeletalMesh* BaseMesh = TypedNodePassThroughMesh->Mesh;
+		if (BaseMesh)
+		{
+			mu::Ptr<mu::NodeMeshConstant> MeshNode = new mu::NodeMeshConstant();
+			Result = MeshNode;
+
+			FString MeshUniqueTags;
+			constexpr bool bIsReference = true;
+			TSoftClassPtr<UAnimInstance> AnimInstance;
+			mu::Ptr<mu::Mesh> MutableMesh = GenerateMutableMesh(BaseMesh, AnimInstance, 0, 0, 0, 0, MeshUniqueTags, GenerationContext, TypedNodeSkel, bIsReference);
+
+			MeshNode->SetValue(MutableMesh);
+		}
+		else
+		{
+			GenerationContext.Compiler->CompilerLog(LOCTEXT("MissingImagePassThrough", "Missing image in pass-through texture node."), Node);
+		}
+	}
+
 	else if (UCustomizableObjectNodeMeshMorph* TypedNodeMorph = Cast<UCustomizableObjectNodeMeshMorph>(Node))
 	{
 		if (const UEdGraphPin* ConnectedPin = FollowInputPin(*TypedNodeMorph->MeshPin()))
diff --git a/Engine/Plugins/Experimental/Mutable/Source/CustomizableObjectEditor/Private/MuCOE/GenerateMutableSource/GenerateMutableSourceMesh.h b/Engine/Plugins/Experimental/Mutable/Source/CustomizableObjectEditor/Private/MuCOE/GenerateMutableSource/GenerateMutableSourceMesh.h
--- a/Engine/Plugins/Experimental/Mutable/Source/CustomizableObjectEditor/Private/MuCOE/GenerateMutableSource/GenerateMutableSourceMesh.h
+++ b/Engine/Plugins/Experimental/Mutable/Source/CustomizableObjectEditor/Private/MuCOE/GenerateMutableSource/GenerateMutableSourceMesh.h
@@ -74,7 +74,8 @@ mu::Ptr<mu::Mesh> ConvertStaticMeshToMutable(const UStaticMesh* StaticMesh, int3
 mu::Ptr<mu::Mesh> GenerateMutableMesh(const UObject* Mesh, const TSoftClassPtr<UAnimInstance>& AnimBp,
                                 int32 LODIndexConnected, int32 SectionIndexConnected,
                                 int32 LODIndex, int32 SectionIndex, const FString& MeshUniqueTags,
-                                FMutableGraphGenerationContext& GenerationContext, const UCustomizableObjectNode* CurrentNode);
+                                FMutableGraphGenerationContext& GenerationContext, const UCustomizableObjectNode* CurrentNode, 
+								bool bIsReference);
 
 
 mu::Ptr<mu::Mesh> BuildMorphedMutableMesh(const UEdGraphPin* BaseSourcePin, const FString& MorphTargetName,
diff --git a/Engine/Plugins/Experimental/Mutable/Source/CustomizableObjectEditor/Private/MuCOE/GenerateMutableSource/GenerateMutableSourceTable.cpp b/Engine/Plugins/Experimental/Mutable/Source/CustomizableObjectEditor/Private/MuCOE/GenerateMutableSource/GenerateMutableSourceTable.cpp
--- a/Engine/Plugins/Experimental/Mutable/Source/CustomizableObjectEditor/Private/MuCOE/GenerateMutableSource/GenerateMutableSourceTable.cpp
+++ b/Engine/Plugins/Experimental/Mutable/Source/CustomizableObjectEditor/Private/MuCOE/GenerateMutableSource/GenerateMutableSourceTable.cpp
@@ -281,7 +281,8 @@ bool FillTableColumn(const UCustomizableObjectNodeTable* TableNode, mu::TablePtr
 			}
 
 			//TODO: Add AnimBp physics to Tables.
-			mu::Ptr<mu::Mesh> MutableMesh = GenerateMutableMesh(SkeletalMesh, TSoftClassPtr<UAnimInstance>(), LODIndexConnected, SectionIndexConnected, LODIndex, SectionIndex, MeshUniqueTags, GenerationContext, TableNode);
+			constexpr bool bIsReference = false;
+			mu::Ptr<mu::Mesh> MutableMesh = GenerateMutableMesh(SkeletalMesh, TSoftClassPtr<UAnimInstance>(), LODIndexConnected, SectionIndexConnected, LODIndex, SectionIndex, MeshUniqueTags, GenerationContext, TableNode, bIsReference);
 
 			if (MutableMesh)
 			{
@@ -387,7 +388,8 @@ bool FillTableColumn(const UCustomizableObjectNodeTable* TableNode, mu::TablePtr
 				CurrentColumn = MutableTable->AddColumn(MutableColumnName, mu::ETableColumnType::Mesh);
 			}
 
-			mu::MeshPtr MutableMesh = GenerateMutableMesh(StaticMesh, TSoftClassPtr<UAnimInstance>(), CurrentLOD, SectionIndex, CurrentLOD, SectionIndex, FString(), GenerationContext, TableNode);
+			constexpr bool bIsReference = false;
+			mu::MeshPtr MutableMesh = GenerateMutableMesh(StaticMesh, TSoftClassPtr<UAnimInstance>(), CurrentLOD, SectionIndex, CurrentLOD, SectionIndex, FString(), GenerationContext, TableNode, bIsReference);
 
 			if (MutableMesh)
 			{
diff --git a/Engine/Plugins/Experimental/Mutable/Source/CustomizableObjectEditor/Private/MuCOE/GraphTraversal.cpp b/Engine/Plugins/Experimental/Mutable/Source/CustomizableObjectEditor/Private/MuCOE/GraphTraversal.cpp
--- a/Engine/Plugins/Experimental/Mutable/Source/CustomizableObjectEditor/Private/MuCOE/GraphTraversal.cpp
+++ b/Engine/Plugins/Experimental/Mutable/Source/CustomizableObjectEditor/Private/MuCOE/GraphTraversal.cpp
@@ -19,6 +19,7 @@
 #include "MuCOE/Nodes/CustomizableObjectNodeMeshReshape.h"
 #include "MuCOE/Nodes/CustomizableObjectNodeMeshSwitch.h"
 #include "MuCOE/Nodes/CustomizableObjectNodeMeshVariation.h"
+#include "MuCOE/Nodes/CustomizableObjectNodePassThroughMesh.h"
 #include "MuCOE/Nodes/CustomizableObjectNodeObject.h"
 #include "MuCOE/Nodes/CustomizableObjectNodeSkeletalMesh.h"
 #include "MuCOE/Nodes/CustomizableObjectNodeStaticMesh.h"
@@ -364,7 +365,12 @@ UCustomizableObjectNodeObject* GetFullGraphRootNodeObject(UCustomizableObjectNod
 const UEdGraphPin* FindMeshBaseSource(const UEdGraphPin& Pin, const bool bOnlyLookForStaticMesh)
 {
 	check(Pin.Direction == EGPD_Output);
-	check(Pin.PinType.PinCategory == UEdGraphSchema_CustomizableObject::PC_Mesh || Pin.PinType.PinCategory == UEdGraphSchema_CustomizableObject::PC_Material);
+	check(Pin.PinType.PinCategory == UEdGraphSchema_CustomizableObject::PC_Mesh 
+		||
+		Pin.PinType.PinCategory == UEdGraphSchema_CustomizableObject::PC_PassThroughMesh
+		||
+		Pin.PinType.PinCategory == UEdGraphSchema_CustomizableObject::PC_Material
+	);
 	
 	const UEdGraphNode* Node = Pin.GetOwningNode();
 	check(Node);
@@ -377,6 +383,14 @@ const UEdGraphPin* FindMeshBaseSource(const UEdGraphPin& Pin, const bool bOnlyLo
 		}
 	}
 
+	else if (Cast<UCustomizableObjectNodePassThroughMesh>(Node))
+	{
+		if (!bOnlyLookForStaticMesh)
+		{
+			return &Pin;
+		}
+	}
+
 	else if (Cast<UCustomizableObjectNodeStaticMesh>(Node))
 	{
 		return &Pin;
diff --git a/Engine/Plugins/Experimental/Mutable/Source/CustomizableObjectEditor/Private/MuCOE/Nodes/CustomizableObjectNodeMaterial.cpp b/Engine/Plugins/Experimental/Mutable/Source/CustomizableObjectEditor/Private/MuCOE/Nodes/CustomizableObjectNodeMaterial.cpp
--- a/Engine/Plugins/Experimental/Mutable/Source/CustomizableObjectEditor/Private/MuCOE/Nodes/CustomizableObjectNodeMaterial.cpp
+++ b/Engine/Plugins/Experimental/Mutable/Source/CustomizableObjectEditor/Private/MuCOE/Nodes/CustomizableObjectNodeMaterial.cpp
@@ -1300,12 +1300,21 @@ bool UCustomizableObjectNodeMaterial::CanConnect(const UEdGraphPin* InOwnedInput
 {
 	const UEdGraphSchema_CustomizableObject* Schema = GetDefault<UEdGraphSchema_CustomizableObject>();
 
-	if (InOwnedInputPin && InOwnedInputPin->PinType.PinCategory == Schema->PC_Image &&
-		InOutputPin && InOutputPin->PinType.PinCategory == Schema->PC_PassThroughImage)
+	if (InOwnedInputPin && InOutputPin)
 	{
-		return true;
+		if (InOwnedInputPin->PinType.PinCategory == Schema->PC_Image &&
+			InOutputPin->PinType.PinCategory == Schema->PC_PassThroughImage)
+		{
+			return true;
+		}
+
+		if (InOwnedInputPin->PinType.PinCategory == Schema->PC_Mesh &&
+			InOutputPin->PinType.PinCategory == Schema->PC_PassThroughMesh)
+		{
+			return true;
+		}
 	}
-	
+
 	return Super::CanConnect(InOwnedInputPin, InOutputPin, bOutIsOtherNodeBlocklisted, bOutArePinsCompatible);
 }
 
diff --git a/Engine/Plugins/Experimental/Mutable/Source/CustomizableObjectEditor/Private/MuCOE/Nodes/CustomizableObjectNodePassThroughMesh.cpp b/Engine/Plugins/Experimental/Mutable/Source/CustomizableObjectEditor/Private/MuCOE/Nodes/CustomizableObjectNodePassThroughMesh.cpp
new file mode 100644
--- /dev/null
+++ b/Engine/Plugins/Experimental/Mutable/Source/CustomizableObjectEditor/Private/MuCOE/Nodes/CustomizableObjectNodePassThroughMesh.cpp
@@ -0,0 +1,50 @@
+// Copyright Epic Games, Inc. All Rights Reserved.
+
+#include "MuCOE/Nodes/CustomizableObjectNodePassThroughMesh.h"
+
+#include "MuCO/CustomizableObjectCustomVersion.h"
+#include "MuCOE/EdGraphSchema_CustomizableObject.h"
+#include "Engine/SkeletalMesh.h"
+
+#define LOCTEXT_NAMESPACE "CustomizableObjectEditor"
+
+
+void UCustomizableObjectNodePassThroughMesh::AllocateDefaultPins(UCustomizableObjectNodeRemapPins* RemapPins)
+{
+	const UEdGraphSchema_CustomizableObject* Schema = GetDefault<UEdGraphSchema_CustomizableObject>();
+
+	FString PinName = TEXT("Mesh");
+	UEdGraphPin* PinMeshPin = CustomCreatePin(EGPD_Output, Schema->PC_PassThroughMesh, FName(*PinName));
+	PinMeshPin->bDefaultValueIsIgnored = true;
+}
+
+
+FText UCustomizableObjectNodePassThroughMesh::GetNodeTitle(ENodeTitleType::Type TitleType) const
+{
+	if (Mesh)
+	{
+		FFormatNamedArguments Args;
+		Args.Add(TEXT("MeshName"), FText::FromString(Mesh->GetName()));
+
+		return FText::Format(LOCTEXT("PassThrough Mesh_Title", "{MeshName}\nPassThrough Mesh"), Args);
+	}
+	else
+	{
+		return LOCTEXT("PassThrough Mesh", "PassThrough Mesh");
+	}
+}
+
+
+FLinearColor UCustomizableObjectNodePassThroughMesh::GetNodeTitleColor() const
+{
+	const UEdGraphSchema_CustomizableObject* Schema = GetDefault<UEdGraphSchema_CustomizableObject>();
+	return Schema->GetPinTypeColor(Schema->PC_PassThroughMesh);
+}
+
+
+FText UCustomizableObjectNodePassThroughMesh::GetTooltipText() const
+{
+	return LOCTEXT("PassThrough_Mesh_Tooltip", "Defines a pass-through Mesh. It will not be modified by Mutable in any way, just referenced as a UE asset. It's much cheaper than a Mutable Mesh, but you cannot make any operations on it, just switch it.");
+}
+
+#undef LOCTEXT_NAMESPACE
diff --git a/Engine/Plugins/Experimental/Mutable/Source/CustomizableObjectEditor/Private/MuCOE/SMutableCodeViewer.cpp b/Engine/Plugins/Experimental/Mutable/Source/CustomizableObjectEditor/Private/MuCOE/SMutableCodeViewer.cpp
--- a/Engine/Plugins/Experimental/Mutable/Source/CustomizableObjectEditor/Private/MuCOE/SMutableCodeViewer.cpp
+++ b/Engine/Plugins/Experimental/Mutable/Source/CustomizableObjectEditor/Private/MuCOE/SMutableCodeViewer.cpp
@@ -2431,7 +2431,7 @@ bool SMutableCodeViewer::IsConstantResourceUsedByOperation(const int32 IndexOnCo
 namespace
 {
 	/** Test implementation to provide image parameters. It will generate some images of a fixed size and format. */
-	class TestImageProvider : public mu::ImageParameterGenerator
+	class TestResourceProvider : public mu::ExternalResourceProvider
 	{
 		static inline const mu::FImageDesc IMAGE_DESC = 
 			mu::FImageDesc(mu::FImageSize(1024, 1024), mu::EImageFormat::IF_RGBA_UBYTE, 1);
@@ -2444,7 +2444,7 @@ namespace
 	public:
 
 
-		TTuple<UE::Tasks::FTask, TFunction<void()>> GetImageAsync(FName Id, uint8 MipmapsToSkip, TFunction<void(mu::Ptr<mu::Image>)>& ResultCallback) override
+		virtual TTuple<UE::Tasks::FTask, TFunction<void()>> GetImageAsync(FName Id, uint8 MipmapsToSkip, TFunction<void(mu::Ptr<mu::Image>)>& ResultCallback) override
 		{
 			MUTABLE_CPUPROFILER_SCOPE(TestImageProvider_GetImage);
 
@@ -2486,13 +2486,13 @@ namespace
 			return MakeTuple(UE::Tasks::MakeCompletedTask<void>(), []() -> void {});
 		}
 
-		mu::FImageDesc GetImageDesc(FName Id, uint8 MipmapsToSkip) override
+		virtual mu::FImageDesc GetImageDesc(FName Id, uint8 MipmapsToSkip) override
 		{
 			return IMAGE_DESC;
 		}
 
 
-		TTuple<UE::Tasks::FTask, TFunction<void()>> GetReferencedImageAsync(const void* ModelPtr, int32 Id, uint8 MipmapsToSkip, TFunction<void(mu::Ptr<mu::Image>)>& ResultCallback)
+		virtual TTuple<UE::Tasks::FTask, TFunction<void()>> GetReferencedImageAsync(const void* ModelPtr, int32 Id, uint8 MipmapsToSkip, TFunction<void(mu::Ptr<mu::Image>)>& ResultCallback)
 		{
 			check(ReferencedTextures.IsValidIndex(Id));
 
@@ -2526,6 +2526,50 @@ namespace
 			return Invoke(TrivialReturn);
 		}
 
+
+		virtual TTuple<UE::Tasks::FTask, TFunction<void()>> GetMeshAsync(FName Id, TFunction<void(mu::Ptr<mu::Mesh>)>& ResultCallback) override
+		{
+			// TODO: This is a stub function where Mesh Parameters would be supported.
+
+			// Thread: worker
+			MUTABLE_CPUPROFILER_SCOPE(FUnrealMutableImageProvider::GetMeshAsync);
+
+			auto TrivialReturn = []() -> TTuple<UE::Tasks::FTask, TFunction<void()>>
+				{
+					return MakeTuple(UE::Tasks::MakeCompletedTask<void>(), []() -> void {});
+				};
+
+			// No USkeletalMesh was provided, cannot do anything, just provide a dummy mesh
+			UE_LOG(LogMutable, Warning, TEXT("No USkeletalMesh was provided for an application-specific mesh parameter."));
+
+			mu::Ptr<mu::Mesh> Dummy = new mu::Mesh;
+			ResultCallback(Dummy);
+			return Invoke(TrivialReturn);
+		}
+
+
+		virtual TTuple<UE::Tasks::FTask, TFunction<void()>> GetReferencedMeshAsync(const void* ModelPtr, int32 Id, TFunction<void(mu::Ptr<mu::Mesh>)>& ResultCallback) override
+		{
+			// TODO: This is a stub function where referenced Meshes would be supported.
+
+			// Thread: worker
+			MUTABLE_CPUPROFILER_SCOPE(FUnrealMutableImageProvider::GetReferencedMeshAsync);
+
+			auto TrivialReturn = []() -> TTuple<UE::Tasks::FTask, TFunction<void()>>
+				{
+					return MakeTuple(UE::Tasks::MakeCompletedTask<void>(), []() -> void {});
+				};
+
+			mu::Ptr<mu::Mesh> Mesh = new mu::Mesh();
+
+			// Not supported outside editor yet.
+			UE_LOG(LogMutable, Warning, TEXT("Failed to get referenced mesh. Not implemented yet."));
+
+			ResultCallback(Mesh);
+			return Invoke(TrivialReturn);
+
+		}
+
 	};
 }
 
@@ -2560,9 +2604,9 @@ void SMutableCodeViewer::Tick(const FGeometry& AllottedGeometry, const double In
 	const mu::Ptr<mu::Settings> Settings = new mu::Settings();
 	const mu::Ptr<mu::System> System = new mu::System(Settings);
 
-	TSharedPtr<TestImageProvider> ImageProvider = MakeShared<TestImageProvider>();
-	ImageProvider->ReferencedTextures = ReferencedTextures;
-	System->SetImageParameterGenerator(ImageProvider);
+	TSharedPtr<TestResourceProvider> ExternalResourceProvider = MakeShared<TestResourceProvider>();
+	ExternalResourceProvider->ReferencedTextures = ReferencedTextures;
+	System->SetExternalResourceProvider(ExternalResourceProvider);
 
 	System->GetPrivate()->BeginBuild(MutableModel);
 
diff --git a/Engine/Plugins/Experimental/Mutable/Source/CustomizableObjectEditor/Private/MuCOE/SMutableMeshViewer.cpp b/Engine/Plugins/Experimental/Mutable/Source/CustomizableObjectEditor/Private/MuCOE/SMutableMeshViewer.cpp
--- a/Engine/Plugins/Experimental/Mutable/Source/CustomizableObjectEditor/Private/MuCOE/SMutableMeshViewer.cpp
+++ b/Engine/Plugins/Experimental/Mutable/Source/CustomizableObjectEditor/Private/MuCOE/SMutableMeshViewer.cpp
@@ -274,6 +274,7 @@ TSharedRef<SWidget> SMutableMeshViewer::GenerateDataTableSlates()
 	const FText FacesCountTitle = LOCTEXT("FacesCountTitle", "Face count : ");
 	const FText BonesCountTitle = LOCTEXT("BonesCountTitle", "Bone count : ");
 	const FText MeshIdPrefixTitle = LOCTEXT("MeshIdPrefixTitle", "Mesh ID prefix : ");
+	const FText MeshFlagsTitle = LOCTEXT("MeshFlagsTitle", "Mesh flags : ");
 	const FText BuffersTitle = LOCTEXT("BuffersTitle", "Buffers");
 	
 	return SNew(SScrollBox)
@@ -387,6 +388,28 @@ TSharedRef<SWidget> SMutableMeshViewer::GenerateDataTableSlates()
 								Text(this, &SMutableMeshViewer::GetMeshIdPrefix)
 						]
 				]
+
+				// Flags
+				+ SVerticalBox::Slot().
+				Padding(0, SimpleSpacing).
+				AutoHeight()
+				[
+					SNew(SHorizontalBox)
+
+						+ SHorizontalBox::Slot()
+						.AutoWidth()
+						[
+							SNew(STextBlock).
+								Text(MeshFlagsTitle)
+						]
+
+						+ SHorizontalBox::Slot()
+						.AutoWidth()
+						[
+							SNew(STextBlock).
+								Text(this, &SMutableMeshViewer::GetMeshFlags)
+						]
+				]
 			]
 
 			// Buffers Data --------------------------------------------------------------
@@ -694,6 +717,11 @@ FText SMutableMeshViewer::GetMeshIdPrefix() const
 	return FText::AsNumber(MutableMesh ? MutableMesh->MeshIDPrefix : 0);
 }
 
+FText SMutableMeshViewer::GetMeshFlags() const
+{
+	return FText::AsNumber(MutableMesh ? uint32(MutableMesh->Flags) : uint32(0));
+}
+
 
 TSharedRef<ITableRow> SMutableMeshViewer::OnGenerateBufferRow(TSharedPtr<FBufferElement, ESPMode::ThreadSafe> InBuffer,
                                                               const TSharedRef<STableViewBase, ESPMode::ThreadSafe>& OwnerTable)
diff --git a/Engine/Plugins/Experimental/Mutable/Source/CustomizableObjectEditor/Private/MuCOE/SMutableMeshViewer.h b/Engine/Plugins/Experimental/Mutable/Source/CustomizableObjectEditor/Private/MuCOE/SMutableMeshViewer.h
--- a/Engine/Plugins/Experimental/Mutable/Source/CustomizableObjectEditor/Private/MuCOE/SMutableMeshViewer.h
+++ b/Engine/Plugins/Experimental/Mutable/Source/CustomizableObjectEditor/Private/MuCOE/SMutableMeshViewer.h
@@ -124,6 +124,7 @@ private:
 	FText GetFaceCount() const;
 	FText GetBoneCount() const;
 	FText GetMeshIdPrefix() const;
+	FText GetMeshFlags() const;
 
 	/*
 	 * Elements used to feed the buffers list (index and buffer channels as an internal list)
diff --git a/Engine/Plugins/Experimental/Mutable/Source/CustomizableObjectEditor/Public/MuCOE/EdGraphSchema_CustomizableObject.h b/Engine/Plugins/Experimental/Mutable/Source/CustomizableObjectEditor/Public/MuCOE/EdGraphSchema_CustomizableObject.h
--- a/Engine/Plugins/Experimental/Mutable/Source/CustomizableObjectEditor/Public/MuCOE/EdGraphSchema_CustomizableObject.h
+++ b/Engine/Plugins/Experimental/Mutable/Source/CustomizableObjectEditor/Public/MuCOE/EdGraphSchema_CustomizableObject.h
@@ -91,6 +91,7 @@ public:
 	static const FName PC_Object;
 	static const FName PC_Material;
 	static const FName PC_Mesh;
+	static const FName PC_PassThroughMesh;
 	static const FName PC_Layout;
 	static const FName PC_Image;
 	static const FName PC_PassThroughImage;
diff --git a/Engine/Plugins/Experimental/Mutable/Source/CustomizableObjectEditor/Public/MuCOE/Nodes/CustomizableObjectNodePassThroughMesh.h b/Engine/Plugins/Experimental/Mutable/Source/CustomizableObjectEditor/Public/MuCOE/Nodes/CustomizableObjectNodePassThroughMesh.h
new file mode 100644
--- /dev/null
+++ b/Engine/Plugins/Experimental/Mutable/Source/CustomizableObjectEditor/Public/MuCOE/Nodes/CustomizableObjectNodePassThroughMesh.h
@@ -0,0 +1,36 @@
+// Copyright Epic Games, Inc. All Rights Reserved.
+
+#pragma once
+
+#include "MuCOE/Nodes/CustomizableObjectNodeMesh.h"
+#include "Engine/SkeletalMesh.h"
+
+#include "CustomizableObjectNodePassThroughMesh.generated.h"
+
+
+UCLASS()
+class CUSTOMIZABLEOBJECTEDITOR_API UCustomizableObjectNodePassThroughMesh : public UCustomizableObjectNodeMesh
+{
+public:
+	GENERATED_BODY()
+
+	UPROPERTY(EditAnywhere, Category = Mesh, Meta = (DisplayName = Mesh))
+	TObjectPtr<class USkeletalMesh> Mesh = nullptr;
+
+	// UCustomizableObjectNodeMesh interface
+	virtual TArray<UCustomizableObjectLayout*> GetLayouts(const UEdGraphPin& OutPin) const override { return {}; }
+	virtual UTexture2D* FindTextureForPin(const UEdGraphPin* Pin) const override { return nullptr; }
+	virtual void GetUVChannelForPin(const UEdGraphPin* Pin, TArray<FVector2f>& OutSegments, int32 UVIndex) const override {}
+	virtual UObject* GetMesh() const override { return Mesh.Get(); }
+	virtual UEdGraphPin* GetMeshPin(int32 LOD, int32 SectionIndex) const override { return {}; }
+	virtual UEdGraphPin* GetLayoutPin(int32 LODIndex, int32 SectionIndex, int32 LayoutIndex) const override { return {}; }
+	virtual void GetPinSection(const UEdGraphPin& Pin, int32& OutLODIndex, int32& OutSectionIndex, int32& OutLayoutIndex) const override {}
+
+	// UCustomizableObjectNode interface
+	virtual void AllocateDefaultPins(UCustomizableObjectNodeRemapPins* RemapPins) override;
+
+	// Begin EdGraphNode interface
+	virtual FText GetNodeTitle(ENodeTitleType::Type TitleType) const override;
+	virtual FLinearColor GetNodeTitleColor() const override;
+	virtual FText GetTooltipText() const override;
+};
diff --git a/Engine/Plugins/Experimental/Mutable/Source/MutableRuntime/Internal/MuR/MeshPrivate.h b/Engine/Plugins/Experimental/Mutable/Source/MutableRuntime/Internal/MuR/MeshPrivate.h
--- a/Engine/Plugins/Experimental/Mutable/Source/MutableRuntime/Internal/MuR/MeshPrivate.h
+++ b/Engine/Plugins/Experimental/Mutable/Source/MutableRuntime/Internal/MuR/MeshPrivate.h
@@ -759,21 +759,4 @@ namespace mu
 	};
 
 
-	/** Optimised mesh formats that are identified in some operations to chose a faster version. */
-	typedef enum
-	{
-		SMF_NONE,
-
-        SMF_PROJECT,
-
-        SMF_PROJECTWRAPPING,
-
-		SMF_COUNT
-	} STATIC_MESH_FORMATS;
-
-	//! Array of functions that can identify each format
-	typedef bool (*STATIC_MESH_FORMAT_ID_FUNC)( const Mesh* );
-	extern STATIC_MESH_FORMAT_ID_FUNC s_staticMeshFormatIdentify[SMF_COUNT];
-
-
 }
diff --git a/Engine/Plugins/Experimental/Mutable/Source/MutableRuntime/Internal/MuR/ModelPrivate.h b/Engine/Plugins/Experimental/Mutable/Source/MutableRuntime/Internal/MuR/ModelPrivate.h
--- a/Engine/Plugins/Experimental/Mutable/Source/MutableRuntime/Internal/MuR/ModelPrivate.h
+++ b/Engine/Plugins/Experimental/Mutable/Source/MutableRuntime/Internal/MuR/ModelPrivate.h
@@ -14,7 +14,7 @@
 #include "MuR/ParametersPrivate.h"
 #include "MuR/MutableRuntimeModule.h"
 
-#define MUTABLE_MAX_RUNTIME_PARAMETERS_PER_STATE	64
+#define MUTABLE_MAX_RUNTIME_PARAMETERS_PER_STATE	65
 #define MUTABLE_GROW_BORDER_VALUE					2
 
 namespace mu
diff --git a/Engine/Plugins/Experimental/Mutable/Source/MutableRuntime/Internal/MuR/Operations.h b/Engine/Plugins/Experimental/Mutable/Source/MutableRuntime/Internal/MuR/Operations.h
--- a/Engine/Plugins/Experimental/Mutable/Source/MutableRuntime/Internal/MuR/Operations.h
+++ b/Engine/Plugins/Experimental/Mutable/Source/MutableRuntime/Internal/MuR/Operations.h
@@ -52,6 +52,7 @@ namespace mu
 
 		//! A referenced, but opaque engine resource
 		IM_REFERENCE,
+		ME_REFERENCE,
 
         //! Select one value or the other depending on a boolean input
         NU_CONDITIONAL,
@@ -1199,7 +1200,7 @@ namespace mu
 
 
 	// Generic data about a Mutable runtime operation.
-    struct OP_DESC
+    struct FOpDesc
     {
 		//! Type of data generated by the instruction
 		DATATYPE type;
@@ -1213,7 +1214,7 @@ namespace mu
 		bool supportedBasePixelFormats[uint8(EImageFormat::IF_COUNT)];
     };
 
-	MUTABLERUNTIME_API extern const OP_DESC& GetOpDesc( OP_TYPE type );
+	MUTABLERUNTIME_API extern const FOpDesc& GetOpDesc( OP_TYPE type );
 
 
     //!
diff --git a/Engine/Plugins/Experimental/Mutable/Source/MutableRuntime/Internal/MuR/SystemPrivate.h b/Engine/Plugins/Experimental/Mutable/Source/MutableRuntime/Internal/MuR/SystemPrivate.h
--- a/Engine/Plugins/Experimental/Mutable/Source/MutableRuntime/Internal/MuR/SystemPrivate.h
+++ b/Engine/Plugins/Experimental/Mutable/Source/MutableRuntime/Internal/MuR/SystemPrivate.h
@@ -1911,7 +1911,7 @@ namespace mu
         //! Data streaming interface, if any.
 		TSharedPtr<ModelReader> StreamInterface;
 
-		TSharedPtr<ImageParameterGenerator> ImageParameterGenerator;
+		TSharedPtr<ExternalResourceProvider> ExternalResourceProvider;
 
 		/** If this is set, it will be tried first instead of the internal formatting function. */
 		FImageOperator::FImagePixelFormatFunc ImagePixelFormatOverride;
diff --git a/Engine/Plugins/Experimental/Mutable/Source/MutableRuntime/Private/MuR/CodeRunner.cpp b/Engine/Plugins/Experimental/Mutable/Source/MutableRuntime/Private/MuR/CodeRunner.cpp
--- a/Engine/Plugins/Experimental/Mutable/Source/MutableRuntime/Private/MuR/CodeRunner.cpp
+++ b/Engine/Plugins/Experimental/Mutable/Source/MutableRuntime/Private/MuR/CodeRunner.cpp
@@ -166,23 +166,52 @@ namespace mu
 	}
 
 
-	TTuple<UE::Tasks::FTask, TFunction<void()>> CodeRunner::LoadExternalImageAsync(FExternalImageId Id, uint8 MipmapsToSkip, TFunction<void(Ptr<Image>)>& ResultCallback)
+	TTuple<UE::Tasks::FTask, TFunction<void()>> CodeRunner::LoadExternalImageAsync(FExternalResourceId Id, uint8 MipmapsToSkip, TFunction<void(Ptr<Image>)>& ResultCallback)
     {
 		MUTABLE_CPUPROFILER_SCOPE(LoadExternalImageAsync);
 
 		check(m_pSystem);
 
-		if (m_pSystem->ImageParameterGenerator)
+		if (m_pSystem->ExternalResourceProvider)
 		{
-			if (Id.ReferenceImageId < 0)
+			if (Id.ReferenceResourceId < 0)
 			{
 				// It's a parameter image
-				return m_pSystem->ImageParameterGenerator->GetImageAsync(Id.ParameterId, MipmapsToSkip, ResultCallback);
+				return m_pSystem->ExternalResourceProvider->GetImageAsync(Id.ParameterId, MipmapsToSkip, ResultCallback);
 			}
 			else
 			{
 				// It's an image reference
-				return m_pSystem->ImageParameterGenerator->GetReferencedImageAsync(m_pModel.Get(), Id.ReferenceImageId, MipmapsToSkip, ResultCallback);
+				return m_pSystem->ExternalResourceProvider->GetReferencedImageAsync(m_pModel.Get(), Id.ReferenceResourceId, MipmapsToSkip, ResultCallback);
+			}
+		}
+		else
+		{
+			// Not found and there is no generator!
+			check(false);
+		}
+
+		return MakeTuple(UE::Tasks::MakeCompletedTask<void>(), []() -> void {});
+	}
+
+
+	TTuple<UE::Tasks::FTask, TFunction<void()>> CodeRunner::LoadExternalMeshAsync(FExternalResourceId Id, TFunction<void(Ptr<Mesh>)>& ResultCallback)
+	{
+		MUTABLE_CPUPROFILER_SCOPE(LoadExternalImageAsync);
+
+		check(m_pSystem);
+
+		if (m_pSystem->ExternalResourceProvider)
+		{
+			if (Id.ReferenceResourceId < 0)
+			{
+				// It's a parameter mesh
+				return m_pSystem->ExternalResourceProvider->GetMeshAsync(Id.ParameterId, ResultCallback);
+			}
+			else
+			{
+				// It's a mesh reference
+				return m_pSystem->ExternalResourceProvider->GetReferencedMeshAsync(m_pModel.Get(), Id.ReferenceResourceId, ResultCallback);
 			}
 		}
 		else
@@ -202,9 +231,9 @@ namespace mu
 
 		check(m_pSystem);
 
-		if (m_pSystem->ImageParameterGenerator)
+		if (m_pSystem->ExternalResourceProvider)
 		{
-			return m_pSystem->ImageParameterGenerator->GetImageDesc(Id, MipmapsToSkip);
+			return m_pSystem->ExternalResourceProvider->GetImageDesc(Id, MipmapsToSkip);
 		}
 		else
 		{
@@ -1016,6 +1045,34 @@ namespace mu
         switch (type)
         {
 
+		case OP_TYPE::ME_REFERENCE:
+		{
+			OP::ResourceReferenceArgs Args = Program.GetOpArgs<OP::ResourceReferenceArgs>(item.At);
+			switch (item.Stage)
+			{
+			case 0:
+			{
+				Ptr<Mesh> Result;
+				if (Args.ForceLoad)
+				{
+					// This should never be reached because it should have been caught as a Task in IssueOp
+					check(false);
+				}
+				else
+				{
+					Result = Mesh::CreateAsReference(Args.ID, false);
+				}
+				StoreMesh(item, Result);
+				break;
+			}
+
+			default:
+				check(false);
+			}
+
+			break;
+		}
+
         case OP_TYPE::ME_APPLYLAYOUT:
         {
 			OP::MeshApplyLayoutArgs args = Program.GetOpArgs<OP::MeshApplyLayoutArgs>(item.At);
@@ -1284,6 +1341,8 @@ namespace mu
 
                 if (pA && pB && pA->GetVertexCount() && pB->GetVertexCount())
                 {
+					check(!pA->IsReference() && !pB->IsReference());
+
 					FMeshMergeScratchMeshes Scratch;
 					Scratch.FirstReformat = CreateMesh();
 					Scratch.SecondReformat = CreateMesh();
@@ -1305,16 +1364,16 @@ namespace mu
 					Release(pB);
 					StoreMesh(item, Result);
                 }
-                else if (pA && pA->GetVertexCount())
+                else if (pA && (pA->GetVertexCount() || pA->IsReference()))
                 {
 					Release(pB);
 					StoreMesh(item, pA);
                 }
-                else if (pB && pB->GetVertexCount())
+                else if (pB && (pB->GetVertexCount() || pB->IsReference()))
                 {
 					Ptr<Mesh> Result = CloneOrTakeOver(pB);
 
-                    check(Result->GetSurfaceCount() == 1);
+                    check(Result->IsReference() || (Result->GetSurfaceCount() == 1) );
 
                     if (Result->GetSurfaceCount() > 0 && args.newSurfaceID)
                     {
@@ -1793,7 +1852,12 @@ namespace mu
                 Ptr<const Mesh> Source = LoadMesh(FCacheAddress(args.source,item));
                 Ptr<const Mesh> Format = LoadMesh(FCacheAddress(args.format,item));
 
-				if (Source)
+				if (Source && Source->IsReference())
+				{
+					Release(Format);
+					StoreMesh(item, Source);
+				}
+				else if (Source)
 				{
 					uint8 Flags = args.Flags;
 					if (!Format && !(Flags & OP::MeshFormatArgs::ResetBufferIndices))
@@ -2990,6 +3054,11 @@ namespace mu
 
 				Ptr<const Mesh> Source = LoadMesh(FCacheAddress(args.source, item));
 
+				if (Source && Source->IsReference())
+				{
+					StoreMesh(item, Source);
+				}
+
 				Ptr<Mesh> Result = CreateMesh();
 
 				bool bOutSuccess = false;
diff --git a/Engine/Plugins/Experimental/Mutable/Source/MutableRuntime/Private/MuR/CodeRunner.h b/Engine/Plugins/Experimental/Mutable/Source/MutableRuntime/Private/MuR/CodeRunner.h
--- a/Engine/Plugins/Experimental/Mutable/Source/MutableRuntime/Private/MuR/CodeRunner.h
+++ b/Engine/Plugins/Experimental/Mutable/Source/MutableRuntime/Private/MuR/CodeRunner.h
@@ -136,20 +136,22 @@ namespace  mu
 
     public:
 
-		struct FExternalImageId
+		struct FExternalResourceId
 		{
-			/** If it is an image reference. */
-			int32 ReferenceImageId = -1;
+			/** If it is an image or mesh reference. */
+			int32 ReferenceResourceId = -1;
 
-			/** If it is an image parameter.*/
+			/** If it is an image or mesh parameter.*/
 			FName ParameterId; 
 		};
 
-		//! Load an external image asynchronously, retuns an event to wait for complition and a cleanup function 
-		//! that must be called once the event has completed.
-		TTuple<UE::Tasks::FTask, TFunction<void()>> LoadExternalImageAsync(FExternalImageId Id, uint8 MipmapsToSkip, TFunction<void(Ptr<Image>)>& ResultCallback);
+		/** Load an external image asynchronously, returns an event to wait for complition and a cleanup function that must be called once the event has completed. */
+		TTuple<UE::Tasks::FTask, TFunction<void()>> LoadExternalImageAsync(FExternalResourceId Id, uint8 MipmapsToSkip, TFunction<void(Ptr<Image>)>& ResultCallback);
  	    mu::FImageDesc GetExternalImageDesc(FName Id, uint8 MipmapsToSkip);
 
+		/** Load an external mesh asynchronously, returns an event to wait for complition and a cleanup function that must be called once the event has completed. */
+		TTuple<UE::Tasks::FTask, TFunction<void()>> LoadExternalMeshAsync(FExternalResourceId Id, TFunction<void(Ptr<Mesh>)>& ResultCallback);
+
 		/** Settings that may affect the execution of some operations, like image conversion quality. */
 		Ptr<const Settings> m_pSettings;
 
diff --git a/Engine/Plugins/Experimental/Mutable/Source/MutableRuntime/Private/MuR/CodeRunnerTasked.cpp b/Engine/Plugins/Experimental/Mutable/Source/MutableRuntime/Private/MuR/CodeRunnerTasked.cpp
--- a/Engine/Plugins/Experimental/Mutable/Source/MutableRuntime/Private/MuR/CodeRunnerTasked.cpp
+++ b/Engine/Plugins/Experimental/Mutable/Source/MutableRuntime/Private/MuR/CodeRunnerTasked.cpp
@@ -98,6 +98,7 @@ namespace mu
 		{   0,   0,   0,   0 },	// ST_PARAMETER
 		    
 		{   0,   0,   0,   0 },	// IM_REFERENCE
+		{   0,   0,   0,   0 },	// ME_REFERENCE
 
 		{   0,   0,   0,   0 },	// NU_CONDITIONAL
 		{   0,   0,   0,   0 },	// SC_CONDITIONAL
@@ -2415,7 +2416,7 @@ namespace mu
 	{
 	public:
 
-		FImageExternalLoadTask(const FScheduledOp& InItem, uint8 InMipmapsToSkip, CodeRunner::FExternalImageId InId);
+		FImageExternalLoadTask(const FScheduledOp& InItem, uint8 InMipmapsToSkip, CodeRunner::FExternalResourceId InId);
 
 		// FIssuedTask interface
 		virtual bool Prepare(CodeRunner*, bool& bOutFailed) override;
@@ -2423,7 +2424,7 @@ namespace mu
 		
 	private:
 		uint8 MipmapsToSkip;
-		CodeRunner::FExternalImageId Id;
+		CodeRunner::FExternalResourceId Id;
 
 		Ptr<Image> Result;
 		
@@ -2431,7 +2432,7 @@ namespace mu
 	};
 
 
-	FImageExternalLoadTask::FImageExternalLoadTask(const FScheduledOp& InOp, uint8 InMipmapsToSkip, CodeRunner::FExternalImageId InId)
+	FImageExternalLoadTask::FImageExternalLoadTask(const FScheduledOp& InOp, uint8 InMipmapsToSkip, CodeRunner::FExternalResourceId InId)
 		: FIssuedTask(InOp)
 	{
 		MipmapsToSkip = InMipmapsToSkip;
@@ -2442,7 +2443,7 @@ namespace mu
 	bool FImageExternalLoadTask::Prepare(CodeRunner* Runner, bool& bOutFailed)
 	{
 		// This runs in the mutable Runner thread
-		MUTABLE_CPUPROFILER_SCOPE(FLoadImageRomsTask_Prepare);
+		MUTABLE_CPUPROFILER_SCOPE(FImageExternalLoadTask_Prepare);
 
 		// LoadExternalImageAsync will always generate some image even if it is a dummy one.
 		bOutFailed = false;
@@ -2472,7 +2473,72 @@ namespace mu
 
 		Runner->StoreImage(Op, Result);
 	}	
-	
+
+
+	/** This task is used to load a mesh parameter (by its FName) or a mesh reference (from its ID).
+	*/
+	class FMeshExternalLoadTask : public CodeRunner::FIssuedTask
+	{
+	public:
+
+		FMeshExternalLoadTask(const FScheduledOp&, CodeRunner::FExternalResourceId);
+
+		// FIssuedTask interface
+		virtual bool Prepare(CodeRunner*, bool& bOutFailed) override;
+		virtual void Complete(CodeRunner* Runner) override;
+
+	private:
+		uint8 MipmapsToSkip;
+		CodeRunner::FExternalResourceId Id;
+
+		Ptr<Mesh> Result;
+
+		TFunction<void()> ExternalCleanUpFunc;
+	};
+
+
+	FMeshExternalLoadTask::FMeshExternalLoadTask(const FScheduledOp& InOp, CodeRunner::FExternalResourceId InId)
+		: FIssuedTask(InOp)
+		, Id(InId)
+	{
+	}
+
+
+	bool FMeshExternalLoadTask::Prepare(CodeRunner* Runner, bool& bOutFailed)
+	{
+		// This runs in the mutable Runner thread
+		MUTABLE_CPUPROFILER_SCOPE(FMeshExternalLoadTask_Prepare);
+
+		// LoadExternalMeshAsync will always generate some mesh even if it is a dummy one.
+		bOutFailed = false;
+
+		// Capturing this here should not be a problem. The lifetime of the callback lambda is tied to 
+		// the task and the later will always outlive the former. 
+
+		// Probably we could simply pass a reference to the result mesh. 
+		TFunction<void(Ptr<Mesh>)> ResultCallback = [this](Ptr<Mesh> InResult)
+			{
+				Result = InResult;
+			};
+
+		Tie(Event, ExternalCleanUpFunc) = Runner->LoadExternalMeshAsync(Id, ResultCallback);
+
+		// return false indicating there is no work to do so Event is not overriden by a DoWork task.
+		return false;
+	}
+
+
+	void FMeshExternalLoadTask::Complete(CodeRunner* Runner)
+	{
+		if (ExternalCleanUpFunc)
+		{
+			Invoke(ExternalCleanUpFunc);
+		}
+
+		Runner->StoreMesh(Op, Result);
+	}
+
+
 	//---------------------------------------------------------------------------------------------
 	//---------------------------------------------------------------------------------------------
 	//---------------------------------------------------------------------------------------------
@@ -2573,7 +2639,7 @@ namespace mu
 
 			const uint8 MipmapsToSkip = item.ExecutionOptions + static_cast<uint8>(ImageLOD);
 
-			CodeRunner::FExternalImageId FullId;
+			CodeRunner::FExternalResourceId FullId;
 			FullId.ParameterId = Id;
 			Issued = MakeShared<FImageExternalLoadTask>(item, MipmapsToSkip, FullId);
 
@@ -2591,14 +2657,31 @@ namespace mu
 
 				const uint8 MipmapsToSkip = item.ExecutionOptions + static_cast<uint8>(ImageLOD);
 
-				FExternalImageId FullId;
-				FullId.ReferenceImageId = Args.ID;
+				FExternalResourceId FullId;
+				FullId.ReferenceResourceId = Args.ID;
 				Issued = MakeShared<FImageExternalLoadTask>(item, MipmapsToSkip, FullId);
 			}
 
 			break;
 		}
 
+		case OP_TYPE::ME_REFERENCE:
+		{
+			OP::ResourceReferenceArgs Args = m_pModel->GetPrivate()->m_program.GetOpArgs<OP::ResourceReferenceArgs>(item.At);
+
+			// We only convert references to meshes if indicated in the operation.
+			if (Args.ForceLoad)
+			{
+				check(item.Stage == 0);
+
+				FExternalResourceId FullId;
+				FullId.ReferenceResourceId = Args.ID;
+				Issued = MakeShared<FMeshExternalLoadTask>(item, FullId);
+			}
+
+			break;
+		}
+
 		case OP_TYPE::IM_PIXELFORMAT:
 		{
 			if (item.Stage == 1)
diff --git a/Engine/Plugins/Experimental/Mutable/Source/MutableRuntime/Private/MuR/Image.cpp b/Engine/Plugins/Experimental/Mutable/Source/MutableRuntime/Private/MuR/Image.cpp
--- a/Engine/Plugins/Experimental/Mutable/Source/MutableRuntime/Private/MuR/Image.cpp
+++ b/Engine/Plugins/Experimental/Mutable/Source/MutableRuntime/Private/MuR/Image.cpp
@@ -41,6 +41,7 @@ namespace mu
 		RelevancyMaxY = 0;
 	}
 
+
 	Ptr<Image> Image::CreateAsReference(uint32 ID, const FImageDesc& Desc, bool bForceLoad)
 	{
 		Ptr<Image> Result = new Image;
@@ -265,92 +266,6 @@ namespace mu
 	}
 
 
-	//int32 Image::CalculateDataSize() const
-	//{
-	//	check(false);
-	//	return DataStorage.GetDataSize();
-	//}
-
-	//int32 Image::CalculateDataSize(int32 Mip) const
-	//{
-	//	check(false);
-	//	return DataStorage.GetLOD(Mip).Num();
-	//}
-
-	//int32 Image::CalculatePixelCount() const
-	//{
-	//	check(false);
-	//	int32 Result = 0;
-
-	//	const int32 LODCount = FMath::Max(1, GetLODCount());
-	//	const FImageFormatData& FormatData = GetImageFormatData(GetFormat());
-	//	
-	//	if (FormatData.BytesPerBlock)
-	//	{
-	//		FImageSize Size = GetSize();
-
-	//		for (int32 L = 0; L < LODCount; ++L)
-	//		{
-	//			int32 BlocksX = FMath::DivideAndRoundUp(Size[0], (uint16)FormatData.PixelsPerBlockX);
-	//			int32 BlocksY = FMath::DivideAndRoundUp(Size[1], (uint16)FormatData.PixelsPerBlockY);
-
-	//			Result += (BlocksX*BlocksY) * FormatData.PixelsPerBlockX * FormatData.PixelsPerBlockY;
-
-	//			Size[0] = FMath::DivideAndRoundUp(Size[0], (uint16)2);
-	//			Size[1] = FMath::DivideAndRoundUp(Size[1], (uint16)2);
-	//		}
-	//	}
-	//	else
-	//	{
-	//		// An RLE image.
-	//		for (int32 L = 0; L < LODCount; ++L)
-	//		{
-	//			FIntVector2 MipSize = CalculateMipSize(L);
-	//			Result += MipSize[0] * MipSize[1];
-	//		}
-	//	}
-
-	//	return Result;
-	//}
-
-
-	////---------------------------------------------------------------------------------------------
-	//int32 Image::CalculatePixelCount(int32 Mip) const
-	//{
-	//	check(false);
-	//	int32 Result = 0;
-	//	
-	//	const FImageFormatData& FormatData = GetImageFormatData(GetFormat());
-	//	if (FormatData.BytesPerBlock)
-	//	{
-	//		FImageSize Size = GetSize();
-
-	//		for (int32 L = 0; L < Mip + 1; ++L)
-	//		{
-	//			int32 BlocksX = FMath::DivideAndRoundUp(Size[0], (uint16)FormatData.PixelsPerBlockX);
-	//			int32 BlocksY = FMath::DivideAndRoundUp(Size[1], (uint16)FormatData.PixelsPerBlockY);
-
-	//			if (L == Mip)
-	//			{
-	//				Result = (BlocksX*BlocksY) * FormatData.PixelsPerBlockX * FormatData.PixelsPerBlockY;
-	//				break;
-	//			}
-
-	//			Size[0] = FMath::DivideAndRoundUp(Size[0], (uint16)2);
-	//			Size[1] = FMath::DivideAndRoundUp(Size[1], (uint16)2);
-	//		}
-	//	}
-	//	else
-	//	{
-	//		// An RLE image.
-	//		FIntVector2 MipSize = CalculateMipSize(Mip);
-	//		Result += MipSize[0] * MipSize[1];
-	//	}
-
-	//	return Result;
-	//}
-
-
 	//---------------------------------------------------------------------------------------------
 	FIntVector2 Image::CalculateMipSize(int32 Mip) const
 	{
diff --git a/Engine/Plugins/Experimental/Mutable/Source/MutableRuntime/Private/MuR/Mesh.cpp b/Engine/Plugins/Experimental/Mutable/Source/MutableRuntime/Private/MuR/Mesh.cpp
--- a/Engine/Plugins/Experimental/Mutable/Source/MutableRuntime/Private/MuR/Mesh.cpp
+++ b/Engine/Plugins/Experimental/Mutable/Source/MutableRuntime/Private/MuR/Mesh.cpp
@@ -40,6 +40,40 @@ MeshPtr Mesh::StaticUnserialise( InputArchive& arch )
 }
 
 
+Ptr<Mesh> Mesh::CreateAsReference(uint32 ID,  bool bForceLoad)
+{
+	Ptr<Mesh> Result = new Mesh;
+	Result->ReferenceID = ID;
+
+	EnumAddFlags(Result->Flags, EMeshFlags::IsResourceReference);
+	if (bForceLoad)
+	{
+		EnumAddFlags(Result->Flags, EMeshFlags::IsResourceForceLoad);
+	}
+
+	return Result;
+}
+
+
+bool Mesh::IsReference() const
+{
+	return EnumHasAnyFlags(Flags, EMeshFlags::IsResourceReference);
+}
+
+
+bool Mesh::IsForceLoad() const
+{
+	return EnumHasAnyFlags(Flags, EMeshFlags::IsResourceForceLoad);
+}
+
+
+uint32 Mesh::GetReferencedMesh() const
+{
+	ensure(IsReference());
+	return ReferenceID;
+}
+
+
 MeshPtr Mesh::Clone() const
 {
     //MUTABLE_CPUPROFILER_SCOPE(MeshClone);
@@ -48,7 +82,7 @@ MeshPtr Mesh::Clone() const
     MeshPtr pResult = new Mesh();
 
     pResult->InternalId = InternalId;
-	pResult->StaticFormatFlags = StaticFormatFlags;
+	pResult->Flags = Flags;
 	pResult->Surfaces = Surfaces;
 	pResult->Skeleton = Skeleton;
 	pResult->PhysicsBody = PhysicsBody;
@@ -83,61 +117,61 @@ MeshPtr Mesh::Clone() const
 }
 
 
-MeshPtr Mesh::Clone(EMeshCopyFlags Flags) const
+Ptr<Mesh> Mesh::Clone(EMeshCopyFlags InFlags) const
 {
     //MUTABLE_CPUPROFILER_SCOPE(MeshClone);
 	LLM_SCOPE_BYNAME(TEXT("MutableRuntime"));
 
-    MeshPtr pResult = new Mesh();
+	Ptr<Mesh> pResult = new Mesh();
 
     pResult->InternalId = InternalId;
-	pResult->StaticFormatFlags = StaticFormatFlags;
 	pResult->MeshIDPrefix = MeshIDPrefix;
+	pResult->Flags = Flags;
 
-	if (EnumHasAnyFlags(Flags, EMeshCopyFlags::WithSurfaces))
+	if (EnumHasAnyFlags(InFlags, EMeshCopyFlags::WithSurfaces))
 	{
 		pResult->Surfaces = Surfaces;
 	}
 
-	if (EnumHasAnyFlags(Flags, EMeshCopyFlags::WithSkeleton))
+	if (EnumHasAnyFlags(InFlags, EMeshCopyFlags::WithSkeleton))
 	{
 		pResult->Skeleton = Skeleton;
 	}
 
-	if (EnumHasAnyFlags(Flags, EMeshCopyFlags::WithPhysicsBody))
+	if (EnumHasAnyFlags(InFlags, EMeshCopyFlags::WithPhysicsBody))
 	{
 		pResult->PhysicsBody = PhysicsBody;
 	}
 
-	if (EnumHasAnyFlags(Flags, EMeshCopyFlags::WithTags))
+	if (EnumHasAnyFlags(InFlags, EMeshCopyFlags::WithTags))
 	{
 		pResult->Tags = Tags;
 	}
 
-	if (EnumHasAnyFlags(Flags, EMeshCopyFlags::WithStreamedResources))
+	if (EnumHasAnyFlags(InFlags, EMeshCopyFlags::WithStreamedResources))
 	{
 		pResult->StreamedResources = StreamedResources;
 	}
 
     // Clone the main buffers
-	if (EnumHasAnyFlags(Flags, EMeshCopyFlags::WithVertexBuffers))
+	if (EnumHasAnyFlags(InFlags, EMeshCopyFlags::WithVertexBuffers))
     {
 		pResult->VertexBuffers = VertexBuffers;
 	}
 
-	if (EnumHasAnyFlags(Flags, EMeshCopyFlags::WithIndexBuffers))
+	if (EnumHasAnyFlags(InFlags, EMeshCopyFlags::WithIndexBuffers))
     {
 		pResult->IndexBuffers = IndexBuffers;
 	}
 
 	// Clone additional buffers
-	if (EnumHasAnyFlags(Flags, EMeshCopyFlags::WithAdditionalBuffers))
+	if (EnumHasAnyFlags(InFlags, EMeshCopyFlags::WithAdditionalBuffers))
 	{
 		pResult->AdditionalBuffers = AdditionalBuffers;
 	}
 
     // Clone the layout	
-	if (EnumHasAnyFlags(Flags, EMeshCopyFlags::WithLayouts))
+	if (EnumHasAnyFlags(InFlags, EMeshCopyFlags::WithLayouts))
 	{
 		pResult->Layouts = Layouts;
 	}
@@ -147,24 +181,24 @@ MeshPtr Mesh::Clone(EMeshCopyFlags Flags) const
 	// physics body doen't need to be deep cloned either as they are also assumed to be shared.
 	
 	// Clone bone poses
-	if (EnumHasAnyFlags(Flags, EMeshCopyFlags::WithPoses))
+	if (EnumHasAnyFlags(InFlags, EMeshCopyFlags::WithPoses))
 	{
 		pResult->BonePoses = BonePoses;
 	}
 
 	// Clone BoneMap
-	if (EnumHasAnyFlags(Flags, EMeshCopyFlags::WithBoneMap))
+	if (EnumHasAnyFlags(InFlags, EMeshCopyFlags::WithBoneMap))
 	{
 		pResult->BoneMap = BoneMap;
 	}
 
 	// Clone SkeletonIDs
-	if (EnumHasAnyFlags(Flags, EMeshCopyFlags::WithSkeletonIDs))
+	if (EnumHasAnyFlags(InFlags, EMeshCopyFlags::WithSkeletonIDs))
 	{
 		pResult->SkeletonIDs = SkeletonIDs;
 	}
 
-	if (EnumHasAnyFlags(Flags, EMeshCopyFlags::WithAdditionalPhysics))
+	if (EnumHasAnyFlags(InFlags, EMeshCopyFlags::WithAdditionalPhysics))
 	{
 		pResult->AdditionalPhysicsBodies = AdditionalPhysicsBodies;
 	}
@@ -173,58 +207,58 @@ MeshPtr Mesh::Clone(EMeshCopyFlags Flags) const
 }
 
 
-void Mesh::CopyFrom(const Mesh& From, EMeshCopyFlags Flags)
+void Mesh::CopyFrom(const Mesh& From, EMeshCopyFlags InFlags)
 {
     //MUTABLE_CPUPROFILER_SCOPE(CopyFrom);
 
     InternalId = From.InternalId;
-	StaticFormatFlags = From.StaticFormatFlags;
+	Flags = From.Flags;
 	MeshIDPrefix = From.MeshIDPrefix;
 
-	if (EnumHasAnyFlags(Flags, EMeshCopyFlags::WithSurfaces))
+	if (EnumHasAnyFlags(InFlags, EMeshCopyFlags::WithSurfaces))
 	{
 		Surfaces = From.Surfaces;
 	}
 
-	if (EnumHasAnyFlags(Flags, EMeshCopyFlags::WithSkeleton))
+	if (EnumHasAnyFlags(InFlags, EMeshCopyFlags::WithSkeleton))
 	{
 		Skeleton = From.Skeleton;
 	}
 
-	if (EnumHasAnyFlags(Flags, EMeshCopyFlags::WithPhysicsBody))
+	if (EnumHasAnyFlags(InFlags, EMeshCopyFlags::WithPhysicsBody))
 	{
 		PhysicsBody = From.PhysicsBody;
 	}
 
-	if (EnumHasAnyFlags(Flags, EMeshCopyFlags::WithTags))
+	if (EnumHasAnyFlags(InFlags, EMeshCopyFlags::WithTags))
 	{
 		Tags = From.Tags;
 	}
 
-	if (EnumHasAnyFlags(Flags, EMeshCopyFlags::WithStreamedResources))
+	if (EnumHasAnyFlags(InFlags, EMeshCopyFlags::WithStreamedResources))
 	{
 		StreamedResources = From.StreamedResources;
 	}
 
     // Copy the main buffers
-	if (EnumHasAnyFlags(Flags, EMeshCopyFlags::WithVertexBuffers))
+	if (EnumHasAnyFlags(InFlags, EMeshCopyFlags::WithVertexBuffers))
     {
 		VertexBuffers = From.VertexBuffers;
 	}
 
-	if (EnumHasAnyFlags(Flags, EMeshCopyFlags::WithIndexBuffers))
+	if (EnumHasAnyFlags(InFlags, EMeshCopyFlags::WithIndexBuffers))
     {
 		IndexBuffers = From.IndexBuffers;
 	}
 
 	// Copy additional buffers
-	if (EnumHasAnyFlags(Flags, EMeshCopyFlags::WithAdditionalBuffers))
+	if (EnumHasAnyFlags(InFlags, EMeshCopyFlags::WithAdditionalBuffers))
 	{
 		AdditionalBuffers = From.AdditionalBuffers;
 	}
 
     // Copy the layout	
-	if (EnumHasAnyFlags(Flags, EMeshCopyFlags::WithLayouts))
+	if (EnumHasAnyFlags(InFlags, EMeshCopyFlags::WithLayouts))
 	{
 		Layouts = From.Layouts;
 	}
@@ -234,24 +268,24 @@ void Mesh::CopyFrom(const Mesh& From, EMeshCopyFlags Flags)
 	// physics body doen't need to be deep copied either as they are also assumed to be shared.
 	
 	// Copy bone poses
-	if (EnumHasAnyFlags(Flags, EMeshCopyFlags::WithPoses))
+	if (EnumHasAnyFlags(InFlags, EMeshCopyFlags::WithPoses))
 	{
 		BonePoses = From.BonePoses;
 	}
 
 	// Copy BoneMap
-	if (EnumHasAnyFlags(Flags, EMeshCopyFlags::WithBoneMap))
+	if (EnumHasAnyFlags(InFlags, EMeshCopyFlags::WithBoneMap))
 	{
 		BoneMap = From.BoneMap;
 	}
 
 	// Copy SkeletonIDs
-	if (EnumHasAnyFlags(Flags, EMeshCopyFlags::WithSkeletonIDs))
+	if (EnumHasAnyFlags(InFlags, EMeshCopyFlags::WithSkeletonIDs))
 	{
 		SkeletonIDs = From.SkeletonIDs;
 	}
 
-	if (EnumHasAnyFlags(Flags, EMeshCopyFlags::WithAdditionalPhysics))
+	if (EnumHasAnyFlags(InFlags, EMeshCopyFlags::WithAdditionalPhysics))
 	{
 		AdditionalPhysicsBodies = From.AdditionalPhysicsBodies;
 	}
@@ -1106,7 +1140,7 @@ void Mesh::FBonePose::Unserialise(InputArchive& arch)
 
 void Mesh::Serialise(OutputArchive& arch) const
 {
-	uint32 ver = 19;
+	uint32 ver = 20;
 	arch << ver;
 
 	arch << IndexBuffers;
@@ -1119,7 +1153,7 @@ void Mesh::Serialise(OutputArchive& arch) const
 	arch << Skeleton;
 	arch << PhysicsBody;
 
-	arch << StaticFormatFlags;
+	arch << uint32(Flags);
 	arch << Surfaces;
 
 	arch << Tags;
@@ -1138,132 +1172,33 @@ void Mesh::Unserialise(InputArchive& arch)
 {
 	uint32 ver;
 	arch >> ver;
-	check(ver <= 19);
+	check(ver == 20);
 
 	arch >> IndexBuffers;
 	arch >> VertexBuffers;
-	
-	if (ver < 19)
-	{
-		FMeshBufferSet Dummy;
-		arch >> Dummy;
-	}
-
 	arch >> AdditionalBuffers;
 	arch >> Layouts;
 
-	if (ver >= 14)
-	{
-		arch >> SkeletonIDs;
-	}
+	arch >> SkeletonIDs;
 
 	arch >> Skeleton;
-	if (ver >= 12)
-	{ 
-		arch >> PhysicsBody;
-	}
-	else
-	{
-		PhysicsBody = nullptr;
-	}
+	arch >> PhysicsBody;
 
-	arch >> StaticFormatFlags;
+	uint32 Temp;
+	arch >> Temp;
+	Flags = reinterpret_cast<EMeshFlags>(Flags);
 
-	if (ver >= 16)
-	{
-		arch >> Surfaces;
-	}
-	else
-	{
-		// Deserialize LegacySurfaces
-		UnserialiseLegacySurfaces(arch, Surfaces);
-	}
+	arch >> Surfaces;
 
-	if (ver <= 16)
-	{
-		struct FACE_GROUP_DEPRECATED
-		{
-			std::string m_name;
-			TArray<int32> m_faces;
-			inline void Unserialise(InputArchive& arch) 
-			{
-				int32 ver = 0;
-				arch >> ver;
-				arch >> m_name;
-				arch >> m_faces;
-			}
+	arch >> Tags;
+	arch >> StreamedResources;
 
-		};
-		TArray<FACE_GROUP_DEPRECATED> FaceGroups;
-		arch >> FaceGroups;
-	}
+	arch >> BonePoses;
+	arch >> BoneMap;
 
-	if (ver <= 16)
-	{
-		TArray < std::string > Temp;
-		arch >> Temp;
-		Tags.SetNum(Temp.Num());
-		for (int32 c = 0; c < Temp.Num(); ++c)
-		{
-			Tags[c] = Temp[c].c_str();
-		}
-	}
-	else
-	{
-		arch >> Tags;
-	}
+	arch >> AdditionalPhysicsBodies;
 
-	if (ver >= 18)
-	{
-		arch >> StreamedResources;
-	}
-
-	if (ver >= 13)
-	{
-		arch >> BonePoses;
-	}
-	else if (Skeleton)
-	{
-		const int32 NumBones = Skeleton->GetBoneCount();
-		BonePoses.SetNum(NumBones);
-		check(Skeleton->m_boneTransforms_DEPRECATED.Num() == NumBones);
-
-		for (int32 BoneIndex = 0; BoneIndex < NumBones; ++BoneIndex)
-		{
-			BonePoses[BoneIndex].BoneId = FBoneName(static_cast<uint32>(BoneIndex));
-			BonePoses[BoneIndex].BoneUsageFlags = EBoneUsageFlags::Skinning;
-			BonePoses[BoneIndex].BoneTransform = Skeleton->m_boneTransforms_DEPRECATED[BoneIndex];
-		}
-	}
-
-	if (ver >= 16)
-	{
-		arch >> BoneMap;
-	}
-	else
-	{
-		const int32 NumBonePoses = BonePoses.Num();
-		BoneMap.SetNum(NumBonePoses);
-		for (int32 BoneIndex = 0; BoneIndex < NumBonePoses; ++BoneIndex)
-		{
-			BoneMap[BoneIndex] = FBoneName(static_cast<uint32>(BoneIndex));
-		}
-
-		for (FMeshSurface& Surface : Surfaces)
-		{
-			Surface.BoneMapCount = NumBonePoses;
-		}
-	}
-
-	if (ver >= 15)
-	{
-		arch >> AdditionalPhysicsBodies;
-	}
-
-	if (ver >= 19)
-	{
-		arch >> MeshIDPrefix;
-	}
+	arch >> MeshIDPrefix;
 }
 
 
@@ -1325,22 +1260,6 @@ bool Mesh::IsSimilar(const Mesh& o, bool bCompareLayouts) const
 }
 
 
-void Mesh::ResetStaticFormatFlags() const
-{
-    StaticFormatFlags = 0;
-
-    for ( int f=0; f<SMF_COUNT; ++f )
-    {
-        if ( s_staticMeshFormatIdentify[f]
-             &&
-             s_staticMeshFormatIdentify[f]( this ) )
-        {
-            StaticFormatFlags |= (1<<f);
-        }
-    }
-}
-
-
 void Mesh::CheckIntegrity() const
 {
 #ifdef MUTABLE_DEBUG
@@ -1505,12 +1424,6 @@ void Mesh::CheckIntegrity() const
 }
 
 
-static bool StaticMeshFormatIdentify_None( const Mesh* )
-{
-    return false;
-}
-
-
 static bool StaticMeshFormatIdentify_Project( const Mesh* pM )
 {
     // This format is used internally for the mesh project
@@ -1663,12 +1576,21 @@ static bool StaticMeshFormatIdentify_ProjectWrapping( const Mesh* pM )
 }
 
 
-STATIC_MESH_FORMAT_ID_FUNC s_staticMeshFormatIdentify[] =
+void Mesh::ResetStaticFormatFlags() const
 {
-    StaticMeshFormatIdentify_None,
-    StaticMeshFormatIdentify_Project,
-    StaticMeshFormatIdentify_ProjectWrapping
-};
+	EnumRemoveFlags(Flags, EMeshFlags::ProjectFormat);
+	EnumRemoveFlags(Flags, EMeshFlags::ProjectWrappingFormat);
+
+	if (StaticMeshFormatIdentify_Project(this))
+	{
+		EnumAddFlags(Flags, EMeshFlags::ProjectFormat);
+	}
+
+	if (StaticMeshFormatIdentify_ProjectWrapping(this))
+	{
+		EnumAddFlags(Flags, EMeshFlags::ProjectWrappingFormat);
+	}
+}
 
 
 namespace
diff --git a/Engine/Plugins/Experimental/Mutable/Source/MutableRuntime/Private/MuR/OpImageProject.cpp b/Engine/Plugins/Experimental/Mutable/Source/MutableRuntime/Private/MuR/OpImageProject.cpp
--- a/Engine/Plugins/Experimental/Mutable/Source/MutableRuntime/Private/MuR/OpImageProject.cpp
+++ b/Engine/Plugins/Experimental/Mutable/Source/MutableRuntime/Private/MuR/OpImageProject.cpp
@@ -966,7 +966,7 @@ void ImageRasterProjectedPlanar(const Mesh* pMesh, Image* pImage,
 		ProcessorFeatures |= EPixelProcessorFeatures::WithMask;
 	}
 
-	if (pMesh->StaticFormatFlags & (1 << SMF_PROJECT))
+	if (EnumHasAnyFlags(pMesh->Flags, EMeshFlags::ProjectFormat))
 	{
 		float UnusedProjectionAngle = 0;
 		FImageRasterInvokeArgs RasterArgs = 
@@ -1180,7 +1180,7 @@ void ImageRasterProjectedWrapping( const Mesh* pMesh, Image* pImage,
 		ProcessorFeatures |= EPixelProcessorFeatures::WithMask;
 	}
 
-    if ( ( pMesh->StaticFormatFlags & (1<<SMF_PROJECTWRAPPING) ) )
+	if (EnumHasAnyFlags(pMesh->Flags, EMeshFlags::ProjectWrappingFormat))
     {
 		float UnusedProjectionAngle = 0;
 
@@ -1394,8 +1394,8 @@ void ImageRasterProjectedCylindrical( const Mesh* pMesh, Image* pImage,
 		ProcessorFeatures |= EPixelProcessorFeatures::WithMask;
 	}
 
-    if ((pMesh->StaticFormatFlags & (1 << SMF_PROJECT)))
-    {
+	if (EnumHasAnyFlags(pMesh->Flags, EMeshFlags::ProjectFormat))
+	{
 		int32 UnusedLayout = 0;
 		int32 UnusedBlock = 0;
 
@@ -3071,12 +3071,12 @@ void MeshProject(Mesh* Result, const Mesh* pMesh, const FProjector& projector, b
 {
 	MUTABLE_CPUPROFILER_SCOPE(MeshProject);
 
-    if (pMesh->StaticFormatFlags & (1<<SMF_PROJECT))
+    if ( EnumHasAnyFlags(pMesh->Flags, EMeshFlags::ProjectFormat) )
     {
         // Mesh-optimised version
         MeshProject_Optimised(Result, pMesh, projector, bOutSuccess);
     }
-    else if (pMesh->StaticFormatFlags & (1<<SMF_PROJECTWRAPPING))
+    else if (EnumHasAnyFlags(pMesh->Flags, EMeshFlags::ProjectWrappingFormat))
     {
         // Mesh-optimised version for wrapping projectors
         // \todo: make sure the projector is a wrapping projector
diff --git a/Engine/Plugins/Experimental/Mutable/Source/MutableRuntime/Private/MuR/Operations.cpp b/Engine/Plugins/Experimental/Mutable/Source/MutableRuntime/Private/MuR/Operations.cpp
--- a/Engine/Plugins/Experimental/Mutable/Source/MutableRuntime/Private/MuR/Operations.cpp
+++ b/Engine/Plugins/Experimental/Mutable/Source/MutableRuntime/Private/MuR/Operations.cpp
@@ -14,10 +14,9 @@ namespace mu
 	//---------------------------------------------------------------------------------------------
 	//---------------------------------------------------------------------------------------------
 	//---------------------------------------------------------------------------------------------
-    #define OP_DESC_COUNT		int32(OP_TYPE::COUNT)
 
     // clang-format off
-    static const OP_DESC s_opDescs[OP_DESC_COUNT] =
+    static const FOpDesc s_opDescs[] =
 	{ 
 		// type				cached	supported base image formats
         { DT_NONE,			false,	{ 0, 0, 0, 0, 0, 0, 0, 0, 0, 0, 0, 0, 0, 0, 0, 0, 0, 0, 0, 0, 0, 0, 0, 0, 0, 0, 0, 0, 0 }		},	// NONE
@@ -42,6 +41,7 @@ namespace mu
 		{ DT_STRING,		false,	{ 0, 0, 0, 0, 0, 0, 0, 0, 0, 0, 0, 0, 0, 0, 0, 0, 0, 0, 0, 0, 0, 0, 0, 0, 0, 0, 0, 0, 0 }		},	// ST_PARAMETER
 
 		{ DT_IMAGE,			true,	{ 0, 0, 0, 0, 0, 0, 0, 0, 0, 0, 0, 0, 0, 0, 0, 0, 0, 0, 0, 0, 0, 0, 0, 0, 0, 0, 0, 0, 0 }		},	// IM_REFERENCE
+		{ DT_MESH,			true,	{ 0, 0, 0, 0, 0, 0, 0, 0, 0, 0, 0, 0, 0, 0, 0, 0, 0, 0, 0, 0, 0, 0, 0, 0, 0, 0, 0, 0, 0 }		},	// ME_REFERENCE
 
 		{ DT_INT,			false,	{ 0, 0, 0, 0, 0, 0, 0, 0, 0, 0, 0, 0, 0, 0, 0, 0, 0, 0, 0, 0, 0, 0, 0, 0, 0, 0, 0, 0, 0 }		},	// NU_CONDITIONAL
 		{ DT_SCALAR,		false,	{ 0, 0, 0, 0, 0, 0, 0, 0, 0, 0, 0, 0, 0, 0, 0, 0, 0, 0, 0, 0, 0, 0, 0, 0, 0, 0, 0, 0, 0 }		},	// SC_CONDITIONAL
@@ -146,10 +146,10 @@ namespace mu
 
     // clang-format on
 
-    static_assert( OP_DESC_COUNT == (int)OP_TYPE::COUNT, "OperationDescMismatch" );
+    static_assert( sizeof(s_opDescs)/sizeof(FOpDesc) == (int32)OP_TYPE::COUNT, "OperationDescMismatch");
 
 	//---------------------------------------------------------------------------------------------
-	const OP_DESC& GetOpDesc( OP_TYPE type )
+	const FOpDesc& GetOpDesc( OP_TYPE type )
 	{
         return s_opDescs[ (int32)type ];
 	}
@@ -383,6 +383,7 @@ namespace mu
         case OP_TYPE::PR_PARAMETER:
         case OP_TYPE::IM_PARAMETER:
 		case OP_TYPE::IM_REFERENCE:
+		case OP_TYPE::ME_REFERENCE:
 			break;
 
         case OP_TYPE::SC_CURVE:
diff --git a/Engine/Plugins/Experimental/Mutable/Source/MutableRuntime/Private/MuR/System.cpp b/Engine/Plugins/Experimental/Mutable/Source/MutableRuntime/Private/MuR/System.cpp
--- a/Engine/Plugins/Experimental/Mutable/Source/MutableRuntime/Private/MuR/System.cpp
+++ b/Engine/Plugins/Experimental/Mutable/Source/MutableRuntime/Private/MuR/System.cpp
@@ -137,7 +137,7 @@ namespace mu
 		// Make it explicit to try to capture metrics
         StreamInterface = nullptr;
 		ExtensionDataStreamer = nullptr;
-        ImageParameterGenerator = nullptr;
+		ExternalResourceProvider = nullptr;
     }
 
 
@@ -206,11 +206,11 @@ namespace mu
 
 	
     //---------------------------------------------------------------------------------------------
-    void System::SetImageParameterGenerator(const TSharedPtr<ImageParameterGenerator>& InInterface )
+    void System::SetExternalResourceProvider(const TSharedPtr<ExternalResourceProvider>& InInterface )
     {
 		LLM_SCOPE_BYNAME(TEXT("MutableRuntime"));
 
-        m_pD->ImageParameterGenerator = InInterface;
+        m_pD->ExternalResourceProvider = InInterface;
     }
 
 
diff --git a/Engine/Plugins/Experimental/Mutable/Source/MutableRuntime/Public/MuR/Image.h b/Engine/Plugins/Experimental/Mutable/Source/MutableRuntime/Public/MuR/Image.h
--- a/Engine/Plugins/Experimental/Mutable/Source/MutableRuntime/Public/MuR/Image.h
+++ b/Engine/Plugins/Experimental/Mutable/Source/MutableRuntime/Public/MuR/Image.h
@@ -46,7 +46,7 @@ namespace mu
         //! \param format Pixel format.
         Image(uint32 SizeX, uint32 SizeY, uint32 lods, EImageFormat Format, EInitializationType InitType);
 
-		/** */
+		/** Create a new empty image that repreents an external resource image. */
 		static Ptr<Image> CreateAsReference(uint32 ID, const FImageDesc& Desc, bool bForceLoad);
 
 		//! Serialisation
@@ -63,6 +63,18 @@ namespace mu
 		/** */
 		void Init(uint32 SizeX, uint32 SizeY, uint32 Lods, EImageFormat Format, EInitializationType InitType);
 
+		/** Clear the image to black colour. */
+		void InitToBlack();
+
+		/** Return true if this is a reference to an engine image. */
+		bool IsReference() const;
+
+		/** If true, this is a reference that must be resolved at compile time. */
+		bool IsForceLoad() const;
+
+		/** Return the id of the engine referenced texture. Only valid if IsReference. */
+		uint32 GetReferencedTexture() const;
+
 		//! Return the width of the image.
         uint16 GetSizeX() const;
 
@@ -77,7 +89,6 @@ namespace mu
 		//! Return the number of levels of detail (mipmaps) in the texture. The base lavel is also
 		//! counted, so the minimum is 1.
 		int32 GetLODCount() const;
-		void SetLODCount(int32 LODCount);
 
 		//! Return a pointer to a instance-owned buffer where the image pixels are.
         const uint8* GetLODData(int32 LODIndex) const;
@@ -86,17 +97,6 @@ namespace mu
         //! Return the size in bytes of a specific LOD of the image.
         int32 GetLODDataSize(int32 LODIndex) const;
 
-		/** Return true if this is a reference to an engine image. */
-		bool IsReference() const;
-
-		/** If true, this is a reference that must be resolved at compile time. */
-		bool IsForceLoad() const;
-
-		/** Return the id of the engine referenced texture. Only valid if IsReference. */
-		uint32 GetReferencedTexture() const;
-
-		/** Clear the image to black colour. */
-		void InitToBlack();
 
 	protected:
 
@@ -140,7 +140,6 @@ namespace mu
 		uint32 ReferenceID = 0;
 
 		/** Pixel data for all lods. */
-
 		FImageDataStorage DataStorage;
 
 		// This used to be the methods in the private implementation of the image interface
@@ -222,25 +221,8 @@ namespace mu
 		//! Calculate the size of the image data in bytes, regardless of what is allocated in
 		//! m_data, only using the image descriptions. For non-block-compressed images, it returns
 		//! 0.
-		int32 CalculateDataSize() const;
 		static int32 CalculateDataSize(int32 SizeX, int32 SizeY, int32 LodCount, EImageFormat Format);
 
-
-		//! Calculate the size of a lod of the image data in bytes, regardless of what is allocated
-		//! in m_data, only using the image descriptions. For non-block-compressed images, it
-		//! returns 0.
-		//int32 CalculateDataSize(int lod) const;
-
-		////! Calculate the number of pixels of the image, regardless of what is allocated in
-		////! m_data, only using the image descriptions.
-		////! For block-compressed images, it includes the wasted pixels in the blocks, in case
-		////! the size is not a multiple of it in all lods.
-		////! It includes the pixels in all lods
-		//int32 CalculatePixelCount() const;
-
-		////! Same as above, but only calculates the pixels of an lod
-		//int32 CalculatePixelCount(int32 LOD) const;
-
 		//! Calculate the size in pixels of a particular mipmap of this image. The size doesn't
 		//! include pixels necessary for completing blocks in block-compressed formats.
 		FIntVector2 CalculateMipSize(int32 LOD) const;
@@ -266,7 +248,7 @@ namespace mu
 		void ReduceLODs(int32 LODsToSkip);
 
 		//! Calculate the number of mipmaps for a particular image size.
-		static int GetMipmapCount(int SizeX, int SizeY);
+		static int32 GetMipmapCount(int32 SizeX, int32 SizeY);
 
 		//! Get the rect inside the image bounding the non-black content of the image.
 		void GetNonBlackRect(FImageRect& rect) const;
diff --git a/Engine/Plugins/Experimental/Mutable/Source/MutableRuntime/Public/MuR/Mesh.h b/Engine/Plugins/Experimental/Mutable/Source/MutableRuntime/Public/MuR/Mesh.h
--- a/Engine/Plugins/Experimental/Mutable/Source/MutableRuntime/Public/MuR/Mesh.h
+++ b/Engine/Plugins/Experimental/Mutable/Source/MutableRuntime/Public/MuR/Mesh.h
@@ -132,6 +132,27 @@ namespace mu
 	
 	ENUM_CLASS_FLAGS(EMeshCopyFlags);
 
+
+	/** Optimised mesh formats that are identified in some operations to chose a faster version. */
+	enum class EMeshFlags : uint32
+	{
+		None = 0,
+
+		/** The mesh is formatted to be used for planar and cilyndrical projection */
+		ProjectFormat = 1 << 0,
+
+		/** The mesh is formatted to be used for wrapping projection */
+		ProjectWrappingFormat = 1 << 1,
+
+		/** The mesh is a reference to an external resource mesh. */
+		IsResourceReference = 1 << 2,
+
+		/** The mesh is a reference to an external resource mesh and must be loaded when first referenced. */
+		IsResourceForceLoad = 1 << 3,
+
+	};
+
+
     //! \brief Mesh object containing any number of buffers with any number of channels.
     //! The buffers can be per-index or per-vertex.
     //! The mesh also includes layout information for every texture channel for internal usage, and
@@ -149,6 +170,9 @@ namespace mu
         //-----------------------------------------------------------------------------------------
         // Life cycle
         //-----------------------------------------------------------------------------------------
+				
+		/** Create a new empty mesh that repreents an external resource mesh. */
+		static Ptr<Mesh> CreateAsReference(uint32 ID, bool bForceLoad);
 
         //! Deep clone this mesh.
         Ptr<Mesh> Clone() const;
@@ -170,6 +194,14 @@ namespace mu
         // Own interface
         //-----------------------------------------------------------------------------------------
 
+		/** Return true if this is a reference to an engine image. */
+		bool IsReference() const;
+
+		/** If true, this is a reference that must be resolved at compile time. */
+		bool IsForceLoad() const;
+
+		/** Return the id of the engine referenced mesh. Only valid if IsReference. */
+		uint32 GetReferencedMesh() const;
 
         //! \name Buffers
         //! \{
@@ -331,10 +363,13 @@ namespace mu
 		/** Non-persistent internal id unique for a mesh generated for a specific state and parameter values. */
 		mutable uint32 InternalId = 0;
 
-		//! This is bit-mask on the STATIC_MESH_FORMATS enumeration, marking what static formats
-		//! are compatible with this one. Usually precalculated at model compilation time.
-		//! It should be reset after any operation that modifies the format.
-		mutable uint32 StaticFormatFlags = 0;
+		/** This is bit - mask on the EMeshFlags enumeration, marking what static formats are compatible with this one and other properties. 
+		 * It should be reset after any operation that modifies the format.
+		 */
+		mutable EMeshFlags Flags = EMeshFlags::None;
+
+		/** Only valid if the right flags are set, this identifies a referenced mesh. */
+		uint32 ReferenceID = 0;
 
 		/** Prefix for the unique IDs related to this mesh (vertices and layout blocks). Useful if the mesh stores them in an implicit, or relative way. 
 		* See MeshVertexIdIterator for details.
diff --git a/Engine/Plugins/Experimental/Mutable/Source/MutableRuntime/Public/MuR/System.h b/Engine/Plugins/Experimental/Mutable/Source/MutableRuntime/Public/MuR/System.h
--- a/Engine/Plugins/Experimental/Mutable/Source/MutableRuntime/Public/MuR/System.h
+++ b/Engine/Plugins/Experimental/Mutable/Source/MutableRuntime/Public/MuR/System.h
@@ -78,18 +78,22 @@ namespace mu
 
 
     /** Interface to request external images used as parameters. */
-    class MUTABLERUNTIME_API ImageParameterGenerator
+    class MUTABLERUNTIME_API ExternalResourceProvider
     {
     public:
 
         //! Ensure virtual destruction
-        virtual ~ImageParameterGenerator() = default;
+        virtual ~ExternalResourceProvider() = default;
 
-        //! Returns the completion event and a cleanup function that must be called once event is completed.
+        /** Returns the completion event and a cleanup function that must be called once event is completed. */
 		virtual TTuple<UE::Tasks::FTask, TFunction<void()>> GetImageAsync(FName Id, uint8 MipmapsToSkip, TFunction<void(Ptr<Image>)>& ResultCallback) = 0;
 		virtual TTuple<UE::Tasks::FTask, TFunction<void()>> GetReferencedImageAsync(const void* ModelPtr, int32 Id, uint8 MipmapsToSkip, TFunction<void(Ptr<Image>)>& ResultCallback) { check(false); return {}; }
 
-        virtual mu::FImageDesc GetImageDesc(FName Id, uint8 MipmapsToSkip) = 0;
+		virtual mu::FImageDesc GetImageDesc(FName Id, uint8 MipmapsToSkip) = 0;
+
+		/** Returns the completion event and a cleanup function that must be called once event is completed. */
+		virtual TTuple<UE::Tasks::FTask, TFunction<void()>> GetMeshAsync(FName Id, TFunction<void(Ptr<Mesh>)>& ResultCallback) = 0;
+		virtual TTuple<UE::Tasks::FTask, TFunction<void()>> GetReferencedMeshAsync(const void* ModelPtr, int32 Id, TFunction<void(Ptr<Mesh>)>& ResultCallback) { check(false); return {}; }
     };
 
 
@@ -127,7 +131,7 @@ namespace mu
 		void SetGeneratedCacheSize(uint32 InCount);
 
         /** Set a new provider for external image data. This is only necessary if image parameters are used in the models. */
-        void SetImageParameterGenerator(const TSharedPtr<ImageParameterGenerator>&);
+        void SetExternalResourceProvider(const TSharedPtr<ExternalResourceProvider>&);
 
 		/** Set a function that will be used to convert image pixel formats instead of the internal conversion. 
 		* \warning The provided function can be called from any thread, and also concurrently.
diff --git a/Engine/Plugins/Experimental/Mutable/Source/MutableTools/Private/MuT/ASTOpReferenceResource.cpp b/Engine/Plugins/Experimental/Mutable/Source/MutableTools/Private/MuT/ASTOpReferenceResource.cpp
--- a/Engine/Plugins/Experimental/Mutable/Source/MutableTools/Private/MuT/ASTOpReferenceResource.cpp
+++ b/Engine/Plugins/Experimental/Mutable/Source/MutableTools/Private/MuT/ASTOpReferenceResource.cpp
@@ -132,38 +132,4 @@ namespace mu
 	}
 
 
-	Ptr<ASTOp> ASTOpReferenceResource::OptimiseSemantic(const FModelOptimizationOptions& options, int32 Pass) const
-	{
-		mu::Ptr<ASTOp> NewOp;
-
-		//switch (type)
-		//{
-
-		//case OP_TYPE::IM_REFERENCE:
-		//{
-		//	// If we are in reference resolution stage
-		//	if (Pass>=2 && bForceLoad)
-		//	{
-		//		MUTABLE_CPUPROFILER_SCOPE(ResolveReference);
-
-		//		check(options.ReferencedResourceProvider);
-
-		//		TFuture<Ptr<Image>> ConstantImageFuture = options.ReferencedResourceProvider(ID);
-		//		Ptr<Image> ConstantImage = ConstantImageFuture.Get();
-
-		//		Ptr<ASTOpConstantResource> ConstantOp = new ASTOpConstantResource;
-		//		ConstantOp->type = OP_TYPE::IM_CONSTANT;
-		//		ConstantOp->SetValue( ConstantImage.get(), options.bUseDiskCache );
-		//		NewOp = ConstantOp;
-		//	}
-		//	break;
-		//}
-
-		//default:
-		//	checkf(false, TEXT("Instruction not supported"));
-		//}
-
-		return NewOp;
-	}
-
 }
\ No newline at end of file
diff --git a/Engine/Plugins/Experimental/Mutable/Source/MutableTools/Private/MuT/ASTOpReferenceResource.h b/Engine/Plugins/Experimental/Mutable/Source/MutableTools/Private/MuT/ASTOpReferenceResource.h
--- a/Engine/Plugins/Experimental/Mutable/Source/MutableTools/Private/MuT/ASTOpReferenceResource.h
+++ b/Engine/Plugins/Experimental/Mutable/Source/MutableTools/Private/MuT/ASTOpReferenceResource.h
@@ -50,7 +50,6 @@ namespace mu
 		void GetLayoutBlockSize(int32* pBlockX, int32* pBlockY) override;
 		bool GetNonBlackRect(FImageRect& maskUsage) const override;
 		Ptr<ImageSizeExpression> GetImageSizeExpression() const override;
-		Ptr<ASTOp> OptimiseSemantic(const FModelOptimizationOptions& options, int32 Pass) const override;
 	};
 
 
diff --git a/Engine/Plugins/Experimental/Mutable/Source/MutableTools/Private/MuT/CodeGenerator_Mesh.cpp b/Engine/Plugins/Experimental/Mutable/Source/MutableTools/Private/MuT/CodeGenerator_Mesh.cpp
--- a/Engine/Plugins/Experimental/Mutable/Source/MutableTools/Private/MuT/CodeGenerator_Mesh.cpp
+++ b/Engine/Plugins/Experimental/Mutable/Source/MutableTools/Private/MuT/CodeGenerator_Mesh.cpp
@@ -32,6 +32,7 @@
 #include "MuT/ASTOpMeshMorph.h"
 #include "MuT/ASTOpMeshAddTags.h"
 #include "MuT/ASTOpSwitch.h"
+#include "MuT/ASTOpReferenceResource.h"
 #include "MuT/CodeGenerator.h"
 #include "MuT/CodeGenerator_FirstPass.h"
 #include "MuT/CompilerPrivate.h"
@@ -1299,6 +1300,19 @@ namespace mu
 			return;
 		}
 
+		if (pMesh->IsReference())
+		{
+			Ptr<ASTOpReferenceResource> ReferenceOp = new ASTOpReferenceResource();
+			ReferenceOp->type = OP_TYPE::ME_REFERENCE;
+			ReferenceOp->ID = pMesh->GetReferencedMesh();
+			ReferenceOp->bForceLoad = pMesh->IsForceLoad();
+
+			OutResult.baseMeshOp = ReferenceOp;
+			OutResult.meshOp = ReferenceOp;
+
+			return;
+		}
+
 		// Separate the tags from the mesh
 		TArray<FString> Tags = pMesh->Tags;
 		if (Tags.Num())
diff --git a/Engine/Plugins/Experimental/Mutable/Source/MutableTools/Private/MuT/ErrorLog.cpp b/Engine/Plugins/Experimental/Mutable/Source/MutableTools/Private/MuT/ErrorLog.cpp
--- a/Engine/Plugins/Experimental/Mutable/Source/MutableTools/Private/MuT/ErrorLog.cpp
+++ b/Engine/Plugins/Experimental/Mutable/Source/MutableTools/Private/MuT/ErrorLog.cpp
@@ -200,7 +200,7 @@ namespace mu
     
     // clang-format off
 
-    const TCHAR* s_opNames[ int32(OP_TYPE::COUNT) ] =
+    const TCHAR* s_opNames[] =
 	{
 		TEXT("NONE             "),
 
@@ -224,6 +224,7 @@ namespace mu
 		TEXT("ST_PARAMETER     "),
 		
 		TEXT("IM_REFERENCE     "),
+		TEXT("ME_REFERENCE     "),
 
 		TEXT("NU_CONDITIONAL   "),
 		TEXT("SC_CONDITIONAL   "),
