[mutable] Fix mesh flags serialization bug preventing some optimizations.
#rnx

diff --git a/Engine/Plugins/Experimental/Mutable/Source/MutableRuntime/Private/MuR/Mesh.cpp b/Engine/Plugins/Experimental/Mutable/Source/MutableRuntime/Private/MuR/Mesh.cpp
--- a/Engine/Plugins/Experimental/Mutable/Source/MutableRuntime/Private/MuR/Mesh.cpp
+++ b/Engine/Plugins/Experimental/Mutable/Source/MutableRuntime/Private/MuR/Mesh.cpp
@@ -1186,7 +1186,7 @@ void Mesh::Unserialise(InputArchive& arch)
 
 	uint32 Temp;
 	arch >> Temp;
-	Flags = reinterpret_cast<EMeshFlags>(Flags);
+	Flags = static_cast<EMeshFlags>(Temp);
 
 	arch >> Surfaces;
 
