[Mutable] Allow FUpdateContextPrivate to be destroyed in the GarbageCollectorThread.

#rnx

diff --git a/Engine/Plugins/Experimental/Mutable/Source/CustomizableObject/Private/MuCO/CustomizableObjectSystem.cpp b/Engine/Plugins/Experimental/Mutable/Source/CustomizableObject/Private/MuCO/CustomizableObjectSystem.cpp
--- a/Engine/Plugins/Experimental/Mutable/Source/CustomizableObject/Private/MuCO/CustomizableObjectSystem.cpp
+++ b/Engine/Plugins/Experimental/Mutable/Source/CustomizableObject/Private/MuCO/CustomizableObjectSystem.cpp
@@ -179,7 +179,7 @@ FUpdateContextPrivate::FUpdateContextPrivate(UCustomizableObjectInstance& InInst
 
 FUpdateContextPrivate::~FUpdateContextPrivate()
 {
-	check(IsInGameThread());
+	check(IsInGameThread() || IsInGarbageCollectorThread());
 
 	if (UCustomizableObjectSystem::IsCreated())
 	{
