[Mutable] Fix crash when converting skeletal meshes with clothing enabled if the physics asset is missing.

#rnx

diff --git a/Engine/Plugins/Experimental/Mutable/Source/CustomizableObjectEditor/Private/MuCOE/GenerateMutableSource/GenerateMutableSourceMesh.cpp b/Engine/Plugins/Experimental/Mutable/Source/CustomizableObjectEditor/Private/MuCOE/GenerateMutableSource/GenerateMutableSourceMesh.cpp
--- a/Engine/Plugins/Experimental/Mutable/Source/CustomizableObjectEditor/Private/MuCOE/GenerateMutableSource/GenerateMutableSourceMesh.cpp
+++ b/Engine/Plugins/Experimental/Mutable/Source/CustomizableObjectEditor/Private/MuCOE/GenerateMutableSource/GenerateMutableSourceMesh.cpp
@@ -1502,9 +1502,12 @@ mu::MeshPtr ConvertSkeletalMeshToMutable(const USkeletalMesh* InSkeletalMesh, co
 			AssetData.OriginalAssetGuid = Asset->GetAssetGuid();
 			AssetData.Name = Asset->GetFName();
 
-			OutPhysicsAssetIndex = GenerationContext.PhysicsAssets.AddUnique(Asset->PhysicsAsset);
-			GenerationContext.AddParticipatingObject(*Asset->PhysicsAsset);
-			
+			if (Asset->PhysicsAsset)
+			{
+				OutPhysicsAssetIndex = GenerationContext.PhysicsAssets.AddUnique(Asset->PhysicsAsset);
+				GenerationContext.AddParticipatingObject(*Asset->PhysicsAsset);
+			}
+
 			// Store raw clothing config serialized raw data, and info to recreate it afterwards.
 			for (const TPair<FName, TObjectPtr<UClothConfigBase>>& ClothConfig : Asset->ClothConfigs)
 			{
