[mutable] Removed dangerous node reconstruction from the Postload stage.

#rnx

diff --git a/Engine/Plugins/Experimental/Mutable/Source/CustomizableObjectEditor/Private/MuCOE/Nodes/CustomizableObjectNodeColorParameter.cpp b/Engine/Plugins/Experimental/Mutable/Source/CustomizableObjectEditor/Private/MuCOE/Nodes/CustomizableObjectNodeColorParameter.cpp
--- a/Engine/Plugins/Experimental/Mutable/Source/CustomizableObjectEditor/Private/MuCOE/Nodes/CustomizableObjectNodeColorParameter.cpp
+++ b/Engine/Plugins/Experimental/Mutable/Source/CustomizableObjectEditor/Private/MuCOE/Nodes/CustomizableObjectNodeColorParameter.cpp
@@ -57,18 +57,5 @@ void UCustomizableObjectNodeColorParameter::OnRenameNode(const FString& NewName)
 	}
 }
 
-
-void UCustomizableObjectNodeColorParameter::BackwardsCompatibleFixup()
-{
-	Super::BackwardsCompatibleFixup();
-
-	const int32 CustomizableObjectCustomVersion = GetLinkerCustomVersion(FCustomizableObjectCustomVersion::GUID);
-
-	if (CustomizableObjectCustomVersion < FCustomizableObjectCustomVersion::AddedRenameOptionToParameterNodes)
-	{
-		ReconstructNode();
-	}
-}
-
 #undef LOCTEXT_NAMESPACE
 
diff --git a/Engine/Plugins/Experimental/Mutable/Source/CustomizableObjectEditor/Private/MuCOE/Nodes/CustomizableObjectNodeEnumParameter.cpp b/Engine/Plugins/Experimental/Mutable/Source/CustomizableObjectEditor/Private/MuCOE/Nodes/CustomizableObjectNodeEnumParameter.cpp
--- a/Engine/Plugins/Experimental/Mutable/Source/CustomizableObjectEditor/Private/MuCOE/Nodes/CustomizableObjectNodeEnumParameter.cpp
+++ b/Engine/Plugins/Experimental/Mutable/Source/CustomizableObjectEditor/Private/MuCOE/Nodes/CustomizableObjectNodeEnumParameter.cpp
@@ -70,20 +70,5 @@ void UCustomizableObjectNodeEnumParameter::OnRenameNode(const FString& NewName)
 	}
 }
 
-
-void UCustomizableObjectNodeEnumParameter::BackwardsCompatibleFixup()
-{
-	Super::BackwardsCompatibleFixup();
-
-	const int32 CustomizableObjectCustomVersion = GetLinkerCustomVersion(FCustomizableObjectCustomVersion::GUID);
-
-	if (CustomizableObjectCustomVersion < FCustomizableObjectCustomVersion::AddedRenameOptionToParameterNodes)
-	{
-		ReconstructNode();
-	}
-}
-
-
-
 #undef LOCTEXT_NAMESPACE
 
diff --git a/Engine/Plugins/Experimental/Mutable/Source/CustomizableObjectEditor/Private/MuCOE/Nodes/CustomizableObjectNodeFloatParameter.cpp b/Engine/Plugins/Experimental/Mutable/Source/CustomizableObjectEditor/Private/MuCOE/Nodes/CustomizableObjectNodeFloatParameter.cpp
--- a/Engine/Plugins/Experimental/Mutable/Source/CustomizableObjectEditor/Private/MuCOE/Nodes/CustomizableObjectNodeFloatParameter.cpp
+++ b/Engine/Plugins/Experimental/Mutable/Source/CustomizableObjectEditor/Private/MuCOE/Nodes/CustomizableObjectNodeFloatParameter.cpp
@@ -87,11 +87,6 @@ void UCustomizableObjectNodeFloatParameter::BackwardsCompatibleFixup()
 	{
 		ReconstructNode();
 	}
-
-	if (CustomizableObjectCustomVersion < FCustomizableObjectCustomVersion::AddedRenameOptionToParameterNodes)
-	{
-		ReconstructNode();
-	}
 }
 
 
diff --git a/Engine/Plugins/Experimental/Mutable/Source/CustomizableObjectEditor/Private/MuCOE/Nodes/CustomizableObjectNodeGroupProjectorParameter.cpp b/Engine/Plugins/Experimental/Mutable/Source/CustomizableObjectEditor/Private/MuCOE/Nodes/CustomizableObjectNodeGroupProjectorParameter.cpp
--- a/Engine/Plugins/Experimental/Mutable/Source/CustomizableObjectEditor/Private/MuCOE/Nodes/CustomizableObjectNodeGroupProjectorParameter.cpp
+++ b/Engine/Plugins/Experimental/Mutable/Source/CustomizableObjectEditor/Private/MuCOE/Nodes/CustomizableObjectNodeGroupProjectorParameter.cpp
@@ -129,11 +129,6 @@ void UCustomizableObjectNodeGroupProjectorParameter::BackwardsCompatibleFixup()
 	}
 
 	if (CustomizableObjectCustomVersion < FCustomizableObjectCustomVersion::GroupProjectorImagePinRemoved)
-	{
-		ReconstructNode();		
-	}
-
-	if (CustomizableObjectCustomVersion < FCustomizableObjectCustomVersion::AddedRenameOptionToParameterNodes)
 	{
 		ReconstructNode();
 	}
diff --git a/Engine/Plugins/Experimental/Mutable/Source/CustomizableObjectEditor/Private/MuCOE/Nodes/CustomizableObjectNodeProjectorParameter.cpp b/Engine/Plugins/Experimental/Mutable/Source/CustomizableObjectEditor/Private/MuCOE/Nodes/CustomizableObjectNodeProjectorParameter.cpp
--- a/Engine/Plugins/Experimental/Mutable/Source/CustomizableObjectEditor/Private/MuCOE/Nodes/CustomizableObjectNodeProjectorParameter.cpp
+++ b/Engine/Plugins/Experimental/Mutable/Source/CustomizableObjectEditor/Private/MuCOE/Nodes/CustomizableObjectNodeProjectorParameter.cpp
@@ -50,11 +50,6 @@ void UCustomizableObjectNodeProjectorParameter::BackwardsCompatibleFixup()
 	{
 		DefaultValue.ProjectionType = ProjectionType_DEPRECATED;
 	}
-
-	if (CustomizableObjectCustomVersion < FCustomizableObjectCustomVersion::AddedRenameOptionToParameterNodes)
-	{
-		ReconstructNode();
-	}
 }
 
 
diff --git a/Engine/Plugins/Experimental/Mutable/Source/CustomizableObjectEditor/Private/MuCOE/Nodes/CustomizableObjectNodeTable.cpp b/Engine/Plugins/Experimental/Mutable/Source/CustomizableObjectEditor/Private/MuCOE/Nodes/CustomizableObjectNodeTable.cpp
--- a/Engine/Plugins/Experimental/Mutable/Source/CustomizableObjectEditor/Private/MuCOE/Nodes/CustomizableObjectNodeTable.cpp
+++ b/Engine/Plugins/Experimental/Mutable/Source/CustomizableObjectEditor/Private/MuCOE/Nodes/CustomizableObjectNodeTable.cpp
@@ -183,11 +183,6 @@ void UCustomizableObjectNodeTable::BackwardsCompatibleFixup()
 			}
 		}
 	}
-
-	if (CustomizableObjectCustomVersion < FCustomizableObjectCustomVersion::AddedRenameOptionToParameterNodes)
-	{
-		ReconstructNode();
-	}
 }
 
 
@@ -410,11 +405,11 @@ void UCustomizableObjectNodeTable::AllocateDefaultPins(UCustomizableObjectNodeRe
 		if (const FSoftObjectProperty* SoftObjectProperty = CastField<FSoftObjectProperty>(ColumnProperty))
 		{
 			// Only process object properties that might have pointers to objects of any of the "SupportedSoftObjectTypes"
-			// Object properties for unrelated types do not reference anything we would make pins from, so there is no need to load them or make sure they are non-null
 			const bool bPotentiallySupportedObject = SoftObjectProperty->PropertyClass && SupportedSoftObjectTypes.ContainsByPredicate([SoftObjectProperty](const TObjectPtr<UClass>& Type)
-				{
-					return SoftObjectProperty->PropertyClass->IsChildOf(Type) || Type->IsChildOf(SoftObjectProperty->PropertyClass);
-				});
+			{
+				return SoftObjectProperty->PropertyClass->IsChildOf(Type) || Type->IsChildOf(SoftObjectProperty->PropertyClass);
+			});
+
 			if (bPotentiallySupportedObject)
 			{
 				UObject* Object = nullptr;
diff --git a/Engine/Plugins/Experimental/Mutable/Source/CustomizableObjectEditor/Private/MuCOE/Nodes/CustomizableObjectNodeTextureParameter.cpp b/Engine/Plugins/Experimental/Mutable/Source/CustomizableObjectEditor/Private/MuCOE/Nodes/CustomizableObjectNodeTextureParameter.cpp
--- a/Engine/Plugins/Experimental/Mutable/Source/CustomizableObjectEditor/Private/MuCOE/Nodes/CustomizableObjectNodeTextureParameter.cpp
+++ b/Engine/Plugins/Experimental/Mutable/Source/CustomizableObjectEditor/Private/MuCOE/Nodes/CustomizableObjectNodeTextureParameter.cpp
@@ -70,11 +70,6 @@ void UCustomizableObjectNodeTextureParameter::BackwardsCompatibleFixup()
 		ReferenceValue = DefaultValue;
 		DefaultValue = {};
 	}
-
-	if (CustomizableObjectCustomVersion < FCustomizableObjectCustomVersion::AddedRenameOptionToParameterNodes)
-	{
-		ReconstructNode();
-	}
 }
 
 #undef LOCTEXT_NAMESPACE
diff --git a/Engine/Plugins/Experimental/Mutable/Source/CustomizableObjectEditor/Public/MuCOE/Nodes/CustomizableObjectNodeColorParameter.h b/Engine/Plugins/Experimental/Mutable/Source/CustomizableObjectEditor/Public/MuCOE/Nodes/CustomizableObjectNodeColorParameter.h
--- a/Engine/Plugins/Experimental/Mutable/Source/CustomizableObjectEditor/Public/MuCOE/Nodes/CustomizableObjectNodeColorParameter.h
+++ b/Engine/Plugins/Experimental/Mutable/Source/CustomizableObjectEditor/Public/MuCOE/Nodes/CustomizableObjectNodeColorParameter.h
@@ -37,6 +37,5 @@ public:
 	// UCustomizableObjectNode interface
 	void AllocateDefaultPins(UCustomizableObjectNodeRemapPins* RemapPins) override;
 	bool IsAffectedByLOD() const override { return false; }
-	virtual void BackwardsCompatibleFixup() override;
 };
 
diff --git a/Engine/Plugins/Experimental/Mutable/Source/CustomizableObjectEditor/Public/MuCOE/Nodes/CustomizableObjectNodeEnumParameter.h b/Engine/Plugins/Experimental/Mutable/Source/CustomizableObjectEditor/Public/MuCOE/Nodes/CustomizableObjectNodeEnumParameter.h
--- a/Engine/Plugins/Experimental/Mutable/Source/CustomizableObjectEditor/Public/MuCOE/Nodes/CustomizableObjectNodeEnumParameter.h
+++ b/Engine/Plugins/Experimental/Mutable/Source/CustomizableObjectEditor/Public/MuCOE/Nodes/CustomizableObjectNodeEnumParameter.h
@@ -57,7 +57,6 @@ public:
 	// UCustomizableObjectNode interface
 	void AllocateDefaultPins(UCustomizableObjectNodeRemapPins* RemapPins) override;
 	virtual bool IsAffectedByLOD() const override { return false; }
-	virtual void BackwardsCompatibleFixup() override;
 
 };
 
