[mutable] Integrated the parameter selector tool to the Object Node details as a property widget.

#rnx

diff --git a/Engine/Plugins/Experimental/Mutable/Source/CustomizableObjectEditor/Private/MuCOE/CustomizableObjectEditorModule.cpp b/Engine/Plugins/Experimental/Mutable/Source/CustomizableObjectEditor/Private/MuCOE/CustomizableObjectEditorModule.cpp
--- a/Engine/Plugins/Experimental/Mutable/Source/CustomizableObjectEditor/Private/MuCOE/CustomizableObjectEditorModule.cpp
+++ b/Engine/Plugins/Experimental/Mutable/Source/CustomizableObjectEditor/Private/MuCOE/CustomizableObjectEditorModule.cpp
@@ -261,6 +261,7 @@ void FCustomizableObjectEditorModule::StartupModule()
 	PropertyModule.RegisterCustomPropertyTypeLayout("CustomizableObjectIdentifier", FOnGetPropertyTypeCustomizationInstance::CreateStatic(&FCustomizableObjectIdentifierCustomization::MakeInstance));
 	PropertyModule.RegisterCustomPropertyTypeLayout(FMeshReshapeBoneReference::StaticStruct()->GetFName(), FOnGetPropertyTypeCustomizationInstance::CreateStatic(&FMeshReshapeBonesReferenceCustomization::MakeInstance));
 	PropertyModule.RegisterCustomPropertyTypeLayout(FBoneToRemove::StaticStruct()->GetFName(), FOnGetPropertyTypeCustomizationInstance::CreateStatic(&FCustomizableObjectLODReductionSettings::MakeInstance));
+	PropertyModule.RegisterCustomPropertyTypeLayout(NAME_StrProperty, FOnGetPropertyTypeCustomizationInstance::CreateStatic(&FCustomizableObjectStateParameterSelector::MakeInstance), MakeShared<FStatePropertyTypeIdentifier>());
 
 	PropertyModule.NotifyCustomizationModuleChanged();
 
@@ -303,11 +304,19 @@ void FCustomizableObjectEditorModule::ShutdownModule()
 	if( FModuleManager::Get().IsModuleLoaded( "PropertyEditor" ) )
 	{
 		FPropertyEditorModule& PropertyModule = FModuleManager::GetModuleChecked<FPropertyEditorModule>("PropertyEditor");
+		
+		// Unregister Property views
 		for (const auto& ClassName : RegisteredCustomDetails)
 		{
 			PropertyModule.UnregisterCustomClassLayout(ClassName);
 		}
 
+		// Unregister Custom properties
+		PropertyModule.UnregisterCustomPropertyTypeLayout("CustomizableObjectIdentifier");
+		PropertyModule.UnregisterCustomPropertyTypeLayout(FMeshReshapeBoneReference::StaticStruct()->GetFName());
+		PropertyModule.UnregisterCustomPropertyTypeLayout(FBoneToRemove::StaticStruct()->GetFName());
+		PropertyModule.UnregisterCustomPropertyTypeLayout(NAME_StrProperty);
+
 		PropertyModule.NotifyCustomizationModuleChanged();
 	}
 
diff --git a/Engine/Plugins/Experimental/Mutable/Source/CustomizableObjectEditor/Private/MuCOE/GraphTraversal.cpp b/Engine/Plugins/Experimental/Mutable/Source/CustomizableObjectEditor/Private/MuCOE/GraphTraversal.cpp
--- a/Engine/Plugins/Experimental/Mutable/Source/CustomizableObjectEditor/Private/MuCOE/GraphTraversal.cpp
+++ b/Engine/Plugins/Experimental/Mutable/Source/CustomizableObjectEditor/Private/MuCOE/GraphTraversal.cpp
@@ -3,7 +3,6 @@
 #include "MuCOE/GraphTraversal.h"
 
 #include "AssetRegistry/AssetRegistryModule.h"
-#include "MuCO/CustomizableObjectPrivate.h"
 #include "MuCOE/CustomizableObjectPin.h"
 #include "MuCOE/EdGraphSchema_CustomizableObject.h"
 #include "MuCOE/Nodes/CustomizableObjectNodeEnumParameter.h"
@@ -552,7 +551,7 @@ void GetNodeGroupObjectNodeMappingImmersive(UCustomizableObject* Object, FAssetR
 		UCustomizableObject* ChildObject = Cast<UCustomizableObject>(AssetData.GetAsset());
 		if (!ChildObject)
 		{
-			continue;			
+			continue;
 		}
 
 		if (ChildObject != Object && !ChildObject->HasAnyFlags(RF_Transient))
@@ -590,6 +589,22 @@ TMultiMap<FGuid, UCustomizableObjectNodeObject*> GetNodeGroupObjectNodeMapping(U
 }
 
 
+void GetAllObjectsInGraph(UCustomizableObject* Object, TSet<UCustomizableObject*>& OutObjects)
+{
+	if (!Object)
+	{
+		return;
+	}
+
+	// Search the root of the CO's graph
+	UCustomizableObject* RootObject = GetRootObject(Object);
+	TMultiMap<FGuid, UCustomizableObjectNodeObject*> DummyMap;
+
+	FAssetRegistryModule& AssetRegistryModule = FModuleManager::LoadModuleChecked<FAssetRegistryModule>("AssetRegistry");
+	GetNodeGroupObjectNodeMappingImmersive(RootObject, AssetRegistryModule, OutObjects, DummyMap);
+}
+
+
 namespace GraphTraversal
 {
 	bool IsRootObject(const UCustomizableObject& Object)
diff --git a/Engine/Plugins/Experimental/Mutable/Source/CustomizableObjectEditor/Private/MuCOE/Nodes/CustomizableObjectNodeObjectDetails.cpp b/Engine/Plugins/Experimental/Mutable/Source/CustomizableObjectEditor/Private/MuCOE/Nodes/CustomizableObjectNodeObjectDetails.cpp
--- a/Engine/Plugins/Experimental/Mutable/Source/CustomizableObjectEditor/Private/MuCOE/Nodes/CustomizableObjectNodeObjectDetails.cpp
+++ b/Engine/Plugins/Experimental/Mutable/Source/CustomizableObjectEditor/Private/MuCOE/Nodes/CustomizableObjectNodeObjectDetails.cpp
@@ -5,17 +5,25 @@
 #include "DetailLayoutBuilder.h"
 #include "Framework/MultiBox/MultiBoxBuilder.h"
 #include "IDetailsView.h"
+#include "MuCO/CustomizableObjectPrivate.h"
 #include "MuCOE/CustomizableObjectEditorUtilities.h"
 #include "MuCOE/CustomizableObjectGraph.h"
+#include "MuCOE/GraphTraversal.h"
 #include "MuCOE/Nodes/CustomizableObjectNode.h"
 #include "MuCOE/Nodes/CustomizableObjectNodeObject.h"
 #include "MuCOE/Nodes/CustomizableObjectNodeObjectGroup.h"
+#include "MuCOE/Nodes/CustomizableObjectNodeColorParameter.h"
+#include "MuCOE/Nodes/CustomizableObjectNodeEnumParameter.h"
+#include "MuCOE/Nodes/CustomizableObjectNodeFloatParameter.h"
+#include "MuCOE/Nodes/CustomizableObjectNodeGroupProjectorParameter.h"
+#include "MuCOE/Nodes/CustomizableObjectNodeTable.h"
+#include "MuCOE/Nodes/CustomizableObjectNodeTextureParameter.h"
 #include "MuCOE/SCustomizableObjectNodeObjectRTMorphTargetOverride.h"
-#include "MuCOE/Widgets/CustomizableObjectLODReductionSettings.h"
 #include "MuCOE/UnrealEditorPortabilityHelpers.h"
+#include "MuCOE/Widgets/CustomizableObjectLODReductionSettings.h"
 #include "PropertyCustomizationHelpers.h"
 #include "ScopedTransaction.h"
-#include "MuCO/CustomizableObjectPrivate.h"
+#include "SSearchableComboBox.h"
 #include "Widgets/Input/SComboButton.h"
 #include "Widgets/Input/STextComboBox.h"
 #include "Widgets/Text/STextBlock.h"
@@ -24,6 +32,122 @@
 #define LOCTEXT_NAMESPACE "CustomizableObjectDetails"
 
 
+TSharedRef<IPropertyTypeCustomization> FCustomizableObjectStateParameterSelector::MakeInstance()
+{
+	return MakeShared<FCustomizableObjectStateParameterSelector>();
+}
+
+
+void FCustomizableObjectStateParameterSelector::CustomizeHeader(TSharedRef<IPropertyHandle> InPropertyHandle, FDetailWidgetRow& InHeaderRow, IPropertyTypeCustomizationUtils& InCustomizationUtils)
+{
+	TArray<UObject*> OuterObjects;
+	InPropertyHandle->GetOuterObjects(OuterObjects);
+	PropertyHandle = InPropertyHandle;
+
+	if (OuterObjects.Num())
+	{
+		BaseObjectNode = Cast<UCustomizableObjectNodeObject>(OuterObjects[0]);
+		
+		if (!BaseObjectNode.IsValid())
+		{
+			return;
+		}
+	}
+
+	FString SelectedParameterName;
+
+	InPropertyHandle->GetValue(SelectedParameterName);
+	GenerateParameterOptions(SelectedParameterName);
+
+	InHeaderRow
+	.NameContent().HAlign(EHorizontalAlignment::HAlign_Fill).VAlign(EVerticalAlignment::VAlign_Center)
+	[
+		InPropertyHandle->CreatePropertyNameWidget()
+	]
+	.ValueContent().MinDesiredWidth(300.0f)
+	[
+		SNew(SBorder)
+		.BorderBackgroundColor(FLinearColor::Transparent)
+		[
+			SNew(SSearchableComboBox)
+			.InitiallySelectedItem(SelectedParameter)
+			.OptionsSource(&ParameterOptions)
+			.OnSelectionChanged(this, &FCustomizableObjectStateParameterSelector::OnParameterNameSelectionChanged)
+			.OnGenerateWidget(this, &FCustomizableObjectStateParameterSelector::OnGenerateStateParameterSelectorComboBox)
+			.Content()
+			[
+				SNew(STextBlock)
+				.Text(this, &FCustomizableObjectStateParameterSelector::GetSelectedParameterName)
+				.Font(IDetailLayoutBuilder::GetDetailFont())
+			]
+		]
+	]
+	.OverrideResetToDefault(FResetToDefaultOverride::Create(FSimpleDelegate::CreateSP(this, &FCustomizableObjectStateParameterSelector::ResetSelectedParameterButtonClicked)));
+}
+
+
+void FCustomizableObjectStateParameterSelector::OnParameterNameSelectionChanged(TSharedPtr<FString> Selection, ESelectInfo::Type SelectInfo)
+{
+	if (Selection.IsValid())
+	{
+		FString Value = Selection == ParameterOptions[0] ? FString() : *Selection;
+		PropertyHandle->SetValue(Value);
+		SelectedParameter = Selection;
+	}
+}
+
+
+FText FCustomizableObjectStateParameterSelector::GetSelectedParameterName() const
+{
+	if (SelectedParameter.IsValid())
+	{
+		return FText::FromString(*SelectedParameter);
+	}
+
+	return FText();
+}
+
+
+void FCustomizableObjectStateParameterSelector::GenerateParameterOptions(const FString& SelectedValue)
+{
+	ParameterOptions.Empty();
+	ParameterOptions.Add(MakeShareable(new FString("- Nothing Selected -")));
+
+	SelectedParameter = ParameterOptions.Last();
+
+	for (const FString& ParameterName : BaseObjectNode->ParameterNames)
+	{
+		ParameterOptions.Add(MakeShareable(new FString(ParameterName)));
+
+		if (ParameterName == SelectedValue)
+		{
+			SelectedParameter = ParameterOptions.Last();
+		}
+	}
+
+	// we should always have something selected
+	check(SelectedParameter.IsValid());
+}
+
+
+TSharedRef<SWidget> FCustomizableObjectStateParameterSelector::OnGenerateStateParameterSelectorComboBox(TSharedPtr<FString> InItem) const
+{
+	return SNew(STextBlock).Text(FText::FromString(*InItem.Get())).Font(IDetailLayoutBuilder::GetDetailFont());
+}
+
+
+void FCustomizableObjectStateParameterSelector::ResetSelectedParameterButtonClicked()
+{
+	check(ParameterOptions.Num() > 0);
+
+	PropertyHandle->SetValue(FString());
+	SelectedParameter = ParameterOptions[0];
+}
+
+
+
+// Details -------------------------------------------------------------------------------
+
 TSharedRef<IDetailCustomization> FCustomizableObjectNodeObjectDetails::MakeInstance()
 {
 	return MakeShareable( new FCustomizableObjectNodeObjectDetails );
@@ -32,13 +156,13 @@ TSharedRef<IDetailCustomization> FCustomizableObjectNodeObjectDetails::MakeInsta
 
 void FCustomizableObjectNodeObjectDetails::CustomizeDetails( IDetailLayoutBuilder& DetailBuilder )
 {
-	Node = nullptr;
+	BaseObjectNode = nullptr;
 	DetailBuilderPtr = &DetailBuilder;
 
 	const IDetailsView* DetailsView = DetailBuilder.GetDetailsView();
 	if ( DetailsView->GetSelectedObjects().Num() )
 	{
-		Node = Cast<UCustomizableObjectNodeObject>( DetailsView->GetSelectedObjects()[0].Get() );
+		BaseObjectNode = Cast<UCustomizableObjectNodeObject>( DetailsView->GetSelectedObjects()[0].Get() );
 	}
 
 	IDetailCategoryBuilder& StatesCategory = DetailBuilder.EditCategory("States");
@@ -49,8 +173,10 @@ void FCustomizableObjectNodeObjectDetails::CustomizeDetails( IDetailLayoutBuilde
 
 	//StatesCategory.CategoryIcon( "ActorClassIcon.CustomizableObject" );
 
-	if (Node)
+	if (BaseObjectNode.IsValid())
 	{
+		FillParameterNamesArray();
+
 		// Properties
 		TSharedRef<IPropertyHandle> StatesProperty = DetailBuilder.GetProperty("States");
 		TSharedRef<IPropertyHandle> ParentObjectProperty = DetailBuilder.GetProperty("ParentObject");
@@ -58,7 +184,7 @@ void FCustomizableObjectNodeObjectDetails::CustomizeDetails( IDetailLayoutBuilde
 		TSharedRef<IPropertyHandle> LODsProperty = DetailBuilder.GetProperty("NumLODs");
 		TSharedRef<IPropertyHandle> ComponentSettingsProperty= DetailBuilder.GetProperty("ComponentSettings");
 
-		FName BonesToRemovePropertyPath = FName("ComponentSettings[" + FString::FromInt(Node->CurrentComponent) + "].LODReductionSettings[" + FString::FromInt(Node->CurrentLOD) + "].BonesToRemove");
+		FName BonesToRemovePropertyPath = FName("ComponentSettings[" + FString::FromInt(BaseObjectNode->CurrentComponent) + "].LODReductionSettings[" + FString::FromInt(BaseObjectNode->CurrentLOD) + "].BonesToRemove");
 		TSharedRef<IPropertyHandle> BonesToRemoveProperty = DetailBuilder.GetProperty(BonesToRemovePropertyPath);
 
 		// Callbacks
@@ -73,21 +199,21 @@ void FCustomizableObjectNodeObjectDetails::CustomizeDetails( IDetailLayoutBuilde
 
 		GroupNodeComboOptions.Empty();
 
-		if (Node->bIsBase)
+		if (BaseObjectNode->bIsBase)
 		{
 			ExternalCategory.AddCustomRow(LOCTEXT("FCustomizableObjectNodeObjectDetails", "Blocks"))
 			[
 				SNew(SObjectPropertyEntryBox)
 				.AllowedClass(UCustomizableObject::StaticClass())
 				.OnObjectChanged(this, &FCustomizableObjectNodeObjectDetails::ParentObjectSelectionChanged)
-				.ObjectPath(Node->ParentObject->GetPathName())
+				.ObjectPath(BaseObjectNode->ParentObject->GetPathName())
 				.ForceVolatile(true)
 			];
 
-			if (Node->ParentObject)
+			if (BaseObjectNode->ParentObject)
 			{
 				TArray<UCustomizableObjectNodeObjectGroup*> GroupNodes;
-				Node->ParentObject->GetPrivate()->GetSource()->GetNodesOfClass<UCustomizableObjectNodeObjectGroup>(GroupNodes);
+				BaseObjectNode->ParentObject->GetPrivate()->GetSource()->GetNodesOfClass<UCustomizableObjectNodeObjectGroup>(GroupNodes);
 
 				TSharedPtr<FString> ItemToSelect;
 
@@ -95,13 +221,13 @@ void FCustomizableObjectNodeObjectDetails::CustomizeDetails( IDetailLayoutBuilde
 				{
 					GroupNodeComboOptions.Add(MakeShareable(new FString(GroupNode->GroupName)));
 
-					if (Node->ParentObjectGroupId == GroupNode->NodeGuid)
+					if (BaseObjectNode->ParentObjectGroupId == GroupNode->NodeGuid)
 					{
 						ItemToSelect = GroupNodeComboOptions.Last();
 					}
 				}
 
-				if (!Node->ParentObjectGroupId.IsValid() && ParentComboOptions.Num() > 0)
+				if (!BaseObjectNode->ParentObjectGroupId.IsValid() && ParentComboOptions.Num() > 0)
 				{
 					ItemToSelect = GroupNodeComboOptions.Last();
 				}
@@ -137,7 +263,7 @@ void FCustomizableObjectNodeObjectDetails::CustomizeDetails( IDetailLayoutBuilde
                             "Realtime Morph Targets Override "))
                 [
                     SNew(SCustomizableObjectNodeSkeletalMeshRTMorphTargetOverride)
-                    .Node(Node)
+                    .Node(BaseObjectNode.Get())
                 ];
 
 				// Component Settings Category ----------
@@ -234,15 +360,15 @@ void FCustomizableObjectNodeObjectDetails::CustomizeDetails( IDetailLayoutBuilde
 
 void FCustomizableObjectNodeObjectDetails::ParentObjectSelectionChanged(const FAssetData & AssetData)
 {
-	if (Node)
+	if (BaseObjectNode.IsValid())
 	{
 		UCustomizableObject* Parent = Cast<UCustomizableObject>(AssetData.GetAsset());
-		Node->SetParentObject(Parent);
+		BaseObjectNode->SetParentObject(Parent);
 
 		// If set the parent to nullt, invalidate also the reference GUID
 		if (!Parent)
 		{
-			Node->ParentObjectGroupId.Invalidate();
+			BaseObjectNode->ParentObjectGroupId.Invalidate();
 		}
 	}
 
@@ -255,18 +381,18 @@ void FCustomizableObjectNodeObjectDetails::ParentObjectSelectionChanged(const FA
 
 void FCustomizableObjectNodeObjectDetails::OnGroupNodeComboBoxSelectionChanged(TSharedPtr<FString> Selection, ESelectInfo::Type SelectInfo, TSharedRef<IPropertyHandle> ParentProperty)
 {
-	if (Selection.IsValid() && Node->ParentObject != nullptr)
+	if (Selection.IsValid() && BaseObjectNode->ParentObject != nullptr)
 	{
 		TArray<UCustomizableObjectNodeObjectGroup*> GroupNodes;
-		Node->ParentObject->GetPrivate()->GetSource()->GetNodesOfClass<UCustomizableObjectNodeObjectGroup>(GroupNodes);
+		BaseObjectNode->ParentObject->GetPrivate()->GetSource()->GetNodesOfClass<UCustomizableObjectNodeObjectGroup>(GroupNodes);
 
 		for (UCustomizableObjectNodeObjectGroup* GroupNode : GroupNodes)
 		{
 			if (*Selection == GroupNode->GroupName)
 			{
 				const FScopedTransaction Transaction(LOCTEXT("ChangedAttachedToExternalObjectTransaction", "Changed Attached to External Object"));
-				Node->Modify();
-				Node->ParentObjectGroupId = GroupNode->NodeGuid;
+				BaseObjectNode->Modify();
+				BaseObjectNode->ParentObjectGroupId = GroupNode->NodeGuid;
 			}
 		}
 	}
@@ -313,9 +439,9 @@ TSharedRef<SWidget> FCustomizableObjectNodeObjectDetails::OnGenerateLODComboBoxF
 
 TSharedRef<SWidget> FCustomizableObjectNodeObjectDetails::OnGenerateComponentMenuForPicker()
 {
-	if (Node)
+	if (BaseObjectNode.IsValid())
 	{
-		int32 NumComponents = Node->NumMeshComponents;
+		int32 NumComponents = BaseObjectNode->NumMeshComponents;
 		FMenuBuilder MenuBuilder(true, NULL);
 
 		for (int32 ComponentIndex = 0; ComponentIndex < NumComponents; ++ComponentIndex)
@@ -334,9 +460,9 @@ TSharedRef<SWidget> FCustomizableObjectNodeObjectDetails::OnGenerateComponentMen
 
 TSharedRef<SWidget> FCustomizableObjectNodeObjectDetails::OnGenerateLODMenuForPicker()
 {
-	if (Node)
+	if (BaseObjectNode.IsValid())
 	{
-		int32 NumLODs = Node->NumLODs;
+		int32 NumLODs = BaseObjectNode->NumLODs;
 		FMenuBuilder MenuBuilder(true, NULL);
 
 		for (int32 LODIndex = 0; LODIndex < NumLODs; ++LODIndex)
@@ -355,10 +481,10 @@ TSharedRef<SWidget> FCustomizableObjectNodeObjectDetails::OnGenerateLODMenuForPi
 
 void FCustomizableObjectNodeObjectDetails::OnSelectedComponentChanged(int32 NewComponentIndex)
 {
-	if (Node)
+	if (BaseObjectNode.IsValid())
 	{
-		Node->CurrentComponent = NewComponentIndex;
-		Node->CurrentLOD = 0;
+		BaseObjectNode->CurrentComponent = NewComponentIndex;
+		BaseObjectNode->CurrentLOD = 0;
 	}
 	
 	DetailBuilderPtr->ForceRefreshDetails();
@@ -367,9 +493,9 @@ void FCustomizableObjectNodeObjectDetails::OnSelectedComponentChanged(int32 NewC
 
 void FCustomizableObjectNodeObjectDetails::OnSelectedLODChanged(int32 NewLODIndex)
 {
-	if (Node)
+	if (BaseObjectNode.IsValid())
 	{
-		Node->CurrentLOD = NewLODIndex;
+		BaseObjectNode->CurrentLOD = NewLODIndex;
 	}
 
 	DetailBuilderPtr->ForceRefreshDetails();
@@ -380,9 +506,9 @@ FText FCustomizableObjectNodeObjectDetails::GetCurrentComponentName() const
 {
 	FText ComponentText;
 
-	if (Node)
+	if (BaseObjectNode.IsValid())
 	{
-		ComponentText = FText::FromString(FString(TEXT("Component ")) + FString::FromInt(Node->CurrentComponent));
+		ComponentText = FText::FromString(FString(TEXT("Component ")) + FString::FromInt(BaseObjectNode->CurrentComponent));
 	}
 
 	return ComponentText;
@@ -393,9 +519,9 @@ FText FCustomizableObjectNodeObjectDetails::GetCurrentLODName() const
 {
 	FText LODText;
 
-	if (Node)
+	if (BaseObjectNode.IsValid())
 	{
-		LODText = FText::FromString(FString(TEXT("LOD ")) + FString::FromInt(Node->CurrentLOD));
+		LODText = FText::FromString(FString(TEXT("LOD ")) + FString::FromInt(BaseObjectNode->CurrentLOD));
 	}
 
 	return LODText;
@@ -406,22 +532,182 @@ void FCustomizableObjectNodeObjectDetails::OnNumComponentsOrLODsChanged(const FP
 {
 	if (DetailBuilderPtr && PropertyChangedEvent.ChangeType == EPropertyChangeType::ValueSet)
 	{
-		int32 ComponentToSelect = FMath::Min(Node->NumMeshComponents - 1, Node->CurrentComponent);
+		int32 ComponentToSelect = FMath::Min(BaseObjectNode->NumMeshComponents - 1, BaseObjectNode->CurrentComponent);
 
-		if (ComponentToSelect != Node->CurrentComponent)
+		if (ComponentToSelect != BaseObjectNode->CurrentComponent)
 		{
-			Node->CurrentComponent = ComponentToSelect;
+			BaseObjectNode->CurrentComponent = ComponentToSelect;
 
 			// Reset the LOD selection
-			Node->CurrentLOD = 0;
+			BaseObjectNode->CurrentLOD = 0;
 		}
 		else
 		{
-			Node->CurrentLOD = FMath::Min(Node->NumLODs - 1, Node->CurrentLOD);
+			BaseObjectNode->CurrentLOD = FMath::Min(BaseObjectNode->NumLODs - 1, BaseObjectNode->CurrentLOD);
 		}
 
 		DetailBuilderPtr->ForceRefreshDetails();
 	}
 }
 
+
+void FCustomizableObjectNodeObjectDetails::FillParameterNamesArray()
+{
+	BaseObjectNode->ParameterNames.Empty();
+
+	UCustomizableObject* CustomizableObject = Cast<UCustomizableObject>(BaseObjectNode->GetOutermostObject());
+
+	if (!CustomizableObject)
+	{
+		return;
+	}
+
+	// Get full graph root customizable object
+	UCustomizableObject* RootObjet = GetRootObject(CustomizableObject);
+
+	// Full tree graph of customizable objects
+	TSet<UCustomizableObject*> CustomObjectTree;
+
+	// Get the whole tree of customizable object
+	GetAllObjectsInGraph(RootObjet, CustomObjectTree);
+
+	// Array to store all the ids of group nodes of type toggle
+	TArray<FGuid> ToggleGroupObjectIds;
+
+	// Array to store all the Child Object Nodes
+	TArray<const UCustomizableObjectNodeObject*> AllObjectNodes;
+
+	for (const UCustomizableObject* Object : CustomObjectTree)
+	{
+		if (!Object || !Object->GetPrivate()->GetSource())
+		{
+			return;
+		}
+
+		UEdGraph* Source = Object->GetPrivate()->GetSource();
+
+		if (!Object || !Object->GetPrivate() || !Object->GetPrivate()->GetSource())
+		{
+			continue;
+		}
+
+		// All type of parameter nodes
+		TArray<UCustomizableObjectNodeColorParameter*> ColorParameterNodes;
+		Source->GetNodesOfClass(ColorParameterNodes);
+
+		for (const UCustomizableObjectNodeColorParameter* ColorParameterNode : ColorParameterNodes)
+		{
+			if (ColorParameterNode)
+			{
+				BaseObjectNode->ParameterNames.Add(ColorParameterNode->ParameterName);
+			}
+		}
+
+		TArray<UCustomizableObjectNodeFloatParameter*> FloatParameterNodes;
+		Source->GetNodesOfClass(FloatParameterNodes);
+
+		for (const UCustomizableObjectNodeFloatParameter* FloatParameterNode : FloatParameterNodes)
+		{
+			if (FloatParameterNode)
+			{
+				BaseObjectNode->ParameterNames.Add(FloatParameterNode->ParameterName);
+			}
+		}
+
+		TArray<UCustomizableObjectNodeEnumParameter*> EnumParameterNodes;
+		Source->GetNodesOfClass(EnumParameterNodes);
+
+		for (const UCustomizableObjectNodeEnumParameter* EnumParameterNode : EnumParameterNodes)
+		{
+			if (EnumParameterNode)
+			{
+				BaseObjectNode->ParameterNames.Add(EnumParameterNode->ParameterName);
+			}
+		}
+
+		TArray<UCustomizableObjectNodeGroupProjectorParameter*>GroupProjectorParameterNodes;
+		Source->GetNodesOfClass(GroupProjectorParameterNodes);
+
+		for (const UCustomizableObjectNodeGroupProjectorParameter* GroupProjectorParameterNode : GroupProjectorParameterNodes)
+		{
+			if (GroupProjectorParameterNode)
+			{
+				BaseObjectNode->ParameterNames.Add(GroupProjectorParameterNode->ParameterName);
+			}
+		}
+
+		TArray<UCustomizableObjectNodeProjectorParameter*> ProjectorParameterNodes;
+		Source->GetNodesOfClass(ProjectorParameterNodes);
+
+		for (const UCustomizableObjectNodeProjectorParameter* ProjectorParameterNode : ProjectorParameterNodes)
+		{
+			if (ProjectorParameterNode)
+			{
+				BaseObjectNode->ParameterNames.Add(ProjectorParameterNode->ParameterName);
+			}
+		}
+
+		TArray<UCustomizableObjectNodeTextureParameter*> TextureParameterNodes;
+		Source->GetNodesOfClass(TextureParameterNodes);
+
+		for (const UCustomizableObjectNodeTextureParameter* TextureParameterNode : TextureParameterNodes)
+		{
+			if (TextureParameterNode)
+			{
+				BaseObjectNode->ParameterNames.Add(TextureParameterNode->ParameterName);
+			}
+		}
+
+		TArray<UCustomizableObjectNodeTable*> TableNodes;
+		Source->GetNodesOfClass(TableNodes);
+
+		for (const UCustomizableObjectNodeTable* TableNode : TableNodes)
+		{
+			if (TableNode)
+			{
+				BaseObjectNode->ParameterNames.Add(TableNode->ParameterName);
+			}
+		}
+
+
+		TArray<UCustomizableObjectNodeObjectGroup*> GroupNodes;
+		Source->GetNodesOfClass(GroupNodes);
+
+		for (const UCustomizableObjectNodeObjectGroup* GroupNode : GroupNodes)
+		{
+			if (GroupNode)
+			{
+				if (GroupNode->GroupType == ECustomizableObjectGroupType::COGT_TOGGLE)
+				{
+					ToggleGroupObjectIds.Add(GroupNode->NodeGuid);
+				}
+				else
+				{
+					BaseObjectNode->ParameterNames.Add(GroupNode->GroupName);
+				}
+			}
+		}
+
+		TArray<UCustomizableObjectNodeObject*> ObjectNodes;
+		Source->GetNodesOfClass(ObjectNodes);
+
+		for (const UCustomizableObjectNodeObject* ObjectNode : ObjectNodes)
+		{
+			if (ObjectNode)
+			{
+				AllObjectNodes.Add(ObjectNode);
+			}
+		}
+	}
+
+	// Now that we know all the group objects of type toggle, we process all the object nodes that can generate a parameter
+	for (const UCustomizableObjectNodeObject* ObjectNode : AllObjectNodes)
+	{
+		if (ToggleGroupObjectIds.Contains(ObjectNode->ParentObjectGroupId))
+		{
+			BaseObjectNode->ParameterNames.Add(ObjectNode->ObjectName);
+		}
+	}
+}
+
 #undef LOCTEXT_NAMESPACE
diff --git a/Engine/Plugins/Experimental/Mutable/Source/CustomizableObjectEditor/Private/MuCOE/Nodes/CustomizableObjectNodeObjectDetails.h b/Engine/Plugins/Experimental/Mutable/Source/CustomizableObjectEditor/Private/MuCOE/Nodes/CustomizableObjectNodeObjectDetails.h
--- a/Engine/Plugins/Experimental/Mutable/Source/CustomizableObjectEditor/Private/MuCOE/Nodes/CustomizableObjectNodeObjectDetails.h
+++ b/Engine/Plugins/Experimental/Mutable/Source/CustomizableObjectEditor/Private/MuCOE/Nodes/CustomizableObjectNodeObjectDetails.h
@@ -3,16 +3,81 @@
 #pragma once
 
 #include "IDetailCustomization.h"
+#include "IPropertyTypeCustomization.h"
+#include "PropertyEditorModule.h"
+#include "PropertyHandle.h"
 
 namespace ESelectInfo { enum Type : int; }
 
 class FString;
 class IDetailLayoutBuilder;
-class IPropertyHandle;
 class SWidget;
+class UCustomizableObject;
+class UCustomizableObjectNodeObject;
 struct FAssetData;
 struct FPropertyChangedEvent;
 
+/**	Class created to identify FStrings with metadata ShowParameterOptions.
+	This allows us to change the default FString property widget for a ComboBox in the
+	CustomizableObjectNodeObject details. */
+class FStatePropertyTypeIdentifier : public IPropertyTypeIdentifier 
+{
+public:
+	virtual bool IsPropertyTypeCustomized(const IPropertyHandle& InPropertyHandle) const override
+	{
+		// Metadata to use a combobox widged instead of a string widge
+		return InPropertyHandle.HasMetaData(TEXT("ShowParameterOptions"));
+	}
+};
+
+
+/** Custom Widget for the property RuntimeParameters of the States(FCustomizableObjectState) */
+class FCustomizableObjectStateParameterSelector : public IPropertyTypeCustomization
+{
+public:
+
+	FCustomizableObjectStateParameterSelector() {};
+
+	static TSharedRef<IPropertyTypeCustomization> MakeInstance();
+
+	// IPropertyTypeCustomization overrides
+	virtual void CustomizeHeader(TSharedRef<IPropertyHandle> InPropertyHandle, FDetailWidgetRow& InHeaderRow, IPropertyTypeCustomizationUtils& InCustomizationUtils) override;
+
+	virtual void CustomizeChildren(TSharedRef<IPropertyHandle> InPropertyHandle, IDetailChildrenBuilder& InChildBuilder, IPropertyTypeCustomizationUtils& InCustomizationUtils) override
+	{}
+
+private:
+
+	/** Generates the option of the combobox */
+	void GenerateParameterOptions(const FString& SelectedValue);
+
+	/** Generates the Text widget of the combobox */
+	TSharedRef<SWidget> OnGenerateStateParameterSelectorComboBox(TSharedPtr<FString> InItem) const;
+
+	/** Callback of OnSelectionChanged of the combobox */
+	void OnParameterNameSelectionChanged(TSharedPtr<FString> Selection, ESelectInfo::Type SelectInfo);
+
+	/** Returns the name of the current selected parameter */
+	FText GetSelectedParameterName() const;
+
+	/** Callback of the function OverrideResetToDefault */
+	void ResetSelectedParameterButtonClicked();
+
+private:
+
+	/** Array with all the possible parameter names */
+	TArray<TSharedPtr<FString>> ParameterOptions;
+
+	/** Weak pointer to the Customizable Object that contains this property */
+	TWeakObjectPtr<UCustomizableObjectNodeObject> BaseObjectNode;
+
+	/** Pointer to the current selected parameter */
+	TSharedPtr<FString> SelectedParameter;
+
+	/** Runtime Parameter Name property of a State */
+	TSharedPtr<IPropertyHandle> PropertyHandle;
+};
+
 
 class FCustomizableObjectNodeObjectDetails : public IDetailCustomization
 {
@@ -27,7 +92,7 @@ public:
 
 private:
 
-	class UCustomizableObjectNodeObject* Node;
+	TWeakObjectPtr<UCustomizableObjectNodeObject> BaseObjectNode;
 	TArray< TSharedPtr<FString> > GroupNodeComboOptions;
 	IDetailLayoutBuilder* DetailBuilderPtr = nullptr;
 
@@ -60,4 +125,6 @@ private:
 
 	FIntPoint GetGridSize() const;
 	TArray<FIntRect> GetBlocks() const;
+
+	void FillParameterNamesArray();
 };
diff --git a/Engine/Plugins/Experimental/Mutable/Source/CustomizableObjectEditor/Public/MuCOE/GraphTraversal.h b/Engine/Plugins/Experimental/Mutable/Source/CustomizableObjectEditor/Public/MuCOE/GraphTraversal.h
--- a/Engine/Plugins/Experimental/Mutable/Source/CustomizableObjectEditor/Public/MuCOE/GraphTraversal.h
+++ b/Engine/Plugins/Experimental/Mutable/Source/CustomizableObjectEditor/Public/MuCOE/GraphTraversal.h
@@ -6,6 +6,7 @@
 
 class UCustomizableObject;
 class UCustomizableObjectNodeObject;
+class UEdGraph;
 class UEdGraphPin;
 class UCustomizableObjectNode;
 
@@ -70,6 +71,9 @@ const UEdGraphPin* FindMeshBaseSource(const UEdGraphPin& Pin, const bool bOnlyLo
  * @param Object Child or root Object to start the search from. */
 TMultiMap<FGuid, UCustomizableObjectNodeObject*> GetNodeGroupObjectNodeMapping(UCustomizableObject* Object);
 
+/** Returns all the Customizable Objects in a graph starting at the root object */
+void GetAllObjectsInGraph(UCustomizableObject* Object, TSet<UCustomizableObject*>& OutObjects);
+
 namespace GraphTraversal
 {
     /** Return ture if the given Customizable Object is Root Object (not a Child Object). */
diff --git a/Engine/Plugins/Experimental/Mutable/Source/CustomizableObjectEditor/Public/MuCOE/Nodes/CustomizableObjectNodeObject.h b/Engine/Plugins/Experimental/Mutable/Source/CustomizableObjectEditor/Public/MuCOE/Nodes/CustomizableObjectNodeObject.h
--- a/Engine/Plugins/Experimental/Mutable/Source/CustomizableObjectEditor/Public/MuCOE/Nodes/CustomizableObjectNodeObject.h
+++ b/Engine/Plugins/Experimental/Mutable/Source/CustomizableObjectEditor/Public/MuCOE/Nodes/CustomizableObjectNodeObject.h
@@ -63,7 +63,7 @@ struct CUSTOMIZABLEOBJECTEDITOR_API FCustomizableObjectState
 	UPROPERTY(EditAnywhere, Category=CustomizableObject)
 	FString Name;
 
-	UPROPERTY(EditAnywhere, Category=CustomizableObject)
+	UPROPERTY(EditAnywhere, Category = CustomizableObject, meta = (ShowParameterOptions))
 	TArray<FString> RuntimeParameters;
 
 	/** Special treatment of texture compression for this state. */
@@ -283,6 +283,9 @@ public:
 	int32 CurrentComponent = 0;
 	int32 CurrentLOD = 0;
 
+	// Array filled in the Details of the node to store all the parameter names of a CO graph (full tree)
+	TArray<FString> ParameterNames;
+
 private:
 	static const FName ChildrenPinName;
 	static const FName OutputPinName;
