[Mutable] Fixed "Failed to remap old pins" format.

#rnx

diff --git a/Engine/Plugins/Experimental/Mutable/Source/CustomizableObjectEditor/Private/MuCOE/CustomizableObjectEditorLogger.cpp b/Engine/Plugins/Experimental/Mutable/Source/CustomizableObjectEditor/Private/MuCOE/CustomizableObjectEditorLogger.cpp
--- a/Engine/Plugins/Experimental/Mutable/Source/CustomizableObjectEditor/Private/MuCOE/CustomizableObjectEditorLogger.cpp
+++ b/Engine/Plugins/Experimental/Mutable/Source/CustomizableObjectEditor/Private/MuCOE/CustomizableObjectEditorLogger.cpp
@@ -46,8 +46,9 @@ private:
 	/** Private constructor */
 	explicit FCustomizableObjectToken(const UEdGraphNode* Node) : NodeBeingReferenced(Node)
 	{
-		const FString String = Node ? *Node->GetName() : TEXT("<None>");
-		CachedText = FText::FromString(String);
+		CachedText = Node ?
+			FText::FromString(Node->GetNodeTitle(ENodeTitleType::FullTitle).ToString().Replace(TEXT("\n"), TEXT(" "))) :
+			LOCTEXT("TokenNoNode", "<None>");
 	}
 
 	/** An object being referenced by this token, if any */
diff --git a/Engine/Plugins/Experimental/Mutable/Source/CustomizableObjectEditor/Private/MuCOE/Nodes/CustomizableObjectNode.cpp b/Engine/Plugins/Experimental/Mutable/Source/CustomizableObjectEditor/Private/MuCOE/Nodes/CustomizableObjectNode.cpp
--- a/Engine/Plugins/Experimental/Mutable/Source/CustomizableObjectEditor/Private/MuCOE/Nodes/CustomizableObjectNode.cpp
+++ b/Engine/Plugins/Experimental/Mutable/Source/CustomizableObjectEditor/Private/MuCOE/Nodes/CustomizableObjectNode.cpp
@@ -253,7 +253,7 @@ void UCustomizableObjectNode::ReconstructNode(UCustomizableObjectNodeRemapPins*
 	}
 
 	bool bOrphanedPin = false;
-	FText FirstOldPin;
+	FName FirstOldPin;
 	for (UEdGraphPin* OldPin : OldPins)
 	{
 		OldPin->Modify();
@@ -265,7 +265,7 @@ void UCustomizableObjectNode::ReconstructNode(UCustomizableObjectNodeRemapPins*
 				bOrphanedPin = !OldPin->bOrphanedPin;
 				if (bOrphanedPin)
 				{
-					FirstOldPin = FText::FromString(OldPin->GetName());
+					FirstOldPin = OldPin->GetFName();
 				}
 			}
 			OrphanPin(*OldPin);
@@ -287,7 +287,7 @@ void UCustomizableObjectNode::ReconstructNode(UCustomizableObjectNodeRemapPins*
 	{
 		if (bOrphanedPin)
 		{
-			FCustomizableObjectEditorLogger::CreateLog(FText::Format(LOCTEXT("OrphanPinsWarningReconstruct", "Failed to remap old pins. Graph:{0} Node:{1} FirstOldPin:{2}"), FText::FromString(Graph->GetPathName()), FText::FromString(GetPathName()), FirstOldPin))
+			FCustomizableObjectEditorLogger::CreateLog(FText::Format(LOCTEXT("OrphanPinsWarningReconstruct", "Failed to remap old pins. First old pin: {0}"), FText::FromName(FirstOldPin)))
 			.BaseObject()
 			.Severity(EMessageSeverity::Warning)
 			.Context(*this)
