[Mutable] Display log messages when a compilation starts/ends.

#rnx

diff --git a/Engine/Plugins/Experimental/Mutable/Source/CustomizableObjectEditor/Private/MuCOE/CustomizableObjectCompiler.cpp b/Engine/Plugins/Experimental/Mutable/Source/CustomizableObjectEditor/Private/MuCOE/CustomizableObjectCompiler.cpp
--- a/Engine/Plugins/Experimental/Mutable/Source/CustomizableObjectEditor/Private/MuCOE/CustomizableObjectCompiler.cpp
+++ b/Engine/Plugins/Experimental/Mutable/Source/CustomizableObjectEditor/Private/MuCOE/CustomizableObjectCompiler.cpp
@@ -213,6 +213,8 @@ void FCustomizableObjectCompiler::Compile(const TSharedRef<FCompilationRequest>&
 
 	SetCompilationState(ECompilationStatePrivate::InProgress, ECompilationResultPrivate::Unknown);
 
+	CompilationStartTime = FPlatformTime::Seconds();
+
 	// Now that we know for sure that the CO is locked and there are no pending updates of instances using the CO,
 	// destroy any live update instances, as they become invalid when recompiling the CO
 	for (TObjectIterator<UCustomizableObjectInstance> It; It; ++It)
@@ -225,6 +227,22 @@ void FCustomizableObjectCompiler::Compile(const TSharedRef<FCompilationRequest>&
 		}
 	}
 
+	// If we don't have the target platform yet (in editor) we need to get it
+	if (!CurrentOptions.TargetPlatform)
+	{
+		check(!CurrentOptions.bIsCooking);
+
+		CurrentOptions.TargetPlatform = GetTargetPlatformManagerRef().GetRunningTargetPlatform();
+		check(CurrentOptions.TargetPlatform != nullptr);
+	}
+
+	UE_LOG(LogMutable, Display, TEXT("Compiling Customizable Object %s for platform %s."), *CurrentObject->GetName(), *CurrentOptions.TargetPlatform->PlatformName());
+
+	if (CurrentOptions.bForceLargeLODBias)
+	{
+		UE_LOG(LogMutable, Display, TEXT("Compiling Customizable Object with %d LODBias."), CurrentOptions.DebugBias);
+	}
+
 	// Create and update compilation progress notification
 	const FText UpdateMsg = FText::FromString(FString::Printf(TEXT("Compiling Customizable Objects:\n%s"), *CurrentObject->GetName()));
 	if (!CompileNotificationHandle.IsValid())
@@ -864,31 +882,9 @@ void FCustomizableObjectCompiler::CompileInternal(bool bAsync)
 		return;
 	}
 
-
-	if (CurrentOptions.bIsCooking && CurrentOptions.TargetPlatform)
-	{
-		UE_LOG(LogMutable, Display, TEXT("Compiling Customizable Object %s for platform %s."), *CurrentObject->GetName(), *CurrentOptions.TargetPlatform->PlatformName());
-	}
-
-	if (CurrentOptions.bIsCooking && CurrentOptions.bForceLargeLODBias)
-	{
-		UE_LOG(LogMutable, Display, TEXT("Compiling Customizable Object with %d LODBias."), CurrentOptions.DebugBias);
-	}
-
 	FMutableGraphGenerationContext GenerationContext(CurrentObject, this, CurrentOptions);
 	GenerationContext.ParamNamesToSelectedOptions = CurrentRequest->GetParameterNamesToSelectedOptions();
 
-	// If we don't have the target platform yet (in editor) we need to get it
-	if (!GenerationContext.Options.TargetPlatform)
-	{
-		// Set the target platform in the context. For now it is the current platform.
-		ITargetPlatformManagerModule* TPM = GetTargetPlatformManager();
-		check(TPM);
-
-		GenerationContext.Options.TargetPlatform = TPM->GetRunningTargetPlatform();
-		check(GenerationContext.Options.TargetPlatform != nullptr);
-	}
-
 	// Clear Messages from previous Compilations
 	CompilationLogsContainer.ClearMessageCounters();
 	CompilationLogsContainer.ClearMessagesArray();
@@ -1223,6 +1219,9 @@ void FCustomizableObjectCompiler::CompleteRequest(ECompilationStatePrivate State
 		{
 			CurrentObject->GetPrivate()->PostCompile();
 		}
+
+		UE_LOG(LogMutable, Display, TEXT("Finished compiling Customizable Object %s. Compilation took %5.3f seconds to complete."),
+			*CurrentObject->GetName(), FPlatformTime::Seconds() - CompilationStartTime);
 	}
 
 	// Remove referenced objects
diff --git a/Engine/Plugins/Experimental/Mutable/Source/CustomizableObjectEditor/Private/MuCOE/CustomizableObjectCompiler.h b/Engine/Plugins/Experimental/Mutable/Source/CustomizableObjectEditor/Private/MuCOE/CustomizableObjectCompiler.h
--- a/Engine/Plugins/Experimental/Mutable/Source/CustomizableObjectEditor/Private/MuCOE/CustomizableObjectCompiler.h
+++ b/Engine/Plugins/Experimental/Mutable/Source/CustomizableObjectEditor/Private/MuCOE/CustomizableObjectCompiler.h
@@ -190,4 +190,7 @@ private:
 	/** Compilation progress bar handle */
 	FProgressNotificationHandle CompileNotificationHandle;
 
+	/** Compilation start time in seconds. */
+	double CompilationStartTime = 0;
+
 };
