[mutable] Added child object node options to the parameters selector, if they are linked to a toggle group node.

#rnx

diff --git a/Engine/Plugins/Experimental/Mutable/Source/CustomizableObjectEditor/Private/MuCOE/Nodes/CustomizableObjectNodeObjectDetails.cpp b/Engine/Plugins/Experimental/Mutable/Source/CustomizableObjectEditor/Private/MuCOE/Nodes/CustomizableObjectNodeObjectDetails.cpp
--- a/Engine/Plugins/Experimental/Mutable/Source/CustomizableObjectEditor/Private/MuCOE/Nodes/CustomizableObjectNodeObjectDetails.cpp
+++ b/Engine/Plugins/Experimental/Mutable/Source/CustomizableObjectEditor/Private/MuCOE/Nodes/CustomizableObjectNodeObjectDetails.cpp
@@ -175,8 +175,6 @@ void FCustomizableObjectNodeObjectDetails::CustomizeDetails( IDetailLayoutBuilde
 
 	if (BaseObjectNode.IsValid())
 	{
-		FillParameterNamesArray();
-
 		// Properties
 		TSharedRef<IPropertyHandle> StatesProperty = DetailBuilder.GetProperty("States");
 		TSharedRef<IPropertyHandle> ParentObjectProperty = DetailBuilder.GetProperty("ParentObject");
@@ -201,6 +199,8 @@ void FCustomizableObjectNodeObjectDetails::CustomizeDetails( IDetailLayoutBuilde
 
 		if (BaseObjectNode->bIsBase)
 		{
+			FillParameterNamesArray();
+
 			ExternalCategory.AddCustomRow(LOCTEXT("FCustomizableObjectNodeObjectDetails", "Blocks"))
 			[
 				SNew(SObjectPropertyEntryBox)
@@ -579,18 +579,14 @@ void FCustomizableObjectNodeObjectDetails::FillParameterNamesArray()
 
 	for (const UCustomizableObject* Object : CustomObjectTree)
 	{
-		if (!Object || !Object->GetPrivate()->GetSource())
-		{
-			return;
-		}
-
-		UEdGraph* Source = Object->GetPrivate()->GetSource();
-
+		// Cheking Private to avoid a static analysis waring
 		if (!Object || !Object->GetPrivate() || !Object->GetPrivate()->GetSource())
 		{
 			continue;
 		}
 
+		UEdGraph* Source = Object->GetPrivate()->GetSource();
+
 		// All type of parameter nodes
 		TArray<UCustomizableObjectNodeColorParameter*> ColorParameterNodes;
 		Source->GetNodesOfClass(ColorParameterNodes);
@@ -703,7 +699,27 @@ void FCustomizableObjectNodeObjectDetails::FillParameterNamesArray()
 	// Now that we know all the group objects of type toggle, we process all the object nodes that can generate a parameter
 	for (const UCustomizableObjectNodeObject* ObjectNode : AllObjectNodes)
 	{
-		if (ToggleGroupObjectIds.Contains(ObjectNode->ParentObjectGroupId))
+		FGuid ParentObjectGroupId = FGuid();
+
+		if (ObjectNode->bIsBase)
+		{
+			ParentObjectGroupId = ObjectNode->ParentObjectGroupId;
+		}
+		else
+		{
+			if (const UEdGraphPin* ObjectPin = ObjectNode->OutputPin())
+			{
+				if (const UEdGraphPin* GroupPin = FollowOutputPin(*ObjectPin))
+				{
+					if (const UCustomizableObjectNodeObjectGroup* GroupNode = Cast<UCustomizableObjectNodeObjectGroup>(GroupPin->GetOwningNode()))
+					{
+						ParentObjectGroupId = GroupNode->NodeGuid;
+					}
+				}
+			}
+		}
+
+		if (ToggleGroupObjectIds.Contains(ParentObjectGroupId))
 		{
 			BaseObjectNode->ParameterNames.Add(ObjectNode->ObjectName);
 		}
