[mutable] Added logs to expose the RHI device being used

#tests LocalExec
#rnx

diff --git a/Engine/Plugins/Experimental/Mutable/Source/MutableValidation/MutableValidation.Build.cs b/Engine/Plugins/Experimental/Mutable/Source/MutableValidation/MutableValidation.Build.cs
--- a/Engine/Plugins/Experimental/Mutable/Source/MutableValidation/MutableValidation.Build.cs
+++ b/Engine/Plugins/Experimental/Mutable/Source/MutableValidation/MutableValidation.Build.cs
@@ -17,7 +17,8 @@ namespace UnrealBuildTool.Rules
 				new string[] {
 					"Settings", 
 					"Engine", 
-					"MutableRuntime"
+					"MutableRuntime",
+					"RHI"
 				});
 			
 			PrivateDependencyModuleNames.AddRange(
diff --git a/Engine/Plugins/Experimental/Mutable/Source/MutableValidation/Private/MuV/COIBakingTestCommandlet.cpp b/Engine/Plugins/Experimental/Mutable/Source/MutableValidation/Private/MuV/COIBakingTestCommandlet.cpp
--- a/Engine/Plugins/Experimental/Mutable/Source/MutableValidation/Private/MuV/COIBakingTestCommandlet.cpp
+++ b/Engine/Plugins/Experimental/Mutable/Source/MutableValidation/Private/MuV/COIBakingTestCommandlet.cpp
@@ -72,7 +72,7 @@ int32 UCOIBakingTestCommandlet::Main(const FString& Params)
 	// Make sure there is nothing else that the engine needs to do before starting our test
 	Wait(60);
 
-	LogMutableSettings();
+	LogGlobalSettings();
 	
 	// Compile it's CO (using current config)
 	UCustomizableObject* InstanceCustomizableObject = TargetInstance->GetCustomizableObject();
diff --git a/Engine/Plugins/Experimental/Mutable/Source/MutableValidation/Private/MuV/COIBulkUpdateTestCommandlet.cpp b/Engine/Plugins/Experimental/Mutable/Source/MutableValidation/Private/MuV/COIBulkUpdateTestCommandlet.cpp
--- a/Engine/Plugins/Experimental/Mutable/Source/MutableValidation/Private/MuV/COIBulkUpdateTestCommandlet.cpp
+++ b/Engine/Plugins/Experimental/Mutable/Source/MutableValidation/Private/MuV/COIBulkUpdateTestCommandlet.cpp
@@ -38,7 +38,7 @@ int32 UCOIBulkUpdateTestCommandlet::Main(const FString& Params)
 	// Load the asset registry system so we can proceed without issues
 	PrepareAssetRegistry();
 
-	LogMutableSettings();
+	LogGlobalSettings();
 	
 	// Cache all UAssets (find a way to not have them in memory yet, not until we need them)
 	TArray<FAssetData> FoundAssetData;
diff --git a/Engine/Plugins/Experimental/Mutable/Source/MutableValidation/Private/MuV/CustomizableObjectValidationCommandlet.cpp b/Engine/Plugins/Experimental/Mutable/Source/MutableValidation/Private/MuV/CustomizableObjectValidationCommandlet.cpp
--- a/Engine/Plugins/Experimental/Mutable/Source/MutableValidation/Private/MuV/CustomizableObjectValidationCommandlet.cpp
+++ b/Engine/Plugins/Experimental/Mutable/Source/MutableValidation/Private/MuV/CustomizableObjectValidationCommandlet.cpp
@@ -74,7 +74,7 @@ int32 UCustomizableObjectValidationCommandlet::Main(const FString& Params)
 	// Make sure there is nothing else that the engine needs to do before starting our test
 	Wait(60);
 
-	LogMutableSettings();
+	LogGlobalSettings();
 	
 	// Compile the Customizable Object ------------------------------------------------------------------------------ //
 	bool bWasCoCompilationSuccessful = false;
diff --git a/Engine/Plugins/Experimental/Mutable/Source/MutableValidation/Private/MuV/ValidationUtils.cpp b/Engine/Plugins/Experimental/Mutable/Source/MutableValidation/Private/MuV/ValidationUtils.cpp
--- a/Engine/Plugins/Experimental/Mutable/Source/MutableValidation/Private/MuV/ValidationUtils.cpp
+++ b/Engine/Plugins/Experimental/Mutable/Source/MutableValidation/Private/MuV/ValidationUtils.cpp
@@ -6,7 +6,7 @@
 #include "Commandlets/Commandlet.h"
 #include "MuCO/CustomizableObject.h"
 #include "MuCO/CustomizableObjectSystem.h"
-#include "Interfaces/ITargetPlatform.h"
+#include "RHIGlobals.h"
 
 void PrepareAssetRegistry()
 {
@@ -22,13 +22,17 @@ void PrepareAssetRegistry()
 }
 
 
-void LogMutableSettings()
+void LogGlobalSettings()
 {
+	// Mutable Settings
 	const int32 WorkingMemory = UCustomizableObjectSystem::GetInstanceChecked()->GetWorkingMemory() ;
 	UE_LOG(LogMutable,Log, TEXT("(int) working_memory_bytes : %d"), WorkingMemory * 1024)
 	UE_LOG(LogMutable, Display, TEXT("The mutable updates will use as working memory the value of %d KB"), WorkingMemory)
 	
 	// Expand this when adding new controls from the .xml file
+	
+	// RHI Settings
+	UE_LOG(LogMutable, Log, TEXT("(string) rhi_adapter_name : %s"), *GRHIAdapterName )
 }
 
 
diff --git a/Engine/Plugins/Experimental/Mutable/Source/MutableValidation/Private/MuV/ValidationUtils.h b/Engine/Plugins/Experimental/Mutable/Source/MutableValidation/Private/MuV/ValidationUtils.h
--- a/Engine/Plugins/Experimental/Mutable/Source/MutableValidation/Private/MuV/ValidationUtils.h
+++ b/Engine/Plugins/Experimental/Mutable/Source/MutableValidation/Private/MuV/ValidationUtils.h
@@ -19,4 +19,5 @@ void Wait(const double ToWaitSeconds);
  * Isolate tests using different configurations.
  * @note Add new logs each time you add a way to change the configuration of the test from the .xml testing file
  */
-void LogMutableSettings();
+void LogGlobalSettings();
+
