[Mutable] Removed warning when an int parameter had no options.

This warning could be easely seen when compiling Child COs.

- Fixed incorrect UI message.

#rnx

diff --git a/Engine/Plugins/Experimental/Mutable/Source/CustomizableObject/Private/MuCO/CustomizableObject.cpp b/Engine/Plugins/Experimental/Mutable/Source/CustomizableObject/Private/MuCO/CustomizableObject.cpp
--- a/Engine/Plugins/Experimental/Mutable/Source/CustomizableObject/Private/MuCO/CustomizableObject.cpp
+++ b/Engine/Plugins/Experimental/Mutable/Source/CustomizableObject/Private/MuCO/CustomizableObject.cpp
@@ -1528,7 +1528,8 @@ int32 UCustomizableObject::FindIntParameterValue(int32 ParamIndex, const FString
 
 int32 UCustomizableObjectPrivate::FindIntParameterValue(int32 ParamIndex, const FString& Value) const
 {
-	int32 MinValueIndex = 0;
+	int32 MinValueIndex = INDEX_NONE;
+	
 	if (ParamIndex >= 0 && ParamIndex < ParameterProperties.Num())
 	{
 		const TArray<FMutableModelParameterValue>& PossibleValues = ParameterProperties[ParamIndex].PossibleValues;
@@ -1548,15 +1549,8 @@ int32 UCustomizableObjectPrivate::FindIntParameterValue(int32 ParamIndex, const
 				}
 			}
 		}
-		else
-		{
-			UE_LOG(LogMutable, Warning, TEXT("No possible values for parameter with index [%d] at FindIntParameterValue at CO %s."), ParamIndex, *GetName());
-		}
-	}
-	else
-	{
-		UE_LOG(LogMutable, Warning, TEXT("Index [%d] out of ParameterProperties bounds at FindIntParameterValue at CO %s."), ParamIndex, *GetName());
 	}
+	
 	return MinValueIndex;
 }
 
diff --git a/Engine/Plugins/Experimental/Mutable/Source/CustomizableObjectEditor/Private/MuCOE/CustomizableInstanceDetails.cpp b/Engine/Plugins/Experimental/Mutable/Source/CustomizableObjectEditor/Private/MuCOE/CustomizableInstanceDetails.cpp
--- a/Engine/Plugins/Experimental/Mutable/Source/CustomizableObjectEditor/Private/MuCOE/CustomizableInstanceDetails.cpp
+++ b/Engine/Plugins/Experimental/Mutable/Source/CustomizableObjectEditor/Private/MuCOE/CustomizableInstanceDetails.cpp
@@ -930,61 +930,65 @@ TSharedRef<SWidget> FCustomizableInstanceDetails::GenerateParameterWidget(const
 TSharedRef<SWidget> FCustomizableInstanceDetails::GenerateIntWidget(const int32 ParamIndexInObject)
 {
 	UCustomizableObject* CustomizableObject = CustomInstance->GetCustomizableObject();
+	
+	const bool bMultidimensional = CustomInstance->GetCustomizableObject()->IsParameterMultidimensional(ParamIndexInObject);
+	if (bMultidimensional)
+	{
+		return SNew(STextBlock).Text(LOCTEXT("MultidimensionalINTParameter_Text", "Multidimensional INT Parameter not supported"));
+	}
 
-	int32 numValues = CustomizableObject->GetIntParameterNumOptions(ParamIndexInObject);
-	bool bIsParamMultidimensional = CustomInstance->GetCustomizableObject()->IsParameterMultidimensional((ParamIndexInObject));
-
-	if (!bIsParamMultidimensional && numValues)
+	const int32 NumValues = CustomizableObject->GetIntParameterNumOptions(ParamIndexInObject) != 0;
+	if (!NumValues)
 	{
-		FString ToolTipText = FString("None");
-		FString ParamName = CustomizableObject->GetParameterName(ParamIndexInObject);
-		FString Value = CustomInstance->GetIntParameterSelectedOption(ParamName, -1);
+		return SNew(STextBlock).Text(LOCTEXT("NoAvailableOptions", "No Available Options"));
+	}
+	
+	FString ToolTipText = FString("None");
+	FString ParamName = CustomizableObject->GetParameterName(ParamIndexInObject);
+	FString Value = CustomInstance->GetIntParameterSelectedOption(ParamName, -1);
 
-		TSharedPtr<TArray<TSharedPtr<FString>>>* FoundOptions = IntParameterOptions.Find(ParamIndexInObject);
-		TArray<TSharedPtr<FString>>& OptionNamesAttribute = FoundOptions && FoundOptions->IsValid() ? 
-															*FoundOptions->Get() :
-															*(IntParameterOptions.Add(ParamIndexInObject, MakeShared<TArray<TSharedPtr<FString>>>()).Get());
+	TSharedPtr<TArray<TSharedPtr<FString>>>* FoundOptions = IntParameterOptions.Find(ParamIndexInObject);
+	TArray<TSharedPtr<FString>>& OptionNamesAttribute = FoundOptions && FoundOptions->IsValid() ? 
+														*FoundOptions->Get() :
+														*(IntParameterOptions.Add(ParamIndexInObject, MakeShared<TArray<TSharedPtr<FString>>>()).Get());
 
-		OptionNamesAttribute.Empty();
+	OptionNamesAttribute.Empty();
 
-		int32 ValueIndex = 0;
+	int32 ValueIndex = 0;
+
+	for (int32 i = 0; i < NumValues; ++i)
+	{
+		FString PossibleValue = CustomizableObject->GetIntParameterAvailableOption(ParamIndexInObject, i);
 
-		for (int32 i = 0; i < numValues; ++i)
+		if (PossibleValue == Value)
 		{
-			FString PossibleValue = CustomizableObject->GetIntParameterAvailableOption(ParamIndexInObject, i);
+			ValueIndex = i;
 
-			if (PossibleValue == Value)
+			const FString* Identifier = CustomizableObject->GetPrivate()->GroupNodeMap.FindKey(FCustomizableObjectIdPair(ParamName, PossibleValue));
+			if (Identifier)
 			{
-				ValueIndex = i;
-
-				const FString* Identifier = CustomizableObject->GetPrivate()->GroupNodeMap.FindKey(FCustomizableObjectIdPair(ParamName, PossibleValue));
-				if (Identifier)
+				if (FString* CustomizableObjectPath = CustomizableObject->GetPrivate()->CustomizableObjectPathMap.Find(*Identifier))
 				{
-					if (FString* CustomizableObjectPath = CustomizableObject->GetPrivate()->CustomizableObjectPathMap.Find(*Identifier))
-					{
-						ToolTipText = *CustomizableObjectPath;
-					}
+					ToolTipText = *CustomizableObjectPath;
 				}
 			}
-
-			OptionNamesAttribute.Add(TSharedPtr<FString>(new FString(CustomizableObject->GetIntParameterAvailableOption(ParamIndexInObject, i))));
 		}
 
-		return SNew(SSearchableComboBox)
-			.ToolTipText(FText::FromString(ToolTipText))
-			.OptionsSource(&OptionNamesAttribute)
-			.InitiallySelectedItem(OptionNamesAttribute[ValueIndex])
-			.Method(EPopupMethod::UseCurrentWindow)
-			.OnSelectionChanged(this, &FCustomizableInstanceDetails::OnIntParameterComboBoxChanged, ParamName)
-			.OnGenerateWidget(this, &FCustomizableInstanceDetails::OnGenerateWidgetIntParameter)
-			.Content()
-			[
-				SNew(STextBlock)
-					.Text(FText::FromString(*OptionNamesAttribute[ValueIndex]))
-			];
+		OptionNamesAttribute.Add(TSharedPtr<FString>(new FString(CustomizableObject->GetIntParameterAvailableOption(ParamIndexInObject, i))));
 	}
 
-	return SNew(STextBlock).Text(LOCTEXT("MultidimensionalINTParameter_Text", "Multidimensional INT Parameter not supported"));
+	return SNew(SSearchableComboBox)
+		.ToolTipText(FText::FromString(ToolTipText))
+		.OptionsSource(&OptionNamesAttribute)
+		.InitiallySelectedItem(OptionNamesAttribute[ValueIndex])
+		.Method(EPopupMethod::UseCurrentWindow)
+		.OnSelectionChanged(this, &FCustomizableInstanceDetails::OnIntParameterComboBoxChanged, ParamName)
+		.OnGenerateWidget(this, &FCustomizableInstanceDetails::OnGenerateWidgetIntParameter)
+		.Content()
+		[
+			SNew(STextBlock)
+				.Text(FText::FromString(*OptionNamesAttribute[ValueIndex]))
+		];
 }
 
 
