[Mutable] Merge UCustomizableObject FCompileOptions with member options.

- All compilation options now are members of the UCustomizableObject or its private.
- Now FCompileOptions will be created on the fly, gathering its options from the different places (CO members + inis).

#rnx

diff --git a/Engine/Plugins/Experimental/Mutable/Config/DefaultMutable.ini b/Engine/Plugins/Experimental/Mutable/Config/DefaultMutable.ini
--- a/Engine/Plugins/Experimental/Mutable/Config/DefaultMutable.ini
+++ b/Engine/Plugins/Experimental/Mutable/Config/DefaultMutable.ini
@@ -2,6 +2,7 @@
 
 [CoreRedirects]
 +StructRedirects=(OldName="/Script/CustomizableObject.ParameterUIData",NewName="/Script/CustomizableObject.MutableParameterData")
++StructRedirects=(OldName="/Script/CustomizableObject.CompilationOptions",NewName="/Script/CustomizableObject.CompilationOptions_DEPRECATED")
 
 [Features]
 
diff --git a/Engine/Plugins/Experimental/Mutable/Source/CustomizableObject/Internal/MuCO/CustomizableObjectCompilerTypes.h b/Engine/Plugins/Experimental/Mutable/Source/CustomizableObject/Internal/MuCO/CustomizableObjectCompilerTypes.h
--- a/Engine/Plugins/Experimental/Mutable/Source/CustomizableObject/Internal/MuCO/CustomizableObjectCompilerTypes.h
+++ b/Engine/Plugins/Experimental/Mutable/Source/CustomizableObject/Internal/MuCO/CustomizableObjectCompilerTypes.h
@@ -2,6 +2,8 @@
 
 #pragma once
 
+#include "UObject/WeakObjectPtr.h"
+
 #include "CustomizableObjectCompilerTypes.generated.h"
 
 
@@ -37,34 +39,25 @@ enum class ECustomizableObjectTextureCompression : uint8
 };
 
 
-USTRUCT()
 struct FCompilationOptions
 {
-	GENERATED_USTRUCT_BODY()
-
 	/** Enum to know what texture compression should be used. This compression is used only in manual compiles in editor.
-	 *  When packaging, ECustomizableObjectTextureCompression::HighQuality is always used.
-	 */
-	UPROPERTY()
+	 *  When packaging, ECustomizableObjectTextureCompression::HighQuality is always used. */
 	ECustomizableObjectTextureCompression TextureCompression = ECustomizableObjectTextureCompression::Fast;
 
 	// From 0 to UE_MUTABLE_MAX_OPTIMIZATION
-	UPROPERTY()
 	int32 OptimizationLevel = UE_MUTABLE_MAX_OPTIMIZATION;
 
 	// Use the disk to store intermediate compilation data. This slows down the object compilation
 	// but it may be necessary for huge objects.
-	UPROPERTY()
 	bool bUseDiskCompilation = false;
 
 	/** High limit of the size in bytes of the packaged data when cooking this object.
 	* This limit is before any pak or filesystem compression. This limit will be broken if a single piece of data is bigger because data is not fragmented for packaging purposes.
 	*/
-	UPROPERTY()
 	uint64 PackagedDataBytesLimit = 256 * 1024 * 1024;
 
 	/** High (inclusive) limit of the size in bytes of a data block to be included into the compiled object directly instead of stored in a streamable file. */
-	UPROPERTY()
 	uint64 EmbeddedDataBytesLimit = 1024;
 
 	// Did we have the extra bones enabled when we compiled?
@@ -106,7 +99,6 @@ struct FCompilationOptions
 
 	// Control image tiled generation
 	int32 ImageTiling = 0;
-
 };
 
 
@@ -127,6 +119,7 @@ enum class ECompilationResultPrivate : uint8
 };
 
 
+#if WITH_EDITOR
 struct CUSTOMIZABLEOBJECT_API FCompilationRequest
 {
 	FCompilationRequest(UCustomizableObject& CustomizableObject, bool bAsync);
@@ -167,5 +160,4 @@ private:
 	TArray<FText> Warnings;
 	TArray<FText> Errors;
 };
-
-//#endif // WITH_EDITORONLY_DATA
\ No newline at end of file
+#endif
diff --git a/Engine/Plugins/Experimental/Mutable/Source/CustomizableObject/Internal/MuCO/CustomizableObjectCustomVersion.h b/Engine/Plugins/Experimental/Mutable/Source/CustomizableObject/Internal/MuCO/CustomizableObjectCustomVersion.h
--- a/Engine/Plugins/Experimental/Mutable/Source/CustomizableObject/Internal/MuCO/CustomizableObjectCustomVersion.h
+++ b/Engine/Plugins/Experimental/Mutable/Source/CustomizableObject/Internal/MuCO/CustomizableObjectCustomVersion.h
@@ -118,6 +118,8 @@ struct CUSTOMIZABLEOBJECT_API FCustomizableObjectCustomVersion
 
 		AddedRenameOptionToParameterNodes,
 
+		CompilationOptions,
+
 		// -----<new versions can be added above this line>-------------------------------------------------
 		VersionPlusOne,
 		LatestVersion = VersionPlusOne - 1
diff --git a/Engine/Plugins/Experimental/Mutable/Source/CustomizableObject/Internal/MuCO/CustomizableObjectPrivate.h b/Engine/Plugins/Experimental/Mutable/Source/CustomizableObject/Internal/MuCO/CustomizableObjectPrivate.h
--- a/Engine/Plugins/Experimental/Mutable/Source/CustomizableObject/Internal/MuCO/CustomizableObjectPrivate.h
+++ b/Engine/Plugins/Experimental/Mutable/Source/CustomizableObject/Internal/MuCO/CustomizableObjectPrivate.h
@@ -918,6 +918,8 @@ public:
 
 #if WITH_EDITORONLY_DATA
 	TObjectPtr<UEdGraph>& GetSource() const;
+
+	FCompilationOptions GetCompileOptions() const;
 #endif
 	
 	/** Cache of generated SkeletalMeshes */
@@ -981,6 +983,32 @@ public:
 	// use this to lookup fast by Name
 	TMap<FString, FMutableParameterIndex> ParameterPropertiesLookupTable;
 
+#if WITH_EDITORONLY_DATA
+	UPROPERTY()
+	ECustomizableObjectTextureCompression TextureCompression = ECustomizableObjectTextureCompression::Fast;
+
+	/** From 0 to UE_MUTABLE_MAX_OPTIMIZATION */
+	UPROPERTY()
+	int32 OptimizationLevel = UE_MUTABLE_MAX_OPTIMIZATION;
+
+	/** Use the disk to store intermediate compilation data. This slows down the object compilation
+	 * but it may be necessary for huge objects. */
+	UPROPERTY()
+	bool bUseDiskCompilation = false;
+
+	/** High limit of the size in bytes of the packaged data when cooking this object.
+	 * This limit is before any pak or filesystem compression. This limit will be broken if a single piece of data is bigger because data is not fragmented for packaging purposes.	*/
+	UPROPERTY()
+	uint64 PackagedDataBytesLimit = 256 * 1024 * 1024;
+	
+	/** High (inclusive) limit of the size in bytes of a data block to be included into the compiled object directly instead of stored in a streamable file. */
+	UPROPERTY()
+	uint64 EmbeddedDataBytesLimit = 1024;
+
+	UPROPERTY()
+	int32 ImageTiling = 0;
+#endif
+	
 	// This is a manual version number for the binary blobs in this asset.
 	// Increasing it invalidates all the previously compiled models.
 	enum ECustomizableObjectVersions
diff --git a/Engine/Plugins/Experimental/Mutable/Source/CustomizableObject/Private/MuCO/CustomizableObject.cpp b/Engine/Plugins/Experimental/Mutable/Source/CustomizableObject/Private/MuCO/CustomizableObject.cpp
--- a/Engine/Plugins/Experimental/Mutable/Source/CustomizableObject/Private/MuCO/CustomizableObject.cpp
+++ b/Engine/Plugins/Experimental/Mutable/Source/CustomizableObject/Private/MuCO/CustomizableObject.cpp
@@ -238,6 +238,8 @@ void UCustomizableObject::PostLoad()
 {
 	Super::PostLoad();
 
+	const int32 CustomizableObjectCustomVersion = GetLinkerCustomVersion(FCustomizableObjectCustomVersion::GUID);
+
 #if WITH_EDITOR
 	if (ReferenceSkeletalMesh_DEPRECATED)
 	{
@@ -245,6 +247,17 @@ void UCustomizableObject::PostLoad()
 		ReferenceSkeletalMesh_DEPRECATED = nullptr;
 	}
 
+#if WITH_EDITORONLY_DATA
+	if (CustomizableObjectCustomVersion < FCustomizableObjectCustomVersion::CompilationOptions)
+	{
+		GetPrivate()->OptimizationLevel = CompileOptions_DEPRECATED.OptimizationLevel;
+		GetPrivate()->TextureCompression = CompileOptions_DEPRECATED.TextureCompression;
+		GetPrivate()->bUseDiskCompilation = CompileOptions_DEPRECATED.bUseDiskCompilation;
+		GetPrivate()->EmbeddedDataBytesLimit = CompileOptions_DEPRECATED.EmbeddedDataBytesLimit;
+		GetPrivate()->PackagedDataBytesLimit = CompileOptions_DEPRECATED.PackagedDataBytesLimit;
+	}
+#endif
+
 	// Register to dirty delegate so we update derived data version ID each time that the package is marked as dirty.
 	if (UPackage* Package = GetOutermost())
 	{
@@ -256,9 +269,7 @@ void UCustomizableObject::PostLoad()
 				}
 			});
 	}
-
-	const int32 CustomizableObjectCustomVersion = GetLinkerCustomVersion(FCustomizableObjectCustomVersion::GUID);
-
+	
 	if (!IsRunningCookCommandlet())
 	{
 		GetPrivate()->Status.NextState(FCustomizableObjectStatusTypes::EState::Loading);
@@ -292,6 +303,8 @@ void UCustomizableObject::Serialize(FArchive& Ar_Asset)
 	
 	Super::Serialize(Ar_Asset);
 
+	Ar_Asset.UsingCustomVersion(FCustomizableObjectCustomVersion::GUID);
+	
 #if WITH_EDITOR
 	if (Ar_Asset.IsCooking())
 	{
@@ -912,7 +925,6 @@ void UCustomizableObjectPrivate::CompileForTargetPlatform(const ITargetPlatform*
 	Options.TextureCompression = ECustomizableObjectTextureCompression::HighQuality;
 	Options.bIsCooking = true;
 	Options.TargetPlatform = TargetPlatform;
-	Options.CustomizableObjectNumBoneInfluences = ICustomizableObjectModule::Get().GetNumBoneInfluences();
 	CompileRequests.Add(CompileRequest);
 
 	EditorModule->CompileCustomizableObject(CompileRequest, true);
@@ -2111,6 +2123,27 @@ TObjectPtr<UEdGraph>& UCustomizableObjectPrivate::GetSource() const
 {
 	return GetPublic()->Source;
 }
+
+
+FCompilationOptions UCustomizableObjectPrivate::GetCompileOptions() const
+{
+	FCompilationOptions Options;
+	Options.TextureCompression = TextureCompression;
+	Options.OptimizationLevel = OptimizationLevel;
+	Options.bUseDiskCompilation = bUseDiskCompilation;
+	Options.PackagedDataBytesLimit = PackagedDataBytesLimit;
+	Options.EmbeddedDataBytesLimit = EmbeddedDataBytesLimit;
+	Options.CustomizableObjectNumBoneInfluences = ICustomizableObjectModule::Get().GetNumBoneInfluences();
+	Options.bRealTimeMorphTargetsEnabled = GetPublic()->bEnableRealTimeMorphTargets;
+	Options.bClothingEnabled = GetPublic()->bEnableClothing;
+	Options.b16BitBoneWeightsEnabled = GetPublic()->bEnable16BitBoneWeights;
+	Options.bSkinWeightProfilesEnabled = GetPublic()->bEnableAltSkinWeightProfiles;
+	Options.bPhysicsAssetMergeEnabled = GetPublic()->bEnablePhysicsAssetMerge;
+	Options.bAnimBpPhysicsManipulationEnabled = GetPublic()->bEnableAnimBpPhysicsAssetsManipualtion;
+	Options.ImageTiling = ImageTiling;
+	
+	return Options;
+}
 #endif
 
 
@@ -2222,7 +2255,7 @@ void UCustomizableObjectBulk::PrepareBulkData(UCustomizableObject* InOuter, cons
 
 	FModelResources& ModelResources = CustomizableObject->GetPrivate()->GetModelResources(true);
 
-	uint64 TargetBulkDataFileBytes = CustomizableObject->CompileOptions.PackagedDataBytesLimit;
+	uint64 TargetBulkDataFileBytes = CustomizableObject->GetPrivate()->GetCompileOptions().PackagedDataBytesLimit;
 	const uint64 MaxChunkSize = UCustomizableObjectSystem::GetInstance()->GetMaxChunkSizeForPlatform(TargetPlatform);
 	TargetBulkDataFileBytes = FMath::Min(TargetBulkDataFileBytes, MaxChunkSize);
 
@@ -2757,80 +2790,84 @@ FArchive& operator<<(FArchive& Ar, FMutableRefSkeletalMeshData& Data)
 
 	return Ar;
 }
+#endif // WITH_EDITORONLY_DATA
 
 
-//-------------------------------------------------------------------------------------------------
-//-------------------------------------------------------------------------------------------------
-//-------------------------------------------------------------------------------------------------
+#if WITH_EDITOR
 FCompilationRequest::FCompilationRequest(UCustomizableObject& InCustomizableObject, bool bAsyncCompile)
 {
 	CustomizableObject = &InCustomizableObject;
-	Options = InCustomizableObject.CompileOptions;
-	Options.bRealTimeMorphTargetsEnabled = InCustomizableObject.bEnableRealTimeMorphTargets;
-	Options.bClothingEnabled = InCustomizableObject.bEnableClothing;
-	Options.b16BitBoneWeightsEnabled = InCustomizableObject.bEnable16BitBoneWeights;
-	Options.bSkinWeightProfilesEnabled = InCustomizableObject.bEnableAltSkinWeightProfiles;
-	Options.bPhysicsAssetMergeEnabled = InCustomizableObject.bEnablePhysicsAssetMerge;
-	Options.bAnimBpPhysicsManipulationEnabled = InCustomizableObject.bEnableAnimBpPhysicsAssetsManipualtion;
+	Options = InCustomizableObject.GetPrivate()->GetCompileOptions();
 	bAsync = bAsyncCompile;
 }
 
+
 UCustomizableObject* FCompilationRequest::GetCustomizableObject()
 {
 	return CustomizableObject.Get();
 }
 
+
 FCompilationOptions& FCompilationRequest::GetCompileOptions()
 {
 	return Options;
 }
 
+
 bool FCompilationRequest::IsAsyncCompilation() const
 {
 	return bAsync;
 }
 
+
 void FCompilationRequest::SetCompilationState(ECompilationStatePrivate InState, ECompilationResultPrivate InResult)
 {
 	State = InState;
 	Result = InResult;
 }
 
+
 ECompilationStatePrivate FCompilationRequest::GetCompilationState() const
 {
 	return State;
 }
 
+
 ECompilationResultPrivate FCompilationRequest::GetCompilationResult() const
 {
 	return Result;
 }
 
+
 TArray<FText>& FCompilationRequest::GetWarnings()
 {
 	return Warnings;
 }
 
+
 TArray<FText>& FCompilationRequest::GetErrors()
 {
 	return Errors;
 }
 
+
 void FCompilationRequest::SetParameterNamesToSelectedOptions(const TMap<FString, FString>& InParamNamesToSelectedOptions)
 {
 	ParamNamesToSelectedOptions = InParamNamesToSelectedOptions;
 }
 
+
 const TMap<FString, FString>& FCompilationRequest::GetParameterNamesToSelectedOptions() const
 {
 	return ParamNamesToSelectedOptions;
 }
 
+
 bool FCompilationRequest::operator==(const FCompilationRequest& Other) const
 {
 	return CustomizableObject == Other.CustomizableObject && Options.TargetPlatform == Other.Options.TargetPlatform;
-};
+}
+#endif
 
-#endif // WITH_EDITORONLY_DATA
 
 #undef LOCTEXT_NAMESPACE
diff --git a/Engine/Plugins/Experimental/Mutable/Source/CustomizableObject/Public/MuCO/CustomizableObject.h b/Engine/Plugins/Experimental/Mutable/Source/CustomizableObject/Public/MuCO/CustomizableObject.h
--- a/Engine/Plugins/Experimental/Mutable/Source/CustomizableObject/Public/MuCO/CustomizableObject.h
+++ b/Engine/Plugins/Experimental/Mutable/Source/CustomizableObject/Public/MuCO/CustomizableObject.h
@@ -10,6 +10,7 @@
 #include "MuCO/CustomizableObjectStreamedResourceData.h"
 #include "MuCO/CustomizableObjectParameterTypeDefinitions.h"
 #include "MuCO/CustomizableObjectUIData.h"
+#include "MuCO/CustomizableObject_Deprecated.h"
 #include "Templates/TypeCompatibleBytes.h"
 
 #include "CustomizableObject.generated.h"
@@ -417,10 +418,11 @@ public:
 	bool bEnableMeshStreaming = false;
 	
 #if WITH_EDITORONLY_DATA
-	// Compilation options to use in editor and for packaging for this object.
+private:
 	UPROPERTY()
-	FCompilationOptions CompileOptions;
+	FCompilationOptions_DEPRECATED CompileOptions_DEPRECATED;
 
+public:
 	UPROPERTY(EditAnywhere, Category = CompileOptions)
 	bool bEnableRealTimeMorphTargets = false;
 
diff --git a/Engine/Plugins/Experimental/Mutable/Source/CustomizableObject/Public/MuCO/CustomizableObject_Deprecated.h b/Engine/Plugins/Experimental/Mutable/Source/CustomizableObject/Public/MuCO/CustomizableObject_Deprecated.h
new file mode 100644
--- /dev/null
+++ b/Engine/Plugins/Experimental/Mutable/Source/CustomizableObject/Public/MuCO/CustomizableObject_Deprecated.h
@@ -0,0 +1,29 @@
+﻿// Copyright Epic Games, Inc. All Rights Reserved.
+
+#pragma once
+
+#include "MuCO/CustomizableObjectCompilerTypes.h"
+
+#include "CustomizableObject_Deprecated.generated.h"
+
+
+USTRUCT()
+struct FCompilationOptions_DEPRECATED
+{
+	GENERATED_USTRUCT_BODY()
+
+	UPROPERTY()
+	ECustomizableObjectTextureCompression TextureCompression = ECustomizableObjectTextureCompression::Fast;
+
+	UPROPERTY()
+	int32 OptimizationLevel = UE_MUTABLE_MAX_OPTIMIZATION;
+
+	UPROPERTY()
+	bool bUseDiskCompilation = false;
+
+	UPROPERTY()
+	uint64 PackagedDataBytesLimit = 256 * 1024 * 1024;
+
+	UPROPERTY()
+	uint64 EmbeddedDataBytesLimit = 1024;
+};
\ No newline at end of file
diff --git a/Engine/Plugins/Experimental/Mutable/Source/CustomizableObjectEditor/Private/MuCOE/CustomizableObjectCompiler.cpp b/Engine/Plugins/Experimental/Mutable/Source/CustomizableObjectEditor/Private/MuCOE/CustomizableObjectCompiler.cpp
--- a/Engine/Plugins/Experimental/Mutable/Source/CustomizableObjectEditor/Private/MuCOE/CustomizableObjectCompiler.cpp
+++ b/Engine/Plugins/Experimental/Mutable/Source/CustomizableObjectEditor/Private/MuCOE/CustomizableObjectCompiler.cpp
@@ -44,12 +44,6 @@ class UTexture2D;
 UCustomizableObjectNodeObject* GetRootNode(UCustomizableObject* Object, bool &bOutMultipleBaseObjectsFound);
 
 
-FCustomizableObjectCompiler::FCustomizableObjectCompiler()
-{
-	CustomizableObjectNumBoneInfluences = ICustomizableObjectModule::Get().GetNumBoneInfluences();
-}
-
-
 bool FCustomizableObjectCompiler::Tick(bool bBlocking)
 {
 	MUTABLE_CPUPROFILER_SCOPE(FCustomizableObjectCompiler::Tick);
@@ -150,7 +144,6 @@ void FCustomizableObjectCompiler::Compile(const TSharedRef<FCompilationRequest>&
 	CurrentRequest = InCompileRequest.ToSharedPtr();
 	CurrentObject = InCompileRequest->GetCustomizableObject();
 	CurrentOptions = InCompileRequest->GetCompileOptions();
-	CurrentOptions.CustomizableObjectNumBoneInfluences = CustomizableObjectNumBoneInfluences;
 
 	if (!CurrentObject)
 	{
@@ -1295,8 +1288,6 @@ mu::NodePtr FCustomizableObjectCompiler::Export(UCustomizableObject* Object, con
 	FSlateNotificationManager::Get().AddNotification(Info);
 
 	FCompilationOptions CompilerOptions = InCompilerOptions;
-	CompilerOptions.CustomizableObjectNumBoneInfluences = CustomizableObjectNumBoneInfluences;
-
 	CompilerOptions.bRealTimeMorphTargetsEnabled = Object->bEnableRealTimeMorphTargets;
 	CompilerOptions.bClothingEnabled = Object->bEnableClothing;
 	CompilerOptions.b16BitBoneWeightsEnabled = Object->bEnable16BitBoneWeights;
diff --git a/Engine/Plugins/Experimental/Mutable/Source/CustomizableObjectEditor/Private/MuCOE/CustomizableObjectCompiler.h b/Engine/Plugins/Experimental/Mutable/Source/CustomizableObjectEditor/Private/MuCOE/CustomizableObjectCompiler.h
--- a/Engine/Plugins/Experimental/Mutable/Source/CustomizableObjectEditor/Private/MuCOE/CustomizableObjectCompiler.h
+++ b/Engine/Plugins/Experimental/Mutable/Source/CustomizableObjectEditor/Private/MuCOE/CustomizableObjectCompiler.h
@@ -25,8 +25,6 @@ class UCustomizableObjectNode;
 class FCustomizableObjectCompiler : public FTickableEditorObject, public FTickableCookObject, public FGCObject
 {
 public:
-
-	FCustomizableObjectCompiler();
 	virtual ~FCustomizableObjectCompiler() override {}
 
 	/** Check for pending compilation process. Returns true if an object has been updated. */
@@ -160,9 +158,6 @@ private:
 	TSharedPtr<FCustomizableObjectSaveDDRunnable> SaveDDTask;
 	TSharedPtr<FRunnableThread> SaveDDThread;
 
-	// Cache configuration settings from ini files 
-	ECustomizableObjectNumBoneInfluences CustomizableObjectNumBoneInfluences = ECustomizableObjectNumBoneInfluences::Four;
-
 	/** Array where to put the names of the already processed child in ProcessChildObjectsRecursively */
 	TArray<FName> ArrayAlreadyProcessedChild;
 
diff --git a/Engine/Plugins/Experimental/Mutable/Source/CustomizableObjectEditor/Private/MuCOE/CustomizableObjectEditor.cpp b/Engine/Plugins/Experimental/Mutable/Source/CustomizableObjectEditor/Private/MuCOE/CustomizableObjectEditor.cpp
--- a/Engine/Plugins/Experimental/Mutable/Source/CustomizableObjectEditor/Private/MuCOE/CustomizableObjectEditor.cpp
+++ b/Engine/Plugins/Experimental/Mutable/Source/CustomizableObjectEditor/Private/MuCOE/CustomizableObjectEditor.cpp
@@ -1248,7 +1248,7 @@ TSharedRef<SWidget> FCustomizableObjectEditor::GenerateCompileOptionsMenuContent
 
 		if (CustomizableObject)
 		{
-			int32 SelectedOptimization = FMath::Clamp(CustomizableObject->CompileOptions.OptimizationLevel, 0, CompileOptimizationStrings.Num() - 1);
+			int32 SelectedOptimization = FMath::Clamp(CustomizableObject->GetPrivate()->OptimizationLevel, 0, CompileOptimizationStrings.Num() - 1);
 			CompileOptimizationCombo =
 				SNew(STextComboBox)
 				.OptionsSource(&CompileOptimizationStrings)
@@ -1265,7 +1265,7 @@ TSharedRef<SWidget> FCustomizableObjectEditor::GenerateCompileOptionsMenuContent
 			CompileTextureCompressionStrings.Add(MakeShareable(new FString(LOCTEXT("MutableTextureCompressionFast", "Fast").ToString())));
 			CompileTextureCompressionStrings.Add(MakeShareable(new FString(LOCTEXT("MutableTextureCompressionHighQuality", "High Quality").ToString())));
 
-			int32 SelectedCompression = FMath::Clamp(int32(CustomizableObject->CompileOptions.TextureCompression), 0, CompileTextureCompressionStrings.Num() - 1);
+			int32 SelectedCompression = FMath::Clamp(int32(CustomizableObject->GetPrivate()->TextureCompression), 0, CompileTextureCompressionStrings.Num() - 1);
 			CompileTextureCompressionCombo =
 				SNew(STextComboBox)
 				.OptionsSource(&CompileTextureCompressionStrings)
@@ -1289,13 +1289,13 @@ TSharedRef<SWidget> FCustomizableObjectEditor::GenerateCompileOptionsMenuContent
 			.DropDownValues(TilingOptions)
 			.Value_Lambda([this]() 
 				{ 
-					return CustomizableObject ? float(CustomizableObject->CompileOptions.ImageTiling) : 0.0f;
+					return CustomizableObject ? float(CustomizableObject->GetPrivate()->ImageTiling) : 0.0f;
 				})
 			.OnValueChanged_Lambda([this](float Value) 
 				{ 
 					if (CustomizableObject)
 					{
-						CustomizableObject->CompileOptions.ImageTiling = int32(Value);
+						CustomizableObject->GetPrivate()->ImageTiling = int32(Value);
 						CustomizableObject->Modify();
 					}
 				});
@@ -1321,13 +1321,13 @@ TSharedRef<SWidget> FCustomizableObjectEditor::GenerateCompileOptionsMenuContent
 			.DropDownValues(EmbeddedOptions)
 			.Value_Lambda([this]()
 				{
-					return CustomizableObject ? float(CustomizableObject->CompileOptions.EmbeddedDataBytesLimit) : 0.0f;
+					return CustomizableObject ? float(CustomizableObject->GetPrivate()->EmbeddedDataBytesLimit) : 0.0f;
 				})
 			.OnValueChanged_Lambda([this](float Value)
 				{
 					if (CustomizableObject)
 					{
-						CustomizableObject->CompileOptions.EmbeddedDataBytesLimit = uint64(Value);
+						CustomizableObject->GetPrivate()->EmbeddedDataBytesLimit = uint64(Value);
 						CustomizableObject->Modify();
 					}
 				});
@@ -1347,13 +1347,13 @@ TSharedRef<SWidget> FCustomizableObjectEditor::GenerateCompileOptionsMenuContent
 			.DropDownValues(PackagedOptions)
 			.Value_Lambda([this]()
 				{
-					return CustomizableObject ? float(CustomizableObject->CompileOptions.PackagedDataBytesLimit) : 0.0f;
+					return CustomizableObject ? float(CustomizableObject->GetPrivate()->PackagedDataBytesLimit) : 0.0f;
 				})
 			.OnValueChanged_Lambda([this](float Value)
 				{
 					if (CustomizableObject)
 					{
-						CustomizableObject->CompileOptions.PackagedDataBytesLimit = uint64(Value);
+						CustomizableObject->GetPrivate()->PackagedDataBytesLimit = uint64(Value);
 						CustomizableObject->Modify();
 					}
 				});
@@ -1523,14 +1523,21 @@ void FCustomizableObjectEditor::ResetCompileOptions()
 {
 	const FScopedTransaction Transaction(LOCTEXT("ResetCompilationOptionsTransaction", "Reset Compilation Options"));
 	CustomizableObject->Modify();
-	CustomizableObject->CompileOptions = FCompilationOptions();
+	
+	UCustomizableObjectPrivate* DefaultObject = Cast<UCustomizableObjectPrivate>(CustomizableObject->GetPrivate()->StaticClass()->GetDefaultObject());
+	CustomizableObject->GetPrivate()->OptimizationLevel = DefaultObject->OptimizationLevel;
+	CustomizableObject->GetPrivate()->bUseDiskCompilation = DefaultObject->bUseDiskCompilation;
+	CustomizableObject->GetPrivate()->TextureCompression = DefaultObject->TextureCompression;
+	CustomizableObject->GetPrivate()->EmbeddedDataBytesLimit = DefaultObject->EmbeddedDataBytesLimit;
+	CustomizableObject->GetPrivate()->PackagedDataBytesLimit = DefaultObject->PackagedDataBytesLimit;
+	CustomizableObject->GetPrivate()->ImageTiling = DefaultObject->ImageTiling;
 }
 
 void FCustomizableObjectEditor::OnChangeCompileOptimizationLevel(TSharedPtr<FString> NewSelection, ESelectInfo::Type SelectInfo)
 {
 	const FScopedTransaction Transaction(LOCTEXT("ChangedOptimizationLevelTransaction", "Changed Optimization Level"));
 	CustomizableObject->Modify();
-	CustomizableObject->CompileOptions.OptimizationLevel = CompileOptimizationStrings.Find(NewSelection);
+	CustomizableObject->GetPrivate()->OptimizationLevel = CompileOptimizationStrings.Find(NewSelection);
 }
 
 
@@ -1538,13 +1545,13 @@ void FCustomizableObjectEditor::CompileOptions_UseDiskCompilation_Toggled()
 {
 	const FScopedTransaction Transaction(LOCTEXT("ChangedEnableCompilingUsingTheDiskAsMemory", "Changed Enable compiling using the disk as memory"));
 	CustomizableObject->Modify();
-	CustomizableObject->CompileOptions.bUseDiskCompilation = !CustomizableObject->CompileOptions.bUseDiskCompilation;
+	CustomizableObject->GetPrivate()->bUseDiskCompilation = !CustomizableObject->GetPrivate()->bUseDiskCompilation;
 }
 
 
 bool FCustomizableObjectEditor::CompileOptions_UseDiskCompilation_IsChecked()
 {
-	return CustomizableObject->CompileOptions.bUseDiskCompilation;
+	return CustomizableObject->GetPrivate()->bUseDiskCompilation;
 }
 
 
@@ -1552,7 +1559,7 @@ void FCustomizableObjectEditor::OnChangeCompileTextureCompressionType(TSharedPtr
 {
 	const FScopedTransaction Transaction(LOCTEXT("ChangedTextureCompressionTransaction", "Changed Texture Compression Type"));
 	CustomizableObject->Modify();
-	CustomizableObject->CompileOptions.TextureCompression = ECustomizableObjectTextureCompression(CompileTextureCompressionStrings.Find(NewSelection));
+	CustomizableObject->GetPrivate()->TextureCompression = ECustomizableObjectTextureCompression(CompileTextureCompressionStrings.Find(NewSelection));
 }
 
 
diff --git a/Engine/Plugins/Experimental/Mutable/Source/CustomizableObjectEditor/Private/MuCOE/CustomizableObjectInstanceBakingUtils.cpp b/Engine/Plugins/Experimental/Mutable/Source/CustomizableObjectEditor/Private/MuCOE/CustomizableObjectInstanceBakingUtils.cpp
--- a/Engine/Plugins/Experimental/Mutable/Source/CustomizableObjectEditor/Private/MuCOE/CustomizableObjectInstanceBakingUtils.cpp
+++ b/Engine/Plugins/Experimental/Mutable/Source/CustomizableObjectEditor/Private/MuCOE/CustomizableObjectInstanceBakingUtils.cpp
@@ -421,7 +421,7 @@ bool BakeCustomizableObjectInstance(
 	// COI Validation completed : Proceed with the baking operation ----------------------------------------------------------------------------------
 	
 	// Notify of better configuration -> Continue operation normally
-	if (InstanceCO->CompileOptions.TextureCompression != ECustomizableObjectTextureCompression::HighQuality)
+	if (InstanceCO->GetPrivate()->GetCompileOptions().TextureCompression != ECustomizableObjectTextureCompression::HighQuality)
 	{
 		FCustomizableObjectEditorLogger::CreateLog(
 		LOCTEXT("CustomizableObjectBakeLowQuality", "The Customizable Object wasn't compiled with high quality textures. For the best baking results, change the Texture Compression setting and recompile it."))
diff --git a/Engine/Plugins/Experimental/Mutable/Source/CustomizableObjectEditor/Private/MuCOE/CustomizableObjectLayout.cpp b/Engine/Plugins/Experimental/Mutable/Source/CustomizableObjectEditor/Private/MuCOE/CustomizableObjectLayout.cpp
--- a/Engine/Plugins/Experimental/Mutable/Source/CustomizableObjectEditor/Private/MuCOE/CustomizableObjectLayout.cpp
+++ b/Engine/Plugins/Experimental/Mutable/Source/CustomizableObjectEditor/Private/MuCOE/CustomizableObjectLayout.cpp
@@ -97,7 +97,7 @@ void UCustomizableObjectLayout::GenerateBlocksFromUVs()
 		//Creating a GenerationContext
 		FCustomizableObjectCompiler Compiler;
 		UCustomizableObject* Object = Node->GetGraphEditor()->GetCustomizableObject();
-		FCompilationOptions Options = Object->CompileOptions;
+		FCompilationOptions Options = Object->GetPrivate()->GetCompileOptions();
 	
 		FMutableGraphGenerationContext GenerationContext(Object, &Compiler, Options);
 	
diff --git a/Engine/Plugins/Experimental/Mutable/Source/CustomizableObjectEditor/Private/MuCOE/SMutableObjectViewer.cpp b/Engine/Plugins/Experimental/Mutable/Source/CustomizableObjectEditor/Private/MuCOE/SMutableObjectViewer.cpp
--- a/Engine/Plugins/Experimental/Mutable/Source/CustomizableObjectEditor/Private/MuCOE/SMutableObjectViewer.cpp
+++ b/Engine/Plugins/Experimental/Mutable/Source/CustomizableObjectEditor/Private/MuCOE/SMutableObjectViewer.cpp
@@ -39,11 +39,13 @@ FString SMutableObjectViewer::GetReferencerName() const
 
 void SMutableObjectViewer::Construct(const FArguments& InArgs, UCustomizableObject* InObject)
 {
+	check(InObject);
+	
 	CustomizableObject = InObject;
 
 	// Initialize the debugger compile options
-	CompileOptions.TextureCompression = ECustomizableObjectTextureCompression::Fast;
-	CompileOptions.OptimizationLevel = 2;
+	CompileOptions = InObject->GetPrivate()->GetCompileOptions();
+
 	{
 		ITargetPlatformManagerModule* TPM = GetTargetPlatformManager();
 		const TArray<ITargetPlatform*>& Platforms = TPM->GetActiveTargetPlatforms();
diff --git a/Engine/Plugins/Experimental/Mutable/Source/CustomizableObjectEditor/Private/MuCOE/Thumbnails/CustomizableObjectThumbnailRenderer.cpp b/Engine/Plugins/Experimental/Mutable/Source/CustomizableObjectEditor/Private/MuCOE/Thumbnails/CustomizableObjectThumbnailRenderer.cpp
--- a/Engine/Plugins/Experimental/Mutable/Source/CustomizableObjectEditor/Private/MuCOE/Thumbnails/CustomizableObjectThumbnailRenderer.cpp
+++ b/Engine/Plugins/Experimental/Mutable/Source/CustomizableObjectEditor/Private/MuCOE/Thumbnails/CustomizableObjectThumbnailRenderer.cpp
@@ -53,7 +53,7 @@ void UCustomizableObjectThumbnailRenderer::Draw(UObject* Object, int32 X, int32
 		if (!CustomizableObject->IsCompiled())
 		{
 			// Use default compilation options?
-			FCompilationOptions Options;
+			FCompilationOptions Options = CustomizableObject->GetPrivate()->GetCompileOptions();
 
 			FCustomizableObjectCompiler Compiler;
 			// Compiler.Compile(CustomizableObject, Options, false); // TODO Can not be sync!
diff --git a/Engine/Plugins/Experimental/Mutable/Source/MutableValidation/Private/MuV/AssetValidator_CustomizableObjects.cpp b/Engine/Plugins/Experimental/Mutable/Source/MutableValidation/Private/MuV/AssetValidator_CustomizableObjects.cpp
--- a/Engine/Plugins/Experimental/Mutable/Source/MutableValidation/Private/MuV/AssetValidator_CustomizableObjects.cpp
+++ b/Engine/Plugins/Experimental/Mutable/Source/MutableValidation/Private/MuV/AssetValidator_CustomizableObjects.cpp
@@ -13,6 +13,7 @@
 #include "Materials/Material.h"
 #include "MuCO/CustomizableObject.h"
 #include "MuCO/CustomizableObjectCompilerTypes.h"
+#include "MuCO/CustomizableObjectPrivate.h"
 #include "MuCO/ICustomizableObjectEditorModule.h"
 #include "MuCOE/GraphTraversal.h"
 #include "UObject/NameTypes.h"
@@ -114,7 +115,7 @@ EDataValidationResult UAssetValidator_CustomizableObjects::IsCustomizableObjectV
 	TArray<FCompilationOptions> CompilationOptionsToTest;
 
 	// Current configuration
-	CompilationOptionsToTest.Add(InCustomizableObject->CompileOptions);
+	CompilationOptionsToTest.Add(InCustomizableObject->GetPrivate()->GetCompileOptions());
 
 	// Configuration with LOD bias applied
 	// constexpr int32 MaxBias = 15;	
diff --git a/Engine/Plugins/Experimental/Mutable/Source/MutableValidation/Private/MuV/COIBakingTestCommandlet.cpp b/Engine/Plugins/Experimental/Mutable/Source/MutableValidation/Private/MuV/COIBakingTestCommandlet.cpp
--- a/Engine/Plugins/Experimental/Mutable/Source/MutableValidation/Private/MuV/COIBakingTestCommandlet.cpp
+++ b/Engine/Plugins/Experimental/Mutable/Source/MutableValidation/Private/MuV/COIBakingTestCommandlet.cpp
@@ -12,6 +12,7 @@
 #include "MuCOE/CustomizableObjectInstanceBakingUtils.h"
 #include "HAL/FileManager.h"
 #include "Interfaces/ITargetPlatformManagerModule.h"
+#include "MuCO/CustomizableObjectPrivate.h"
 
 
 /** Flag useful to know if we are currently updating an instance or not */
@@ -83,7 +84,7 @@ int32 UCOIBakingTestCommandlet::Main(const FString& Params)
 	}
 	
 	// Set the target platform to be using for the compilation. Must not be a nullptr
-	FCompilationOptions CompilationOptions = InstanceCustomizableObject->CompileOptions;
+	FCompilationOptions CompilationOptions = InstanceCustomizableObject->GetPrivate()->GetCompileOptions();
 	ITargetPlatformManagerModule& TPM = GetTargetPlatformManagerRef();
 	CompilationOptions.TargetPlatform = TPM.GetRunningTargetPlatform();
 	CompilationOptions.bUseDiskCompilation = false;
diff --git a/Engine/Plugins/Experimental/Mutable/Source/MutableValidation/Private/MuV/COIBulkUpdateTestCommandlet.cpp b/Engine/Plugins/Experimental/Mutable/Source/MutableValidation/Private/MuV/COIBulkUpdateTestCommandlet.cpp
--- a/Engine/Plugins/Experimental/Mutable/Source/MutableValidation/Private/MuV/COIBulkUpdateTestCommandlet.cpp
+++ b/Engine/Plugins/Experimental/Mutable/Source/MutableValidation/Private/MuV/COIBulkUpdateTestCommandlet.cpp
@@ -10,6 +10,7 @@
 #include "MuCO/CustomizableObject.h"
 #include "MuCO/CustomizableObjectInstance.h"
 #include "Interfaces/ITargetPlatformManagerModule.h"
+#include "MuCO/CustomizableObjectPrivate.h"
 
 
 int32 UCOIBulkUpdateTestCommandlet::Main(const FString& Params)
@@ -140,7 +141,7 @@ int32 UCOIBulkUpdateTestCommandlet::Main(const FString& Params)
 		const FString CustomizableObjectName = CustomizableObject->GetName();
 		
 		// Set the compilation platform based on what the system is currently running on
-		FCompilationOptions CompilationOptions = CustomizableObject->CompileOptions;
+		FCompilationOptions CompilationOptions = CustomizableObject->GetPrivate()->GetCompileOptions();
 		CompilationOptions.TargetPlatform = TargetCompilationPlatform;
 		CompilationOptions.bUseDiskCompilation = false;
 		CompilationOptions.OptimizationLevel = UE_MUTABLE_MAX_OPTIMIZATION;
diff --git a/Engine/Plugins/Experimental/Mutable/Source/MutableValidation/Private/MuV/CustomizableObjectCompilationUtility.cpp b/Engine/Plugins/Experimental/Mutable/Source/MutableValidation/Private/MuV/CustomizableObjectCompilationUtility.cpp
--- a/Engine/Plugins/Experimental/Mutable/Source/MutableValidation/Private/MuV/CustomizableObjectCompilationUtility.cpp
+++ b/Engine/Plugins/Experimental/Mutable/Source/MutableValidation/Private/MuV/CustomizableObjectCompilationUtility.cpp
@@ -33,7 +33,7 @@ bool FCustomizableObjectCompilationUtility::CompileCustomizableObject(UCustomiza
 	}
 	else
 	{
-		CompilationOptions = InCustomizableObject->CompileOptions;
+		CompilationOptions = InCustomizableObject->GetPrivate()->GetCompileOptions();
 		UE_LOG(LogMutable,Display,TEXT("Compiling CO using it's own compilation options."));
 	}
 
diff --git a/Engine/Plugins/Experimental/Mutable/Source/MutableValidation/Private/MuV/CustomizableObjectValidationCommandlet.cpp b/Engine/Plugins/Experimental/Mutable/Source/MutableValidation/Private/MuV/CustomizableObjectValidationCommandlet.cpp
--- a/Engine/Plugins/Experimental/Mutable/Source/MutableValidation/Private/MuV/CustomizableObjectValidationCommandlet.cpp
+++ b/Engine/Plugins/Experimental/Mutable/Source/MutableValidation/Private/MuV/CustomizableObjectValidationCommandlet.cpp
@@ -82,7 +82,7 @@ int32 UCustomizableObjectValidationCommandlet::Main(const FString& Params)
 		LLM_SCOPE_BYNAME(TEXT("CustomizableObjectValidationCommandlet/Compile"));
 		
 		// Override some configurations that may have been changed by the user
-		FCompilationOptions CompilationOptions = ToTestCustomizableObject->CompileOptions;
+		FCompilationOptions CompilationOptions = ToTestCustomizableObject->GetPrivate()->GetCompileOptions();
 		CompilationOptions.bSilentCompilation = false;
 		CompilationOptions.OptimizationLevel = UE_MUTABLE_MAX_OPTIMIZATION;
 		CompilationOptions.TextureCompression = ECustomizableObjectTextureCompression::Fast;
