[Mutable] Remove non-const getters for the different UCustomizableObjectDescriptor parameter arrays.
Remove repeated code to get and set parameter values.
Remove the custom widget SCustomizableInstanceProperties.

#rnx

diff --git a/Engine/Plugins/Experimental/Mutable/Source/CustomizableObject/Internal/MuCO/CustomizableObjectInstanceDescriptor.h b/Engine/Plugins/Experimental/Mutable/Source/CustomizableObject/Internal/MuCO/CustomizableObjectInstanceDescriptor.h
--- a/Engine/Plugins/Experimental/Mutable/Source/CustomizableObject/Internal/MuCO/CustomizableObjectInstanceDescriptor.h
+++ b/Engine/Plugins/Experimental/Mutable/Source/CustomizableObject/Internal/MuCO/CustomizableObjectInstanceDescriptor.h
@@ -86,28 +86,16 @@ struct CUSTOMIZABLEOBJECT_API FCustomizableObjectInstanceDescriptor
 	// Parameters
 	// ------------------------------------------------------------
 
-	TArray<FCustomizableObjectBoolParameterValue>& GetBoolParameters();
-
 	const TArray<FCustomizableObjectBoolParameterValue>& GetBoolParameters() const;
 	
-	TArray<FCustomizableObjectIntParameterValue>& GetIntParameters();
-
 	const TArray<FCustomizableObjectIntParameterValue>& GetIntParameters() const;
 
-	TArray<FCustomizableObjectFloatParameterValue>& GetFloatParameters();
-
 	const TArray<FCustomizableObjectFloatParameterValue>& GetFloatParameters() const;
 	
-	TArray<FCustomizableObjectTextureParameterValue>& GetTextureParameters();
-	
 	const TArray<FCustomizableObjectTextureParameterValue>& GetTextureParameters() const;
 
-	TArray<FCustomizableObjectVectorParameterValue>& GetVectorParameters();
-	
 	const TArray<FCustomizableObjectVectorParameterValue>& GetVectorParameters() const;
 
-	TArray<FCustomizableObjectProjectorParameterValue>& GetProjectorParameters();
-	
 	const TArray<FCustomizableObjectProjectorParameterValue>& GetProjectorParameters() const;
 	
 	/** Return true if there are any parameters. */
@@ -282,40 +270,44 @@ struct CUSTOMIZABLEOBJECT_API FCustomizableObjectInstanceDescriptor
 	
 	void SetRandomValuesFromStream(const FRandomStream& InStream);
 
+	void SetDefaultValue(int32 ParamIndex);
 	void SetDefaultValues();
 	
 	// ------------------------------------------------------------
 	// Multilayer Projectors
 	// ------------------------------------------------------------
-	
-	/** Given Multilayer Projector name, create a new Multilayer Projector Helper (if non-existent). See FMultilayerProjector.
-	 *
-	 * @return ture if successfully created (or was already created).
-	 */
-	bool CreateMultiLayerProjector(const FName& ProjectorParamName);
-	
-	/** Given Multilayer Projector name, remove a Multilayer Projector Helper. See FMultilayerProjector. */
-	void RemoveMultilayerProjector(const FName& ProjectorParamName);
-	
+
+	/** @return true if ParamName belongs to a multilayer projector parameter. */
+	bool IsMultilayerProjector(const FString& ParamName) const;
+
 	// Layers
+	/** @return number of layers of the projector with name ParamName,-1 if invalid or not found. */
+	int32 NumProjectorLayers(const FName& ParamName) const;
 
-	/** See FMultilayerProjector::NumLayers. */
-	int32 MultilayerProjectorNumLayers(const FName& ProjectorParamName) const;
+	/** Creates a new layer for the multilayer projector with name ParamName. */
+	void CreateLayer(const FName& ParamName, int32 Index);
 
-	/** See FMultilayerProjector::CreateLayer. */
-	void MultilayerProjectorCreateLayer(const FName& ProjectorParamName, int32 Index);
+	/** Removes the layer at Index from the multilayer projector with name ParamName. */
+	void RemoveLayerAt(const FName& ParamName, int32 Index);
 
-	/** See FMultilayerProjector::RemoveLayerAt. */
-	void MultilayerProjectorRemoveLayerAt(const FName& ProjectorParamName, int32 Index);
+	/** @return copy of the layer at Index for the multilayer projector with name ParamName. */
+	FMultilayerProjectorLayer GetLayer(const FName& ParamName, int32 Index) const;
 
-	/** See FMultilayerProjector::GetLayer. */
-	FMultilayerProjectorLayer MultilayerProjectorGetLayer(const FName& ProjectorParamName, int32 Index) const;
+	/** Updates the parameters of the layer at Index from the multilayer projector with name ParamName. */
+	void UpdateLayer(const FName& ParamName, int32 Index, const FMultilayerProjectorLayer& Layer);
 
-	/** See FMultilayerProjector::UpdateLayer. */
-	void MultilayerProjectorUpdateLayer(const FName& ProjectorParamName, int32 Index, const FMultilayerProjectorLayer& Layer);
 
 	// Virtual layers
 
+	/** Given Multilayer Projector name, create a new Multilayer Projector Helper (if non-existent) for virtual layers. See FMultilayerProjector.
+	 *
+	 * @return ture if successfully created (or was already created).
+	 */
+	bool CreateMultiLayerProjector(const FName& ProjectorParamName);
+
+	/** Given Multilayer Projector name, remove a Multilayer Projector Helper. See FMultilayerProjector. */
+	void RemoveMultilayerProjector(const FName& ProjectorParamName);
+
 	/** See FMultilayerProjector::GetVirtualLayers. */
 	TArray<FName> MultilayerProjectorGetVirtualLayers(const FName& ProjectorParamName) const;
 	
diff --git a/Engine/Plugins/Experimental/Mutable/Source/CustomizableObject/Private/MuCO/CustomizableObjectInstance.cpp b/Engine/Plugins/Experimental/Mutable/Source/CustomizableObject/Private/MuCO/CustomizableObjectInstance.cpp
--- a/Engine/Plugins/Experimental/Mutable/Source/CustomizableObject/Private/MuCO/CustomizableObjectInstance.cpp
+++ b/Engine/Plugins/Experimental/Mutable/Source/CustomizableObject/Private/MuCO/CustomizableObjectInstance.cpp
@@ -2449,31 +2449,31 @@ void UCustomizableObjectInstance::RemoveMultilayerProjector(const FName& Project
 
 int32 UCustomizableObjectInstance::MultilayerProjectorNumLayers(const FName& ProjectorParamName) const
 {
-	return Descriptor.MultilayerProjectorNumLayers(ProjectorParamName);
+	return Descriptor.NumProjectorLayers(ProjectorParamName);
 }
 
 
 void UCustomizableObjectInstance::MultilayerProjectorCreateLayer(const FName& ProjectorParamName, int32 Index)
 {
-	Descriptor.MultilayerProjectorCreateLayer(ProjectorParamName, Index);
+	Descriptor.CreateLayer(ProjectorParamName, Index);
 }
 
 
 void UCustomizableObjectInstance::MultilayerProjectorRemoveLayerAt(const FName& ProjectorParamName, int32 Index)
 {
-	Descriptor.MultilayerProjectorRemoveLayerAt(ProjectorParamName, Index);
+	Descriptor.RemoveLayerAt(ProjectorParamName, Index);
 }
 
 
 FMultilayerProjectorLayer UCustomizableObjectInstance::MultilayerProjectorGetLayer(const FName& ProjectorParamName, int32 Index) const
 {
-	return Descriptor.MultilayerProjectorGetLayer(ProjectorParamName, Index);
+	return Descriptor.GetLayer(ProjectorParamName, Index);
 }
 
 
 void UCustomizableObjectInstance::MultilayerProjectorUpdateLayer(const FName& ProjectorParamName, int32 Index, const FMultilayerProjectorLayer& Layer)
 {
-	Descriptor.MultilayerProjectorUpdateLayer(ProjectorParamName, Index, Layer);
+	Descriptor.UpdateLayer(ProjectorParamName, Index, Layer);
 }
 
 
@@ -2735,6 +2735,16 @@ void UCustomizableObjectInstance::SetRandomValuesFromStream(const FRandomStream&
 	Descriptor.SetRandomValuesFromStream(InStream);
 }
 
+void UCustomizableObjectInstance::SetDefaultValue(const FString& ParamName)
+{
+	UCustomizableObject* CustomizableObject = GetCustomizableObject();
+	if (!CustomizableObject)
+	{
+		return;
+	}
+
+	Descriptor.SetDefaultValue(CustomizableObject->FindParameter(ParamName));
+}
 
 void UCustomizableObjectInstance::SetDefaultValues()
 {
diff --git a/Engine/Plugins/Experimental/Mutable/Source/CustomizableObject/Private/MuCO/CustomizableObjectInstanceDescriptor.cpp b/Engine/Plugins/Experimental/Mutable/Source/CustomizableObject/Private/MuCO/CustomizableObjectInstanceDescriptor.cpp
--- a/Engine/Plugins/Experimental/Mutable/Source/CustomizableObject/Private/MuCO/CustomizableObjectInstanceDescriptor.cpp
+++ b/Engine/Plugins/Experimental/Mutable/Source/CustomizableObject/Private/MuCO/CustomizableObjectInstanceDescriptor.cpp
@@ -1110,75 +1110,39 @@ const TArray<uint16>& FCustomizableObjectInstanceDescriptor::GetRequestedLODLeve
 }
 
 
-TArray<FCustomizableObjectBoolParameterValue>& FCustomizableObjectInstanceDescriptor::GetBoolParameters()
-{
-	return BoolParameters;
-}
-
-
 const TArray<FCustomizableObjectBoolParameterValue>& FCustomizableObjectInstanceDescriptor::GetBoolParameters() const
 {
-	return const_cast<FCustomizableObjectInstanceDescriptor*>(this)->GetBoolParameters();
-}
-
-
-TArray<FCustomizableObjectIntParameterValue>& FCustomizableObjectInstanceDescriptor::GetIntParameters()
-{
-	return IntParameters;
+	return BoolParameters;
 }
 
 
 const TArray<FCustomizableObjectIntParameterValue>& FCustomizableObjectInstanceDescriptor::GetIntParameters() const
 {
-	return const_cast<FCustomizableObjectInstanceDescriptor*>(this)->GetIntParameters();	
-}
-
-
-TArray<FCustomizableObjectFloatParameterValue>& FCustomizableObjectInstanceDescriptor::GetFloatParameters()
-{
-	return FloatParameters;	
+	return IntParameters;	
 }
 
 
 const TArray<FCustomizableObjectFloatParameterValue>& FCustomizableObjectInstanceDescriptor::GetFloatParameters() const
 {
-	return const_cast<FCustomizableObjectInstanceDescriptor*>(this)->GetFloatParameters();	
-}
-
-
-TArray<FCustomizableObjectTextureParameterValue>& FCustomizableObjectInstanceDescriptor::GetTextureParameters()
-{
-	return TextureParameters;	
+	return FloatParameters;	
 }
 
 
 const TArray<FCustomizableObjectTextureParameterValue>& FCustomizableObjectInstanceDescriptor::GetTextureParameters() const
 {
-	return const_cast<FCustomizableObjectInstanceDescriptor*>(this)->GetTextureParameters();	
-}
-
-
-TArray<FCustomizableObjectVectorParameterValue>& FCustomizableObjectInstanceDescriptor::GetVectorParameters()
-{
-	return VectorParameters;	
+	return TextureParameters;	
 }
 
 
 const TArray<FCustomizableObjectVectorParameterValue>& FCustomizableObjectInstanceDescriptor::GetVectorParameters() const
 {
-	return const_cast<FCustomizableObjectInstanceDescriptor*>(this)->GetVectorParameters();
-}
-
-
-TArray<FCustomizableObjectProjectorParameterValue>& FCustomizableObjectInstanceDescriptor::GetProjectorParameters()
-{
-	return ProjectorParameters;	
+	return VectorParameters;
 }
 
 
 const TArray<FCustomizableObjectProjectorParameterValue>& FCustomizableObjectInstanceDescriptor::GetProjectorParameters() const
 {
-	return const_cast<FCustomizableObjectInstanceDescriptor*>(this)->GetProjectorParameters();	
+	return ProjectorParameters;	
 }
 
 
@@ -2395,7 +2359,7 @@ void FCustomizableObjectInstanceDescriptor::SetRandomValuesFromStream(const FRan
 }
 
 
-void FCustomizableObjectInstanceDescriptor::SetDefaultValues()
+void FCustomizableObjectInstanceDescriptor::SetDefaultValue(int32 ParamIndex)
 {
 	if (!CustomizableObject)
 	{
@@ -2403,124 +2367,286 @@ void FCustomizableObjectInstanceDescriptor::SetDefaultValues()
 		return;
 	}
 
-	for (FCustomizableObjectBoolParameterValue& Value : BoolParameters)
-	{	
-		Value.ParameterValue =	CustomizableObject->GetBoolParameterDefaultValue(Value.ParameterName);
+	if (ParamIndex >= CustomizableObject->GetParameterCount())
+	{
+		return;
 	}
 
-	for (FCustomizableObjectIntParameterValue& Value : IntParameters)
-	{
-		const int32 ParameterIndex = CustomizableObject->FindParameter(Value.ParameterName);
-		Value.ParameterValueName = CustomizableObject->FindIntParameterValueName(ParameterIndex, CustomizableObject->GetIntParameterDefaultValue(Value.ParameterValueName));
+	const FString ParamName = CustomizableObject->GetParameterName(ParamIndex);
+	const EMutableParameterType ParamType = CustomizableObject->GetParameterType(ParamIndex);
+	
+	const int32 TypedIndex = FindTypedParameterIndex(ParamName, ParamType);
 
-		for (FString& RangeValue : Value.ParameterRangeValueNames)
-		{
-			RangeValue = Value.ParameterValueName;
-		}
+	switch (ParamType)
+	{
+	case EMutableParameterType::Bool:
+	{
+		const bool DefaultValue = CustomizableObject->GetBoolParameterDefaultValue(ParamName);
+		BoolParameters[TypedIndex].ParameterValue = DefaultValue;
+		break;
 	}
 
-	for (FCustomizableObjectFloatParameterValue& Value : FloatParameters)
+	case EMutableParameterType::Int:
 	{
-		Value.ParameterValue = CustomizableObject->GetFloatParameterDefaultValue(Value.ParameterName);
+		const FString DefaultValue = CustomizableObject->FindIntParameterValueName(ParamIndex, CustomizableObject->GetIntParameterDefaultValue(ParamName));
+
+		FCustomizableObjectIntParameterValue& IntParameter = IntParameters[TypedIndex];
+		IntParameter.ParameterValueName = DefaultValue;
 
-		for (float& RangeValue : Value.ParameterRangeValues)
+		for (FString& RangeValue : IntParameter.ParameterRangeValueNames)
 		{
-			RangeValue = Value.ParameterValue;
+			RangeValue = DefaultValue;
 		}
+		break;
 	}
 
-	for (FCustomizableObjectTextureParameterValue& Value : TextureParameters)
+	case EMutableParameterType::Float:
 	{
-		Value.ParameterValue = CustomizableObject->GetTextureParameterDefaultValue(Value.ParameterName);
+		const float DefaultValue = CustomizableObject->GetFloatParameterDefaultValue(ParamName);
+
+		FCustomizableObjectFloatParameterValue& FloatParameter = FloatParameters[TypedIndex];
+		FloatParameter.ParameterValue = DefaultValue;
 
-		for (FName& RangeValue : Value.ParameterRangeValues)
+		for (float& RangeValue : FloatParameter.ParameterRangeValues)
 		{
-			RangeValue = Value.ParameterValue;
+			RangeValue = DefaultValue;
 		}
+		break;
 	}
 
-	for (FCustomizableObjectVectorParameterValue& Value : VectorParameters)
+	case EMutableParameterType::Color:
 	{
-		Value.ParameterValue = CustomizableObject->GetColorParameterDefaultValue(Value.ParameterName);
+		const FLinearColor DefaultValue = CustomizableObject->GetColorParameterDefaultValue(ParamName);
+		VectorParameters[TypedIndex].ParameterValue = DefaultValue;
+		break;
 	}
 
-	for (FCustomizableObjectProjectorParameterValue& Value : ProjectorParameters)
+	case EMutableParameterType::Projector:
 	{
-		Value.Value = CustomizableObject->GetProjectorParameterDefaultValue(Value.ParameterName);
+		const FCustomizableObjectProjector DefaultValue = CustomizableObject->GetProjectorParameterDefaultValue(ParamName);
+
+		FCustomizableObjectProjectorParameterValue& ProjectorParameter = ProjectorParameters[TypedIndex];
+		ProjectorParameter.Value = DefaultValue;
 
-		for (FCustomizableObjectProjector& RangeValue : Value.RangeValues)
+		for (FCustomizableObjectProjector& RangeValue : ProjectorParameter.RangeValues)
+		{
+			RangeValue = DefaultValue;
+		}
+		break;
+	}
+
+	case EMutableParameterType::Texture:
+	{
+		const FName	DefaultValue = CustomizableObject->GetTextureParameterDefaultValue(ParamName);
+		
+		if (ensure(TextureParameters.IsValidIndex(TypedIndex)))
 		{
-			RangeValue = Value.Value;
+			FCustomizableObjectTextureParameterValue& TextureParameter = TextureParameters[TypedIndex];
+			TextureParameter.ParameterValue = DefaultValue;
+
+			for (FName& RangeValue : TextureParameter.ParameterRangeValues)
+			{
+				RangeValue = DefaultValue;
+			}
 		}
+		break;
+	}
+
+	default:
+		// Parameter type replication not implemented.
+		unimplemented();
 	}
 }
 
 
-bool FCustomizableObjectInstanceDescriptor::CreateMultiLayerProjector(const FName& ProjectorParamName)
+void FCustomizableObjectInstanceDescriptor::SetDefaultValues()
+{
+	if (!CustomizableObject)
+	{
+		CustomizableObjectNullErrorMessage();
+		return;
+	}
+
+	const int32 NumParameters = CustomizableObject->GetParameterCount();
+	for (int32 ParamIndex = 0; ParamIndex < NumParameters; ++ParamIndex)
+	{
+		SetDefaultValue(ParamIndex);
+	}
+}
+
+bool FCustomizableObjectInstanceDescriptor::IsMultilayerProjector(const FString& ParamName) const
 {
-	if (!FMultilayerProjector::AreDescriptorParametersValid(*this, ProjectorParamName.ToString()))
+	if (!CustomizableObject)
 	{
-#if WITH_EDITOR
-		UE_LOG(LogMutable, Error, TEXT("%s"), *FMultilayerProjector::DESCRIPTOR_PARAMETERS_INVALID);
-#else
-		ensureAlwaysMsgf(false, TEXT("%s"), *FMultilayerProjector::DESCRIPTOR_PARAMETERS_INVALID);
-#endif
+		CustomizableObjectNullErrorMessage();
 		return false;
 	}
 
-	if (!MultilayerProjectors.Contains(ProjectorParamName))
+	// Projector.
+	if (const int32 ProjectorParameterIndex = FindTypedParameterIndex(ParamName, EMutableParameterType::Projector);
+		ProjectorParameterIndex == INDEX_NONE)
 	{
-		const FMultilayerProjector MultilayerProjector(ProjectorParamName);
-		MultilayerProjectors.Add(ProjectorParamName, MultilayerProjector);
+		ensureAlwaysMsgf(false, TEXT("%s"), *FMultilayerProjector::MULTILAYER_PROJECTOR_PARAMETERS_INVALID);
+		return false;
+	}
+
+	// Num layers.
+	if (const int32 FloatParameterIndex = FindTypedParameterIndex(ParamName + FMultilayerProjector::NUM_LAYERS_PARAMETER_POSTFIX, EMutableParameterType::Float);
+		FloatParameterIndex == INDEX_NONE)
+	{
+		ensureAlwaysMsgf(false, TEXT("%s"), *FMultilayerProjector::MULTILAYER_PROJECTOR_PARAMETERS_INVALID);
+		return false;
+	}
+
+	// Selected Image.
+	if (const int32 IntParameterIndex = FindTypedParameterIndex(ParamName + FMultilayerProjector::IMAGE_PARAMETER_POSTFIX, EMutableParameterType::Int);
+		IntParameterIndex == INDEX_NONE)
+	{
+		ensureAlwaysMsgf(false, TEXT("%s"), *FMultilayerProjector::MULTILAYER_PROJECTOR_PARAMETERS_INVALID);
+		return false;
+	}
+
+	// Opacity.	
+	if (const int32 FloatParameterIndex = FindTypedParameterIndex(ParamName + FMultilayerProjector::OPACITY_PARAMETER_POSTFIX, EMutableParameterType::Float);
+		FloatParameterIndex == INDEX_NONE)
+	{
+		ensureAlwaysMsgf(false, TEXT("%s"), *FMultilayerProjector::MULTILAYER_PROJECTOR_PARAMETERS_INVALID);
+		return false;
 	}
 
 	return true;
 }
 
 
-void FCustomizableObjectInstanceDescriptor::RemoveMultilayerProjector(const FName& ProjectorParamName)
+int32 FCustomizableObjectInstanceDescriptor::NumProjectorLayers(const FName& ProjectorParamName) const
 {
-	MultilayerProjectors.Remove(ProjectorParamName);
+	const FString ParamName = ProjectorParamName.ToString();
+
+	if (IsMultilayerProjector(ParamName))
+	{
+		return GetFloatParameterSelectedOption(ParamName + FMultilayerProjector::NUM_LAYERS_PARAMETER_POSTFIX);
+	}
+
+	return INDEX_NONE;
 }
 
 
-int32 FCustomizableObjectInstanceDescriptor::MultilayerProjectorNumLayers(const FName& ProjectorParamName) const
+void FCustomizableObjectInstanceDescriptor::CreateLayer(const FName& ProjectorParamName, int32 RangeIndex)
 {
-	check(MultilayerProjectors.Contains(ProjectorParamName)); // Multilayer Projector not created.
-	
-	return MultilayerProjectors[ProjectorParamName].NumLayers(*this);
+	const FString ParamName = ProjectorParamName.ToString();
+
+	if (!IsMultilayerProjector(ParamName))
+	{
+		return;
+	}
+
+	// Num Layers.
+	SetFloatParameterSelectedOption(ParamName + FMultilayerProjector::NUM_LAYERS_PARAMETER_POSTFIX, NumProjectorLayers(ProjectorParamName) + 1);
+
+	// Projector Range. New value is defaulted.
+	AddValueToProjectorRange(ParamName);
+
+	// Selected Image Range.
+	{
+		const FString ImageParamName = ParamName + FMultilayerProjector::IMAGE_PARAMETER_POSTFIX;
+		AddValueToTextureRange(ImageParamName);
+
+		const int32 ParamIndex = CustomizableObject->FindParameter(ImageParamName);
+		const int32 DefaultValueIndex = CustomizableObject->GetIntParameterDefaultValue(ImageParamName);
+		const FString DefaultValueName = CustomizableObject->GetIntParameterAvailableOption(ParamIndex, DefaultValueIndex);
+		SetIntParameterSelectedOption(ImageParamName, DefaultValueName, RangeIndex);
+	}
+
+	// Opacity Range.
+	{
+		const FString OpacityParamName = ParamName + FMultilayerProjector::OPACITY_PARAMETER_POSTFIX;
+		AddValueToFloatRange(OpacityParamName);
+
+		const int32 ParamIndex = CustomizableObject->FindParameter(OpacityParamName);
+		const float DefaultValue = CustomizableObject->GetFloatParameterDefaultValue(OpacityParamName);
+		SetFloatParameterSelectedOption(OpacityParamName, DefaultValue, RangeIndex);
+	}
 }
 
 
-void FCustomizableObjectInstanceDescriptor::MultilayerProjectorCreateLayer(const FName& ProjectorParamName, int32 Index)
+void FCustomizableObjectInstanceDescriptor::RemoveLayerAt(const FName& ProjectorParamName, int32 RangeIndex)
 {
-	check(MultilayerProjectors.Contains(ProjectorParamName)); // Multilayer Projector not created.
-	
-	MultilayerProjectors[ProjectorParamName].CreateLayer(*this, Index);
+	const FString ParamName = ProjectorParamName.ToString();
+	const int32 NumLayers = NumProjectorLayers(ProjectorParamName);
+	check(RangeIndex >= 0 && RangeIndex < NumLayers); // Layer out of range.
+
+	if (!IsMultilayerProjector(ParamName))
+	{
+		return;
+	}
+
+	// Num Layers.
+	SetFloatParameterSelectedOption(ParamName + FMultilayerProjector::NUM_LAYERS_PARAMETER_POSTFIX, NumLayers - 1);
+
+	// Projector Range.
+	RemoveValueFromProjectorRange(ParamName, RangeIndex);
+
+	// Selected Image Range.
+	RemoveValueFromIntRange(ParamName + FMultilayerProjector::IMAGE_PARAMETER_POSTFIX, RangeIndex);
+
+	// Opacity Range.
+	RemoveValueFromFloatRange(ParamName + FMultilayerProjector::OPACITY_PARAMETER_POSTFIX);
 }
 
 
-void FCustomizableObjectInstanceDescriptor::MultilayerProjectorRemoveLayerAt(const FName& ProjectorParamName, int32 Index)
+FMultilayerProjectorLayer FCustomizableObjectInstanceDescriptor::GetLayer(const FName& ProjectorParamName, int32 Index) const
 {
-	check(MultilayerProjectors.Contains(ProjectorParamName)); // Multilayer Projector not created.
-	
-	MultilayerProjectors[ProjectorParamName].RemoveLayerAt(*this, Index);
+	const FString ParamName = ProjectorParamName.ToString();
+
+	FMultilayerProjectorLayer ProjectorLayer;
+
+	if (IsMultilayerProjector(ParamName))
+	{
+		ProjectorLayer.Read(*this, ParamName, Index);
+	}
+
+	return ProjectorLayer;
 }
 
 
-FMultilayerProjectorLayer FCustomizableObjectInstanceDescriptor::MultilayerProjectorGetLayer(const FName& ProjectorParamName, int32 Index) const
+void FCustomizableObjectInstanceDescriptor::UpdateLayer(const FName& ProjectorParamName, int32 Index, const FMultilayerProjectorLayer& Layer)
 {
-	check(MultilayerProjectors.Contains(ProjectorParamName)); // Multilayer Projector not created.
-	
-	return MultilayerProjectors[ProjectorParamName].GetLayer(*this, Index);
+	const FString& ParamName = ProjectorParamName.ToString();
+
+	if (!IsMultilayerProjector(ParamName))
+	{
+		return;
+	}
+
+	Layer.Write(*this, ParamName, Index);
 }
 
 
-void FCustomizableObjectInstanceDescriptor::MultilayerProjectorUpdateLayer(const FName& ProjectorParamName, int32 Index, const FMultilayerProjectorLayer& Layer)
+bool FCustomizableObjectInstanceDescriptor::CreateMultiLayerProjector(const FName& ProjectorParamName)
 {
-	check(MultilayerProjectors.Contains(ProjectorParamName)); // Multilayer Projector not created.
-	
-	MultilayerProjectors[ProjectorParamName].UpdateLayer(*this, Index, Layer);
+	if (!IsMultilayerProjector(ProjectorParamName.ToString()))
+	{
+#if WITH_EDITOR
+		UE_LOG(LogMutable, Error, TEXT("%s"), *FMultilayerProjector::MULTILAYER_PROJECTOR_PARAMETERS_INVALID);
+#else
+		ensureAlwaysMsgf(false, TEXT("%s"), *FMultilayerProjector::MULTILAYER_PROJECTOR_PARAMETERS_INVALID);
+#endif
+		return false;
+	}
+
+	if (!MultilayerProjectors.Contains(ProjectorParamName))
+	{
+		const FMultilayerProjector MultilayerProjector(ProjectorParamName);
+		MultilayerProjectors.Add(ProjectorParamName, MultilayerProjector);
+	}
+
+	return true;
+}
+
+
+void FCustomizableObjectInstanceDescriptor::RemoveMultilayerProjector(const FName& ProjectorParamName)
+{
+	MultilayerProjectors.Remove(ProjectorParamName);
 }
 
 
diff --git a/Engine/Plugins/Experimental/Mutable/Source/CustomizableObject/Private/MuCO/MultilayerProjector.cpp b/Engine/Plugins/Experimental/Mutable/Source/CustomizableObject/Private/MuCO/MultilayerProjector.cpp
--- a/Engine/Plugins/Experimental/Mutable/Source/CustomizableObject/Private/MuCO/MultilayerProjector.cpp
+++ b/Engine/Plugins/Experimental/Mutable/Source/CustomizableObject/Private/MuCO/MultilayerProjector.cpp
@@ -8,74 +8,51 @@
 #include UE_INLINE_GENERATED_CPP_BY_NAME(MultilayerProjector)
 
 
-const FString FMultilayerProjector::DESCRIPTOR_PARAMETERS_INVALID = TEXT("Descriptor Parameters Invalid.");
+const FString FMultilayerProjector::MULTILAYER_PROJECTOR_PARAMETERS_INVALID = TEXT("Invalid Multilayer Projector Parameters.");
 
 
-void FMultilayerProjectorLayer::Read(const FCustomizableObjectInstanceDescriptor& Descriptor, const FMultilayerProjector& MultilayerProjector, const int32 Index)
+void FMultilayerProjectorLayer::Read(const FCustomizableObjectInstanceDescriptor& Descriptor, const FString& ParamName, const int32 LayerIndex)
 {
-	check(Index >= 0 && Index < MultilayerProjector.NumLayers(Descriptor)); // Layer out of range.
-	
-	const FString ParamName = MultilayerProjector.ParamName.ToString();
-
-	if (!FMultilayerProjector::AreDescriptorParametersValid(Descriptor, ParamName))
+	if (!Descriptor.IsMultilayerProjector(ParamName))
 	{
-		ensureAlwaysMsgf(false, TEXT("%s"), *FMultilayerProjector::DESCRIPTOR_PARAMETERS_INVALID);
+		ensureAlwaysMsgf(false, TEXT("%s"), *FMultilayerProjector::MULTILAYER_PROJECTOR_PARAMETERS_INVALID);
 		return;
 	}
-	
+
+	check(LayerIndex >= 0 && LayerIndex < Descriptor.NumProjectorLayers(*ParamName)); // Layer out of range.
+
 	{
-		const int32 ProjectorParamIndex = Descriptor.FindTypedParameterIndex(ParamName, EMutableParameterType::Projector);
-		const FCustomizableObjectProjector& Projector = Descriptor.GetProjectorParameters()[ProjectorParamIndex].RangeValues[Index];
-		Position = static_cast<FVector3d>(Projector.Position);
-		Direction = static_cast<FVector3d>(Projector.Direction);
-		Up = static_cast<FVector3d>(Projector.Up);
-		Scale = static_cast<FVector3d>(Projector.Scale);
-		Angle = Projector.Angle;
+		ECustomizableObjectProjectorType DummyType;
+		Descriptor.GetProjectorValue(ParamName, Position, Direction, Up, Scale, Angle, DummyType, LayerIndex);
 	}
 
 	{
 		const int32 ImageParamIndex = Descriptor.FindTypedParameterIndex(ParamName + FMultilayerProjector::IMAGE_PARAMETER_POSTFIX, EMutableParameterType::Int);
-		Image = Descriptor.GetIntParameters()[ImageParamIndex].ParameterRangeValueNames[Index];
+		Image = Descriptor.GetIntParameters()[ImageParamIndex].ParameterRangeValueNames[LayerIndex];
 	}
 
 	{
 		const int32 OpacityParamIndex = Descriptor.FindTypedParameterIndex(ParamName + FMultilayerProjector::OPACITY_PARAMETER_POSTFIX, EMutableParameterType::Float);
-		Opacity = Descriptor.GetFloatParameters()[OpacityParamIndex].ParameterRangeValues[Index];
+		Opacity = Descriptor.GetFloatParameters()[OpacityParamIndex].ParameterRangeValues[LayerIndex];
 	}
 }
 
 
-void FMultilayerProjectorLayer::Write(FCustomizableObjectInstanceDescriptor& Descriptor, const FMultilayerProjector& MultilayerProjector, const int32 Index) const
+void FMultilayerProjectorLayer::Write(FCustomizableObjectInstanceDescriptor& Descriptor, const FString& ParamName, int32 LayerIndex) const
 {
-	check(Index >= 0 && Index < MultilayerProjector.NumLayers(Descriptor)); // Layer out of range.
-
-	const FString ParamName = MultilayerProjector.ParamName.ToString();
-
-	if (!FMultilayerProjector::AreDescriptorParametersValid(Descriptor, ParamName))
+	if (!Descriptor.IsMultilayerProjector(ParamName))
 	{
-		ensureAlwaysMsgf(false, TEXT("%s"), *FMultilayerProjector::DESCRIPTOR_PARAMETERS_INVALID);
+		ensureAlwaysMsgf(false, TEXT("%s"), *FMultilayerProjector::MULTILAYER_PROJECTOR_PARAMETERS_INVALID);
 		return;
 	}
 
-	{
-		const int32 ProjectorParamIndex = Descriptor.FindTypedParameterIndex(ParamName, EMutableParameterType::Projector);
-		FCustomizableObjectProjector& Projector = Descriptor.GetProjectorParameters()[ProjectorParamIndex].RangeValues[Index];
-		Projector.Position = static_cast<FVector3f>(Position);
-		Projector.Direction = static_cast<FVector3f>(Direction);
-		Projector.Up = static_cast<FVector3f>(Up);
-		Projector.Scale = static_cast<FVector3f>(Scale);
-		Projector.Angle = Angle;
-	}
+	check(LayerIndex >= 0 && LayerIndex < Descriptor.NumProjectorLayers(*ParamName)); // Layer out of range.
 
-	{
-		const int32 ImageParamIndex = Descriptor.FindTypedParameterIndex(ParamName + FMultilayerProjector::IMAGE_PARAMETER_POSTFIX, EMutableParameterType::Int);
-		Descriptor.GetIntParameters()[ImageParamIndex].ParameterRangeValueNames[Index] = Image;
-	}
 
-	{
-		const int32 OpacityParamIndex = Descriptor.FindTypedParameterIndex(ParamName + FMultilayerProjector::OPACITY_PARAMETER_POSTFIX, EMutableParameterType::Float);
-		Descriptor.GetFloatParameters()[OpacityParamIndex].ParameterRangeValues[Index] = Opacity;
-	}
+	Descriptor.SetProjectorValue(ParamName, Position, Direction, Up, Scale, Angle, LayerIndex);
+	Descriptor.SetIntParameterSelectedOption(ParamName + FMultilayerProjector::IMAGE_PARAMETER_POSTFIX, Image, LayerIndex);
+	Descriptor.SetFloatParameterSelectedOption(ParamName + FMultilayerProjector::OPACITY_PARAMETER_POSTFIX, Opacity, LayerIndex);
+
 }
 
 
@@ -110,135 +87,6 @@ const FString FMultilayerProjector::IMAGE_PARAMETER_POSTFIX = FString("_Selected
 
 const FString FMultilayerProjector::POSE_PARAMETER_POSTFIX = FString("_SelectedPoses");
 
-
-int32 FMultilayerProjector::NumLayers(const FCustomizableObjectInstanceDescriptor& Descriptor) const
-{
-	const FString NumLayersParamName = ParamName.ToString() + NUM_LAYERS_PARAMETER_POSTFIX;
-	
-	const int32 FloatParameterIndex = Descriptor.FindTypedParameterIndex(NumLayersParamName, EMutableParameterType::Float);
-	check(FloatParameterIndex != -1); // Parameter not found.
-
-	return Descriptor.GetFloatParameters()[FloatParameterIndex].ParameterValue;
-}
-
-
-void FMultilayerProjector::CreateLayer(FCustomizableObjectInstanceDescriptor& Descriptor, const int32 Index) const
-{
-	check(Index >= 0 && Index <= NumLayers(Descriptor)); // Layer is non-contiguous or out of range.
-
-	const FString ParamNameString = ParamName.ToString();
-
-	if (!AreDescriptorParametersValid(Descriptor, ParamNameString))
-	{
-		ensureAlwaysMsgf(false, TEXT("%s"), *FMultilayerProjector::DESCRIPTOR_PARAMETERS_INVALID);
-		return;
-	}
-
-	const UCustomizableObject* Object = Descriptor.GetCustomizableObject(); 
-	
-	if (!Object)
-	{
-		return;
-	}
-	
-	// Num Layers.
-	{
-        const FString NumLayersParamName = ParamNameString + NUM_LAYERS_PARAMETER_POSTFIX;
-		const int32 FloatParameterIndex = Descriptor.FindTypedParameterIndex(NumLayersParamName, EMutableParameterType::Float);
-
-        Descriptor.GetFloatParameters()[FloatParameterIndex].ParameterValue += 1;
-    }
-	
-	// Projector Range.
-	{
-		const int32 ProjectorParameterIndex = Descriptor.FindTypedParameterIndex(ParamNameString, EMutableParameterType::Projector);
-		
-		FCustomizableObjectProjectorParameterValue& ProjectorParameter = Descriptor.GetProjectorParameters()[ProjectorParameterIndex];
-		const FCustomizableObjectProjector Projector = Descriptor.GetCustomizableObject()->GetProjectorParameterDefaultValue(ParamNameString);
-		ProjectorParameter.RangeValues.Insert(Projector, Index);
-	}
-	
-	// Selected Image Range.
-	{
-		const int32 IntParameterIndex = Descriptor.FindTypedParameterIndex(ParamNameString + FMultilayerProjector::IMAGE_PARAMETER_POSTFIX, EMutableParameterType::Int);
-
-		FCustomizableObjectIntParameterValue& IntParameter = Descriptor.GetIntParameters()[IntParameterIndex];
-		const int32 ParamIndexInObject = Object->FindParameter(IntParameter.ParameterName);
-
-		const FString DefaultValue = Object->GetIntParameterAvailableOption(ParamIndexInObject, 0); // TODO: Define the default option in the editor instead of taking the first available, like it's currently defined for GetProjectorDefaultValue()
-		IntParameter.ParameterRangeValueNames.Insert(DefaultValue, Index);
-	}
-	
-	// Opacity Range.
-	{
-		const int32 FloatParameterIndex = Descriptor.FindTypedParameterIndex(ParamNameString + FMultilayerProjector::OPACITY_PARAMETER_POSTFIX, EMutableParameterType::Float);
-
-		FCustomizableObjectFloatParameterValue& FloatParameter = Descriptor.GetFloatParameters()[FloatParameterIndex];
-		FloatParameter.ParameterRangeValues.Insert(0.5, Index); // TODO: Define the default float in the editor instead of [0.5f], like it's currently defined for GetProjectorDefaultValue()
-	}
-}
-
-
-void FMultilayerProjector::RemoveLayerAt(FCustomizableObjectInstanceDescriptor& Descriptor, const int32 Index) const
-{
-	check(Index >= 0 && Index < NumLayers(Descriptor)); // Layer out of range.
-
-	const FString ParamNameString = ParamName.ToString();
-
-	if (!AreDescriptorParametersValid(Descriptor, ParamName.ToString()))
-	{
-		ensureAlwaysMsgf(false, TEXT("%s"), *FMultilayerProjector::DESCRIPTOR_PARAMETERS_INVALID);
-		return;
-	}
-	
-	// Num Layers.
-	{
-		const FString NumLayersParamName = ParamNameString + NUM_LAYERS_PARAMETER_POSTFIX;
-        const int32 FloatParameterIndex = Descriptor.FindTypedParameterIndex(NumLayersParamName, EMutableParameterType::Float);
-		
-    	Descriptor.GetFloatParameters()[FloatParameterIndex].ParameterValue -= 1;
-    }
-    
-	// Projector Range.
-	{
-		const int32 ProjectorParameterIndex = Descriptor.FindTypedParameterIndex(ParamNameString, EMutableParameterType::Projector);
-		
-		FCustomizableObjectProjectorParameterValue& ProjectorParameter = Descriptor.GetProjectorParameters()[ProjectorParameterIndex];
-		ProjectorParameter.RangeValues.RemoveAt(Index);
-	}
-	
-	// Selected Image Range.
-	{
-		const int32 IntParameterIndex = Descriptor.FindTypedParameterIndex(ParamNameString + IMAGE_PARAMETER_POSTFIX, EMutableParameterType::Int);
-		
-		FCustomizableObjectIntParameterValue& IntParameter = Descriptor.GetIntParameters()[IntParameterIndex];
-		IntParameter.ParameterRangeValueNames.RemoveAt(Index);
-	}
-	
-	// Opacity Range.
-	{
-		const int32 FloatParameterIndex = Descriptor.FindTypedParameterIndex(ParamNameString + OPACITY_PARAMETER_POSTFIX, EMutableParameterType::Float);
-		
-		FCustomizableObjectFloatParameterValue& FloatParameter = Descriptor.GetFloatParameters()[FloatParameterIndex];
-		FloatParameter.ParameterRangeValues.RemoveAt(Index);
-	}
-}
-
-
-FMultilayerProjectorLayer FMultilayerProjector::GetLayer(const FCustomizableObjectInstanceDescriptor& Descriptor, const int32 Index) const
-{
-	FMultilayerProjectorLayer MultilayerProjectorLayer;
-	MultilayerProjectorLayer.Read(Descriptor, *this, Index);
-	return MultilayerProjectorLayer;
-}
-
-
-void FMultilayerProjector::UpdateLayer(FCustomizableObjectInstanceDescriptor& Descriptor, const int32 Index, const FMultilayerProjectorLayer& Layer) const
-{
-	Layer.Write(Descriptor, *this, Index);
-}
-
-
 TArray<FName> FMultilayerProjector::GetVirtualLayers() const
 {
 	TArray<FName> VirtualLayers;
@@ -251,9 +99,9 @@ void FMultilayerProjector::CreateVirtualLayer(FCustomizableObjectInstanceDescrip
 {
 	if (!VirtualLayersMapping.Contains(Id))
 	{
-		const int32 Index = NumLayers(Descriptor);
+		const int32 Index = Descriptor.NumProjectorLayers(ParamName);
 		
-		CreateLayer(Descriptor, Index);
+		Descriptor.CreateLayer(ParamName, Index);
 		VirtualLayersMapping.Add(Id, Index);
 		VirtualLayersOrder.Add(Id, NEW_VIRTUAL_LAYER_ORDER);
 	}
@@ -275,7 +123,7 @@ FMultilayerProjectorVirtualLayer FMultilayerProjector::FindOrCreateVirtualLayer(
 		}
 		else
 		{
-			Layer = GetLayer(Descriptor, *Index);
+			Layer = Descriptor.GetLayer(ParamName, *Index);
 			bEnabled = true;
 		}
 
@@ -283,14 +131,14 @@ FMultilayerProjectorVirtualLayer FMultilayerProjector::FindOrCreateVirtualLayer(
 	}
 	else
 	{
-		const int32 NewIndex = NumLayers(Descriptor);
+		const int32 NewIndex = Descriptor.NumProjectorLayers(ParamName);
 		constexpr int32 NewOrder = NEW_VIRTUAL_LAYER_ORDER;
 		
-		CreateLayer(Descriptor, NewIndex);
+		Descriptor.CreateLayer(ParamName, NewIndex);
 		VirtualLayersMapping.Add(Id, NewIndex);
 		VirtualLayersOrder.Add(Id, NewOrder);
 
-		Layer = GetLayer(Descriptor, NewIndex);
+		Layer = Descriptor.GetLayer(ParamName, NewIndex);
 		bEnabled = true;
 		Order = NewOrder;
 	}
@@ -310,7 +158,7 @@ void FMultilayerProjector::RemoveVirtualLayer(FCustomizableObjectInstanceDescrip
 	}
 	else
 	{
-		RemoveLayerAt(Descriptor, *Index);
+		Descriptor.RemoveLayerAt(ParamName, *Index);
 		
 		for (TMap<FName, int32>::TIterator It =  VirtualLayersMapping.CreateIterator(); It; ++It)
 		{
@@ -334,7 +182,7 @@ FMultilayerProjectorVirtualLayer FMultilayerProjector::GetVirtualLayer(const FCu
 	const int32* Index = VirtualLayersMapping.Find(Id);
 	check(Index); // Virtual Layer not created.
 
-	const FMultilayerProjectorLayer Layer = GetLayer(Descriptor, *Index);
+	const FMultilayerProjectorLayer Layer = Descriptor.GetLayer(ParamName, *Index);
 	const bool bEnabled = *Index != VIRTUAL_LAYER_DISABLED;
 	const int32 Order = VirtualLayersOrder[Id];
 	
@@ -363,7 +211,7 @@ void FMultilayerProjector::UpdateVirtualLayer(FCustomizableObjectInstanceDescrip
 			int32 NewIndex = CalculateVirtualLayerIndex(Id, Layer.Order);
 			if (OldIndex != NewIndex) // Move required. Could be optimized by moving only the in-between values.
 			{
-				RemoveLayerAt(Descriptor, OldIndex);
+				Descriptor.RemoveLayerAt(ParamName, OldIndex);
 				UpdateMappingVirtualLayerDisabled(Id, OldIndex);
 
 				if (OldIndex < NewIndex)
@@ -371,14 +219,14 @@ void FMultilayerProjector::UpdateVirtualLayer(FCustomizableObjectInstanceDescrip
 					NewIndex -= 1;
 				}
 				
-				CreateLayer(Descriptor, NewIndex);
+				Descriptor.CreateLayer(ParamName, NewIndex);
 				UpdateMappingVirtualLayerEnabled(Id, NewIndex);
 			}
 			
 			*Order = Layer.Order;
 		}
 
-		UpdateLayer(Descriptor, *Index, static_cast<FMultilayerProjectorLayer>(Layer)); // Update enabled layer.
+		Descriptor.UpdateLayer(ParamName, *Index, static_cast<FMultilayerProjectorLayer>(Layer)); // Update enabled layer.
 	}
 	
 	// Enable or disable virtual layer.
@@ -386,16 +234,16 @@ void FMultilayerProjector::UpdateVirtualLayer(FCustomizableObjectInstanceDescrip
 	{
 		const int32 NewIndex = CalculateVirtualLayerIndex(Id, VirtualLayersOrder[Id]);
 
-		CreateLayer(Descriptor, NewIndex);
+		Descriptor.CreateLayer(ParamName, NewIndex);
 		UpdateMappingVirtualLayerEnabled(Id, NewIndex);
 
-		UpdateLayer(Descriptor, NewIndex, static_cast<FMultilayerProjectorLayer>(Layer));
+		Descriptor.UpdateLayer(ParamName, NewIndex, static_cast<FMultilayerProjectorLayer>(Layer));
 		
 		DisableVirtualLayers.Remove(Id);
 	}
 	else if (!Layer.bEnabled && bEnabled)
 	{
-		RemoveLayerAt(Descriptor, *Index);
+		Descriptor.RemoveLayerAt(ParamName, *Index);
 		UpdateMappingVirtualLayerDisabled(Id, *Index);
 		
 		DisableVirtualLayers.Add(Id, Layer);
@@ -467,43 +315,6 @@ void FMultilayerProjector::UpdateMappingVirtualLayerDisabled(const FName& Id, co
 }
 
 
-bool FMultilayerProjector::AreDescriptorParametersValid(const FCustomizableObjectInstanceDescriptor& Descriptor, const FString& ParamName)
-{
-	// Num layers.
-	{
-		const FString NumLayersParamName = ParamName + NUM_LAYERS_PARAMETER_POSTFIX;
-		const int32 FloatParameterIndex = Descriptor.FindTypedParameterIndex(NumLayersParamName, EMutableParameterType::Float);
-		if (FloatParameterIndex != INDEX_NONE)
-		{
-			return false;
-		}
-	}
-    
-	// Projector.
-	if (const int32 ProjectorParameterIndex = Descriptor.FindTypedParameterIndex(ParamName, EMutableParameterType::Projector);
-		ProjectorParameterIndex != INDEX_NONE)
-	{
-		return false;
-	}
-	
-	// Selected Image.
-	if (const int32 IntParameterIndex = Descriptor.FindTypedParameterIndex(ParamName + IMAGE_PARAMETER_POSTFIX, EMutableParameterType::Int);
-		IntParameterIndex != INDEX_NONE)
-	{
-		return false;
-	}
-		
-	// Opacity.	
-	if (const int32 FloatParameterIndex = Descriptor.FindTypedParameterIndex(ParamName + OPACITY_PARAMETER_POSTFIX, EMutableParameterType::Float);
-		FloatParameterIndex != INDEX_NONE)
-	{
-		return false;
-	}
-
-	return true;
-}
-
-
 uint32 GetTypeHash(const FMultilayerProjector& Key)
 {
 	uint32 Hash = GetTypeHash(Key.ParamName);
diff --git a/Engine/Plugins/Experimental/Mutable/Source/CustomizableObject/Public/MuCO/CustomizableObjectInstance.h b/Engine/Plugins/Experimental/Mutable/Source/CustomizableObject/Public/MuCO/CustomizableObjectInstance.h
--- a/Engine/Plugins/Experimental/Mutable/Source/CustomizableObject/Public/MuCO/CustomizableObjectInstance.h
+++ b/Engine/Plugins/Experimental/Mutable/Source/CustomizableObject/Public/MuCO/CustomizableObjectInstance.h
@@ -295,6 +295,10 @@ public:
 	UFUNCTION(BlueprintCallable, Category = CustomizableObjectInstance)
 	void SetRandomValuesFromStream(const FRandomStream& InStream);
 
+	/** Sets a parameter to its default value */
+	UFUNCTION(BlueprintCallable, Category = CustomizableObjectInstance)
+	void SetDefaultValue(const FString& ParamName);
+
 	/** Set all parameters to their default value. */
 	UFUNCTION(BlueprintCallable, Category = CustomizableObjectInstance)
 	void SetDefaultValues();
diff --git a/Engine/Plugins/Experimental/Mutable/Source/CustomizableObject/Public/MuCO/MultilayerProjector.h b/Engine/Plugins/Experimental/Mutable/Source/CustomizableObject/Public/MuCO/MultilayerProjector.h
--- a/Engine/Plugins/Experimental/Mutable/Source/CustomizableObject/Public/MuCO/MultilayerProjector.h
+++ b/Engine/Plugins/Experimental/Mutable/Source/CustomizableObject/Public/MuCO/MultilayerProjector.h
@@ -24,7 +24,7 @@ struct FMultilayerProjectorLayer
 	 * @param MultilayerProjector Multilayer Projector helper.
 	 * @param Index Index where to read the Layer in the Instance Parameters.
 	 */
-	void Read(const FCustomizableObjectInstanceDescriptor& Descriptor, const FMultilayerProjector& MultilayerProjector, int32 Index);
+	void Read(const FCustomizableObjectInstanceDescriptor& Descriptor, const FString& ParamName, int32 Index);
 
 	/** Write to Layer to the Instance Parameters.
 	 *
@@ -32,7 +32,7 @@ struct FMultilayerProjectorLayer
 	 * @param MultilayerProjector Multilayer Projector helper.
 	 * @param Index Index where to write the Layer in the Instance Parameters.
 	 */
-	void Write(FCustomizableObjectInstanceDescriptor& Descriptor, const FMultilayerProjector& MultilayerProjector, int32 Index) const;
+	void Write(FCustomizableObjectInstanceDescriptor& Descriptor, const FString& ParamName, int32 Index) const;
 	
 	/** Layer position. */
 	FVector Position;
@@ -78,10 +78,7 @@ struct FMultilayerProjectorVirtualLayer : public FMultilayerProjectorLayer
 };
 
 
-/** Multilayer Projector Helper. Eases the management of Layers and Virtual Layers.
- *
- * Layer: Management of Multilayer Projector Layers by index.
- * All layers indices has to be consecutive.
+/** Multilayer Projector Helper. Eases the management of Virtual Layers.
  * 
  * Virtual Layer: Management of Multilayer Projector Layers by name.
  * - Allows to enabled and disabled layers.
@@ -97,7 +94,7 @@ struct CUSTOMIZABLEOBJECT_API FMultilayerProjector
 	friend uint32 GetTypeHash(const FMultilayerProjector& Key);
 
 	// Log texts
-	static const FString DESCRIPTOR_PARAMETERS_INVALID;
+	static const FString MULTILAYER_PROJECTOR_PARAMETERS_INVALID;
 
 	// Parameters encoding
 	static const FString NUM_LAYERS_PARAMETER_POSTFIX;
@@ -110,23 +107,6 @@ struct CUSTOMIZABLEOBJECT_API FMultilayerProjector
 
 	explicit FMultilayerProjector(const FName& InParamName);
 
-	// Layers
-	
-	/** Returns the number of layers. */
-	int32 NumLayers(const FCustomizableObjectInstanceDescriptor& Descriptor) const;
-
-	/** Insert the Layer at the given index moving all the following layers one position. */
-	void CreateLayer(FCustomizableObjectInstanceDescriptor& Descriptor, int32 Index) const;
-
-	/** Remove the Layer at the given index moving all the following layers one position. */
-	void RemoveLayerAt(FCustomizableObjectInstanceDescriptor& Descriptor, int32 Index) const;
-
-	/** Get the properties of the Layer at the given index. */
-	FMultilayerProjectorLayer GetLayer(const FCustomizableObjectInstanceDescriptor& Descriptor, int32 Index) const;
-
-	/** Update the Layer properties at the given index. */
-	void UpdateLayer(FCustomizableObjectInstanceDescriptor& Descriptor, int32 Index, const FMultilayerProjectorLayer& Layer) const;
-
 	// Virtual layers
 
 	/** Get all created Virtual Layers. */
@@ -194,9 +174,6 @@ private:
 	 * @param Index Index which the layer has been disabled.
 	 */
 	void UpdateMappingVirtualLayerDisabled(const FName& Id, int32 Index);
-	
-	/** Return false if the Instance does not contain the necessary Instance Parameters. */
-	static bool AreDescriptorParametersValid(const FCustomizableObjectInstanceDescriptor& Descriptor, const FString& ParamName);
 };
 
 
diff --git a/Engine/Plugins/Experimental/Mutable/Source/CustomizableObjectEditor/Private/MuCOE/CustomizableInstanceDetails.cpp b/Engine/Plugins/Experimental/Mutable/Source/CustomizableObjectEditor/Private/MuCOE/CustomizableInstanceDetails.cpp
--- a/Engine/Plugins/Experimental/Mutable/Source/CustomizableObjectEditor/Private/MuCOE/CustomizableInstanceDetails.cpp
+++ b/Engine/Plugins/Experimental/Mutable/Source/CustomizableObjectEditor/Private/MuCOE/CustomizableInstanceDetails.cpp
@@ -19,7 +19,6 @@
 #include "MuCO/CustomizableObject.h"
 #include "MuCO/CustomizableObjectInstance.h"
 #include "MuCO/CustomizableObjectSystem.h"
-#include "MuCOE/SCustomizableInstanceProperties.h"
 #include "MuCOE/UnrealEditorPortabilityHelpers.h"
 
 #include "Slate/DeferredCleanupSlateBrush.h"
@@ -30,6 +29,7 @@
 #include "Widgets/Input/SButton.h"
 #include "Widgets/Input/SCheckBox.h"
 #include "Widgets/Input/SSearchBox.h"
+#include "Widgets/Input/SSpinBox.h"
 #include "Widgets/Input/SSlider.h"
 #include "Widgets/Input/STextComboBox.h"
 #include "Widgets/Layout/SExpandableArea.h"
@@ -47,12 +47,6 @@ namespace UIMetadataKeyWords
 }
 
 
-TAutoConsoleVariable<bool> CVarUseOldInstanceUI(
-	TEXT("mutable.UseOldInstanceUI"),
-	false,
-	TEXT("Enables the old Parameters UI for the Instance Properties tab."));
-
-
 TSharedRef<IDetailCustomization> FCustomizableInstanceDetails::MakeInstance()
 {
 	return MakeShareable(new FCustomizableInstanceDetails);
@@ -118,19 +112,6 @@ void FCustomizableInstanceDetails::CustomizeDetails(const TSharedPtr<IDetailLayo
 	// Store which Texture Parameter values can be selected.
 	GenerateTextureParameterOptions();
 
-	// In case that something of the new UI doesn't work as expected
-	if(CVarUseOldInstanceUI.GetValueOnGameThread())
-	{
-		OldParametersCategory.AddCustomRow(LOCTEXT("CustomizableInstanceDetails_OldUI", "Old Instance Parameters"))
-		[
-			SNew(SCustomizableInstanceProperties)
-			.CustomInstance(CustomInstance.Get())
-			.InstanceDetails(SharedThis(this))
-		];
-
-		return;
-	}
-
 	// State Selector Widget
 	VisibilitySettingsCategory.AddCustomRow(LOCTEXT("CustomizableInstanceDetails_StateSelector", "State"))
 	.NameContent()
@@ -1410,7 +1391,7 @@ FReply FCustomizableInstanceDetails::OnProjectorCopyTransform(const FString Para
 
 	if ((ParameterIndexInObject >= 0) && (ProjectorParamIndex >= 0))
 	{
-		TArray<FCustomizableObjectProjectorParameterValue>& ProjectorParameters = CustomInstance->GetPrivate()->GetDescriptor().GetProjectorParameters();
+		const TArray<FCustomizableObjectProjectorParameterValue>& ProjectorParameters = CustomInstance->GetPrivate()->GetDescriptor().GetProjectorParameters();
 		FCustomizableObjectProjector Value;
 
 		if (RangeIndex == -1)
@@ -1493,7 +1474,7 @@ TSharedRef<SWidget> FCustomizableInstanceDetails::GenerateMultidimensionalProjec
 	FString ParamName = CustomizableObject->GetParameterName(ParamIndexInObject);
 
 	const TSharedPtr<ICustomizableObjectInstanceEditor> Editor = GetEditorChecked();
-	TArray<FCustomizableObjectProjectorParameterValue>& ProjectorParameters = CustomInstance->GetPrivate()->GetDescriptor().GetProjectorParameters();
+	const TArray<FCustomizableObjectProjectorParameterValue>& ProjectorParameters = CustomInstance->GetPrivate()->GetDescriptor().GetProjectorParameters();
 	const int32 ProjectorParamIndex = CustomInstance->FindProjectorParameterNameIndex(ParamName);
 
 	check(ProjectorParamIndex < ProjectorParameters.Num());
@@ -1821,27 +1802,8 @@ FReply FCustomizableInstanceDetails::OnProjectorLayerAdded(FString ParamName)
 	FScopedTransaction LocalTransaction(LOCTEXT("AddProjectorLayer", "Add Projector Layer"));
 	CustomInstance->Modify();
 
-	const int32 NumLayers = CustomInstance->AddValueToProjectorRange(ParamName) + 1;
-	if (NumLayers == 0)
-	{
-		return FReply::Handled();
-	}
-
-	const FString TextureSwitchEnumParamName = ParamName + FMultilayerProjector::IMAGE_PARAMETER_POSTFIX;
-	CustomInstance->AddValueToIntRange(TextureSwitchEnumParamName);
-	check(CustomInstance->FindIntParameterNameIndex(TextureSwitchEnumParamName) != INDEX_NONE);
-	check(NumLayers == CustomInstance->GetPrivate()->GetDescriptor().GetIntParameters()[CustomInstance->FindIntParameterNameIndex(TextureSwitchEnumParamName)].ParameterRangeValueNames.Num());
-
-	TArray<FCustomizableObjectFloatParameterValue>& FloatParameters = CustomInstance->GetPrivate()->GetDescriptor().GetFloatParameters();
-
-	const FString NumLayersParamName = ParamName + FMultilayerProjector::NUM_LAYERS_PARAMETER_POSTFIX;
-	check(CustomInstance->FindFloatParameterNameIndex(NumLayersParamName) != INDEX_NONE);
-	FloatParameters[CustomInstance->FindFloatParameterNameIndex(NumLayersParamName)].ParameterValue = NumLayers;
-
-	const FString OpacitySliderParamName = ParamName + FMultilayerProjector::OPACITY_PARAMETER_POSTFIX;
-	CustomInstance->AddValueToFloatRange(OpacitySliderParamName);
-	check(CustomInstance->FindFloatParameterNameIndex(OpacitySliderParamName) != INDEX_NONE);
-	check(NumLayers == FloatParameters[CustomInstance->FindFloatParameterNameIndex(OpacitySliderParamName)].ParameterRangeValues.Num());
+	const int32 NumLayers = CustomInstance->MultilayerProjectorNumLayers(*ParamName) + 1;
+	CustomInstance->MultilayerProjectorCreateLayer(*ParamName, NumLayers);
 
 	UpdateInstance();
 
@@ -1862,30 +1824,7 @@ FReply FCustomizableInstanceDetails::OnProjectorLayerRemoved(const FString Param
 		GetEditorChecked()->HideGizmo();
 	}
 
-	const int32 projectorParameterIndex = CustomInstance->FindProjectorParameterNameIndex(ParamName);
-	if (projectorParameterIndex == INDEX_NONE
-		|| CustomInstance->GetPrivate()->GetDescriptor().GetProjectorParameters()[projectorParameterIndex].RangeValues.Num() <= 0)
-	{
-		return FReply::Handled();
-	}
-
-	const int32 NumLayers = CustomInstance->RemoveValueFromProjectorRange(ParamName, RangeIndex) + 1;
-
-	const FString TextureSwitchEnumParamName = ParamName + FMultilayerProjector::IMAGE_PARAMETER_POSTFIX;
-	CustomInstance->RemoveValueFromIntRange(TextureSwitchEnumParamName, RangeIndex);
-	check(CustomInstance->FindIntParameterNameIndex(TextureSwitchEnumParamName) != INDEX_NONE);
-	check(NumLayers == CustomInstance->GetPrivate()->GetDescriptor().GetIntParameters()[CustomInstance->FindIntParameterNameIndex(TextureSwitchEnumParamName)].ParameterRangeValueNames.Num());
-
-	TArray<FCustomizableObjectFloatParameterValue>& FloatParameters = CustomInstance->GetPrivate()->GetDescriptor().GetFloatParameters();
-
-	const FString NumLayersParamName = ParamName + FMultilayerProjector::NUM_LAYERS_PARAMETER_POSTFIX;
-	check(CustomInstance->FindFloatParameterNameIndex(NumLayersParamName) != INDEX_NONE);
-	FloatParameters[CustomInstance->FindFloatParameterNameIndex(NumLayersParamName)].ParameterValue = NumLayers;
-
-	const FString OpacitySliderParamName = ParamName + FMultilayerProjector::OPACITY_PARAMETER_POSTFIX;
-	CustomInstance->RemoveValueFromFloatRange(OpacitySliderParamName, RangeIndex);
-	check(CustomInstance->FindFloatParameterNameIndex(OpacitySliderParamName) != INDEX_NONE);
-	check(NumLayers == FloatParameters[CustomInstance->FindFloatParameterNameIndex(OpacitySliderParamName)].ParameterRangeValues.Num());
+	CustomInstance->MultilayerProjectorRemoveLayerAt(*ParamName, RangeIndex);
 
 	UpdateInstance();
 
@@ -2010,13 +1949,7 @@ FReply FCustomizableInstanceDetails::OnResetAllParameters()
 	FScopedTransaction LocalTransaction(LOCTEXT("OnResetAllParameters", "Reset All Parameters"));
 	CustomInstance->Modify();
 
-	const UCustomizableObject* CustomObject = CustomInstance->GetCustomizableObject();
-	const int32 NumObjectParameter = CustomObject->GetParameterCount();
-
-	for (int32 ParameterIndex = 0; ParameterIndex < NumObjectParameter; ++ParameterIndex)
-	{
-		SetParameterValueToDefault(ParameterIndex);
-	}
+	CustomInstance->SetDefaultValues();
 
 	CustomInstance->GetPrivate()->SelectedProfileIndex = INDEX_NONE;
 	UpdateInstance();
@@ -2034,124 +1967,10 @@ void FCustomizableInstanceDetails::OnResetParameterButtonClicked(int32 Parameter
 
 	FScopedTransaction LocalTransaction(FText::Format(LOCTEXT("OnResetParameter", "Reset Parameter: {0}"), FText::FromString(ParameterName)));
 	CustomInstance->Modify();
-	SetParameterValueToDefault(ParameterIndex);
+	CustomInstance->SetDefaultValue(ParameterName);
 	UpdateInstance();
 }
 
-
-void FCustomizableInstanceDetails::SetParameterValueToDefault(int32 ParameterIndex)
-{
-	UCustomizableObject* CustomObject = CustomInstance->GetCustomizableObject();
-	FString ParameterName = CustomObject->GetParameterName(ParameterIndex);
-	EMutableParameterType ParameterType = CustomObject->GetParameterType(ParameterIndex);
-
-	// Checking if there are parameteres with the same name
-	if (ParameterType != CustomObject->GetParameterTypeByName(ParameterName))
-	{
-		return;
-	}
-
-	switch (ParameterType)
-	{
-	case EMutableParameterType::None:
-		break;
-	case EMutableParameterType::Bool:
-	{
-		CustomInstance->SetBoolParameterSelectedOption(ParameterName, CustomObject->GetBoolParameterDefaultValue(ParameterName));
-		break;
-	}
-	case EMutableParameterType::Int:
-	{
-		int32 DefaultValue = CustomObject->GetIntParameterDefaultValue(ParameterName);
-		FString ValueName = CustomObject->FindIntParameterValueName(ParameterIndex, DefaultValue);
-
-		if (!ValueName.IsEmpty())
-		{
-			if (CustomObject->IsParameterMultidimensional(ParameterName))
-			{
-				int32 NumRanges = CustomInstance->GetIntValueRange(ParameterName);
-
-				for (int32 RangeIndex = 0; RangeIndex < NumRanges; ++RangeIndex)
-				{
-					CustomInstance->SetIntParameterSelectedOption(ParameterName, ValueName, RangeIndex);
-				}
-			}
-			else
-			{
-				CustomInstance->SetIntParameterSelectedOption(ParameterName, ValueName);
-			}
-		}
-		break;
-	}
-	case EMutableParameterType::Float:
-	{
-		float DefaultValue = CustomObject->GetFloatParameterDefaultValue(ParameterName);
-
-		if (CustomObject->IsParameterMultidimensional(ParameterName))
-		{
-			int32 NumRanges = CustomInstance->GetFloatValueRange(ParameterName);
-
-			for (int32 RangeIndex = 0; RangeIndex < NumRanges; ++RangeIndex)
-			{
-				CustomInstance->SetFloatParameterSelectedOption(ParameterName, DefaultValue, RangeIndex);
-			}
-		}
-		else
-		{
-			CustomInstance->SetFloatParameterSelectedOption(ParameterName, DefaultValue);
-		}
-		break;
-	}
-	case EMutableParameterType::Color:
-	{
-		CustomInstance->SetColorParameterSelectedOption(ParameterName, CustomObject->GetColorParameterDefaultValue(ParameterName));
-		break;
-	}
-	case EMutableParameterType::Projector:
-	{
-		FCustomizableObjectProjector DefaultValue = CustomObject->GetProjectorParameterDefaultValue(ParameterName);
-
-		if (CustomObject->IsParameterMultidimensional(ParameterName))
-		{
-			int32 NumRanges = CustomInstance->GetProjectorValueRange(ParameterName);
-
-			for (int32 RangeIndex = 0; RangeIndex < NumRanges; ++RangeIndex)
-			{
-				CustomInstance->SetProjectorValue(ParameterName, FVector(DefaultValue.Position), FVector(DefaultValue.Direction), FVector(DefaultValue.Up), FVector(DefaultValue.Scale), DefaultValue.Angle, RangeIndex);
-			}
-		}
-		else
-		{
-			CustomInstance->SetProjectorValue(ParameterName, FVector(DefaultValue.Position), FVector(DefaultValue.Direction), FVector(DefaultValue.Up), FVector(DefaultValue.Scale), DefaultValue.Angle);
-		}
-
-		break;
-	}
-	case EMutableParameterType::Texture:
-	{
-		FName DefaultValue = CustomObject->GetTextureParameterDefaultValue(ParameterName);
-
-		if (CustomObject->IsParameterMultidimensional(ParameterName))
-		{
-			int32 NumRanges = CustomInstance->GetTextureValueRange(ParameterName);
-
-			for (int32 RangeIndex = 0; RangeIndex < NumRanges; ++RangeIndex)
-			{
-				CustomInstance->SetTextureParameterSelectedOption(ParameterName, DefaultValue.ToString(), RangeIndex);
-			}
-		}
-		else
-		{
-			CustomInstance->SetTextureParameterSelectedOption(ParameterName, DefaultValue.ToString());
-		}
-		break;
-	}
-	}
-
-	return;
-}
-
-
 // TRANSACTION SYSTEM -----------------------------------------------------------------------------------------------------------------
 
 void FCustomizableInstanceDetails::BeginTransaction(const FText& TransactionDesc, bool bModifyCustomizableObject)
diff --git a/Engine/Plugins/Experimental/Mutable/Source/CustomizableObjectEditor/Private/MuCOE/CustomizableInstanceDetails.h b/Engine/Plugins/Experimental/Mutable/Source/CustomizableObjectEditor/Private/MuCOE/CustomizableInstanceDetails.h
--- a/Engine/Plugins/Experimental/Mutable/Source/CustomizableObjectEditor/Private/MuCOE/CustomizableInstanceDetails.h
+++ b/Engine/Plugins/Experimental/Mutable/Source/CustomizableObjectEditor/Private/MuCOE/CustomizableInstanceDetails.h
@@ -131,7 +131,6 @@ private:
 	FReply OnPasteAllParameters();
 	FReply OnResetAllParameters();
 	void OnResetParameterButtonClicked(int32 ParameterIndex);
-	void SetParameterValueToDefault(int32 ParameterIndex);
 
 	// Transaction System
 	void BeginTransaction(const FText& TransactionDesc, bool bModifyCustomizableObject = false);
diff --git a/Engine/Plugins/Experimental/Mutable/Source/CustomizableObjectEditor/Private/MuCOE/SCustomizableInstanceProperties.cpp b/Engine/Plugins/Experimental/Mutable/Source/CustomizableObjectEditor/Private/MuCOE/SCustomizableInstanceProperties.cpp
deleted file mode 100644
diff --git a/Engine/Plugins/Experimental/Mutable/Source/CustomizableObjectEditor/Private/MuCOE/StatePerformanceTesting.cpp b/Engine/Plugins/Experimental/Mutable/Source/CustomizableObjectEditor/Private/MuCOE/StatePerformanceTesting.cpp
--- a/Engine/Plugins/Experimental/Mutable/Source/CustomizableObjectEditor/Private/MuCOE/StatePerformanceTesting.cpp
+++ b/Engine/Plugins/Experimental/Mutable/Source/CustomizableObjectEditor/Private/MuCOE/StatePerformanceTesting.cpp
@@ -353,7 +353,7 @@ void FRuntimeTest::LoadStateData(int32 StateIndex)
 		}
 		case EMutableParameterType::Bool:
 		{
-			TArray<FCustomizableObjectBoolParameterValue>& BoolParameters = Instance->GetPrivate()->GetDescriptor().GetBoolParameters();
+			const TArray<FCustomizableObjectBoolParameterValue>& BoolParameters = Instance->GetPrivate()->GetDescriptor().GetBoolParameters();
 			FCustomizableObjectBoolParameterValue Parameter;
 			Parameter.ParameterName = RunTimeParameter;
 			const uint32 MaxIndex = BoolParameters.Num();
@@ -375,7 +375,7 @@ void FRuntimeTest::LoadStateData(int32 StateIndex)
 		}
 		case EMutableParameterType::Int:
 		{
-			TArray<FCustomizableObjectIntParameterValue>& IntParameters = Instance->GetPrivate()->GetDescriptor().GetIntParameters();
+			const TArray<FCustomizableObjectIntParameterValue>& IntParameters = Instance->GetPrivate()->GetDescriptor().GetIntParameters();
 			FCustomizableObjectIntParameterValue Parameter;
 			Parameter.ParameterName = RunTimeParameter;
 			const uint32 MaxIndex = IntParameters.Num();
@@ -402,7 +402,7 @@ void FRuntimeTest::LoadStateData(int32 StateIndex)
 		}
 		case EMutableParameterType::Float:
 		{
-			TArray<FCustomizableObjectFloatParameterValue>& FloatParameters = Instance->GetPrivate()->GetDescriptor().GetFloatParameters();
+			const TArray<FCustomizableObjectFloatParameterValue>& FloatParameters = Instance->GetPrivate()->GetDescriptor().GetFloatParameters();
 			FCustomizableObjectFloatParameterValue Parameter;
 			Parameter.ParameterName = RunTimeParameter;
 			const uint32 MaxIndex = FloatParameters.Num();
@@ -425,7 +425,7 @@ void FRuntimeTest::LoadStateData(int32 StateIndex)
 		}
 		case EMutableParameterType::Color:
 		{
-			TArray<FCustomizableObjectVectorParameterValue>& VectorParameters = Instance->GetPrivate()->GetDescriptor().GetVectorParameters();
+			const TArray<FCustomizableObjectVectorParameterValue>& VectorParameters = Instance->GetPrivate()->GetDescriptor().GetVectorParameters();
 			FCustomizableObjectVectorParameterValue Parameter;
 			Parameter.ParameterName = RunTimeParameter;
 			const uint32 MaxIndex = VectorParameters.Num();
@@ -447,7 +447,7 @@ void FRuntimeTest::LoadStateData(int32 StateIndex)
 		}
 		case EMutableParameterType::Projector:
 		{
-			TArray<FCustomizableObjectVectorParameterValue>& VectorParameters = Instance->GetPrivate()->GetDescriptor().GetVectorParameters();
+			const TArray<FCustomizableObjectVectorParameterValue>& VectorParameters = Instance->GetPrivate()->GetDescriptor().GetVectorParameters();
 			FCustomizableObjectVectorParameterValue Parameter;
 			Parameter.ParameterName = RunTimeParameter;
 			const uint32 MaxIndex = VectorParameters.Num();
@@ -470,7 +470,7 @@ void FRuntimeTest::LoadStateData(int32 StateIndex)
 		}
 		case EMutableParameterType::Texture:
 		{
-			TArray<FCustomizableObjectTextureParameterValue>& TextureParameters = Instance->GetPrivate()->GetDescriptor().GetTextureParameters();
+			const TArray<FCustomizableObjectTextureParameterValue>& TextureParameters = Instance->GetPrivate()->GetDescriptor().GetTextureParameters();
 			FCustomizableObjectTextureParameterValue Parameter;
 			Parameter.ParameterName = RunTimeParameter;
 			const uint32 MaxIndex = TextureParameters.Num();
diff --git a/Engine/Plugins/Experimental/Mutable/Source/CustomizableObjectEditor/Public/MuCOE/SCustomizableInstanceProperties.h b/Engine/Plugins/Experimental/Mutable/Source/CustomizableObjectEditor/Public/MuCOE/SCustomizableInstanceProperties.h
deleted file mode 100644
