[Mutable] Fix FUpdateContextPrivate being destroyed in a BackgroundThread.

FUpdateContextPrivate is not thread safe and should only be destoryed in the GameThread or the GarbageCollectorThread.

#rnx

diff --git a/Engine/Plugins/Experimental/Mutable/Source/CustomizableObject/Private/MuCO/CustomizableObjectSystem.cpp b/Engine/Plugins/Experimental/Mutable/Source/CustomizableObject/Private/MuCO/CustomizableObjectSystem.cpp
--- a/Engine/Plugins/Experimental/Mutable/Source/CustomizableObject/Private/MuCO/CustomizableObjectSystem.cpp
+++ b/Engine/Plugins/Experimental/Mutable/Source/CustomizableObject/Private/MuCO/CustomizableObjectSystem.cpp
@@ -2854,7 +2854,10 @@ namespace impl
 			const mu::Ptr<mu::System> MutableSystem = CustomizableObjectSystemPrivateData->MutableSystem;
 			CustomizableObjectSystemPrivateData->MutableTaskGraph.AddMutableThreadTask(
 				TEXT("Task_Mutable_ReleaseInstance"),
-				[OperationData, MutableSystem]() {Task_Mutable_ReleaseInstance(OperationData->InstanceID, MutableSystem, OperationData->bLiveUpdateMode); });
+				[InstanceID = OperationData->InstanceID, MutableSystem, bLiveUpdateMode = OperationData->bLiveUpdateMode]()
+				{
+					Task_Mutable_ReleaseInstance(InstanceID, MutableSystem, bLiveUpdateMode);
+				});
 		}
 
 
