[Mutable] Exposed new editor function to allow creating Customizable Objects.

#rnx

diff --git a/Engine/Plugins/Experimental/Mutable/Source/CustomizableObjectEditor/Private/MuCOE/CustomizableObjectEditorFunctionLibrary.cpp b/Engine/Plugins/Experimental/Mutable/Source/CustomizableObjectEditor/Private/MuCOE/CustomizableObjectEditorFunctionLibrary.cpp
--- a/Engine/Plugins/Experimental/Mutable/Source/CustomizableObjectEditor/Private/MuCOE/CustomizableObjectEditorFunctionLibrary.cpp
+++ b/Engine/Plugins/Experimental/Mutable/Source/CustomizableObjectEditor/Private/MuCOE/CustomizableObjectEditorFunctionLibrary.cpp
@@ -1,10 +1,17 @@
 // Copyright Epic Games, Inc. All Rights Reserved.
 
 #include "MuCOE/CustomizableObjectEditorFunctionLibrary.h"
+
+#include "AssetToolsModule.h"
+#include "ContentBrowserModule.h"
+#include "IContentBrowserSingleton.h"
+#include "AssetRegistry/AssetRegistryModule.h"
 #include "MuCOE/CustomizableObjectCompiler.h"
-#include "MuCO/CustomizableObjectSystem.h"
 #include "MuCO/CustomizableObject.h"
 #include "MuCO/ICustomizableObjectEditorModule.h"
+#include "MuCOE/CustomizableObjectFactory.h"
+#include "UObject/SavePackage.h"
+
 
 ECustomizableObjectCompilationState UCustomizableObjectEditorFunctionLibrary::CompileCustomizableObjectSynchronously(
 	UCustomizableObject* CustomizableObject,
@@ -45,3 +52,68 @@ ECustomizableObjectCompilationState UCustomizableObjectEditorFunctionLibrary::Co
 
 	return bCompilationSuccess ? ECustomizableObjectCompilationState::Completed : ECustomizableObjectCompilationState::Failed;
 }
+
+
+UCustomizableObject* UCustomizableObjectEditorFunctionLibrary::NewCustomizableObject(const FNewCustomizableObjectParameters& Parameters)
+{
+	const FString PackageName = Parameters.PackagePath + "/" + Parameters.AssetName;
+	if (FindPackage(nullptr, *PackageName))
+	{
+		UE_LOG(LogMutable, Error, TEXT("Package [%s] already exists."), *PackageName);
+		return nullptr;
+	}
+
+	FAssetToolsModule& AssetToolsModule = FModuleManager::Get().LoadModuleChecked<FAssetToolsModule>("AssetTools");
+	
+	FString OutAssetName;
+	FString OutPackageName;
+	AssetToolsModule.Get().CreateUniqueAssetName(PackageName, TEXT(""), OutPackageName, OutAssetName);
+
+	if (OutPackageName != PackageName)
+	{
+		UE_LOG(LogMutable, Error, TEXT("Invalid package name [%s]. Possible valid package name: [%s]"), *PackageName, *OutPackageName);
+		return nullptr;
+	}
+	
+	if (OutAssetName != Parameters.AssetName)
+	{
+		UE_LOG(LogMutable, Error, TEXT("Invalid asset name [%s]. Possible valid asset name: [%s]"), *Parameters.AssetName, *OutAssetName);
+		return nullptr;
+	}
+		
+	UPackage* Package = CreatePackage(*PackageName);
+	if (!Package)
+	{
+		UE_LOG(LogMutable, Error, TEXT("Could not create package [%s]."), *PackageName);
+		return nullptr;
+	}
+	
+	UCustomizableObjectFactory* Factory = NewObject<UCustomizableObjectFactory>();
+	Factory->CreationSettings.bEmptyObject = false;
+	Factory->CreationSettings.bIsChildObject = Parameters.ParentObject != nullptr;
+	Factory->CreationSettings.ParentObject = Parameters.ParentObject;
+	Factory->CreationSettings.GroupNodeName = Parameters.ParentGroupNode;
+
+	const FString PackagePath = FPackageName::GetLongPackagePath(PackageName);
+
+	UObject* Object = AssetToolsModule.Get().CreateAsset(Parameters.AssetName, PackagePath, UCustomizableObject::StaticClass(), Factory);
+	if (!Object)
+	{
+		UE_LOG(LogMutable, Error, TEXT("Could not create Asset [%s]."), *Parameters.AssetName);
+		return nullptr;
+	}
+
+	FSavePackageArgs SavePackageArgs;
+	SavePackageArgs.TopLevelFlags = RF_Public | RF_Standalone;
+	UPackage::Save(Package, Object, *FPackageName::LongPackageNameToFilename(PackageName, FPackageName::GetAssetPackageExtension()), SavePackageArgs);
+
+	FAssetRegistryModule::AssetCreated(Object);
+
+	FContentBrowserModule& ContentBrowserModule = FModuleManager::LoadModuleChecked<FContentBrowserModule>("ContentBrowser");
+
+	TArray<UObject*> Objects;
+	Objects.Add(Object);
+	ContentBrowserModule.Get().SyncBrowserToAssets(Objects);
+	
+	return CastChecked<UCustomizableObject>(Object);
+}
diff --git a/Engine/Plugins/Experimental/Mutable/Source/CustomizableObjectEditor/Public/MuCOE/CustomizableObjectEditorFunctionLibrary.h b/Engine/Plugins/Experimental/Mutable/Source/CustomizableObjectEditor/Public/MuCOE/CustomizableObjectEditorFunctionLibrary.h
--- a/Engine/Plugins/Experimental/Mutable/Source/CustomizableObjectEditor/Public/MuCOE/CustomizableObjectEditorFunctionLibrary.h
+++ b/Engine/Plugins/Experimental/Mutable/Source/CustomizableObjectEditor/Public/MuCOE/CustomizableObjectEditorFunctionLibrary.h
@@ -26,6 +26,29 @@ enum class ECustomizableObjectCompilationState : uint8
 };
 
 
+USTRUCT(BlueprintType)
+struct FNewCustomizableObjectParameters
+{
+	GENERATED_BODY()
+
+	/** Must not end with slash. For example "/Game" */
+	UPROPERTY(BlueprintReadWrite, Category = NewCustomizableObjectParameters)
+	FString PackagePath;
+
+	/** For example "SampleAssetName" */
+	UPROPERTY(BlueprintReadWrite, Category = NewCustomizableObjectParameters)
+	FString AssetName;
+
+	/** Parent to attach the child Customizable Object to. */
+	UPROPERTY(BlueprintReadWrite, Category = NewCustomizableObjectParameters)
+	TObjectPtr<UCustomizableObject> ParentObject;
+
+	/** Group to attach the child Customizable Object to. Only used if ParentObject is provided. */
+	UPROPERTY(BlueprintReadWrite, Category = NewCustomizableObjectParameters)
+	FString ParentGroupNode;
+};
+
+
 /**
  * Functions we want to be able to call on CustomizableObjects at edit time - could
  * be exposed to cook as well.
@@ -49,4 +72,8 @@ public:
 		UCustomizableObject* CustomizableObject, 
 		ECustomizableObjectOptimizationLevel OptimizationLevel = ECustomizableObjectOptimizationLevel::Minimal, 
 		ECustomizableObjectTextureCompression TextureCompression = ECustomizableObjectTextureCompression::Fast);
+
+	/** Create a new Customizable Object inside a package. */
+	UFUNCTION(BlueprintCallable, Category = "CustomizableObject")
+	static UCustomizableObject* NewCustomizableObject(const FNewCustomizableObjectParameters& Parameters);
 };
diff --git a/Engine/Plugins/Experimental/Mutable/Source/CustomizableObjectEditor/Public/MuCOE/CustomizableObjectFactory.h b/Engine/Plugins/Experimental/Mutable/Source/CustomizableObjectEditor/Public/MuCOE/CustomizableObjectFactory.h
--- a/Engine/Plugins/Experimental/Mutable/Source/CustomizableObjectEditor/Public/MuCOE/CustomizableObjectFactory.h
+++ b/Engine/Plugins/Experimental/Mutable/Source/CustomizableObjectEditor/Public/MuCOE/CustomizableObjectFactory.h
@@ -54,14 +54,13 @@ class UCustomizableObjectFactory : public UFactory
 
 	UCustomizableObjectFactory();
 
-	// Begin UObject Interface
-	virtual bool ConfigureProperties() override;
-
+public:
 	// Begin UFactory Interface
+	virtual bool ConfigureProperties() override;
 	virtual bool DoesSupportClass(UClass * Class) override;
 	virtual UClass* ResolveSupportedClass() override;
-	virtual UObject* FactoryCreateNew(UClass* Class,UObject* InParent,FName Name,EObjectFlags Flags,UObject* Context,FFeedbackContext* Warn) override;
-
+	virtual UObject* FactoryCreateNew(UClass* Class, UObject* InParent,FName Name,EObjectFlags Flags,UObject* Context,FFeedbackContext* Warn) override;
+	
 	FCustomizableObjectOptions CreationSettings;
 };
 
