[Mutable] Fix int parameter options selectors not working.

#rnx

diff --git a/Engine/Plugins/Experimental/Mutable/Source/CustomizableObjectEditor/Private/MuCOE/CustomizableInstanceDetails.cpp b/Engine/Plugins/Experimental/Mutable/Source/CustomizableObjectEditor/Private/MuCOE/CustomizableInstanceDetails.cpp
--- a/Engine/Plugins/Experimental/Mutable/Source/CustomizableObjectEditor/Private/MuCOE/CustomizableInstanceDetails.cpp
+++ b/Engine/Plugins/Experimental/Mutable/Source/CustomizableObjectEditor/Private/MuCOE/CustomizableInstanceDetails.cpp
@@ -918,7 +918,7 @@ TSharedRef<SWidget> FCustomizableInstanceDetails::GenerateIntWidget(const int32
 		return SNew(STextBlock).Text(LOCTEXT("MultidimensionalINTParameter_Text", "Multidimensional INT Parameter not supported"));
 	}
 
-	const int32 NumValues = CustomizableObject->GetIntParameterNumOptions(ParamIndexInObject) != 0;
+	const int32 NumValues = CustomizableObject->GetIntParameterNumOptions(ParamIndexInObject);
 	if (!NumValues)
 	{
 		return SNew(STextBlock).Text(LOCTEXT("NoAvailableOptions", "No Available Options"));
