[mutable] Avoid unnecessary data copies when cooking.
- remove size field from the model streamable block because it wasn't necessary in-game.
#rnx

diff --git a/Engine/Plugins/Experimental/Mutable/Source/CustomizableObject/Internal/MuCO/CustomizableObjectPrivate.h b/Engine/Plugins/Experimental/Mutable/Source/CustomizableObject/Internal/MuCO/CustomizableObjectPrivate.h
--- a/Engine/Plugins/Experimental/Mutable/Source/CustomizableObject/Internal/MuCO/CustomizableObjectPrivate.h
+++ b/Engine/Plugins/Experimental/Mutable/Source/CustomizableObject/Internal/MuCO/CustomizableObjectPrivate.h
@@ -492,9 +492,13 @@ struct FRealTimeMorphStreamable
 	UPROPERTY()
 	FMutableStreamableBlock Block;
 
+	UPROPERTY()
+	uint32 Size=0;
+
 	friend FArchive& operator<<(FArchive& Ar, FRealTimeMorphStreamable& Elem)
 	{
 		Ar << Elem.NameResolutionMap;
+		Ar << Elem.Size;
 		Ar << Elem.Block;
 
 		return Ar;
@@ -516,6 +520,9 @@ struct FClothingStreamable
 	UPROPERTY()
 	int32 PhysicsAssetIndex = INDEX_NONE;
 
+	UPROPERTY()
+	uint32 Size = 0;
+
 	UPROPERTY()
 	FMutableStreamableBlock Block;
 
@@ -524,6 +531,7 @@ struct FClothingStreamable
 		Ar << Elem.ClothingAssetIndex;
 		Ar << Elem.ClothingAssetLOD;
 		Ar << Elem.PhysicsAssetIndex;
+		Ar << Elem.Size;
 		Ar << Elem.Block;
 
 		return Ar;
@@ -740,7 +748,7 @@ struct FModelResources
 	
 	/** Map of Hash to Streaming blocks, used to stream a block of data representing a resource from the BulkData */
 	UPROPERTY()
-	TMap<uint64, FMutableStreamableBlock> HashToStreamableBlock;
+	TMap<uint32, FMutableStreamableBlock> HashToStreamableBlock;
 
 	/** Max number of components in the compiled Model. */
 	UPROPERTY()
@@ -759,13 +767,39 @@ struct FModelResources
 	uint8 FirstLODAvailable = 0;
 };
 
+
+struct FModelStreamableData
+{
+	void Get(uint32 Key, uint8* Destination) const
+	{
+		check(Destination);
+		const TArray64<uint8>* Buffer = Data.Find(Key);
+		check(Buffer);
+		FMemory::Memcpy(Destination, Buffer->GetData(), Buffer->Num());
+	}
+
+	void Set(uint32 Key, const uint8* Source, int64 Size)
+	{
+		check(Source);
+		check(Size);
+		TArray64<uint8>& Buffer = Data.Add(Key);
+		check(Buffer.Num()==0);
+		Buffer.SetNumUninitialized(Size);
+		FMemory::Memcpy(Buffer.GetData(), Source, Size);
+	}
+
+	// Temp, to be replaced with disk storage
+	TMap<uint32, TArray64<uint8> > Data;
+};
+
+
 struct CUSTOMIZABLEOBJECT_API FMutableCachedPlatformData
 {
 	/** */
 	TArray64<uint8> ModelData;
 
 	/** */
-	TArray64<uint8> StreamableData;
+	FModelStreamableData ModelStreamableData;
 
 	/** */
 	TArray64<uint8> MorphData;
@@ -869,9 +903,6 @@ public:
 
 	/** Load compiled data for the running platform from disk, this is used to load Editor Compilations. */
 	void LoadCompiledDataFromDisk();
-
-	/** Cache platform data for cook */
-	void CachePlatformData(const ITargetPlatform* InTargetPlatform, TArray64<uint8>& InObjectBytes, TArray64<uint8>& InBulkBytes, TArray64<uint8>& InMorphBytes, TArray64<uint8>& InClothingBytes);
 	
 	/** Loads data previously compiled in BeginCacheForCookedPlatformData onto the UProperties in *this,
 	  * in preparation for saving the cooked package for *this or for a CustomizableObjectInstance using *this.
@@ -1027,6 +1058,8 @@ public:
 
 		MeshReferenceSupport,
 
+		ImproveMemoryUsageForStreamableBlocks,
+
 		LastCustomizableObjectVersion
 	};
 	static constexpr int32 CurrentSupportedVersion = ECustomizableObjectVersions::LastCustomizableObjectVersion;
diff --git a/Engine/Plugins/Experimental/Mutable/Source/CustomizableObject/Internal/MuCO/UnrealMutableModelDiskStreamer.h b/Engine/Plugins/Experimental/Mutable/Source/CustomizableObject/Internal/MuCO/UnrealMutableModelDiskStreamer.h
--- a/Engine/Plugins/Experimental/Mutable/Source/CustomizableObject/Internal/MuCO/UnrealMutableModelDiskStreamer.h
+++ b/Engine/Plugins/Experimental/Mutable/Source/CustomizableObject/Internal/MuCO/UnrealMutableModelDiskStreamer.h
@@ -3,6 +3,7 @@
 #pragma once
 
 #include "MuCO/CustomizableObject.h"
+#include "MuCO/CustomizableObjectPrivate.h"
 #include "MuR/Serialisation.h"
 
 class FArchive;
@@ -55,7 +56,7 @@ public:
 	void EndStreaming();
 
 	// mu::ModelReader interface
-	OPERATION_ID BeginReadBlock(const mu::Model*, uint64 key0, void* pBuffer, uint64 size, TFunction<void(bool bSuccess)>* CompletionCallback) override;
+	OPERATION_ID BeginReadBlock(const mu::Model*, uint32 key0, void* pBuffer, uint64 size, TFunction<void(bool bSuccess)>* CompletionCallback) override;
 	bool IsReadCompleted(OPERATION_ID) override;
 	void EndRead(OPERATION_ID) override;
 
@@ -72,7 +73,7 @@ protected:
 	{
 		TWeakPtr<const mu::Model> Model;
 		TMap<OPERATION_ID, FReadRequest> CurrentReadRequests;
-		TMap<uint64, FMutableStreamableBlock> StreamableBlocks;
+		TMap<uint32, FMutableStreamableBlock> StreamableBlocks;
 		FString BulkFilePrefix;
 
 		TMap<uint32, TSharedPtr<IAsyncReadFileHandle>> ReadFileHandles;
@@ -104,14 +105,14 @@ private:
 
 
 // Implementation of a mutable streamer using bulk storage.
-class CUSTOMIZABLEOBJECT_API FUnrealMutableModelBulkWriter : public mu::ModelWriter
+class CUSTOMIZABLEOBJECT_API FUnrealMutableModelBulkWriterEditor : public mu::ModelWriter
 {
 public:
 	// 
-	FUnrealMutableModelBulkWriter(FArchive* InMainDataArchive = nullptr, FArchive* InStreamedDataArchive = nullptr);
+	FUnrealMutableModelBulkWriterEditor(FArchive* InMainDataArchive = nullptr, FArchive* InStreamedDataArchive = nullptr);
 
 	// mu::ModelWriter interface
-	void OpenWriteFile(uint64 key0) override;
+	void OpenWriteFile(uint32 BlockKey) override;
 	void Write(const void* pBuffer, uint64 size) override;
 	void CloseWriteFile() override;
 
@@ -127,4 +128,30 @@ protected:
 
 };
 
+
+// Implementation of a mutable streamer using bulk storage.
+class CUSTOMIZABLEOBJECT_API FUnrealMutableModelBulkWriterCook : public mu::ModelWriter
+{
+public:
+	// 
+	FUnrealMutableModelBulkWriterCook(FArchive* InMainDataArchive = nullptr, FModelStreamableData* InStreamedData = nullptr);
+
+	// mu::ModelWriter interface
+	void OpenWriteFile(uint32 BlockKey) override;
+	void Write(const void* pBuffer, uint64 size) override;
+	void CloseWriteFile() override;
+
+protected:
+
+	// Non-owned pointer to an archive where we'll store the main model data (non-streamable)
+	FArchive* MainDataArchive = nullptr;
+
+	// Non-owned pointer to an archive where we'll store the resouces (streamable)
+	FModelStreamableData* StreamedData = nullptr;
+
+	uint32 CurrentKey = 0;
+
+};
+
+
 #endif
diff --git a/Engine/Plugins/Experimental/Mutable/Source/CustomizableObject/Private/MuCO/CustomizableObject.cpp b/Engine/Plugins/Experimental/Mutable/Source/CustomizableObject/Private/MuCO/CustomizableObject.cpp
--- a/Engine/Plugins/Experimental/Mutable/Source/CustomizableObject/Private/MuCO/CustomizableObject.cpp
+++ b/Engine/Plugins/Experimental/Mutable/Source/CustomizableObject/Private/MuCO/CustomizableObject.cpp
@@ -150,14 +150,14 @@ void UCustomizableObject::PreSave(FObjectPreSaveContext ObjectSaveContext)
 	if (ObjectSaveContext.IsCooking() && !bIsChildObject)
 	{
 		const ITargetPlatform* TargetPlatform = ObjectSaveContext.GetTargetPlatform();
+
 		// Load cached data before saving
 		if (GetPrivate()->TryLoadCompiledCookDataForPlatform(TargetPlatform))
 		{
 			// Create an export object to manage the streamable data
 			if (!BulkData)
 			{
-				BulkData = UE::Mutable::Private::MoveOldObjectAndCreateNew<UCustomizableObjectBulk>(
-					UCustomizableObjectBulk::StaticClass(), this);
+				BulkData = UE::Mutable::Private::MoveOldObjectAndCreateNew<UCustomizableObjectBulk>(UCustomizableObjectBulk::StaticClass(), this);
 			}
 			BulkData->Mark(OBJECTMARK_TagExp);
 
@@ -216,7 +216,6 @@ void UCustomizableObjectPrivate::SetIsChildObject(bool bIsChildObject)
 }
 
 
-#if WITH_EDITOR
 bool UCustomizableObjectPrivate::TryLoadCompiledCookDataForPlatform(const ITargetPlatform* TargetPlatform)
 {
 	const FMutableCachedPlatformData* PlatformData = CachedPlatformsData.Find(TargetPlatform->PlatformName());
@@ -229,7 +228,6 @@ bool UCustomizableObjectPrivate::TryLoadCompiledCookDataForPlatform(const ITarge
 	LoadCompiledData(MemoryReader, TargetPlatform, true);
 	return true;
 }
-#endif
 
 #endif // End WITH_EDITOR
 
@@ -823,8 +821,8 @@ void UCustomizableObjectPrivate::LoadCompiledDataFromDisk()
 	if (FileManager.FileExists(*ModelFileName) && FileManager.FileExists(*StreamableFileName))
 	{
 		// Check CompiledData
-		IFileHandle* CompiledDataFileHandle = FPlatformFileManager::Get().GetPlatformFile().OpenRead(*ModelFileName);
-		IFileHandle* StreamableDataFileHandle = FPlatformFileManager::Get().GetPlatformFile().OpenRead(*StreamableFileName);
+		TUniquePtr<IFileHandle> CompiledDataFileHandle( FPlatformFileManager::Get().GetPlatformFile().OpenRead(*ModelFileName) );
+		TUniquePtr<IFileHandle> StreamableDataFileHandle( FPlatformFileManager::Get().GetPlatformFile().OpenRead(*StreamableFileName) );
 
 		MutableCompiledDataStreamHeader CompiledDataHeader;
 		MutableCompiledDataStreamHeader StreamableDataHeader;
@@ -863,45 +861,15 @@ void UCustomizableObjectPrivate::LoadCompiledDataFromDisk()
 				LoadCompiledData(MemoryReader, RunningPlatform);
 			}
 		}
-
-		delete CompiledDataFileHandle;
-		delete StreamableDataFileHandle;
 	}
 
-	if (!GetModel()) // Not failed to load the model
+	if (!GetModel()) // Failed to load the model
 	{
 		Status.NextState(FCustomizableObjectStatusTypes::EState::NoModel);
 	}
 }
 
 
-void UCustomizableObjectPrivate::CachePlatformData(
-		const ITargetPlatform* InTargetPlatform, 
-		TArray64<uint8>& InModelBytes, 
-		TArray64<uint8>& InBulkBytes, 
-		TArray64<uint8>& InMorphBytes, 
-		TArray64<uint8>& InClothingBytes)
-{
-	MUTABLE_CPUPROFILER_SCOPE(CachePlatformData)
-
-	FString PlatformName = InTargetPlatform ? InTargetPlatform->PlatformName() : FPlatformProperties::PlatformName();
-
-	check(!CachedPlatformsData.Find(PlatformName));
-
-	FMutableCachedPlatformData& Data = CachedPlatformsData.Add(PlatformName);
-
-	// Cache CO data and mu::Model
-	FMemoryWriter64 MemoryWriter(Data.ModelData);
-	SaveCompiledData(MemoryWriter, true);
-	Data.ModelData.Append(InModelBytes);
-
-	// Cache streamable bulk data
-	Data.StreamableData = MoveTemp(InBulkBytes);
-	Data.MorphData = MoveTemp(InMorphBytes);
-	Data.ClothingData = MoveTemp(InClothingBytes);
-}
-
-
 void UCustomizableObjectPrivate::CompileForTargetPlatform(const ITargetPlatform* TargetPlatform)
 {
 	if (!TargetPlatform)
@@ -1054,18 +1022,18 @@ void UCustomizableObjectPrivate::SaveEmbeddedData(FArchive& Ar)
 {
 	UE_LOG(LogMutable, Verbose, TEXT("Saving embedded data for Customizable Object [%s] now at position %d."), *GetName(), int(Ar.Tell()));
 
-	int32 InternalVersion = GetModel() ? CurrentSupportedVersion : -1;
+	TSharedPtr<mu::Model> Model = GetModel();
+
+	int32 InternalVersion = Model ? CurrentSupportedVersion : -1;
 	Ar << InternalVersion;
 
-	if (GetModel())
+	if (Model)
 	{	
 		// Serialise the entire model, but unload the streamable data first.
 		{
-			GetModel()->UnloadExternalData();
-
 			UnrealMutableOutputStream Stream(Ar);
 			mu::OutputArchive Arch(&Stream);
-			mu::Model::Serialise(GetModel().Get(), Arch);
+			mu::Model::Serialise(Model.Get(), Arch);
 		}
 
 		UE_LOG(LogMutable, Verbose, TEXT("Saved embedded data for Customizable Object [%s] now at position %d."), *GetName(), int(Ar.Tell()));
@@ -2176,6 +2144,56 @@ TUniquePtr<IAsyncReadFileHandle> UCustomizableObjectBulk::OpenFileAsyncRead(uint
 
 #if WITH_EDITOR
 
+
+int64 UCustomizableObjectBulk::FFile::GetSize() const
+{
+	int64 FileSize = 0;
+	for (const FBlock& Block : Blocks)
+	{
+		FileSize += Block.Size;
+	}
+	return FileSize;
+}
+
+
+void UCustomizableObjectBulk::FFile::GetFileData(struct FMutableCachedPlatformData* PlatformData, uint8* DataDestination) const
+{
+	check(PlatformData);
+	check(DataDestination);
+
+	const uint8* SourceData = nullptr;
+	if (DataType == EDataType::Model)
+	{
+		for (int32 BlockIndex = 0; BlockIndex < Blocks.Num(); ++BlockIndex)
+		{
+			const FBlock& Block = Blocks[BlockIndex];
+			PlatformData->ModelStreamableData.Get(Block.Id, DataDestination);
+			DataDestination += Block.Size;
+		}
+		return;
+	}
+	else if (DataType == EDataType::RealTimeMorph)
+	{
+		SourceData = PlatformData->MorphData.GetData();
+	}
+	else if (DataType == EDataType::Clothing)
+	{
+		SourceData = PlatformData->ClothingData.GetData();
+	}
+	else
+	{
+		checkf(false, TEXT("Unknown file DataType found."));
+	}
+
+	for (int32 BlockIndex = 0; BlockIndex<Blocks.Num(); ++BlockIndex)
+	{
+		const FBlock& Block = Blocks[BlockIndex];
+		FMemory::Memcpy(DataDestination, SourceData + Block.Offset, Block.Size);
+		DataDestination += Block.Size;
+	}
+}
+
+
 void UCustomizableObjectBulk::CookAdditionalFilesOverride(const TCHAR* PackageFilename,
 	const ITargetPlatform* TargetPlatform,
 	TFunctionRef<void(const TCHAR* Filename, void* Data, int64 Size)> WriteAdditionalFile)
@@ -2195,42 +2213,13 @@ void UCustomizableObjectBulk::CookAdditionalFilesOverride(const TCHAR* PackageFi
 	for(int32 FileIndex = 0; FileIndex < NumBulkDataFiles; ++FileIndex)
 	{
 		const FFile& CurrentFile = BulkDataFiles[FileIndex];
+		int64 FileSize = CurrentFile.GetSize();
 
-		const uint8* SourceData = nullptr;
-		if (CurrentFile.DataType == EDataType::Model)
-		{
-			SourceData = PlatformData->StreamableData.GetData();
-		}
-		else if (CurrentFile.DataType == EDataType::RealTimeMorph)
-		{
-			SourceData = PlatformData->MorphData.GetData();
-		}
-		else if (CurrentFile.DataType == EDataType::Clothing)
-		{
-			SourceData = PlatformData->ClothingData.GetData();
-		}
-		else
-		{
-			checkf(false, TEXT("Unknown file DataType found."));
-			continue;
-		}
-
-		int64 FileSize = 0;
-		for (const FBlock& Block : CurrentFile.Blocks)
-		{
-			FileSize += Block.Size;
-		}
-
-		// Generate the bulk data file in memory
+		// Get the file data in memory
 		TArray64<uint8> FileBulkData;
-		FileBulkData.SetNum(FileSize);
+		FileBulkData.SetNumUninitialized(FileSize);
 		uint8* FileData = FileBulkData.GetData();
-
-		for (const FBlock& Block : CurrentFile.Blocks)
-		{
-			FMemory::Memcpy(FileData, SourceData + Block.Offset, Block.Size);
-			FileData += Block.Size;
-		}
+		CurrentFile.GetFileData(PlatformData, FileData);
 
 		// Path to the asset
 		const FString CookedFilePath = FPaths::GetPath(PackageFilename);
@@ -2275,29 +2264,16 @@ void UCustomizableObjectBulk::PrepareBulkData(UCustomizableObject* InOuter, cons
 	FClassifyNode RootNode;
 	// Create blocks data, filtering out the ones that are too big and will go on its own file in any case.
 	{
-		const int32 NumBlocks = Model->GetRomCount();
-		RootNode.Blocks.Reserve(NumBlocks);
+		const int32 NumRoms = Model->GetRomCount();
+		RootNode.Blocks.Reserve(NumRoms);
 		
-		uint64 SourceOffset = 0;
-		for (int32 BlockIndex = 0; BlockIndex < NumBlocks; ++BlockIndex)
+		for (int32 RomIndex = 0; RomIndex < NumRoms; ++RomIndex)
 		{
-			uint32 BlockId = Model->GetRomId(BlockIndex);
-			const FMutableStreamableBlock& StreamableBlock = ModelResources.HashToStreamableBlock[BlockId];
-			const uint32 BlockSize = StreamableBlock.Size;
+			uint32 BlockId = Model->GetRomId(RomIndex);
+			const uint32 BlockSize = Model->GetRomSize(RomIndex);
 
-			FBlock CurrentBlock = { EDataType::Model, BlockId, BlockSize, SourceOffset };
-			if (BlockSize > TargetBulkDataFileBytes)
-			{
-				// It will go to its own file
-				BulkDataFiles.Add(FFile{EDataType::Model, 0, {CurrentBlock}});
-			}
-			else
-			{
-				// It may merge with other small blocks
-				RootNode.Blocks.Add(CurrentBlock);
-			}
-
-			SourceOffset += BlockSize;
+			FBlock CurrentBlock = { EDataType::Model, BlockId, BlockSize, 0 };
+			RootNode.Blocks.Add(CurrentBlock);
 		}
 	}
 
@@ -2313,21 +2289,11 @@ void UCustomizableObjectBulk::PrepareBulkData(UCustomizableObject* InOuter, cons
 		const int32 NumBlocks = RealTimeMorphTargetsBlocks.Num();
 		for (const TPair<uint32, FRealTimeMorphStreamable>& MorphStreamable : RealTimeMorphStreamables)
 		{
-			const uint32 BlockSize = MorphStreamable.Value.Block.Size;
+			const uint32 BlockSize = MorphStreamable.Value.Size;
 			
 			check(SourceOffset == MorphStreamable.Value.Block.Offset);
 			FBlock CurrentBlock = { EDataType::RealTimeMorph, MorphStreamable.Key, BlockSize, SourceOffset };
-			
-			if (BlockSize > TargetBulkDataFileBytes)
-			{
-				// It will go to its own file
-				BulkDataFiles.Add(FFile{ EDataType::RealTimeMorph, 0, {CurrentBlock} });
-			}
-			else
-			{
-				// It may merge with other small blocks
-				RootNode.Blocks.Add(CurrentBlock);
-			}
+			RootNode.Blocks.Add(CurrentBlock);
 
 			SourceOffset += BlockSize;
 		}
@@ -2345,21 +2311,11 @@ void UCustomizableObjectBulk::PrepareBulkData(UCustomizableObject* InOuter, cons
 		const int32 NumBlocks = ClothingBlocks.Num();
 		for (const TPair<uint32, FClothingStreamable>& ClothStreamable : ClothingStreamables)
 		{
-			const uint32 BlockSize = ClothStreamable.Value.Block.Size;
+			const uint32 BlockSize = ClothStreamable.Value.Size;
 			
 			check(SourceOffset == ClothStreamable.Value.Block.Offset);
 			FBlock CurrentBlock = { EDataType::Clothing, ClothStreamable.Key, BlockSize, SourceOffset };
-			
-			if (BlockSize > TargetBulkDataFileBytes)
-			{
-				// It will go to its own file
-				BulkDataFiles.Add(FFile{ EDataType::Clothing, 0, {CurrentBlock} });
-			}
-			else
-			{
-				// It may merge with other small blocks
-				RootNode.Blocks.Add(CurrentBlock);
-			}
+			RootNode.Blocks.Add(CurrentBlock);
 
 			SourceOffset += BlockSize;
 		}
@@ -2379,8 +2335,7 @@ void UCustomizableObjectBulk::PrepareBulkData(UCustomizableObject* InOuter, cons
 			FBlock CurrentBlock = RootNode.Blocks[BlockIndex];
 
 			// Next file?
-			// Store different data types in different files. Blocks should be sorted by DataType
-			// so data is properly packeted
+			// Store different data types in different files. Blocks should be sorted by DataType so data is properly packeted
 			if (CurrentFile.DataType != CurrentBlock.DataType)
 			{
 				break;
@@ -2407,11 +2362,11 @@ void UCustomizableObjectBulk::PrepareBulkData(UCustomizableObject* InOuter, cons
 	{
 		// Generate the id for this file
 		FFile& CurrentFile = BulkDataFiles[FileIndex];
-		uint32 FileId = 0;
+		uint32 FileId = uint32(CurrentFile.DataType);
 		for (int32 FileBlockIndex = 0; FileBlockIndex < CurrentFile.Blocks.Num(); ++FileBlockIndex)
 		{
 			FBlock ThisBlock = CurrentFile.Blocks[FileBlockIndex];
-			FileId = (FileBlockIndex == 0) ? ThisBlock.Id : HashCombine(HashCombine(FileId, ThisBlock.Id), uint32(CurrentFile.DataType));
+			FileId = HashCombine(FileId, ThisBlock.Id);
 		}
 
 		// Ensure the FileId is unique
@@ -2441,18 +2396,15 @@ void UCustomizableObjectBulk::PrepareBulkData(UCustomizableObject* InOuter, cons
 			{
 				FBlock ThisBlock = CurrentFile.Blocks[FileBlockIndex];
 
-				FMutableStreamableBlock& StreamableBlock = ModelResources.HashToStreamableBlock[ThisBlock.Id];
-				check(StreamableBlock.Size == ThisBlock.Size);
-				StreamableBlock.FileId = FileId;
-				StreamableBlock.Offset = OffsetInFile;
+				FMutableStreamableBlock* StreamableBlock = ModelResources.HashToStreamableBlock.Find(ThisBlock.Id);
+				StreamableBlock->FileId = FileId;
+				StreamableBlock->Offset = OffsetInFile;
 				OffsetInFile += ThisBlock.Size;
 			}
 		}
 		else if (CurrentFile.DataType == EDataType::RealTimeMorph)
 		{
-			constexpr bool bGetCooked = true;
-			TMap<uint32, FRealTimeMorphStreamable>& MorphBlocks = 
-					CustomizableObject->GetPrivate()->GetModelResources(bGetCooked).RealTimeMorphStreamables;
+			TMap<uint32, FRealTimeMorphStreamable>& MorphBlocks = ModelResources.RealTimeMorphStreamables;
 			// Set it to all streamable blocks
 			uint32 OffsetInFile = 0;
 			for (int32 FileBlockIndex = 0; FileBlockIndex < CurrentFile.Blocks.Num(); ++FileBlockIndex)
@@ -2460,7 +2412,6 @@ void UCustomizableObjectBulk::PrepareBulkData(UCustomizableObject* InOuter, cons
 				FBlock ThisBlock = CurrentFile.Blocks[FileBlockIndex];
 
 				FMutableStreamableBlock& StreamableBlock = MorphBlocks[ThisBlock.Id].Block;
-				check(StreamableBlock.Size == ThisBlock.Size);
 				StreamableBlock.FileId = FileId;
 				StreamableBlock.Offset = OffsetInFile;
 				OffsetInFile += ThisBlock.Size;
@@ -2468,9 +2419,7 @@ void UCustomizableObjectBulk::PrepareBulkData(UCustomizableObject* InOuter, cons
 		}
 		else if (CurrentFile.DataType == EDataType::Clothing)
 		{
-			constexpr bool bGetCooked = true;
-			TMap<uint32, FClothingStreamable>& ClothBlocks = 
-					CustomizableObject->GetPrivate()->GetModelResources(bGetCooked).ClothingStreamables;
+			TMap<uint32, FClothingStreamable>& ClothBlocks = ModelResources.ClothingStreamables;
 			// Set it to all streamable blocks
 			uint32 OffsetInFile = 0;
 			for (int32 FileBlockIndex = 0; FileBlockIndex < CurrentFile.Blocks.Num(); ++FileBlockIndex)
@@ -2478,7 +2427,6 @@ void UCustomizableObjectBulk::PrepareBulkData(UCustomizableObject* InOuter, cons
 				FBlock ThisBlock = CurrentFile.Blocks[FileBlockIndex];
 
 				FMutableStreamableBlock& StreamableBlock = ClothBlocks[ThisBlock.Id].Block;
-				check(StreamableBlock.Size == ThisBlock.Size);
 				StreamableBlock.FileId = FileId;
 				StreamableBlock.Offset = OffsetInFile;
 				OffsetInFile += ThisBlock.Size;
diff --git a/Engine/Plugins/Experimental/Mutable/Source/CustomizableObject/Private/MuCO/CustomizableObjectInstance.cpp b/Engine/Plugins/Experimental/Mutable/Source/CustomizableObject/Private/MuCO/CustomizableObjectInstance.cpp
--- a/Engine/Plugins/Experimental/Mutable/Source/CustomizableObject/Private/MuCO/CustomizableObjectInstance.cpp
+++ b/Engine/Plugins/Experimental/Mutable/Source/CustomizableObject/Private/MuCO/CustomizableObjectInstance.cpp
@@ -5429,12 +5429,12 @@ UE::Tasks::FTask UCustomizableInstancePrivate::LoadAdditionalAssetsAndData(
 				MeshData.NameResolutionMap = RealTimeMorphStreamable.NameResolutionMap;
 
 				const TArray<FMorphTargetVertexData>& SourceData = ModelResources.EditorOnlyMorphTargetReconstructionData;	
-				const uint32 NumElems = Block.Size / sizeof(FMorphTargetVertexData);
+				const uint32 NumElems = RealTimeMorphStreamable.Size / sizeof(FMorphTargetVertexData);
 				const uint32 OffsetInElems = Block.Offset / sizeof(FMorphTargetVertexData);
 				MeshData.Data.SetNumUninitialized(NumElems);
 
-				check(SourceData.Num()*sizeof(FMorphTargetVertexData) >= Block.Offset + Block.Size);
-				FMemory::Memcpy(MeshData.Data.GetData(), SourceData.GetData() + OffsetInElems, Block.Size);
+				check(SourceData.Num()*sizeof(FMorphTargetVertexData) >= Block.Offset + RealTimeMorphStreamable.Size);
+				FMemory::Memcpy(MeshData.Data.GetData(), SourceData.GetData() + OffsetInElems, RealTimeMorphStreamable.Size);
 			}
 		}
 		{
@@ -5453,12 +5453,12 @@ UE::Tasks::FTask UCustomizableInstancePrivate::LoadAdditionalAssetsAndData(
 				MeshData.PhysicsAssetIndex = ClothingStreamable.PhysicsAssetIndex;
 
 				const TArray<FCustomizableObjectMeshToMeshVertData>& SourceData = ModelResources.EditorOnlyClothingMeshToMeshVertData;
-				const uint32 NumElems = Block.Size / sizeof(FCustomizableObjectMeshToMeshVertData);
+				const uint32 NumElems = ClothingStreamable.Size / sizeof(FCustomizableObjectMeshToMeshVertData);
 				const uint32 OffsetInElems = Block.Offset / sizeof(FCustomizableObjectMeshToMeshVertData);
 				MeshData.Data.SetNumUninitialized(NumElems);
 
-				check(SourceData.Num()*sizeof(FCustomizableObjectMeshToMeshVertData) >= Block.Offset + Block.Size);
-				FMemory::Memcpy(MeshData.Data.GetData(), SourceData.GetData() + OffsetInElems, Block.Size);
+				check(SourceData.Num()*sizeof(FCustomizableObjectMeshToMeshVertData) >= Block.Offset + ClothingStreamable.Size);
+				FMemory::Memcpy(MeshData.Data.GetData(), SourceData.GetData() + OffsetInElems, ClothingStreamable.Size);
 			}
 		}
 #else	
@@ -5486,7 +5486,8 @@ UE::Tasks::FTask UCustomizableInstancePrivate::LoadAdditionalAssetsAndData(
 			MUTABLE_CPUPROFILER_SCOPE(RealTimeMorphStreamingRequest_Alloc);
 
 			const int32 BlockId = RealTimeMorphStreamableBlocksToStream[I];
-			const FMutableStreamableBlock& Block = ModelResources.RealTimeMorphStreamables[BlockId].Block; 
+			const FRealTimeMorphStreamable& Streamable = ModelResources.RealTimeMorphStreamables[BlockId];
+			const FMutableStreamableBlock& Block = Streamable.Block; 
 		
 			FInstanceUpdateData::FMorphTargetMeshData& ReadDestData = 
 					OperationData->InstanceUpdateData.RealTimeMorphTargetMeshData.FindOrAdd(BlockId);
@@ -5499,8 +5500,8 @@ UE::Tasks::FTask UCustomizableInstancePrivate::LoadAdditionalAssetsAndData(
 
 			ReadDestData.NameResolutionMap = ModelResources.RealTimeMorphStreamables[BlockId].NameResolutionMap;
 
-			check(Block.Size % sizeof(FMorphTargetVertexData) == 0);
-			uint32 NumElems = Block.Size / sizeof(FMorphTargetVertexData);
+			check(Streamable.Size % sizeof(FMorphTargetVertexData) == 0);
+			uint32 NumElems = Streamable.Size / sizeof(FMorphTargetVertexData);
 
 			ReadDestData.Data.SetNumUninitialized(NumElems);
 
@@ -5528,7 +5529,8 @@ UE::Tasks::FTask UCustomizableInstancePrivate::LoadAdditionalAssetsAndData(
 			MUTABLE_CPUPROFILER_SCOPE(ClothingStreamingRequest_Alloc);
 
 			const int32 BlockId = ClothingStreamableBlocksToStream[I];
-			const FMutableStreamableBlock& Block = ModelResources.ClothingStreamables[BlockId].Block; 
+			const FClothingStreamable& ClothingStreamable = ModelResources.ClothingStreamables[BlockId];
+			const FMutableStreamableBlock& Block = ClothingStreamable.Block;
 		
 			FInstanceUpdateData::FClothingMeshData& ReadDestData = 
 					OperationData->InstanceUpdateData.ClothingMeshData.FindOrAdd(BlockId);
@@ -5539,14 +5541,12 @@ UE::Tasks::FTask UCustomizableInstancePrivate::LoadAdditionalAssetsAndData(
 				continue;
 			}
 
-			const FClothingStreamable& ClothingStreamable = ModelResources.ClothingStreamables[BlockId]; 
-
 			ReadDestData.ClothingAssetIndex = ClothingStreamable.ClothingAssetIndex;
 			ReadDestData.ClothingAssetLOD = ClothingStreamable.ClothingAssetLOD;
 			ReadDestData.PhysicsAssetIndex = ClothingStreamable.PhysicsAssetIndex;
 
-			check(Block.Size % sizeof(FCustomizableObjectMeshToMeshVertData) == 0);
-			const uint32 NumElems = Block.Size / sizeof(FCustomizableObjectMeshToMeshVertData);
+			check(ClothingStreamable.Size % sizeof(FCustomizableObjectMeshToMeshVertData) == 0);
+			const uint32 NumElems = ClothingStreamable.Size / sizeof(FCustomizableObjectMeshToMeshVertData);
 
 			ReadDestData.Data.SetNumUninitialized(NumElems);
 
diff --git a/Engine/Plugins/Experimental/Mutable/Source/CustomizableObject/Private/MuCO/UnrealMutableModelDiskStreamer.cpp b/Engine/Plugins/Experimental/Mutable/Source/CustomizableObject/Private/MuCO/UnrealMutableModelDiskStreamer.cpp
--- a/Engine/Plugins/Experimental/Mutable/Source/CustomizableObject/Private/MuCO/UnrealMutableModelDiskStreamer.cpp
+++ b/Engine/Plugins/Experimental/Mutable/Source/CustomizableObject/Private/MuCO/UnrealMutableModelDiskStreamer.cpp
@@ -103,6 +103,7 @@ bool FUnrealMutableModelBulkReader::PrepareStreamingForObject(UCustomizableObjec
 #endif
 
 		const FModelResources& ModelResources = CustomizableObject->GetPrivate()->GetModelResources();
+		// \TODO: unnecessary runtime map copy. This should be stored in a TSharedPtr and only copy the pointer.
 		NewData.StreamableBlocks = ModelResources.HashToStreamableBlock;
 		if (NewData.StreamableBlocks.IsEmpty())
 		{
@@ -191,7 +192,7 @@ FAutoConsoleVariableRef CVarStreamPriority(
 	TEXT(""));
 
 
-mu::ModelReader::OPERATION_ID FUnrealMutableModelBulkReader::BeginReadBlock(const mu::Model* Model, uint64 Key, void* pBuffer, uint64 size, TFunction<void(bool bSuccess)>* CompletionCallback)
+mu::ModelReader::OPERATION_ID FUnrealMutableModelBulkReader::BeginReadBlock(const mu::Model* Model, uint32 Key, void* pBuffer, uint64 size, TFunction<void(bool bSuccess)>* CompletionCallback)
 {
 	MUTABLE_CPUPROFILER_SCOPE(FUnrealMutableModelBulkStreamer::OpenReadFile);
 
@@ -339,7 +340,7 @@ void FUnrealMutableModelBulkReader::EndRead(mu::ModelReader::OPERATION_ID Operat
 #if WITH_EDITOR
 
 //-------------------------------------------------------------------------------------------------
-FUnrealMutableModelBulkWriter::FUnrealMutableModelBulkWriter(FArchive* InMainDataArchive, FArchive* InStreamedDataArchive)
+FUnrealMutableModelBulkWriterEditor::FUnrealMutableModelBulkWriterEditor(FArchive* InMainDataArchive, FArchive* InStreamedDataArchive)
 {
 	MainDataArchive = InMainDataArchive;
 	StreamedDataArchive = InStreamedDataArchive;
@@ -347,9 +348,9 @@ FUnrealMutableModelBulkWriter::FUnrealMutableModelBulkWriter(FArchive* InMainDat
 }
 
 
-void FUnrealMutableModelBulkWriter::OpenWriteFile(uint64 key0)
+void FUnrealMutableModelBulkWriterEditor::OpenWriteFile(uint32 BlockKey)
 {
-	if (key0 == 0) // Model
+	if (BlockKey == 0) // Model
 	{
 		check(MainDataArchive);
 		CurrentWriteFile = MainDataArchive;
@@ -362,17 +363,52 @@ void FUnrealMutableModelBulkWriter::OpenWriteFile(uint64 key0)
 }
 
 
-void FUnrealMutableModelBulkWriter::Write(const void* pBuffer, uint64 size)
+void FUnrealMutableModelBulkWriterEditor::Write(const void* pBuffer, uint64 size)
 {
 	check(CurrentWriteFile);
 	CurrentWriteFile->Serialize(const_cast<void*>(pBuffer), size);
 }
 
 
-void FUnrealMutableModelBulkWriter::CloseWriteFile()
+void FUnrealMutableModelBulkWriterEditor::CloseWriteFile()
 {
 	CurrentWriteFile = nullptr;
 }
 
+
+
+//-------------------------------------------------------------------------------------------------
+FUnrealMutableModelBulkWriterCook::FUnrealMutableModelBulkWriterCook(FArchive* InMainDataArchive, FModelStreamableData* InStreamedData)
+{
+	MainDataArchive = InMainDataArchive;
+	StreamedData = InStreamedData;
+	CurrentKey = 0;
+}
+
+
+void FUnrealMutableModelBulkWriterCook::OpenWriteFile(uint32 BlockKey)
+{
+	CurrentKey = BlockKey;
+}
+
+
+void FUnrealMutableModelBulkWriterCook::Write(const void* pBuffer, uint64 size)
+{
+	if (CurrentKey == 0)
+	{
+		MainDataArchive->Serialize(const_cast<void*>(pBuffer), size);
+	}
+	else
+	{
+		StreamedData->Set(CurrentKey, reinterpret_cast<const uint8*>(pBuffer), size);
+	}
+}
+
+
+void FUnrealMutableModelBulkWriterCook::CloseWriteFile()
+{
+}
+
+
 #endif
 
diff --git a/Engine/Plugins/Experimental/Mutable/Source/CustomizableObject/Public/MuCO/CustomizableObject.h b/Engine/Plugins/Experimental/Mutable/Source/CustomizableObject/Public/MuCO/CustomizableObject.h
--- a/Engine/Plugins/Experimental/Mutable/Source/CustomizableObject/Public/MuCO/CustomizableObject.h
+++ b/Engine/Plugins/Experimental/Mutable/Source/CustomizableObject/Public/MuCO/CustomizableObject.h
@@ -206,9 +206,6 @@ struct CUSTOMIZABLEOBJECT_API FMutableStreamableBlock
 
 	UPROPERTY()
 	uint32 FileId = 0;
-
-	UPROPERTY()
-	uint32 Size = 0;
 	
 	UPROPERTY()
 	uint64 Offset = 0;
@@ -216,7 +213,6 @@ struct CUSTOMIZABLEOBJECT_API FMutableStreamableBlock
 	friend FArchive& operator<<(FArchive& Ar, FMutableStreamableBlock& Data)
 	{
 		Ar << Data.FileId;
-		Ar << Data.Size;
 		Ar << Data.Offset;
 
 		return Ar;
@@ -323,6 +319,12 @@ private:
 
 		/** List of blocks that are contained in the file, in order. */
 		TArray<FBlock> Blocks;
+
+		/** Get the total size of blocks in this file. */
+		int64 GetSize() const;
+
+		/** Copy the requested block to the requested buffer and return its size. */
+		void GetFileData(struct FMutableCachedPlatformData*, uint8* DataDestination) const;
 	};
 
 	/** Helper to store the size of each BulkData partition. Only valid while cooking */
diff --git a/Engine/Plugins/Experimental/Mutable/Source/CustomizableObjectEditor/Private/MuCOE/CustomizableObjectCompileRunnable.cpp b/Engine/Plugins/Experimental/Mutable/Source/CustomizableObjectEditor/Private/MuCOE/CustomizableObjectCompileRunnable.cpp
--- a/Engine/Plugins/Experimental/Mutable/Source/CustomizableObjectEditor/Private/MuCOE/CustomizableObjectCompileRunnable.cpp
+++ b/Engine/Plugins/Experimental/Mutable/Source/CustomizableObjectEditor/Private/MuCOE/CustomizableObjectCompileRunnable.cpp
@@ -272,7 +272,7 @@ FCustomizableObjectSaveDDRunnable::FCustomizableObjectSaveDDRunnable(UCustomizab
 		StreamableDataFullFileName = FolderPath + CustomizableObject->GetPrivate()->GetCompiledDataFileName(false, InOptions.TargetPlatform);
 
 		// Serialize Customizable Object's data
-		FMemoryWriter64 MemoryWriter(Bytes);
+		FMemoryWriter64 MemoryWriter(ModelBytes);
 		CustomizableObject->GetPrivate()->SaveCompiledData(MemoryWriter, Options.bIsCooking);
 	}
 #if WITH_EDITORONLY_DATA
@@ -281,17 +281,16 @@ FCustomizableObjectSaveDDRunnable::FCustomizableObjectSaveDDRunnable(UCustomizab
 		// Do a copy of the Morph and Clothing Data generated at compile time. Only needed when cooking.
 		
 		constexpr bool bGetCookedFalse = false;
+		FModelResources& ModelResources = CustomizableObject->GetPrivate()->GetModelResources(bGetCookedFalse);
 		
 		static_assert(TCanBulkSerialize<FMorphTargetVertexData>::Value);
-		const TArray<FMorphTargetVertexData>& MorphVertexData = 
-				CustomizableObject->GetPrivate()->GetModelResources(bGetCookedFalse).EditorOnlyMorphTargetReconstructionData;
+		const TArray<FMorphTargetVertexData>& MorphVertexData = ModelResources.EditorOnlyMorphTargetReconstructionData;
 
 		MorphDataBytes.SetNum(MorphVertexData.Num() * sizeof(FMorphTargetVertexData));
 		FMemory::Memcpy(MorphDataBytes.GetData(), MorphVertexData.GetData(), MorphDataBytes.Num());
 		
 		static_assert(TCanBulkSerialize<FCustomizableObjectMeshToMeshVertData>::Value);
-		const TArray<FCustomizableObjectMeshToMeshVertData>& ClothingVertexData = 
-				CustomizableObject->GetPrivate()->GetModelResources(bGetCookedFalse).EditorOnlyClothingMeshToMeshVertData;
+		const TArray<FCustomizableObjectMeshToMeshVertData>& ClothingVertexData = ModelResources.EditorOnlyClothingMeshToMeshVertData;
 
 		ClothingDataBytes.SetNum(ClothingVertexData.Num() * sizeof(FCustomizableObjectMeshToMeshVertData));
 		FMemory::Memcpy(ClothingDataBytes.GetData(), ClothingVertexData.GetData(), ClothingDataBytes.Num());
@@ -313,16 +312,15 @@ uint32 FCustomizableObjectSaveDDRunnable::Run()
 	if (Options.bIsCooking)
 	{
 		// Serialize mu::Model and streamable resources 
-		FMemoryWriter64 ModelMemoryWriter(Bytes, false, true);
-		FMemoryWriter64 StreamableMemoryWriter(BulkDataBytes, false, true);
-
+		FMemoryWriter64 ModelMemoryWriter(ModelBytes, false, true);
 		ModelMemoryWriter << bModelSerialized;
 		if (bModelSerialized)
 		{
-			FUnrealMutableModelBulkWriter Streamer(&ModelMemoryWriter, &StreamableMemoryWriter);
-			mu::Model::Serialise(Model.Get(), Streamer);
+			FUnrealMutableModelBulkWriterCook Streamer(&ModelMemoryWriter, &ModelStreamableData);
+			constexpr bool bDropData = true;
+			mu::Model::Serialise(Model.Get(), Streamer, bDropData);
 
-			//Morph and Clothing are already in the corresponding buffer copied from the compilation thread.
+			// Morph and Clothing are already in the corresponding buffer copied from the compilation thread.
 		}
 	}
 	else if (bModelSerialized) // Save CO data + mu::Model and streamable resources to disk
@@ -361,14 +359,15 @@ uint32 FCustomizableObjectSaveDDRunnable::Run()
 			*StreamableMemoryWriter << CustomizableObjectHeader;
 
 			// Serialize Customizable Object's Data to disk
-			ModelMemoryWriter->Serialize(Bytes.GetData(), Bytes.Num() * sizeof(uint8));
-			Bytes.Empty();
+			ModelMemoryWriter->Serialize(ModelBytes.GetData(), ModelBytes.Num() * sizeof(uint8));
+			ModelBytes.Empty();
 
 			// Serialize mu::Model and streamable resources
 			*ModelMemoryWriter << bModelSerialized;
 
-			FUnrealMutableModelBulkWriter Streamer(ModelMemoryWriter.Get(), StreamableMemoryWriter.Get());
-			mu::Model::Serialise(Model.Get(), Streamer);
+			FUnrealMutableModelBulkWriterEditor Streamer(ModelMemoryWriter.Get(), StreamableMemoryWriter.Get());
+			constexpr bool bDropData = true;
+			mu::Model::Serialise(Model.Get(), Streamer, bDropData);
 
 			// Save to disk
 			ModelMemoryWriter->Flush();
diff --git a/Engine/Plugins/Experimental/Mutable/Source/CustomizableObjectEditor/Private/MuCOE/CustomizableObjectCompileRunnable.h b/Engine/Plugins/Experimental/Mutable/Source/CustomizableObjectEditor/Private/MuCOE/CustomizableObjectCompileRunnable.h
--- a/Engine/Plugins/Experimental/Mutable/Source/CustomizableObjectEditor/Private/MuCOE/CustomizableObjectCompileRunnable.h
+++ b/Engine/Plugins/Experimental/Mutable/Source/CustomizableObjectEditor/Private/MuCOE/CustomizableObjectCompileRunnable.h
@@ -110,8 +110,6 @@ private:
 	FString CompileDataFullFileName;
 	FString StreamableDataFullFileName;
 
-	TSharedPtr<mu::Model, ESPMode::ThreadSafe> Model;
-
 	bool bIsCooking = false;
 
 	// Whether the thread has finished running
@@ -119,11 +117,13 @@ private:
 
 public:
 
+	TSharedPtr<mu::Model, ESPMode::ThreadSafe> Model;
+
 	// Bytes where the model is stored
-	TArray64<uint8> Bytes;
+	TArray64<uint8> ModelBytes;
 
-	// Bytes where the streamed data files are stored
-	TArray64<uint8> BulkDataBytes;
+	// Model streamed data
+	FModelStreamableData ModelStreamableData;
 
 	// Bytes store streameable files coming form the CO itself.
 	TArray64<uint8> MorphDataBytes;
diff --git a/Engine/Plugins/Experimental/Mutable/Source/CustomizableObjectEditor/Private/MuCOE/CustomizableObjectCompiler.cpp b/Engine/Plugins/Experimental/Mutable/Source/CustomizableObjectEditor/Private/MuCOE/CustomizableObjectCompiler.cpp
--- a/Engine/Plugins/Experimental/Mutable/Source/CustomizableObjectEditor/Private/MuCOE/CustomizableObjectCompiler.cpp
+++ b/Engine/Plugins/Experimental/Mutable/Source/CustomizableObjectEditor/Private/MuCOE/CustomizableObjectCompiler.cpp
@@ -995,10 +995,11 @@ void FCustomizableObjectCompiler::CompileInternal(bool bAsync)
 			FRealTimeMorphStreamable& ResourceMeshData = ModelResources.RealTimeMorphStreamables.FindOrAdd(MeshData.Key);
 			
 			check(ResourceMeshData.NameResolutionMap.IsEmpty());
-			check(ResourceMeshData.Block.Size == 0);
+			check(ResourceMeshData.Size == 0);
 
 			ResourceMeshData.NameResolutionMap = MeshData.Value.NameResolutionMap;
-			ResourceMeshData.Block = FMutableStreamableBlock { uint32(0), DataSizeInBytes, RealTimeMorphDataOffsetInBytes };
+			ResourceMeshData.Size = DataSizeInBytes;
+			ResourceMeshData.Block = FMutableStreamableBlock { uint32(0), RealTimeMorphDataOffsetInBytes };
 
 			RealTimeMorphDataOffsetInBytes += DataSizeInBytes;
 			ModelResources.EditorOnlyMorphTargetReconstructionData.Append(MeshData.Value.Data);
@@ -1022,12 +1023,13 @@ void FCustomizableObjectCompiler::CompileInternal(bool bAsync)
 			
 			check(ResourceMeshData.ClothingAssetIndex == INDEX_NONE);
 			check(ResourceMeshData.ClothingAssetLOD == INDEX_NONE);
-			check(ResourceMeshData.Block.Size == 0);
+			check(ResourceMeshData.Size == 0);
 
 			ResourceMeshData.ClothingAssetIndex = MeshData.Value.ClothingAssetIndex;
 			ResourceMeshData.ClothingAssetLOD = MeshData.Value.ClothingAssetLOD;
 			ResourceMeshData.PhysicsAssetIndex = MeshData.Value.PhysicsAssetIndex;
-			ResourceMeshData.Block = FMutableStreamableBlock { uint32(0), DataSizeInBytes, ClothingDataOffsetInBytes };
+			ResourceMeshData.Size = DataSizeInBytes;
+			ResourceMeshData.Block = FMutableStreamableBlock { uint32(0), ClothingDataOffsetInBytes };
 
 			ClothingDataOffsetInBytes += DataSizeInBytes;
 			ModelResources.EditorOnlyClothingMeshToMeshVertData.Append(MeshData.Value.Data);
@@ -1355,7 +1357,9 @@ void FCustomizableObjectCompiler::FinishCompilationTask()
 		FModelResources& ModelResources = CurrentObject->GetPrivate()->GetModelResources(false);
 
 		const int32 NumStreamingFiles = CurrentModel->GetRomCount();
-		ModelResources.HashToStreamableBlock.Empty(NumStreamingFiles);
+
+		TSharedPtr<TMap<uint32, FMutableStreamableBlock>> HashToStreamableBlock = MakeShared<TMap<uint32, FMutableStreamableBlock>>();
+		HashToStreamableBlock->Empty(NumStreamingFiles);
 
 		uint64 Offset = 0;
 		for (int32 FileIndex = 0; FileIndex < NumStreamingFiles; ++FileIndex)
@@ -1363,9 +1367,11 @@ void FCustomizableObjectCompiler::FinishCompilationTask()
 			const uint32 ResourceId = CurrentModel->GetRomId(FileIndex);
 			const uint32 ResourceSize = CurrentModel->GetRomSize(FileIndex);
 
-			ModelResources.HashToStreamableBlock.Add(ResourceId, FMutableStreamableBlock{ 0, ResourceSize, Offset });
+			HashToStreamableBlock->Add(ResourceId, FMutableStreamableBlock{ 0, Offset });
 			Offset += ResourceSize;
 		}
+
+		ModelResources.HashToStreamableBlock = *(HashToStreamableBlock.Get());
 	}
 
 	// Generate ParameterProperties and IntParameterLookUpTable
@@ -1392,12 +1398,24 @@ void FCustomizableObjectCompiler::FinishSavingDerivedDataTask()
 
 	if (CurrentOptions.bIsCooking)
 	{
-		CurrentObject->GetPrivate()->CachePlatformData(
-				SaveDDTask->GetTargetPlatform(), 
-				SaveDDTask->Bytes, 
-				SaveDDTask->BulkDataBytes,
-				SaveDDTask->MorphDataBytes,
-				SaveDDTask->ClothingDataBytes);
+		MUTABLE_CPUPROFILER_SCOPE(CachePlatformData);
+		const ITargetPlatform* TargetPlatform = CurrentOptions.TargetPlatform;
+
+		FString PlatformName = TargetPlatform ? TargetPlatform->PlatformName() : FPlatformProperties::PlatformName();
+
+		check(!CurrentObject->GetPrivate()->CachedPlatformsData.Find(PlatformName));
+
+		FMutableCachedPlatformData& Data = CurrentObject->GetPrivate()->CachedPlatformsData.Add(PlatformName);
+
+		// Cache CO data and mu::Model
+		FMemoryWriter64 MemoryWriter(Data.ModelData);
+		CurrentObject->GetPrivate()->SaveCompiledData(MemoryWriter, true);
+		Data.ModelData.Append(SaveDDTask->ModelBytes);
+
+		// Cache streamable bulk data
+		Data.ModelStreamableData = MoveTemp(SaveDDTask->ModelStreamableData);
+		Data.MorphData = MoveTemp(SaveDDTask->MorphDataBytes);
+		Data.ClothingData = MoveTemp(SaveDDTask->ClothingDataBytes);
 	}
 
 	// Order matters
diff --git a/Engine/Plugins/Experimental/Mutable/Source/MutableRuntime/Private/MuR/Model.cpp b/Engine/Plugins/Experimental/Mutable/Source/MutableRuntime/Private/MuR/Model.cpp
--- a/Engine/Plugins/Experimental/Mutable/Source/MutableRuntime/Private/MuR/Model.cpp
+++ b/Engine/Plugins/Experimental/Mutable/Source/MutableRuntime/Private/MuR/Model.cpp
@@ -139,14 +139,21 @@ namespace mu
 
 
     //---------------------------------------------------------------------------------------------
-    void Model::Serialise( Model* p, ModelWriter& streamer )
+    void Model::Serialise( Model* p, ModelWriter& streamer, bool bDropData )
     {
 		LLM_SCOPE_BYNAME(TEXT("MutableRuntime"));
 
 		mu::FProgram& program = p->m_pD->m_program;
 
-		TArray<TPair<int32, mu::ImagePtrConst>> InitialImages = program.ConstantImageLODs;
-		TArray<TPair<int32, mu::MeshPtrConst>> InitialMeshes = program.ConstantMeshes;
+		TArray<TPair<int32, mu::ImagePtrConst>> InitialImages;
+		TArray<TPair<int32, mu::MeshPtrConst>> InitialMeshes;
+		if (!bDropData)
+		{
+			InitialImages = program.ConstantImageLODs;
+			InitialMeshes = program.ConstantMeshes;
+		}
+
+		OutputMemoryStream MemStream(16 * 1024 * 1024);
 
 		// Save images and unload from memory
 		for (int32 ResourceIndex = 0; ResourceIndex < program.ConstantImageLODs.Num(); ++ResourceIndex)
@@ -165,7 +172,7 @@ namespace mu
 			check(RomData.ResourceIndex == ResourceIndex);
 
 			// Serialize to memory, to find out final size of this rom
-			OutputMemoryStream MemStream(1024*1024);
+			MemStream.Reset();
 			OutputArchive MemoryArch(&MemStream);
 			Image::Serialise(ResData.Value.get(), MemoryArch);
 			check(RomData.Size == MemStream.GetBufferSize());
@@ -195,7 +202,7 @@ namespace mu
 			check(RomData.ResourceIndex == ResourceIndex);
 
 			// Serialize to memory, to find out final size of this rom
-			OutputMemoryStream MemStream(1024 * 1024);
+			MemStream.Reset();
 			OutputArchive MemoryArch(&MemStream);
 			Mesh::Serialise(ResData.Value.get(), MemoryArch);
 			check(RomData.Size == MemStream.GetBufferSize());
@@ -220,8 +227,11 @@ namespace mu
 		}
 
 		// Restore full data
-		program.ConstantImageLODs = InitialImages;
-		program.ConstantMeshes = InitialMeshes;
+		if (!bDropData)
+		{
+			program.ConstantImageLODs = InitialImages;
+			program.ConstantMeshes = InitialMeshes;
+		}
 	}
 
 
diff --git a/Engine/Plugins/Experimental/Mutable/Source/MutableRuntime/Private/MuR/Serialisation.cpp b/Engine/Plugins/Experimental/Mutable/Source/MutableRuntime/Private/MuR/Serialisation.cpp
--- a/Engine/Plugins/Experimental/Mutable/Source/MutableRuntime/Private/MuR/Serialisation.cpp
+++ b/Engine/Plugins/Experimental/Mutable/Source/MutableRuntime/Private/MuR/Serialisation.cpp
@@ -176,11 +176,18 @@ namespace mu
     }
 
 
-    //---------------------------------------------------------------------------------------------
+	//---------------------------------------------------------------------------------------------
 	uint64 OutputMemoryStream::GetBufferSize() const
-    {
-        return m_pD->m_buffer.Num();
-    }
+	{
+		return m_pD->m_buffer.Num();
+	}
+
+
+	//---------------------------------------------------------------------------------------------
+	void OutputMemoryStream::Reset()
+	{
+		m_pD->m_buffer.SetNum(0,EAllowShrinking::No);
+	}
 
 
     //-------------------------------------------------------------------------------------------------
diff --git a/Engine/Plugins/Experimental/Mutable/Source/MutableRuntime/Public/MuR/Model.h b/Engine/Plugins/Experimental/Mutable/Source/MutableRuntime/Public/MuR/Model.h
--- a/Engine/Plugins/Experimental/Mutable/Source/MutableRuntime/Public/MuR/Model.h
+++ b/Engine/Plugins/Experimental/Mutable/Source/MutableRuntime/Public/MuR/Model.h
@@ -57,10 +57,12 @@ namespace mu
 		static void Serialise( const Model* p, OutputArchive& arch );
 		static TSharedPtr<Model> StaticUnserialise( InputArchive& arch );
 
-		//! Special serialise operation that serialises the data in separate "files". An object
-        //! with the ModelStreamer interface is responsible of storing this data and providing
-		//! the "file" concept.
-        static void Serialise( Model* p, ModelWriter& arch );
+		/** Special serialise operation that serialises the data in separate "files". An object
+         * with the ModelStreamer interface is responsible of storing this data and providing
+		 * the "file" concept. 
+		 * If bDropData is set to true, the rom data will be freed as it is serialized.
+		 */
+        static void Serialise( Model* p, ModelWriter& arch, bool bDropData=false );
 
 		//! Return true if the model has external data in other files. This kind of models will
 		//! require data streaming when used.
diff --git a/Engine/Plugins/Experimental/Mutable/Source/MutableRuntime/Public/MuR/Serialisation.h b/Engine/Plugins/Experimental/Mutable/Source/MutableRuntime/Public/MuR/Serialisation.h
--- a/Engine/Plugins/Experimental/Mutable/Source/MutableRuntime/Public/MuR/Serialisation.h
+++ b/Engine/Plugins/Experimental/Mutable/Source/MutableRuntime/Public/MuR/Serialisation.h
@@ -101,7 +101,7 @@ namespace mu
 		//! \return a previously unused identifier, now used for this operation, that can be used in
 		//! calls to the other methods of this interface. If the return value is negative it indicates
 		//! an unrecoverable error.
-		virtual OPERATION_ID BeginReadBlock(const mu::Model*, uint64 key0, void* pBuffer, uint64 size, TFunction<void(bool bSuccess)>* CompletionCallback = nullptr) = 0;
+		virtual OPERATION_ID BeginReadBlock(const mu::Model*, uint32 BlockKey, void* pBuffer, uint64 size, TFunction<void(bool bSuccess)>* CompletionCallback = nullptr) = 0;
 
         //! Check if a data request operation has been completed.
         //! This is a weak check than *may* return true if the given operation has completed, but
@@ -144,7 +144,7 @@ namespace mu
 		//!         files in a model.
 		//! \param key key identifying the model data fragment that is requested.
 		//!         This key interpretation depends on the implementation of the ModelStreamer,
-		virtual void OpenWriteFile(uint64 key0) = 0;
+		virtual void OpenWriteFile(uint32 BlockKey) = 0;
 
 		//! \brief Write a piece of data to the currently open file.
 		//!
@@ -312,6 +312,9 @@ namespace mu
         //! Get the amount of data in the stream, in bytes.
         uint64 GetBufferSize() const;
 
+		/** Clear the internal buffer. */
+		void Reset();
+
     private:
 
         class Private;
