[mutable] Optimizations for the compilation of objects using the disk as cache.
- Lower the size threshold to dump resources to disk.
- Add compression step before dumping, and don't dump if compressed is small.
- Optimize image constant mip splitting step to prevent unnecessary compression and decompression loops.
#rnx

diff --git a/Engine/Plugins/Experimental/Mutable/Source/MutableRuntime/Internal/MuR/SerialisationPrivate.h b/Engine/Plugins/Experimental/Mutable/Source/MutableRuntime/Internal/MuR/SerialisationPrivate.h
--- a/Engine/Plugins/Experimental/Mutable/Source/MutableRuntime/Internal/MuR/SerialisationPrivate.h
+++ b/Engine/Plugins/Experimental/Mutable/Source/MutableRuntime/Internal/MuR/SerialisationPrivate.h
@@ -38,8 +38,6 @@ namespace mu
         //! objects is in use.
         InputMemoryStream( const void* pBuffer, uint64 size );
 
-        ~InputMemoryStream();
-
 
         //-----------------------------------------------------------------------------------------
         // InputStream interface
@@ -49,8 +47,9 @@ namespace mu
 
     private:
 
-        class Private;
-        Private* m_pD;
+		const void* Buffer = nullptr;
+		uint64 Size = 0;
+		uint64 Pos = 0;
 
     };
 
@@ -81,75 +80,20 @@ namespace mu
 
     private:
 
-        uint64 m_writtenBytes;
-
-    };
-
-
-    //---------------------------------------------------------------------------------------------
-    class MUTABLERUNTIME_API InputMemoryStream::Private
-    {
-    public:
-
-        Private()
-        {
-            m_pBuffer = 0;
-            m_size = 0;
-            m_pos = 0;
-        }
+        uint64 WrittenBytes;
 
-        const void* m_pBuffer;
-
-		uint64 m_size;
-
-		uint64 m_pos;
-    };
-
-
-    //---------------------------------------------------------------------------------------------
-    class MUTABLERUNTIME_API OutputMemoryStream::Private
-    {
-    public:
-        TArray<uint8,TSizedHeapAllocator<64>> m_buffer;
     };
 
 
-	//---------------------------------------------------------------------------------------------
-	class MUTABLERUNTIME_API InputArchive::Private
-	{
-	public:
-
-		//! Not owned
-		InputStream* m_pStream;
-
-		//! Already read pointers
-		TArray< Ptr<RefCounted> > m_history;
-
-	};
-
-	//---------------------------------------------------------------------------------------------
-	class MUTABLERUNTIME_API OutputArchive::Private
-	{
-	public:
-
-		//! Not owned
-		OutputStream* m_pStream;
-
-		//! Already written pointers and their ids
-        TMap< const void*, int32 > m_history;
-
-	};
-
-
 #define MUTABLE_IMPLEMENT_POD_SERIALISABLE(Type)				     \
     void DLLEXPORT operator<<(OutputArchive& Arch, const Type& T)    \
     {																 \
-        Arch.GetPrivate()->m_pStream->Write(&T, sizeof(Type));		 \
+        Arch.Stream->Write(&T, sizeof(Type));						 \
     }																 \
                                                                      \
     void DLLEXPORT operator>>(InputArchive& Arch, Type& T)		     \
     {																 \
-        Arch.GetPrivate()->m_pStream->Read(&T, sizeof(Type));		 \
+        Arch.Stream->Read(&T, sizeof(Type));						 \
     }																 \
 		
 
@@ -161,7 +105,7 @@ namespace mu
 		Arch << Num;                                                           \
 		if (Num)                                                               \
 		{                                                                      \
-			Arch.GetPrivate()->m_pStream->Write(&V[0], Num * sizeof(Type));    \
+			Arch.Stream->Write(&V[0], Num * sizeof(Type));					   \
 		}                                                                      \
 	}                                                                          \
                                                                                \
@@ -173,7 +117,7 @@ namespace mu
 		V.SetNum(Num);                                                         \
 		if (Num)                                                               \
 		{                                                                      \
-			Arch.GetPrivate()->m_pStream->Read(&V[0], Num * sizeof(Type));     \
+			Arch.Stream->Read(&V[0], Num * sizeof(Type));					   \
 		}                                                                      \
 	}                                                                          \
 
@@ -233,13 +177,13 @@ namespace mu
         void DLLEXPORT operator<<(OutputArchive& Arch, const Type& T)   \
 		{																\
             uint32 V = (uint32)T;                                       \
-            Arch.GetPrivate()->m_pStream->Write(&V, sizeof(uint32));    \
+            Arch.Stream->Write(&V, sizeof(uint32));					\
 		}																\
 																		\
         void DLLEXPORT operator>>(InputArchive& Arch, Type& T)   		\
 		{																\
             uint32 V;													\
-            Arch.GetPrivate()->m_pStream->Read(&V, sizeof(uint32));     \
+            Arch.Stream->Read(&V, sizeof(uint32));					\
 			T = (Type)V;												\
 		}																\
 
@@ -337,12 +281,12 @@ namespace mu
 		}
 		else
 		{
-            int32* it = arch.GetPrivate()->m_history.Find(p.get());
+            int32* it = arch.History.Find(p.get());
 
 			if ( !it )
 			{
-                int32 id = arch.GetPrivate()->m_history.Num();
-                arch.GetPrivate()->m_history.Add( p.get(), id );
+                int32 id = arch.History.Num();
+                arch.History.Add( p.get(), id );
 				arch << id;
 				T::Serialise( p.get(), arch );
 			}
@@ -365,9 +309,9 @@ namespace mu
 		}
 		else
 		{
-			if ( id < arch.GetPrivate()->m_history.Num() )
+			if ( id < arch.History.Num() )
 			{
-				p = static_cast<T*>( arch.GetPrivate()->m_history[id].get() );
+				p = static_cast<T*>( arch.History[id].get() );
 
 				// If the pointer was 0 it means the position in history is used, but not set yet
 				// option 1: we have a smart pointer loop which is very bad.
@@ -379,12 +323,12 @@ namespace mu
 			{
                 // Ids come in order, but they may have been absorbed outside in some serialisations
                 // like proxies.
-                //check( id == (int)arch.GetPrivate()->m_history.size() );
-                arch.GetPrivate()->m_history.SetNum( id+1 );
+                //check( id == (int)arch.>History.size() );
+                arch.History.SetNum( id+1 );
 
 				Ptr<T> t = T::StaticUnserialise( arch );
 				p = t;
-				arch.GetPrivate()->m_history[id] = t;
+				arch.History[id] = t;
 			}
 		}
 	}
diff --git a/Engine/Plugins/Experimental/Mutable/Source/MutableRuntime/Private/MuR/Serialisation.cpp b/Engine/Plugins/Experimental/Mutable/Source/MutableRuntime/Private/MuR/Serialisation.cpp
--- a/Engine/Plugins/Experimental/Mutable/Source/MutableRuntime/Private/MuR/Serialisation.cpp
+++ b/Engine/Plugins/Experimental/Mutable/Source/MutableRuntime/Private/MuR/Serialisation.cpp
@@ -74,7 +74,7 @@ namespace mu
     	v.resize( size );
     	if (size)
     	{
-    		arch.GetPrivate()->m_pStream->Read( &v[0], (unsigned)size*sizeof(char) );
+    		arch.Stream->Read( &v[0], (unsigned)size*sizeof(char) );
     	}
     }
 
@@ -82,14 +82,14 @@ namespace mu
 	void operator<<(OutputArchive& Arch, const bool& T)
     {
     	uint8 S = T ? 1 : 0;
-    	Arch.GetPrivate()->m_pStream->Write(&S, sizeof(uint8));
+    	Arch.Stream->Write(&S, sizeof(uint8));
     }
 
 
 	void operator>>(InputArchive& Arch, bool& T)
     {
     	uint8 S;
-    	Arch.GetPrivate()->m_pStream->Read(&S, sizeof(uint8));
+    	Arch.Stream->Read(&S, sizeof(uint8));
     	T = S != 0;
     }
 
@@ -97,33 +97,23 @@ namespace mu
     //---------------------------------------------------------------------------------------------
     //---------------------------------------------------------------------------------------------
     //---------------------------------------------------------------------------------------------
-    InputMemoryStream::InputMemoryStream( const void* pBuffer, uint64 size )
+    InputMemoryStream::InputMemoryStream( const void* InBuffer, uint64 InSize )
     {
-        m_pD = new Private();
-        m_pD->m_pBuffer = pBuffer;
-        m_pD->m_size = size;
+        Buffer = InBuffer;
+        Size = InSize;
     }
 
 
     //---------------------------------------------------------------------------------------------
-    InputMemoryStream::~InputMemoryStream()
+    void InputMemoryStream::Read( void* Data, uint64 InSize)
     {
-        check( m_pD );
-        delete m_pD;
-        m_pD = 0;
-    }
-
-
-    //---------------------------------------------------------------------------------------------
-    void InputMemoryStream::Read( void* pData, uint64 size )
-    {
-        if (size)
+        if (InSize)
         {
-            check( m_pD->m_pos + size <= m_pD->m_size );
+            check( Pos + InSize <= Size );
 
-            const uint8* pSource = ((const uint8*)(m_pD->m_pBuffer))+m_pD->m_pos;
-            FMemory::Memcpy( pData, pSource, (SIZE_T)size );
-            m_pD->m_pos += size;
+            const uint8* Source = reinterpret_cast<const uint8*>(Buffer)+Pos;
+            FMemory::Memcpy( Data, Source, (SIZE_T)InSize);
+            Pos += InSize;
         }
     }
 
@@ -131,33 +121,23 @@ namespace mu
     //---------------------------------------------------------------------------------------------
     //---------------------------------------------------------------------------------------------
     //---------------------------------------------------------------------------------------------
-    OutputMemoryStream::OutputMemoryStream(uint64 reserve )
+    OutputMemoryStream::OutputMemoryStream(uint64 Reserve )
     {
-        m_pD = new Private();
-        if (reserve)
+        if (Reserve)
         {
-             m_pD->m_buffer.Reserve( reserve );
+             Buffer.Reserve( Reserve );
         }
     }
 
 
     //---------------------------------------------------------------------------------------------
-    OutputMemoryStream::~OutputMemoryStream()
-    {
-        check( m_pD );
-        delete m_pD;
-        m_pD = 0;
-    }
-
-
-    //---------------------------------------------------------------------------------------------
-    void OutputMemoryStream::Write( const void* pData, uint64 size )
+    void OutputMemoryStream::Write( const void* Data, uint64 Size)
     {
-        if (size)
+        if (Size)
         {
-            uint64 pos = m_pD->m_buffer.Num();
-            m_pD->m_buffer.SetNum( pos + size, EAllowShrinking::No );
-			FMemory::Memcpy( &m_pD->m_buffer[pos], pData, size );
+            uint64 Pos = Buffer.Num();
+			Buffer.SetNum( Pos + Size, EAllowShrinking::No );
+			FMemory::Memcpy( Buffer.GetData()+Pos, Data, Size);
         }
     }
 
@@ -167,9 +147,9 @@ namespace mu
     {
         const void* pResult = 0;
 
-        if ( m_pD->m_buffer.Num() )
+        if (Buffer.Num() )
         {
-            pResult = &m_pD->m_buffer[0];
+            pResult = Buffer.GetData();
         }
 
         return pResult;
@@ -178,15 +158,15 @@ namespace mu
 
 	//---------------------------------------------------------------------------------------------
 	uint64 OutputMemoryStream::GetBufferSize() const
-	{
-		return m_pD->m_buffer.Num();
-	}
+    {
+        return Buffer.Num();
+    }
 
 
 	//---------------------------------------------------------------------------------------------
 	void OutputMemoryStream::Reset()
 	{
-		m_pD->m_buffer.SetNum(0,EAllowShrinking::No);
+		Buffer.SetNum(0,EAllowShrinking::No);
 	}
 
 
@@ -195,47 +175,30 @@ namespace mu
     //-------------------------------------------------------------------------------------------------
     OutputSizeStream::OutputSizeStream()
     {
-        m_writtenBytes = 0;
+        WrittenBytes = 0;
     }
 
 
     //-------------------------------------------------------------------------------------------------
     void OutputSizeStream::Write( const void*, uint64 size )
     {
-        m_writtenBytes += size;
+        WrittenBytes += size;
     }
 
 
     //-------------------------------------------------------------------------------------------------
 	uint64 OutputSizeStream::GetBufferSize() const
     {
-        return m_writtenBytes;
+        return WrittenBytes;
     }
 
 
     //---------------------------------------------------------------------------------------------
 	//---------------------------------------------------------------------------------------------
 	//---------------------------------------------------------------------------------------------
-	InputArchive::InputArchive( InputStream* pStream )
+	InputArchive::InputArchive( InputStream* InStream )
 	{
-		m_pD = new Private();
-		m_pD->m_pStream = pStream;
-	}
-
-
-	//---------------------------------------------------------------------------------------------
-	InputArchive::~InputArchive()
-	{
-        check( m_pD );
-		delete m_pD;
-		m_pD = 0;
-	}
-
-
-	//---------------------------------------------------------------------------------------------
-	InputArchive::Private* InputArchive::GetPrivate() const
-	{
-		return m_pD;
+		Stream = InStream;
 	}
 
 
@@ -249,55 +212,19 @@ namespace mu
 	//---------------------------------------------------------------------------------------------
 	//---------------------------------------------------------------------------------------------
 	//---------------------------------------------------------------------------------------------
-	OutputArchive::OutputArchive( OutputStream* pStream )
-	{
-		m_pD = new Private();
-		m_pD->m_pStream = pStream;
-	}
-
-
-	//---------------------------------------------------------------------------------------------
-	OutputArchive::~OutputArchive()
-	{
-        check( m_pD );
-		delete m_pD;
-		m_pD = 0;
-	}
-
-
-	//---------------------------------------------------------------------------------------------
-	OutputArchive::Private* OutputArchive::GetPrivate() const
+	OutputArchive::OutputArchive( OutputStream* InStream )
 	{
-		return m_pD;
+		Stream = InStream;
 	}
 
 
-    //---------------------------------------------------------------------------------------------
-    //---------------------------------------------------------------------------------------------
-    //---------------------------------------------------------------------------------------------
-    class InputArchiveWithProxies::Private
-    {
-    public:
-        TArray< Ptr<ResourceProxy<Image>> > m_proxyHistory;
-        ProxyFactory* m_pFactory = nullptr;
-    };
-
-
     //---------------------------------------------------------------------------------------------
     //---------------------------------------------------------------------------------------------
     //---------------------------------------------------------------------------------------------
     InputArchiveWithProxies::InputArchiveWithProxies( InputStream* s, ProxyFactory* f )
         : InputArchive( s )
     {
-        m_pD = new Private();
-        m_pD->m_pFactory = f;
-    }
-
-
-    //---------------------------------------------------------------------------------------------
-    InputArchiveWithProxies::~InputArchiveWithProxies()
-    {
-        delete m_pD;
+        Factory = f;
     }
 
 
@@ -325,9 +252,9 @@ namespace mu
             }
             else
             {
-                if ( id < m_pD->m_proxyHistory.Num() )
+                if ( id < ProxyHistory.Num() )
                 {
-                    p = m_pD->m_proxyHistory[id];
+                    p = ProxyHistory[id];
 
                     // If the pointer was nullptr it means the position in history is used, but not set
                     // yet: we have a smart pointer loop which is very bad.
@@ -336,10 +263,10 @@ namespace mu
                 else
                 {
                     // Ids come in order.
-                    m_pD->m_proxyHistory.SetNum(id+1);
+                    ProxyHistory.SetNum(id+1);
 
-                    p = m_pD->m_pFactory->NewImageProxy(*this);
-                    m_pD->m_proxyHistory[id] = p;
+                    p = Factory->NewImageProxy(*this);
+                    ProxyHistory[id] = p;
                 }
             }
         }
diff --git a/Engine/Plugins/Experimental/Mutable/Source/MutableRuntime/Public/MuR/Serialisation.h b/Engine/Plugins/Experimental/Mutable/Source/MutableRuntime/Public/MuR/Serialisation.h
--- a/Engine/Plugins/Experimental/Mutable/Source/MutableRuntime/Public/MuR/Serialisation.h
+++ b/Engine/Plugins/Experimental/Mutable/Source/MutableRuntime/Public/MuR/Serialisation.h
@@ -8,12 +8,14 @@
 #include "MuR/MutableMemory.h"
 #include "MuR/Ptr.h"
 #include "MuR/RefCounted.h"
+#include "MuR/MutableMath.h"
 
 #include "Math/Vector.h"
 #include "Math/IntVector.h"
 #include "Math/Vector4.h"
-#include "MuR/MutableMath.h"
-
+#include "Containers/Array.h"
+#include "Containers/Map.h"
+#include "Containers/ContainerAllocationPolicies.h"
 
 namespace mu
 {    
@@ -161,94 +163,84 @@ namespace mu
 	};
 
 
-
-    //! Interface for any input stream to be use with InputArchives.
-    //! \ingroup tools
+    /** Interface for any input stream to be use with InputArchives. */
     class MUTABLERUNTIME_API InputStream
     {
     public:
 
-        //! Ensure virtual destruction
-        virtual ~InputStream() {}
+		/** Ensure virtual destruction. */
+		virtual ~InputStream() {}
 
-        //! Read a byte buffer
-        //! \param pData destination buffer, must have at least size bytes allocated.
-        //! \param size amount of bytes to read from the stream.
-        virtual void Read( void* pData, uint64 size ) = 0;
+        /** Read a byte buffer
+         * \param pData destination buffer, must have at least size bytes allocated.
+         * \param size amount of bytes to read from the stream.
+		 */
+        virtual void Read( void* Data, uint64 Size ) = 0;
     };
 
 
-    //! Interface for any output stream to be used with OutputArchives
-    //! \ingroup tools
+    /** Interface for any output stream to be used with OutputArchives. */
     class MUTABLERUNTIME_API OutputStream
     {
     public:
 
-        //! Ensure virtual destruction
+        /** Ensure virtual destruction. */
         virtual ~OutputStream() {}
 
-        //! Write a byte buffer
-        //! \param pData source buffer where data will be read from.
-        //! \param size amount of data to write to the stream.
-        virtual void Write( const void* pData, uint64 size ) = 0;
+        /** Write a byte buffer
+         * \param pData source buffer where data will be read from.
+         * \param size amount of data to write to the stream.
+		 */
+        virtual void Write( const void* Data, uint64 Size ) = 0;
 
     };
 
 
-    //! Archive containing data to be deserialised.
-    //! \ingroup tools
+    /** Archive containing data to be deserialised. */
     class MUTABLERUNTIME_API InputArchive
     {
     public:
 
-        //! Construct form an input stream. The stream will not be owned by the archive and the
-        //! caller must make sure it is not modified or destroyed while serialisation is happening.
+        /** Construct form an input stream.The stream will not be owned by the archive and the
+         * caller must make sure it is not modified or destroyed while serialisation is happening.
+		 */
         InputArchive( InputStream* );
 
-        //!
-        virtual ~InputArchive();
+		/** Ensure virtual destruction. */
+		virtual ~InputArchive() {}
 
         //
         virtual Ptr<ResourceProxy<Image>> NewImageProxy();
 
-        // Interface pattern
-        class Private;
+		/** Not owned. */
+		InputStream* Stream = nullptr;
 
-        Private* GetPrivate() const;
-
-    private:
-
-        Private* m_pD;
+		/** Already read pointers. */
+		TArray< Ptr<RefCounted> > History;
 
     };
 
 
-    //! Archive where data can be serialised to.
-    //! \ingroup tools
+    /** Archive where data can be serialised to. */
     class MUTABLERUNTIME_API OutputArchive
     {
     public:
 
-        //! Construct form an output stream. The stream will not be owned by the archive and the
-        //! caller must make sure it is not modified or destroyed while serialisation is happening.
+        /** Construct form an output stream.The stream will not be owned by the archive and the
+         * caller must make sure it is not modified or destroyed while serialisation is happening.
+		 */
         OutputArchive( OutputStream* );
 
-        //!
-        ~OutputArchive();
-
-        // Interface pattern
-        class Private;
+		/** Not owned. */
+		OutputStream* Stream;
 
-        Private* GetPrivate() const;
-
-    private:
-
-        Private* m_pD;
+		/** Already written pointers and their ids. */
+		TMap< const void*, int32 > History;
 
     };
 
-    //!
-    //! \ingroup runtime
+
+    /** */
     class MUTABLERUNTIME_API InputArchiveWithProxies : public InputArchive
     {
     public:
@@ -262,22 +254,19 @@ namespace mu
 
     public:
 
-        //! Construct form an input stream. The stream will not be owned by the archive and the
-        //! caller must make sure it is not modified or destroyed while serialisation is happening.
+        /** Construct form an input stream.The stream will not be owned by the archive and the
+         * caller must make sure it is not modified or destroyed while serialisation is happening.
+		 */
         InputArchiveWithProxies( InputStream*, ProxyFactory*  );
 
-        ~InputArchiveWithProxies() override;
-
-        //
-        Ptr<ResourceProxy<Image>> NewImageProxy() override;
+        // InputArchive interface
+        virtual Ptr<ResourceProxy<Image>> NewImageProxy() override;
 
     private:
 
-        class Private;
-
-        Private* m_pD;
-    };
-
+		TArray< Ptr<ResourceProxy<Image>> > ProxyHistory;
+		ProxyFactory* Factory = nullptr;
+	};
 
 
     //!
@@ -285,31 +274,22 @@ namespace mu
     {
     public:
 
-        //-----------------------------------------------------------------------------------------
-        // Life cycle
-        //-----------------------------------------------------------------------------------------
-
-        //! Create the stream with an optional buffer size in bytes.
-        //! The internal buffer will be enlarged as much as necessary.
-        OutputMemoryStream( uint64 reserve = 0 );
+        /** Create the stream with an optional buffer size in bytes.
+		* The internal buffer will be enlarged as much as necessary.
+		*/
+        OutputMemoryStream( uint64 Reserve = 0 );
 
-        ~OutputMemoryStream();
-
-
-        //-----------------------------------------------------------------------------------------
         // OutputStream interface
-        //-----------------------------------------------------------------------------------------
-        void Write( const void* pData, uint64 size ) override;
+        virtual void Write( const void* Data, uint64 Size ) override;
 
-        //-----------------------------------------------------------------------------------------
         // Own interface
-        //-----------------------------------------------------------------------------------------
 
-        //! Get the serialised data buffer pointer. This pointer invalidates after a Write
-        //! operation has been done, and you need to get it again.
+        /** Get the serialised data buffer pointer. This pointer invalidates after a Write
+         * operation has been done, and you need to get it again.
+		 */
         const void* GetBuffer() const;
 
-        //! Get the amount of data in the stream, in bytes.
+        /** Get the amount of data in the stream, in bytes. */
         uint64 GetBufferSize() const;
 
 		/** Clear the internal buffer. */
@@ -317,9 +297,7 @@ namespace mu
 
     private:
 
-        class Private;
-
-        Private* m_pD;
+		TArray64<uint8> Buffer;
 
     };
 
diff --git a/Engine/Plugins/Experimental/Mutable/Source/MutableTools/Private/MuT/ASTOpConstantResource.cpp b/Engine/Plugins/Experimental/Mutable/Source/MutableTools/Private/MuT/ASTOpConstantResource.cpp
--- a/Engine/Plugins/Experimental/Mutable/Source/MutableTools/Private/MuT/ASTOpConstantResource.cpp
+++ b/Engine/Plugins/Experimental/Mutable/Source/MutableTools/Private/MuT/ASTOpConstantResource.cpp
@@ -6,92 +6,154 @@
 #include "MuR/Layout.h"
 #include "MuR/Mesh.h"
 #include "MuR/ModelPrivate.h"
+#include "MuR/ImagePrivate.h"
 #include "MuR/MutableMath.h"
 #include "MuR/PhysicsBody.h"
 #include "MuR/Serialisation.h"
 #include "MuR/Skeleton.h"
 #include "MuR/Types.h"
+
 #include "Containers/Array.h"
 #include "HAL/PlatformMath.h"
 #include "HAL/PlatformFileManager.h"
 #include "Hash/CityHash.h"
 #include "Misc/AssertionMacros.h"
 #include "GenericPlatform/GenericPlatformFile.h"
+#include "Compression/OodleDataCompression.h"
 
 #include <inttypes.h> // Required for 64-bit printf macros
 
 namespace mu
 {
 
-	//---------------------------------------------------------------------------------------------
-	//---------------------------------------------------------------------------------------------
-	//---------------------------------------------------------------------------------------------
+	/** Proxy class for a temporary resource while compiling. 
+	* The resource may be stored in different ways:
+	* - as is, in memory with its own pointer.
+	* - in a compressed buffer
+	* - saved to a disk file compressed or uncompressed.
+	*/
 	template<class R>
 	class MUTABLETOOLS_API ResourceProxyTempFile : public ResourceProxy<R>
 	{
 	private:
+
+		/** Actual resource to store. If the pointer is valid, it wasn't worth dumping to disk or compressing. */
 		Ptr<const R> Resource;
+
+		/** Temp filename used if it was necessary. */
 		FString FileName;
-		uint64 FileSize = 0;
-		FCriticalSection Mutex;
 
+		/** Size of the resource in memory. */
+		uint32 UncompressedSize = 0;
+
+		/** Size of the saved file. It may be the size of the resource in memory, or its compressed size. */
+		uint32 FileSize = 0;
+
+		/** Valid if the resource was compressed and stored in memory instead of dumped to disk. */
+		TArray<uint8> CompressedBuffer;
+
+		/** Shared context with cache settings and stats. */
 		FProxyFileContext& Options;
 
+		/** Prevent concurrent access to a signel resource. */
+		FCriticalSection Mutex;
+
 	public:
-		ResourceProxyTempFile(const R* resource, FProxyFileContext& InOptions)
+
+		ResourceProxyTempFile(const R* InResource, FProxyFileContext& InOptions)
 			: Options(InOptions)
 		{
-			IPlatformFile& PlatformFile = FPlatformFileManager::Get().GetPlatformFile();
-
-			if (!resource)
+			if (!InResource)
 			{
 				return;
 			}
 
-			FScopeLock Lock(&Mutex);
+			IPlatformFile& PlatformFile = FPlatformFileManager::Get().GetPlatformFile();
 
-			OutputMemoryStream stream;
+			OutputMemoryStream stream(128*1024);
 			OutputArchive arch(&stream);
-			R::Serialise(resource, arch);
+			R::Serialise(InResource, arch);
+
+			UncompressedSize = stream.GetBufferSize();
 
 			if (stream.GetBufferSize() <= Options.MinProxyFileSize)
 			{
-				Resource = resource;
+				// Not worth compressing or caching to disk
+				Resource = InResource;
 			}
 			else
 			{
-				FileSize = stream.GetBufferSize();
-
-				FString Prefix = FPlatformProcess::UserTempDir();
-
-				uint32 PID = FPlatformProcess::GetCurrentProcessId();
-				Prefix += FString::Printf(TEXT("mut.temp.%u"), PID);
-
-				FString FinalTempPath;
-				IFileHandle* ResourceFile = nullptr;
-				uint64 AttemptCount = 0;
-				while (!ResourceFile && AttemptCount < Options.MaxFileCreateAttempts)
+				// Compress
+				int64 CompressedSize = 0;
+				constexpr bool bEnableCompression = true;
+				if (bEnableCompression)
 				{
-					uint64 ThisThreadFileIndex = Options.CurrentFileIndex.load();
-					while (!Options.CurrentFileIndex.compare_exchange_strong(ThisThreadFileIndex, ThisThreadFileIndex + 1));
-
-					FinalTempPath = Prefix + FString::Printf(TEXT(".%.16" PRIx64), ThisThreadFileIndex);
-					ResourceFile = PlatformFile.OpenWrite(*FinalTempPath);
-					++AttemptCount;
+					int64 CompressedBufferSize = FOodleDataCompression::CompressedBufferSizeNeeded(stream.GetBufferSize());
+					CompressedBufferSize = FMath::Max(CompressedBufferSize, int64(stream.GetBufferSize() / 2));
+					CompressedBuffer.SetNumUninitialized(CompressedBufferSize);
+
+					CompressedSize = FOodleDataCompression::CompressParallel(
+						CompressedBuffer.GetData(), CompressedBufferSize,
+						stream.GetBuffer(), stream.GetBufferSize(),
+						FOodleDataCompression::ECompressor::Kraken,
+						FOodleDataCompression::ECompressionLevel::SuperFast,
+						true // CompressIndependentChunks
+					);
 				}
 
-				if (!ResourceFile)
+				bool bCompressed = CompressedSize != 0;
+
+				if (bCompressed && uint64(CompressedSize) <= Options.MinProxyFileSize)
 				{
-					UE_LOG(LogMutableCore, Error, TEXT("Failed to create temporary file. Disk full?"));
-					check(false);
+					// Keep the compressed data, and don't store to file
+					CompressedBuffer.SetNum(CompressedSize,EAllowShrinking::Yes);
 				}
+				else
+				{
+					// Save
+					FString Prefix = FPlatformProcess::UserTempDir();
 
-				ResourceFile->Write((const uint8*)stream.GetBuffer(), stream.GetBufferSize());
-				delete ResourceFile;
+					uint32 PID = FPlatformProcess::GetCurrentProcessId();
+					Prefix += FString::Printf(TEXT("mut.temp.%u"), PID);
+
+					FString FinalTempPath;
+					IFileHandle* ResourceFile = nullptr;
+					uint64 AttemptCount = 0;
+					while (!ResourceFile && AttemptCount < Options.MaxFileCreateAttempts)
+					{
+						uint64 ThisThreadFileIndex = Options.CurrentFileIndex.load();
+						while (!Options.CurrentFileIndex.compare_exchange_strong(ThisThreadFileIndex, ThisThreadFileIndex + 1));
+
+						FinalTempPath = Prefix + FString::Printf(TEXT(".%.16" PRIx64), ThisThreadFileIndex);
+						ResourceFile = PlatformFile.OpenWrite(*FinalTempPath);
+						++AttemptCount;
+					}
+
+					if (!ResourceFile)
+					{
+						UE_LOG(LogMutableCore, Error, TEXT("Failed to create temporary file. Disk full?"));
+						check(false);
+					}
 
-				FileName = FinalTempPath;
-				Options.FilesWritten++;
-				Options.BytesWritten += stream.GetBufferSize();
+					if (bCompressed)
+					{
+						FileSize = CompressedSize;
+						ResourceFile->Write(CompressedBuffer.GetData(), FileSize);
+					}
+					else
+					{
+						FileSize = UncompressedSize;
+						ResourceFile->Write((const uint8*)stream.GetBuffer(), FileSize);
+					}
+
+					CompressedBuffer.SetNum(0, EAllowShrinking::Yes);
+
+					delete ResourceFile;
+
+					FileName = FinalTempPath;
+					Options.FilesWritten++;
+					Options.BytesWritten += FileSize;
+				}
 			}
 		}
 
@@ -111,29 +173,61 @@ namespace mu
 		{
 			FScopeLock Lock(&Mutex);
 
-			Ptr<const R> r;
+			Ptr<const R> Result;
 			if (Resource)
 			{
-				// cached
-				r = Resource;
+				// Cached as is
+				Result = Resource;
 			}
-			else if (!FileName.IsEmpty())
+			else if (!CompressedBuffer.Num() && !FileName.IsEmpty())
 			{
-				TArray<char> buf;
-				buf.SetNumUninitialized(FileSize);
 				IFileHandle* resourceFile = FPlatformFileManager::Get().GetPlatformFile().OpenRead(*FileName);
 				check(resourceFile);
-				resourceFile->Read((uint8*)buf.GetData(), FileSize);
+
+				CompressedBuffer.SetNumUninitialized(FileSize);
+				resourceFile->Read(CompressedBuffer.GetData(), FileSize);
 				delete resourceFile;
 
-				InputMemoryStream stream(buf.GetData(), FileSize);
-				InputArchive arch(&stream);
-				r = R::StaticUnserialise(arch);
+				bool bCompressed = FileSize != UncompressedSize;
+
+				if (!bCompressed)
+				{
+					InputMemoryStream stream(CompressedBuffer.GetData(), FileSize);
+					InputArchive arch(&stream);
+					Result = R::StaticUnserialise(arch);
+
+					CompressedBuffer.SetNum(0, EAllowShrinking::Yes);
+				}
 
 				Options.FilesRead++;
 				Options.BytesRead += FileSize;
 			}
-			return r;
+
+			if (CompressedBuffer.Num())
+			{
+				// Cached compressed
+				TArray<uint8> UncompressedBuf;
+				UncompressedBuf.SetNumUninitialized(UncompressedSize);
+
+				bool bSuccess = FOodleDataCompression::DecompressParallel(
+					UncompressedBuf.GetData(), UncompressedSize,
+					CompressedBuffer.GetData(), CompressedBuffer.Num());
+				check(bSuccess);
+
+				if (bSuccess)
+				{
+					InputMemoryStream stream(UncompressedBuf.GetData(), UncompressedSize);
+					InputArchive arch(&stream);
+					Result = R::StaticUnserialise(arch);
+				}
+
+				if (!FileName.IsEmpty())
+				{
+					CompressedBuffer.SetNum(0, EAllowShrinking::Yes);
+				}
+			}
+
+			return Result;
 		}
 
 	};
@@ -217,7 +311,6 @@ namespace mu
 					MipsToStore = pImage->GetLODCount();
 				}
 
-				// TODO: If the image already has mips, we will be duplicating them...
 				if (pImage->GetLODCount() == 1)
 				{
 					pMip = pImage;
@@ -228,7 +321,11 @@ namespace mu
 				}
 			}
 
-			for (int Mip = 0; Mip < MipsToStore; ++Mip)
+			// Temporary uncompressed version of the image, if we need to generate the mips and the source is compressed.
+			Ptr<const Image> UncompressedMip;
+			EImageFormat UncompressedFormat = GetUncompressedFormat( pMip->GetFormat() );
+
+			for (int32 Mip = 0; Mip < MipsToStore; ++Mip)
 			{
 				check(pMip->GetFormat() == pImage->GetFormat());
 
@@ -259,14 +356,30 @@ namespace mu
 				if (Mip + 1 < MipsToStore)
 				{
 					Ptr<Image> NewMip;
-					if (Mip > pImage->GetLODCount())
+					if (Mip+1 < pImage->GetLODCount())
 					{
-						// Generate from the last mip.
-						NewMip = ImOp.ExtractMip(pMip.get(), 1);
+						// Extract directly from source image
+						NewMip = ImOp.ExtractMip(pImage.get(), Mip + 1);
 					}
 					else
 					{
-						NewMip = ImOp.ExtractMip(pImage.get(), Mip + 1);
+						// Generate from the last mip.
+						if (UncompressedFormat!=pMip->GetFormat())
+						{
+							int32 Quality = 4; // TODO
+
+							if (!UncompressedMip)
+							{
+								UncompressedMip = ImOp.ImagePixelFormat(Quality, pMip.get(), UncompressedFormat);
+							}
+
+							UncompressedMip = ImOp.ExtractMip(UncompressedMip.get(), 1);
+							NewMip = ImOp.ImagePixelFormat(Quality, UncompressedMip.get(), pMip->GetFormat());
+						}
+						else
+						{
+							NewMip = ImOp.ExtractMip(pMip.get(), 1);
+						}
 					}
 					check(NewMip);
 
diff --git a/Engine/Plugins/Experimental/Mutable/Source/MutableTools/Private/MuT/CompilerPrivate.h b/Engine/Plugins/Experimental/Mutable/Source/MutableTools/Private/MuT/CompilerPrivate.h
--- a/Engine/Plugins/Experimental/Mutable/Source/MutableTools/Private/MuT/CompilerPrivate.h
+++ b/Engine/Plugins/Experimental/Mutable/Source/MutableTools/Private/MuT/CompilerPrivate.h
@@ -19,8 +19,8 @@ namespace mu
 
 		/** Options */
 
-		/** Minimum data size in bytes to dumpt it to the disk. */
-		uint64 MinProxyFileSize = 1024 * 1024;
+		/** Minimum data size in bytes to dump it to the disk. */
+		uint64 MinProxyFileSize = 64 * 1024;
 
 		/** When creating temporary files, number of retries in case the OS-level call fails. */
 		uint64 MaxFileCreateAttempts = 256;
