[mutable] Fix problem getting view of a single lod of an image inside the lod tail.
#rnx

diff --git a/Engine/Plugins/Experimental/Mutable/Source/MutableRuntime/Public/MuR/ImageDataStorage.h b/Engine/Plugins/Experimental/Mutable/Source/MutableRuntime/Public/MuR/ImageDataStorage.h
--- a/Engine/Plugins/Experimental/Mutable/Source/MutableRuntime/Public/MuR/ImageDataStorage.h
+++ b/Engine/Plugins/Experimental/Mutable/Source/MutableRuntime/Public/MuR/ImageDataStorage.h
@@ -319,7 +319,7 @@ public:
 
 		if (FirstCompactedTailLOD < LODEnd)
 		{
-			TArrayView<const uint8> FirstLODInRangeView = GetLOD(FirstCompactedTailLOD);
+			TArrayView<const uint8> FirstLODInRangeView = GetLOD(LODBegin);
 			TArrayView<const uint8> LastLODInRangeView = GetLOD(LODEnd - 1);
 			
 			const int32 TailInLODRangeNumBytes = LastLODInRangeView.GetData() - FirstLODInRangeView.GetData() + LastLODInRangeView.Num();
