[mutable] Avoid unnecessary pixel conversion when converting single image mips.
#rnx

diff --git a/Engine/Plugins/Experimental/Mutable/Source/MutableTools/Private/MuT/UnrealPixelFormatOverride.cpp b/Engine/Plugins/Experimental/Mutable/Source/MutableTools/Private/MuT/UnrealPixelFormatOverride.cpp
--- a/Engine/Plugins/Experimental/Mutable/Source/MutableTools/Private/MuT/UnrealPixelFormatOverride.cpp
+++ b/Engine/Plugins/Experimental/Mutable/Source/MutableTools/Private/MuT/UnrealPixelFormatOverride.cpp
@@ -156,7 +156,9 @@ void MutableToImageCore(const mu::Image* InMutable, FImage& CoreImage, int32 LOD
 	{
 		// Unsupported format: force conversion
 		mu::FImageOperator ImOp = mu::FImageOperator::GetDefault(mu::FImageOperator::FImagePixelFormatFunc());
-		Mutable = ImOp.ImagePixelFormat(4, InMutable, mu::EImageFormat::IF_BGRA_UBYTE);
+		Mutable = ImOp.ImagePixelFormat(4, InMutable, mu::EImageFormat::IF_BGRA_UBYTE,LOD);
+		// We are extracting one LOD, so always access LOD 0 of the resulting mutable image 
+		LOD = 0;
 		CoreImageFormat = ERawImageFormat::BGRA8;
 		break;
 	}
