[mutable] Added a button to remove any projector layer from the DGGUI.
+Exposed new functions to the mutable API.
+Removed unnecessary functions from the mutable API.

#rnx

#virtualized

diff --git a/Engine/Plugins/Experimental/Mutable/Source/CustomizableObject/Internal/MuCO/CustomizableObjectInstanceDescriptor.h b/Engine/Plugins/Experimental/Mutable/Source/CustomizableObject/Internal/MuCO/CustomizableObjectInstanceDescriptor.h
--- a/Engine/Plugins/Experimental/Mutable/Source/CustomizableObject/Internal/MuCO/CustomizableObjectInstanceDescriptor.h
+++ b/Engine/Plugins/Experimental/Mutable/Source/CustomizableObject/Internal/MuCO/CustomizableObjectInstanceDescriptor.h
@@ -224,17 +224,13 @@ struct CUSTOMIZABLEOBJECT_API FCustomizableObjectInstanceDescriptor
 	 * The added value is initialized with the default projector as set up in the editor and is the last one of the range. */
 	int32 AddValueToProjectorRange(const FString& ParamName);
 
-	/** Remove the last of the integer range of values from the parameter ParamName, returns the index of the last valid integer, -1 if no values left. */
-	int32 RemoveValueFromIntRange(const FString& ParamName);
+	/** Remove the RangeIndex element of the integer range of values from the parameter ParamName. If Range index is -1, removes the last element.
+		Returns the index of the last valid integer, -1 if no values left. */
+	int32 RemoveValueFromIntRange(const FString& ParamName, int32 RangeIndex = -1);
 
-	/** Remove the RangeIndex element of the integer range of values from the parameter ParamName, returns the index of the last valid integer, -1 if no values left. */
-	int32 RemoveValueFromIntRange(const FString& ParamName, int32 RangeIndex);
-
-	/** Remove the last of the float range of values from the parameter ParamName, returns the index of the last valid float, -1 if no values left. */
-	int32 RemoveValueFromFloatRange(const FString& ParamName);
-
-	/** Remove the RangeIndex element of the float range of values from the parameter ParamName, returns the index of the last valid float, -1 if no values left. */
-	int32 RemoveValueFromFloatRange(const FString& ParamName, int32 RangeIndex);
+	/** Remove the RangeIndex element of the float range of values from the parameter ParamName. If Range index is -1, removes the last element.
+		Returns the index of the last valid float, -1 if no values left. */
+	int32 RemoveValueFromFloatRange(const FString& ParamName, int32 RangeIndex = -1);
 
 	/** Remove the last of the texture range of values from the parameter ParamName, returns the index of the last valid float, -1 if no values left. */
 	int32 RemoveValueFromTextureRange(const FString& ParamName);
@@ -242,11 +238,10 @@ struct CUSTOMIZABLEOBJECT_API FCustomizableObjectInstanceDescriptor
 	/** Remove the RangeIndex element of the texture range of values from the parameter ParamName, returns the index of the last valid float, -1 if no values left. */
 	int32 RemoveValueFromTextureRange(const FString& ParamName, int32 RangeIndex);
 
-	/** Remove the last of the projector range of values from the parameter ParamName, returns the index of the last valid projector, -1 if no values left. */
-	int32 RemoveValueFromProjectorRange(const FString& ParamName);
-
-	/** Remove the RangeIndex element of the projector range of values from the parameter ParamName, returns the index of the last valid projector, -1 if no values left. */
-	int32 RemoveValueFromProjectorRange(const FString& ParamName, int32 RangeIndex);
+	/** Remove the RangeIndex element of the projector range of values from the parameter ParamName. If Range index is -1, removes the last element.
+		Returns the index of the last valid projector, -1 if no values left.
+	*/
+	int32 RemoveValueFromProjectorRange(const FString& ParamName, int32 RangeIndex = -1);
 
 	// ------------------------------------------------------------
    	// States
diff --git a/Engine/Plugins/Experimental/Mutable/Source/CustomizableObject/Private/MuCO/CustomizableObjectInstance.cpp b/Engine/Plugins/Experimental/Mutable/Source/CustomizableObject/Private/MuCO/CustomizableObjectInstance.cpp
--- a/Engine/Plugins/Experimental/Mutable/Source/CustomizableObject/Private/MuCO/CustomizableObjectInstance.cpp
+++ b/Engine/Plugins/Experimental/Mutable/Source/CustomizableObject/Private/MuCO/CustomizableObjectInstance.cpp
@@ -2399,23 +2399,12 @@ int32 UCustomizableObjectInstance::AddValueToProjectorRange(const FString& Param
 }
 
 
-int32 UCustomizableObjectInstance::RemoveValueFromIntRange(const FString& ParamName)
-{
-	return Descriptor.RemoveValueFromIntRange(ParamName);
-}
-
-
 int32 UCustomizableObjectInstance::RemoveValueFromIntRange(const FString& ParamName, int32 RangeIndex)
 {
 	return Descriptor.RemoveValueFromIntRange(ParamName, RangeIndex);
 
 }
 
-int32 UCustomizableObjectInstance::RemoveValueFromFloatRange(const FString& ParamName)
-{
-	return Descriptor.RemoveValueFromFloatRange(ParamName);
-}
-
 
 int32 UCustomizableObjectInstance::RemoveValueFromFloatRange(const FString& ParamName, const int32 RangeIndex)
 {
@@ -2423,12 +2412,6 @@ int32 UCustomizableObjectInstance::RemoveValueFromFloatRange(const FString& Para
 }
 
 
-int32 UCustomizableObjectInstance::RemoveValueFromProjectorRange(const FString& ParamName)
-{
-	return Descriptor.RemoveValueFromProjectorRange(ParamName);
-}
-
-
 int32 UCustomizableObjectInstance::RemoveValueFromProjectorRange(const FString& ParamName, const int32 RangeIndex)
 {
 	return Descriptor.RemoveValueFromProjectorRange(ParamName, RangeIndex);
diff --git a/Engine/Plugins/Experimental/Mutable/Source/CustomizableObject/Private/MuCO/CustomizableObjectInstanceDescriptor.cpp b/Engine/Plugins/Experimental/Mutable/Source/CustomizableObject/Private/MuCO/CustomizableObjectInstanceDescriptor.cpp
--- a/Engine/Plugins/Experimental/Mutable/Source/CustomizableObject/Private/MuCO/CustomizableObjectInstanceDescriptor.cpp
+++ b/Engine/Plugins/Experimental/Mutable/Source/CustomizableObject/Private/MuCO/CustomizableObjectInstanceDescriptor.cpp
@@ -2090,22 +2090,6 @@ int32 FCustomizableObjectInstanceDescriptor::AddValueToProjectorRange(const FStr
 }
 
 
-int32 FCustomizableObjectInstanceDescriptor::RemoveValueFromIntRange(const FString& ParamName)
-{
-	const int32 IntParameterIndex = FindTypedParameterIndex(ParamName, EMutableParameterType::Int);
-	if (IntParameterIndex != INDEX_NONE)
-	{
-		FCustomizableObjectIntParameterValue& IntParameter = IntParameters[IntParameterIndex];
-		if (IntParameter.ParameterRangeValueNames.Num() > 0)
-		{
-			IntParameter.ParameterRangeValueNames.Pop();
-			return IntParameter.ParameterRangeValueNames.Num() - 1;
-		}
-	}
-	return -1;
-}
-
-
 int32 FCustomizableObjectInstanceDescriptor::RemoveValueFromIntRange(const FString& ParamName, const int32 RangeIndex)
 {
 	const int32 IntParameterIndex = FindTypedParameterIndex(ParamName, EMutableParameterType::Int);
@@ -2114,7 +2098,8 @@ int32 FCustomizableObjectInstanceDescriptor::RemoveValueFromIntRange(const FStri
 		FCustomizableObjectIntParameterValue& IntParameter = IntParameters[IntParameterIndex];
 		if (IntParameter.ParameterRangeValueNames.Num() > 0)
 		{
-			IntParameter.ParameterRangeValueNames.RemoveAt(RangeIndex);
+			int32 IndexToRemove = IntParameter.ParameterRangeValueNames.IsValidIndex(RangeIndex) ? RangeIndex : IntParameter.ParameterRangeValueNames.Num() - 1;
+			IntParameter.ParameterRangeValueNames.RemoveAt(IndexToRemove);
 			return IntParameter.ParameterRangeValueNames.Num() - 1;
 		}
 	}
@@ -2122,22 +2107,6 @@ int32 FCustomizableObjectInstanceDescriptor::RemoveValueFromIntRange(const FStri
 }
 
 
-int32 FCustomizableObjectInstanceDescriptor::RemoveValueFromFloatRange(const FString& ParamName)
-{
-	const int32 FloatParameterIndex = FindTypedParameterIndex(ParamName, EMutableParameterType::Float);
-	if (FloatParameterIndex != INDEX_NONE)
-	{
-		FCustomizableObjectFloatParameterValue& FloatParameter = FloatParameters[FloatParameterIndex];
-		if (FloatParameter.ParameterRangeValues.Num() > 0)
-		{
-			FloatParameter.ParameterRangeValues.Pop();
-			return FloatParameter.ParameterRangeValues.Num() - 1;
-		}
-	}
-	return -1;
-}
-
-
 int32 FCustomizableObjectInstanceDescriptor::RemoveValueFromFloatRange(const FString& ParamName, const int32 RangeIndex)
 {
 	const int32 FloatParameterIndex = FindTypedParameterIndex(ParamName, EMutableParameterType::Float);
@@ -2146,7 +2115,8 @@ int32 FCustomizableObjectInstanceDescriptor::RemoveValueFromFloatRange(const FSt
 		FCustomizableObjectFloatParameterValue& FloatParameter = FloatParameters[FloatParameterIndex];
 		if (FloatParameter.ParameterRangeValues.Num() > 0)
 		{
-			FloatParameter.ParameterRangeValues.RemoveAt(RangeIndex);
+			int32 IndexToRemove = FloatParameter.ParameterRangeValues.IsValidIndex(RangeIndex) ? RangeIndex : FloatParameter.ParameterRangeValues.Num() - 1;
+			FloatParameter.ParameterRangeValues.RemoveAt(IndexToRemove);
 			return FloatParameter.ParameterRangeValues.Num() - 1;
 		}
 	}
@@ -2186,31 +2156,18 @@ int32 FCustomizableObjectInstanceDescriptor::RemoveValueFromTextureRange(const F
 }
 
 
-int32 FCustomizableObjectInstanceDescriptor::RemoveValueFromProjectorRange(const FString& ParamName)
-{
-	const int32 ProjectorParameterIndex = FindTypedParameterIndex(ParamName, EMutableParameterType::Projector);
-	if (ProjectorParameterIndex != INDEX_NONE)
-	{
-		FCustomizableObjectProjectorParameterValue& ProjectorParameter = ProjectorParameters[ProjectorParameterIndex];
-		if (ProjectorParameter.RangeValues.Num() > 0)
-		{
-			ProjectorParameter.RangeValues.Pop();
-			return ProjectorParameter.RangeValues.Num() - 1;
-		}
-	}
-	return -1;
-}
-
-
 int32 FCustomizableObjectInstanceDescriptor::RemoveValueFromProjectorRange(const FString& ParamName, const int32 RangeIndex)
 {
 	const int32 ProjectorParameterIndex = FindTypedParameterIndex(ParamName, EMutableParameterType::Projector);
 	if (ProjectorParameterIndex != INDEX_NONE)
 	{
 		FCustomizableObjectProjectorParameterValue& ProjectorParameter = ProjectorParameters[ProjectorParameterIndex];
+
 		if (ProjectorParameter.RangeValues.Num() > 0)
 		{
-			ProjectorParameter.RangeValues.RemoveAt(RangeIndex);
+			int32 IndexToRemove = ProjectorParameter.RangeValues.IsValidIndex(RangeIndex) ? RangeIndex : ProjectorParameter.RangeValues.Num() - 1;
+			ProjectorParameter.RangeValues.RemoveAt(IndexToRemove);
+			
 			return ProjectorParameter.RangeValues.Num() - 1;
 		}
 	}
diff --git a/Engine/Plugins/Experimental/Mutable/Source/CustomizableObject/Public/MuCO/CustomizableObjectInstance.h b/Engine/Plugins/Experimental/Mutable/Source/CustomizableObject/Public/MuCO/CustomizableObjectInstance.h
--- a/Engine/Plugins/Experimental/Mutable/Source/CustomizableObject/Public/MuCO/CustomizableObjectInstance.h
+++ b/Engine/Plugins/Experimental/Mutable/Source/CustomizableObject/Public/MuCO/CustomizableObjectInstance.h
@@ -511,26 +511,20 @@ public:
 	UFUNCTION(BlueprintCallable, Category = CustomizableObjectInstance)
 	int32 AddValueToProjectorRange(const FString& ParamName);
 
-	/** Remove the last of the integer range of values from the parameter ParamName, returns the index of the last valid integer, -1 if no values left. */
+	/** Remove the RangeIndex element of the integer range of values from the parameter ParamName. If RangeValue is -1 removes the last of the integer range of values.
+		Returns the index of the last valid integer, -1 if no values left. */
 	UFUNCTION(BlueprintCallable, Category = CustomizableObjectInstance)
-	int32 RemoveValueFromIntRange(const FString& ParamName);
+	int32 RemoveValueFromIntRange(const FString& ParamName, int32 RangeIndex = -1);
 
-	/** Remove the RangeIndex element of the integer range of values from the parameter ParamName, returns the index of the last valid integer, -1 if no values left. */
-	int32 RemoveValueFromIntRange(const FString& ParamName, int32 RangeIndex);
-
-	/** Remove the last of the float range of values from the parameter ParamName, returns the index of the last valid float, -1 if no values left. */
+	/** Remove the RangeIndex element of the float range of values from the parameter ParamName. If RangeValue is -1 removes the last of the float range of values.
+		Returns the index of the last valid float, -1 if no values left. */
 	UFUNCTION(BlueprintCallable, Category = CustomizableObjectInstance)
-	int32 RemoveValueFromFloatRange(const FString& ParamName);
-
-	/** Remove the RangeIndex element of the float range of values from the parameter ParamName, returns the index of the last valid float, -1 if no values left. */
-	int32 RemoveValueFromFloatRange(const FString& ParamName, int32 RangeIndex);
+	int32 RemoveValueFromFloatRange(const FString& ParamName, int32 RangeIndex = -1);
 
-	/** Remove the last of the projector range of values from the parameter ParamName, returns the index of the last valid projector, -1 if no values left. */
+	/** Remove the RangeIndex element of the projector range of values from the parameter ParamName. If RangeValue is -1 removes the last of the projector range of values.
+		Returns the index of the last valid projector, -1 if no values left. */
 	UFUNCTION(BlueprintCallable, Category = CustomizableObjectInstance)
-	int32 RemoveValueFromProjectorRange(const FString& ParamName);
-
-	/** Remove the RangeIndex element of the projector range of values from the parameter ParamName, returns the index of the last valid projector, -1 if no values left. */
-	int32 RemoveValueFromProjectorRange(const FString& ParamName, int32 RangeIndex);
+	int32 RemoveValueFromProjectorRange(const FString& ParamName, int32 RangeIndex = -1);
 
 	// ------------------------------------------------------------
 	// Multilayer Projectors
