[mutable] Fix ImageStorage batch iteration for lod ranges in tail.

#rnx

diff --git a/Engine/Plugins/Experimental/Mutable/Source/MutableRuntime/Public/MuR/ImageDataStorage.h b/Engine/Plugins/Experimental/Mutable/Source/MutableRuntime/Public/MuR/ImageDataStorage.h
--- a/Engine/Plugins/Experimental/Mutable/Source/MutableRuntime/Public/MuR/ImageDataStorage.h
+++ b/Engine/Plugins/Experimental/Mutable/Source/MutableRuntime/Public/MuR/ImageDataStorage.h
@@ -268,7 +268,7 @@ public:
 
 		if (FirstCompactedTailLOD < LODEnd)
 		{
-			TArrayView<const uint8> FirstLODInRangeView = GetLOD(FirstCompactedTailLOD);
+			TArrayView<const uint8> FirstLODInRangeView = GetLOD(FMath::Max(FirstCompactedTailLOD, LODBegin));
 			TArrayView<const uint8> LastLODInRangeView = GetLOD(LODEnd - 1);
 
 			const int32 TailInLODRangeNumBytes = LastLODInRangeView.GetData() - FirstLODInRangeView.GetData() + LastLODInRangeView.Num();
@@ -277,7 +277,6 @@ public:
 		return NumBatches;
 	}
 
-	
 	/**
 	 * Returns a non modifiable view to the portion of the [LODBegin, LODEnd) for the BatchId. 
 	 *
@@ -319,7 +318,7 @@ public:
 
 		if (FirstCompactedTailLOD < LODEnd)
 		{
-			TArrayView<const uint8> FirstLODInRangeView = GetLOD(FirstCompactedTailLOD);
+			TArrayView<const uint8> FirstLODInRangeView = GetLOD(FMath::Max(FirstCompactedTailLOD, LODBegin));
 			TArrayView<const uint8> LastLODInRangeView = GetLOD(LODEnd - 1);
 			
 			const int32 TailInLODRangeNumBytes = LastLODInRangeView.GetData() - FirstLODInRangeView.GetData() + LastLODInRangeView.Num();
@@ -341,7 +340,7 @@ public:
 		// If the BatchId is not found, return an empty view.
 		return TArrayView<const uint8>();
 	}
-
+	
 	/**
 	 * Non const version of GetBatchLODRange().
 	 */
