[Mutable] Rename editor viewport tabs and some node properties.
- Remove the unused mask field in the group projector node.

#rnx

diff --git a/Engine/Plugins/Experimental/Mutable/Source/CustomizableObject/Internal/MuCO/CustomizableObjectInstancePrivate.h b/Engine/Plugins/Experimental/Mutable/Source/CustomizableObject/Internal/MuCO/CustomizableObjectInstancePrivate.h
--- a/Engine/Plugins/Experimental/Mutable/Source/CustomizableObject/Internal/MuCO/CustomizableObjectInstancePrivate.h
+++ b/Engine/Plugins/Experimental/Mutable/Source/CustomizableObject/Internal/MuCO/CustomizableObjectInstancePrivate.h
@@ -174,13 +174,13 @@ public:
 	GENERATED_BODY()
 
 	/** The generated skeletal meshes for this Instance, one for each component */
-	UPROPERTY(Transient, VisibleAnywhere, Category = NoCategory)
+	UPROPERTY(Transient, VisibleAnywhere, Category = NoCategory, meta=(DisplayName = "Meshes"))
 	TArray<TObjectPtr<USkeletalMesh>> SkeletalMeshes;
 
-	UPROPERTY(Transient, VisibleAnywhere, Category = NoCategory)
+	UPROPERTY(Transient, VisibleAnywhere, Category = NoCategory, meta=(DisplayName = "Materials"))
 	TArray<FGeneratedMaterial> GeneratedMaterials;
 
-	UPROPERTY( Transient )
+	UPROPERTY(Transient)
 	TArray<FGeneratedTexture> GeneratedTextures;
 
 	// Indices of the parameters that are relevant for the given parameter values.
diff --git a/Engine/Plugins/Experimental/Mutable/Source/CustomizableObjectEditor/Private/MuCOE/CustomizableInstanceDetails.cpp b/Engine/Plugins/Experimental/Mutable/Source/CustomizableObjectEditor/Private/MuCOE/CustomizableInstanceDetails.cpp
--- a/Engine/Plugins/Experimental/Mutable/Source/CustomizableObjectEditor/Private/MuCOE/CustomizableInstanceDetails.cpp
+++ b/Engine/Plugins/Experimental/Mutable/Source/CustomizableObjectEditor/Private/MuCOE/CustomizableInstanceDetails.cpp
@@ -107,6 +107,7 @@ void FCustomizableInstanceDetails::CustomizeDetails(const TSharedPtr<IDetailLayo
 	FAddPropertyParams PrivatePropertyParams;
 	PrivatePropertyParams.HideRootObjectNode(true);
 
+	ResourcesCategory.InitiallyCollapsed(true);
 	IDetailPropertyRow* PrivateDataRow = ResourcesCategory.AddExternalObjects(Private, EPropertyLocation::Default, PrivatePropertyParams);
 
 	// Store which Texture Parameter values can be selected.
diff --git a/Engine/Plugins/Experimental/Mutable/Source/CustomizableObjectEditor/Private/MuCOE/CustomizableObjectEditor.cpp b/Engine/Plugins/Experimental/Mutable/Source/CustomizableObjectEditor/Private/MuCOE/CustomizableObjectEditor.cpp
--- a/Engine/Plugins/Experimental/Mutable/Source/CustomizableObjectEditor/Private/MuCOE/CustomizableObjectEditor.cpp
+++ b/Engine/Plugins/Experimental/Mutable/Source/CustomizableObjectEditor/Private/MuCOE/CustomizableObjectEditor.cpp
@@ -417,7 +417,7 @@ TSharedRef<SDockTab> FCustomizableObjectEditor::SpawnTab_Viewport( const FSpawnT
 	check( Args.GetTabId() == ViewportTabId );
 
 	TSharedRef<SDockTab> DockTab = SNew(SDockTab)
-		.Label( FText::FromString( GetTabPrefix() + LOCTEXT("CustomizableObjectViewport_TabTitle", "Preview Inst. Viewport").ToString() ) )
+		.Label( FText::FromString( GetTabPrefix() + LOCTEXT("CustomizableObjectViewport_TabTitle", "Viewport").ToString() ) )
 		[
 			SNew( SVerticalBox )
 
@@ -471,7 +471,7 @@ TSharedRef<SDockTab> FCustomizableObjectEditor::SpawnTab_InstanceProperties( con
 	ScrollBox->SetScrollBarRightClickDragAllowed(true);
 
 	TSharedRef<SDockTab> DockTab = SNew(SDockTab)
-		.Label( FText::FromString( GetTabPrefix() + LOCTEXT( "CustomizableInstanceProperties_TabTitle", "Preview Inst. Properties" ).ToString() ) )
+		.Label( FText::FromString( GetTabPrefix() + LOCTEXT( "CustomizableInstanceProperties_TabTitle", "Preview Instance" ).ToString() ) )
 		[
 			ScrollBox
 		];
@@ -631,7 +631,7 @@ TSharedRef<SDockTab> FCustomizableObjectEditor::SpawnTab_AdvancedPreviewSettings
 {
 	check(Args.GetTabId() == AdvancedPreviewSettingsTabId);
 	TSharedRef<SDockTab> DockTab = SNew(SDockTab)
-		.Label(LOCTEXT("StaticMeshPreviewScene_TabTitle", "Preview Settings"))
+		.Label(LOCTEXT("StaticMeshPreviewScene_TabTitle", "Preview Scene Settings"))
 		[
 			AdvancedPreviewSettingsWidget.ToSharedRef()
 		];
diff --git a/Engine/Plugins/Experimental/Mutable/Source/CustomizableObjectEditor/Private/MuCOE/CustomizableObjectNodeObjectGroupDetails.cpp b/Engine/Plugins/Experimental/Mutable/Source/CustomizableObjectEditor/Private/MuCOE/CustomizableObjectNodeObjectGroupDetails.cpp
--- a/Engine/Plugins/Experimental/Mutable/Source/CustomizableObjectEditor/Private/MuCOE/CustomizableObjectNodeObjectGroupDetails.cpp
+++ b/Engine/Plugins/Experimental/Mutable/Source/CustomizableObjectEditor/Private/MuCOE/CustomizableObjectNodeObjectGroupDetails.cpp
@@ -37,11 +37,43 @@ void FCustomizableObjectNodeObjectGroupDetails::CustomizeDetails(IDetailLayoutBu
 
 	if (NodeGroup)
 	{
+		DetailBuilder.HideProperty("DefaultValue");
+
+		IDetailCategoryBuilder& GroupInfoCategory = DetailBuilder.EditCategory("GroupInfo");
+
+		// Forcing property order
+		GroupInfoCategory.AddProperty("GroupName");
+		GroupInfoCategory.AddProperty("GroupType");
+
+		// Getting group node children names
+		GenerateChildrenObjectNames();
+
+		GroupInfoCategory.AddCustomRow(LOCTEXT("NodeObjectGroupDetails_ComboBox", "Default Value Selector"))
+			.Visibility(TAttribute<EVisibility>(this, &FCustomizableObjectNodeObjectGroupDetails::DefaultValueSelectorVisibility))
+			.NameContent()
+			[
+				SNew(STextBlock)
+				.Text(LOCTEXT("NodeObjectGroupDetails_ComboBox_Text", "Default Value"))
+			.ToolTipText(LOCTEXT("NodeObjectGroupDetails_ComboBox_Tooltip", "Select the default value of the group."))
+			.Font(DetailBuilder.GetDetailFont())
+			]
+		.ValueContent()
+			.HAlign(EHorizontalAlignment::HAlign_Left)
+			[
+				SAssignNew(DefaultValueSelector, STextComboBox)
+				.InitiallySelectedItem(InitialNameOption)
+			.OptionsSource(&ChildrenNameOptions)
+			.OnComboBoxOpening(this, &FCustomizableObjectNodeObjectGroupDetails::GenerateChildrenObjectNames)
+			.OnSelectionChanged(this, &FCustomizableObjectNodeObjectGroupDetails::OnSetDefaultValue)
+			.Font(DetailBuilder.GetDetailFont())
+			.ToolTipText(this, &FCustomizableObjectNodeObjectGroupDetails::DefaultValueComboBoxTooltip)
+			];
+
 		if (TSharedPtr<FCustomizableObjectEditor> GraphEditor = StaticCastSharedPtr<FCustomizableObjectEditor>(NodeGroup->GetGraphEditor()))
 		{
 			if (UCustomizableObject* NodeGroupCO = CastChecked<UCustomizableObject>(NodeGroup->GetCustomizableObjectGraph()->GetOuter()))
 			{
-				IDetailCategoryBuilder& BlocksCategory = DetailBuilder.EditCategory("Group Info");
+				IDetailCategoryBuilder& BlocksCategory = DetailBuilder.EditCategory("External Objects");
 
 				TMultiMap<FGuid, UCustomizableObjectNodeObject*> ObjectNodesObjects = GetNodeGroupObjectNodeMapping(NodeGroupCO);
 				TArray<UCustomizableObjectNodeObject*> ChildNodes;
@@ -67,38 +99,6 @@ void FCustomizableObjectNodeObjectGroupDetails::CustomizeDetails(IDetailLayoutBu
 				}
 			}
 		}
-
-		DetailBuilder.HideProperty("DefaultValue");
-
-		IDetailCategoryBuilder& CustomizableObjectCategory = DetailBuilder.EditCategory("CustomizableObject");
-
-		// Forcing property order
-		CustomizableObjectCategory.AddProperty("GroupName");
-		CustomizableObjectCategory.AddProperty("GroupType");
-
-		// Getting group node children names
-		GenerateChildrenObjectNames();
-
-		CustomizableObjectCategory.AddCustomRow(LOCTEXT("NodeObjectGroupDetails_ComboBox", "Default Value Selector"))
-		.Visibility(TAttribute<EVisibility>(this, &FCustomizableObjectNodeObjectGroupDetails::DefaultValueSelectorVisibility))
-		.NameContent()
-		[
-			SNew(STextBlock)
-			.Text(LOCTEXT("NodeObjectGroupDetails_ComboBox_Text", "Default Value"))
-			.ToolTipText(LOCTEXT("NodeObjectGroupDetails_ComboBox_Tooltip", "Select the default value of the group."))
-			.Font(DetailBuilder.GetDetailFont())
-		]
-		.ValueContent()
-		.HAlign(EHorizontalAlignment::HAlign_Left)
-		[
-			SAssignNew(DefaultValueSelector, STextComboBox)
-			.InitiallySelectedItem(InitialNameOption)
-			.OptionsSource(&ChildrenNameOptions)
-			.OnComboBoxOpening(this, &FCustomizableObjectNodeObjectGroupDetails::GenerateChildrenObjectNames)
-			.OnSelectionChanged(this, &FCustomizableObjectNodeObjectGroupDetails::OnSetDefaultValue)
-			.Font(DetailBuilder.GetDetailFont())
-			.ToolTipText(this, &FCustomizableObjectNodeObjectGroupDetails::DefaultValueComboBoxTooltip)
-		];
 	}
 }
 
diff --git a/Engine/Plugins/Experimental/Mutable/Source/CustomizableObjectEditor/Private/MuCOE/PinViewer/SPinViewer.cpp b/Engine/Plugins/Experimental/Mutable/Source/CustomizableObjectEditor/Private/MuCOE/PinViewer/SPinViewer.cpp
--- a/Engine/Plugins/Experimental/Mutable/Source/CustomizableObjectEditor/Private/MuCOE/PinViewer/SPinViewer.cpp
+++ b/Engine/Plugins/Experimental/Mutable/Source/CustomizableObjectEditor/Private/MuCOE/PinViewer/SPinViewer.cpp
@@ -338,7 +338,7 @@ void PinViewerAttachToDetailCustomization(IDetailLayoutBuilder& DetailBuilder)
 	{
 		UCustomizableObjectNode* Node = Cast<UCustomizableObjectNode>(SelectedObjects[0].Get());
 
-		IDetailCategoryBuilder& PinViewerCategoryBuilder = DetailBuilder.EditCategory("PinViewer", FText::GetEmpty(), ECategoryPriority::Uncommon);
+		IDetailCategoryBuilder& PinViewerCategoryBuilder = DetailBuilder.EditCategory("PinViewer", LOCTEXT("PinViewer", "Pins"), ECategoryPriority::Uncommon);
 		PinViewerCategoryBuilder.AddCustomRow(LOCTEXT("PinViewerDetailsCategory", "PinViwer")).ShouldAutoExpand(true)
 		[
 			SNew(SPinViewer).Node(Node)
diff --git a/Engine/Plugins/Experimental/Mutable/Source/CustomizableObjectEditor/Public/MuCOE/Nodes/CustomizableObjectNodeGroupProjectorParameter.h b/Engine/Plugins/Experimental/Mutable/Source/CustomizableObjectEditor/Public/MuCOE/Nodes/CustomizableObjectNodeGroupProjectorParameter.h
--- a/Engine/Plugins/Experimental/Mutable/Source/CustomizableObjectEditor/Public/MuCOE/Nodes/CustomizableObjectNodeGroupProjectorParameter.h
+++ b/Engine/Plugins/Experimental/Mutable/Source/CustomizableObjectEditor/Public/MuCOE/Nodes/CustomizableObjectNodeGroupProjectorParameter.h
@@ -55,9 +55,6 @@ public:
 	UPROPERTY(EditAnywhere, Category = ProjectorGroup, Meta = (ToolTip = "Only used when connected to a Group node. Specifies which material channel in the Group node's child material nodes will be connected to the projection."))
 	FString MaterialChannelNameToConnect;
 
-	UPROPERTY(EditAnywhere, Category = ProjectorGroup, Meta = (ToolTip = "Only used when connected to a Group node. Specifies which material channel will be used to mask out the projection."))
-	FString MaskedOutAreaMaterialChannelName;
-
 	/** Reference Texture used to decide the texture properties of the mutable-generated textures
 	* connected to this material. If null, it will try to be guessed at compile time from
 	* the graph. */
diff --git a/Engine/Plugins/Experimental/Mutable/Source/CustomizableObjectEditor/Public/MuCOE/Nodes/CustomizableObjectNodeObjectGroup.h b/Engine/Plugins/Experimental/Mutable/Source/CustomizableObjectEditor/Public/MuCOE/Nodes/CustomizableObjectNodeObjectGroup.h
--- a/Engine/Plugins/Experimental/Mutable/Source/CustomizableObjectEditor/Public/MuCOE/Nodes/CustomizableObjectNodeObjectGroup.h
+++ b/Engine/Plugins/Experimental/Mutable/Source/CustomizableObjectEditor/Public/MuCOE/Nodes/CustomizableObjectNodeObjectGroup.h
@@ -24,13 +24,13 @@ public:
 
 	UCustomizableObjectNodeObjectGroup();
 
-	UPROPERTY(EditAnywhere, Category=CustomizableObject)
+	UPROPERTY(EditAnywhere, Category=GroupInfo, meta = (DisplayName= "Name"))
 	FString GroupName;
 
-	UPROPERTY(EditAnywhere, Category=CustomizableObject)
+	UPROPERTY(EditAnywhere, Category=GroupInfo, meta = (DisplayName = "Type"))
 	ECustomizableObjectGroupType GroupType;
 	
-	UPROPERTY(EditAnywhere, Category = CustomizableObject)
+	UPROPERTY(EditAnywhere, Category=GroupInfo)
 	FString DefaultValue;
 
 	UPROPERTY(EditAnywhere, Category = UI, meta = (DisplayName = "Parameter UI Metadata"))
