[mutable] Add concurrency to a data generation step of the compilation.
- fix some comment typos
- comment out some broken debug code.
#rnx

diff --git a/Engine/Plugins/Experimental/Mutable/Source/MutableRuntime/Internal/MuR/ModelPrivate.h b/Engine/Plugins/Experimental/Mutable/Source/MutableRuntime/Internal/MuR/ModelPrivate.h
--- a/Engine/Plugins/Experimental/Mutable/Source/MutableRuntime/Internal/MuR/ModelPrivate.h
+++ b/Engine/Plugins/Experimental/Mutable/Source/MutableRuntime/Internal/MuR/ModelPrivate.h
@@ -35,10 +35,10 @@ namespace mu
 		//! Size of the rom
 		uint32 Size;
 
-		//! Index of the resource in its type-sepcific array
+		//! Index of the resource in its type-specific array
 		uint32 ResourceIndex;
 
-		//! Index of the resource in its type-sepcific array (one of DATATYPE values)
+		//! Index of the resource in its type-specific array (one of DATATYPE values)
 		uint32 ResourceType;        
     };
 
diff --git a/Engine/Plugins/Experimental/Mutable/Source/MutableRuntime/Private/MuR/OpImagePixelFormat.cpp b/Engine/Plugins/Experimental/Mutable/Source/MutableRuntime/Private/MuR/OpImagePixelFormat.cpp
--- a/Engine/Plugins/Experimental/Mutable/Source/MutableRuntime/Private/MuR/OpImagePixelFormat.cpp
+++ b/Engine/Plugins/Experimental/Mutable/Source/MutableRuntime/Private/MuR/OpImagePixelFormat.cpp
@@ -384,14 +384,14 @@ void FImageOperator::ImagePixelFormat(bool& bOutSuccess, int32 CompressionQualit
 				
 				CompressRLE_L(Base, Result);
 
-#ifdef MUTABLE_DEBUG_RLE		
-				if (bOutSuccess)
-				{
-					// Verify
-					Ptr<Image> Test = ImagePixelFormat(CompressionQuality, Result, EImageFormat::IF_L_UBYTE);
-					check(Test->DataStorage == Base->DataStorage);
-				}
-#endif
+//#ifdef MUTABLE_DEBUG_RLE		
+//				if (bOutSuccess)
+//				{
+//					// Verify
+//					Ptr<Image> Test = ImagePixelFormat(CompressionQuality, Result, EImageFormat::IF_L_UBYTE);
+//					check(Test->DataStorage == Base->DataStorage);
+//				}
+//#endif
 				break;
 			}
 
@@ -418,14 +418,14 @@ void FImageOperator::ImagePixelFormat(bool& bOutSuccess, int32 CompressionQualit
 
 				CompressRLE_L(TempBase.get(), Result);
 
-#ifdef MUTABLE_DEBUG_RLE
-				if (bOutSuccess)
-				{
-					// Verify
-					Ptr<Image> Test = ImagePixelFormat(CompressionQuality, Result, EImageFormat::IF_L_UBYTE);
-					check(Test->DataStorage == TempBase->DataStorage);
-				}
-#endif
+//#ifdef MUTABLE_DEBUG_RLE
+//				if (bOutSuccess)
+//				{
+//					// Verify
+//					Ptr<Image> Test = ImagePixelFormat(CompressionQuality, Result, EImageFormat::IF_L_UBYTE);
+//					check(Test->DataStorage == TempBase->DataStorage);
+//				}
+//#endif
 				ReleaseImage(TempBase);
 
 				break;
@@ -465,14 +465,14 @@ void FImageOperator::ImagePixelFormat(bool& bOutSuccess, int32 CompressionQualit
 				
 				CompressRLE_L1(Base, Result);
 				
-#ifdef MUTABLE_DEBUG_RLE		
-				if (bOutSuccess)
-				{
-					// Verify
-					Ptr<Image> Test = ImagePixelFormat(CompressionQuality, Result, EImageFormat::IF_L_UBYTE);
-					check(Test->DataStorage == Base->DataStorage);
-				}
-#endif
+//#ifdef MUTABLE_DEBUG_RLE		
+//				if (bOutSuccess)
+//				{
+//					// Verify
+//					Ptr<Image> Test = ImagePixelFormat(CompressionQuality, Result, EImageFormat::IF_L_UBYTE);
+//					check(Test->DataStorage == Base->DataStorage);
+//				}
+//#endif
 
 				break;
 			}
@@ -485,14 +485,14 @@ void FImageOperator::ImagePixelFormat(bool& bOutSuccess, int32 CompressionQualit
 				Ptr<Image> TempBase = ImagePixelFormat(CompressionQuality, Base, EImageFormat::IF_L_UBYTE, -1);
 				CompressRLE_L1(TempBase.get(), Result);
 
-#ifdef MUTABLE_DEBUG_RLE
-				if (bOutSuccess)
-				{
-					// Verify
-					Ptr<Image> Test = ImagePixelFormat(CompressionQuality, Result, EImageFormat::IF_L_UBYTE);
-					check(Test->DataStorage == TempBase->DataStorage);
-				}
-#endif
+//#ifdef MUTABLE_DEBUG_RLE
+//				if (bOutSuccess)
+//				{
+//					// Verify
+//					Ptr<Image> Test = ImagePixelFormat(CompressionQuality, Result, EImageFormat::IF_L_UBYTE);
+//					check(Test->DataStorage == TempBase->DataStorage);
+//				}
+//#endif
 				ReleaseImage(TempBase);
 				break;
 			}
diff --git a/Engine/Plugins/Experimental/Mutable/Source/MutableTools/Private/MuT/Compiler.cpp b/Engine/Plugins/Experimental/Mutable/Source/MutableTools/Private/MuT/Compiler.cpp
--- a/Engine/Plugins/Experimental/Mutable/Source/MutableTools/Private/MuT/Compiler.cpp
+++ b/Engine/Plugins/Experimental/Mutable/Source/MutableTools/Private/MuT/Compiler.cpp
@@ -3,14 +3,6 @@
 
 #include "MuT/Compiler.h"
 
-#include "Containers/Array.h"
-#include "Containers/Map.h"
-#include "HAL/LowLevelMemTracker.h"
-#include "HAL/PlatformCrt.h"
-#include "Hash/CityHash.h"
-#include "Logging/LogCategory.h"
-#include "Logging/LogMacros.h"
-#include "Misc/AssertionMacros.h"
 #include "MuR/Image.h"
 #include "MuR/Mesh.h"
 #include "MuR/Model.h"
@@ -32,7 +24,17 @@
 #include "MuT/Node.h"
 #include "MuT/NodePrivate.h"
 #include "MuT/Table.h"
+
 #include "Trace/Detail/Channel.h"
+#include "Containers/Array.h"
+#include "Containers/Map.h"
+#include "HAL/LowLevelMemTracker.h"
+#include "HAL/PlatformCrt.h"
+#include "Hash/CityHash.h"
+#include "Logging/LogCategory.h"
+#include "Logging/LogMacros.h"
+#include "Misc/AssertionMacros.h"
+#include "Async/ParallelFor.h"
 
 #include <inttypes.h>	// Required for 64-bit printf macros
 
@@ -468,21 +470,14 @@ namespace mu
 	//---------------------------------------------------------------------------------------------
 	namespace
 	{
-		inline void EnsureUniqueRomId(uint32& RomId, const mu::FProgram& program)
+		inline void EnsureUniqueRomId(TSet<uint32>& UsedIds, uint32& RomId)
 		{
 			while (true)
 			{
-				bool bUnique = true;
-				for (const FRomData& Rom : program.m_roms)
-				{
-					if (Rom.Id == RomId)
-					{
-						bUnique = false;
-						break;
-					}
-				}
+				bool bAlreadyUsed = false;
+				UsedIds.FindOrAdd( RomId, &bAlreadyUsed);
 
-				if (bUnique)
+				if (!bAlreadyUsed)
 				{
 					break;
 				}
@@ -499,81 +494,135 @@ namespace mu
 	void Compiler::Private::GenerateRoms(Model* p, int32 EmbeddedDataBytesLimit)
 	{
 		LLM_SCOPE_BYNAME(TEXT("MutableRuntime"));
-		MUTABLE_CPUPROFILER_SCOPE(Mutable_GenerateRoms);
+		MUTABLE_CPUPROFILER_SCOPE(GenerateRoms);
 
 		mu::FProgram& program = p->GetPrivate()->m_program;
 
 		int32 NumRoms = 0;
 		int32 NumEmbedded = 0;
 
+		TSet<uint32> UsedIds;
+		UsedIds.Reserve(program.ConstantImageLODs.Num()+ program.ConstantMeshes.Num());
+
+		TArray<FRomData> RomDatas;
+		RomDatas.SetNumZeroed(FMath::Max(program.ConstantImageLODs.Num(), program.ConstantMeshes.Num()));
+
 		// Save images and unload from memory
-		for (int32 ResourceIndex = 0; ResourceIndex < program.ConstantImageLODs.Num(); ++ResourceIndex)
 		{
-			TPair<int32, mu::ImagePtrConst>& ResData = program.ConstantImageLODs[ResourceIndex];
+			MUTABLE_CPUPROFILER_SCOPE(GenerateRoms_ImageIds);
 
-			// This shouldn't have been serialised with rom support before.
-			check(ResData.Key < 0);
+			ParallelFor(program.ConstantImageLODs.Num(),
+				[EmbeddedDataBytesLimit, &program, &RomDatas](uint32 ResourceIndex)
+				{
+					TPair<int32, Ptr<const Image>>& ResData = program.ConstantImageLODs[ResourceIndex];
+
+					// This shouldn't have been serialised with rom support before.
+					check(ResData.Key < 0);
 
-			// Serialize to memory, to find out final size of this rom
-			OutputMemoryStream MemStream(1024 * 1024);
-			OutputArchive MemoryArch(&MemStream);
-			Image::Serialise(ResData.Value.get(), MemoryArch);
+					// Serialize to memory, to find out final size of this rom
+					OutputMemoryStream MemStream(1024 * 1024);
+					OutputArchive MemoryArch(&MemStream);
+					Image::Serialise(ResData.Value.get(), MemoryArch);
 
-			// If the resource uses less memory than the threshold, don't save it in a separate rom.
-			if (MemStream.GetBufferSize() <= EmbeddedDataBytesLimit)
+					// If the resource uses less memory than the threshold, don't save it in a separate rom.
+					if (MemStream.GetBufferSize() <= EmbeddedDataBytesLimit)
+					{
+						return;
+					}
+
+					FRomData& RomData = RomDatas[ResourceIndex];
+					RomData.ResourceType = DT_IMAGE;
+					RomData.ResourceIndex = ResourceIndex;
+					RomData.Size = MemStream.GetBufferSize();
+
+					uint64 FullHash = CityHash64(static_cast<const char*>(MemStream.GetBuffer()), MemStream.GetBufferSize());
+					RomData.Id = GetTypeHash(FullHash);
+				});
+		}
+
+		{
+			MUTABLE_CPUPROFILER_SCOPE(GenerateRoms_ImageIdsUnique);
+
+			for (int32 ResourceIndex = 0; ResourceIndex < program.ConstantImageLODs.Num(); ++ResourceIndex)
 			{
-				NumEmbedded++;
-				continue;
-			}
+				FRomData& RomData = RomDatas[ResourceIndex];
 
-			NumRoms++;
+				// If the resource uses less memory than the threshold, don't save it in a separate rom.
+				if (int32(RomData.Size) <= EmbeddedDataBytesLimit)
+				{
+					NumEmbedded++;
+					continue;
+				}
+				NumRoms++;
 
-			FRomData RomData;
-			RomData.ResourceType = DT_IMAGE;
-			RomData.ResourceIndex = ResourceIndex;
-			RomData.Size = MemStream.GetBufferSize();
+				// Ensure that the Id is unique
+				EnsureUniqueRomId(UsedIds, RomData.Id);
 
-			// Ensure that the Id is unique
-			RomData.Id = CityHash32(static_cast<const char*>(MemStream.GetBuffer()), MemStream.GetBufferSize());
-			EnsureUniqueRomId(RomData.Id, program);
+				int32 RomIndex = program.m_roms.Add(RomData);
 
-			int32 RomIndex = program.m_roms.Add(RomData);
-			ResData.Key = RomIndex;
+				TPair<int32, Ptr<const Image>>& ResData = program.ConstantImageLODs[ResourceIndex];
+				ResData.Key = RomIndex;
+			}
 		}
 
 		// Save meshes and unload from memory
-		for (int32 ResourceIndex = 0; ResourceIndex < program.ConstantMeshes.Num(); ++ResourceIndex)
+		FMemory::Memzero( RomDatas.GetData(), RomDatas.GetAllocatedSize() );
 		{
-			TPair<int32, mu::MeshPtrConst>& ResData = program.ConstantMeshes[ResourceIndex];
+			MUTABLE_CPUPROFILER_SCOPE(GenerateRoms_MeshIds);
+
+			ParallelFor(program.ConstantMeshes.Num(),
+				[EmbeddedDataBytesLimit, &program, &RomDatas](uint32 ResourceIndex)
+				{
+					TPair<int32, Ptr<const Mesh>>& ResData = program.ConstantMeshes[ResourceIndex];
 
-			// This shouldn't have been serialised with rom support before.
-			check(ResData.Key < 0);
+					// This shouldn't have been serialised with rom support before.
+					check(ResData.Key < 0);
 
-			// Serialize to memory, to find out final size of this rom
-			OutputMemoryStream MemStream(1024 * 1024);
-			OutputArchive MemoryArch(&MemStream);
-			Mesh::Serialise(ResData.Value.get(), MemoryArch);
+					// Serialize to memory, to find out final size of this rom
+					OutputMemoryStream MemStream(1024 * 1024);
+					OutputArchive MemoryArch(&MemStream);
+					Mesh::Serialise(ResData.Value.get(), MemoryArch);
 
-			// If the resource uses less memory than the threshold, don't save it in a separate rom.
-			if (MemStream.GetBufferSize() <= EmbeddedDataBytesLimit)
+					// If the resource uses less memory than the threshold, don't save it in a separate rom.
+					if (MemStream.GetBufferSize() <= EmbeddedDataBytesLimit)
+					{
+						return;
+					}
+
+					FRomData& RomData = RomDatas[ResourceIndex];
+					RomData.ResourceType = DT_MESH;
+					RomData.ResourceIndex = ResourceIndex;
+					RomData.Size = MemStream.GetBufferSize();
+
+					// Ensure that the Id is unique
+					uint64 FullHash = CityHash64(static_cast<const char*>(MemStream.GetBuffer()), MemStream.GetBufferSize());
+					RomData.Id = GetTypeHash(FullHash);
+				});
+		}
+
+		{
+			MUTABLE_CPUPROFILER_SCOPE(GenerateRoms_MeshIdsUnique);
+
+			for (int32 ResourceIndex = 0; ResourceIndex < program.ConstantMeshes.Num(); ++ResourceIndex)
 			{
-				NumEmbedded++;
-				continue;
-			}
+				FRomData& RomData = RomDatas[ResourceIndex];
 
-			NumRoms++;
+				// If the resource uses less memory than the threshold, don't save it in a separate rom.
+				if (int32(RomData.Size) <= EmbeddedDataBytesLimit)
+				{
+					NumEmbedded++;
+					continue;
+				}
+				NumRoms++;
 
-			FRomData RomData;
-			RomData.ResourceType = DT_MESH;
-			RomData.ResourceIndex = ResourceIndex;
-			RomData.Size = MemStream.GetBufferSize();
+				// Ensure that the Id is unique
+				EnsureUniqueRomId(UsedIds, RomData.Id);
 
-			// Ensure that the Id is unique
-			RomData.Id = CityHash32(static_cast<const char*>(MemStream.GetBuffer()), MemStream.GetBufferSize());
-			EnsureUniqueRomId(RomData.Id, program);
+				int32 RomIndex = program.m_roms.Add(RomData);
 
-			int32 RomIndex = program.m_roms.Add(RomData);
-			ResData.Key = RomIndex;
+				TPair<int32, Ptr<const Mesh>>& ResData = program.ConstantMeshes[ResourceIndex];
+				ResData.Key = RomIndex;
+			}
 		}
 
 		UE_LOG(LogMutableCore, Log, TEXT("Generated roms for model with %d embedded constants and %d in roms."), NumEmbedded, NumRoms);
