[mutable] Add CVar to force fast low-quality texture compression during cook.
#rnx

diff --git a/Engine/Plugins/Experimental/Mutable/Source/CustomizableObjectEditor/Private/MuCOE/CustomizableObjectCompileRunnable.cpp b/Engine/Plugins/Experimental/Mutable/Source/CustomizableObjectEditor/Private/MuCOE/CustomizableObjectCompileRunnable.cpp
--- a/Engine/Plugins/Experimental/Mutable/Source/CustomizableObjectEditor/Private/MuCOE/CustomizableObjectCompileRunnable.cpp
+++ b/Engine/Plugins/Experimental/Mutable/Source/CustomizableObjectEditor/Private/MuCOE/CustomizableObjectCompileRunnable.cpp
@@ -34,6 +34,12 @@ TAutoConsoleVariable<bool> CVarMutableCompilerDiskCache(
 	TEXT("Force the use of disk cache to reduce memory usage when compiling CustomizableObjects both in editor and cook commandlets."),
 	ECVF_Default);
 
+TAutoConsoleVariable<bool> CVarMutableCompilerFastCompression(
+	TEXT("mutable.ForceFastTextureCompression"),
+	false,
+	TEXT("Force the use of lower quality but faster compression during cook."),
+	ECVF_Default);
+
 
 FCustomizableObjectCompileRunnable::FCustomizableObjectCompileRunnable(mu::Ptr<mu::Node> Root)
 	: MutableRoot(Root)
@@ -136,7 +142,13 @@ uint32 FCustomizableObjectCompileRunnable::Run()
 	}
 
 	// Texture compression override, if necessary
-	if (Options.TextureCompression== ECustomizableObjectTextureCompression::HighQuality)
+	bool bUseHighQualityCompression = (Options.TextureCompression == ECustomizableObjectTextureCompression::HighQuality);
+	if (CVarMutableCompilerFastCompression->GetBool())
+	{
+		bUseHighQualityCompression = false;
+	}
+
+	if (bUseHighQualityCompression)
 	{
 		CompilerOptions->SetImagePixelFormatOverride( UnrealPixelFormatFunc );
 	}
