[Mutable] Improve NodeClipMeshWithMesh.
* Allow transformation of skeletal meshes used for clipping.
* Add a default material for clipping meshes.
* Isolate the section used for clipping in the preview viewport and override its material with Mutable's default material for clipping.
* Fix errors preventing the use of static meshes gathered through data tables.

#rnx

#virtualized

diff --git a/Engine/Plugins/Experimental/Mutable/Config/DefaultMutable.ini b/Engine/Plugins/Experimental/Mutable/Config/DefaultMutable.ini
--- a/Engine/Plugins/Experimental/Mutable/Config/DefaultMutable.ini
+++ b/Engine/Plugins/Experimental/Mutable/Config/DefaultMutable.ini
@@ -18,6 +18,10 @@
 ; use 8 bone influences per vertex instead of the default 4. This is used at object compile time.
 bExtraBoneInfluencesEnabled=False
 
+[EditorDefaults]
+ClipMorphMaterialName=/Engine/EditorMaterials/LevelGridMaterial.LevelGridMaterial
+ClipMeshMaterialName=/Mutable/Materials/ClipMeshMaterial.ClipMeshMaterial
+
 ;--------------------------------------------------------------------------------------------------
 ; These are scalability settings that can be defined in any xxxScalability.ini file.
 [ScalabilitySettings]
diff --git a/Engine/Plugins/Experimental/Mutable/Source/CustomizableObjectEditor/Private/MuCOE/CustomizableObjectEditor.cpp b/Engine/Plugins/Experimental/Mutable/Source/CustomizableObjectEditor/Private/MuCOE/CustomizableObjectEditor.cpp
--- a/Engine/Plugins/Experimental/Mutable/Source/CustomizableObjectEditor/Private/MuCOE/CustomizableObjectEditor.cpp
+++ b/Engine/Plugins/Experimental/Mutable/Source/CustomizableObjectEditor/Private/MuCOE/CustomizableObjectEditor.cpp
@@ -1070,24 +1070,45 @@ void FCustomizableObjectEditor::HideGizmoClipMorph()
 
 void FCustomizableObjectEditor::ShowGizmoClipMesh(UCustomizableObjectNodeMeshClipWithMesh& Node)
 {
-	UStaticMesh* ClipMesh = nullptr;
+	UObject* ClipMesh = nullptr;
+	int32 LODIndex = 0;
+	int32 SectionIndex = 0;
+	int32 MaterialSlotIndex = 0;
 
 	if (const UEdGraphPin* ConnectedPin = FollowInputPin(*Node.ClipMeshPin()))
 	{
 		if (const UEdGraphNode* ConnectedNode = ConnectedPin->GetOwningNode())
 		{
-			if (const UCustomizableObjectNodeStaticMesh* TypedNode = Cast<UCustomizableObjectNodeStaticMesh>(ConnectedNode))
+			int32 DummyIndex = 0;
+
+			if (const UCustomizableObjectNodeStaticMesh* StaticMeshNode = Cast<UCustomizableObjectNodeStaticMesh>(ConnectedNode))
+			{
+				ClipMesh = StaticMeshNode->GetMesh();
+				StaticMeshNode->GetPinSection(*ConnectedPin, LODIndex, SectionIndex, DummyIndex);
+				MaterialSlotIndex = SectionIndex;
+			}
+			else if (const UCustomizableObjectNodeSkeletalMesh* SkeletalMeshNode = Cast<UCustomizableObjectNodeSkeletalMesh>(ConnectedNode))
 			{
-				ClipMesh = TypedNode->StaticMesh;
+				ClipMesh = SkeletalMeshNode->GetMesh();
+				SkeletalMeshNode->GetPinSection(*ConnectedPin, LODIndex, SectionIndex, DummyIndex);
+				MaterialSlotIndex = SkeletalMeshNode->GetSkeletalMaterialIndexFor(*ConnectedPin);
 			}
 			else if (const UCustomizableObjectNodeTable* TableNode = Cast<UCustomizableObjectNodeTable>(ConnectedNode))
 			{
-				ClipMesh = TableNode->GetColumnDefaultAssetByType<UStaticMesh>(ConnectedPin);
+				ClipMesh = TableNode->GetColumnDefaultAssetByType<UObject>(ConnectedPin);
+
+				TableNode->GetPinLODAndSection(ConnectedPin, LODIndex, SectionIndex);
+				MaterialSlotIndex = SectionIndex;
+
+				if (TableNode->GetPinMeshType(ConnectedPin) == ETableMeshPinType::SKELETAL_MESH)
+				{
+					MaterialSlotIndex = TableNode->GetDefaultSkeletalMaterialIndexFor(*ConnectedPin);
+				}
 			}
 		}
 	}
 
-	if (ClipMesh)
+	if (ClipMesh && LODIndex >= 0 && MaterialSlotIndex >= 0)
 	{
 		if (GizmoType != EGizmoType::ClipMesh)
 		{
@@ -1097,7 +1118,7 @@ void FCustomizableObjectEditor::ShowGizmoClipMesh(UCustomizableObjectNodeMeshCli
 
 		SelectSingleNode(Node);
 
-		Viewport->ShowGizmoClipMesh(Node, *ClipMesh);
+		Viewport->ShowGizmoClipMesh(Node, *ClipMesh, LODIndex, SectionIndex, MaterialSlotIndex);
 	}
 }
 
@@ -1163,10 +1184,8 @@ void FCustomizableObjectEditor::PostUndo(bool bSuccess)
 {
 	if (bSuccess)
 	{
-		GraphEditor->ClearSelectionSet();
 		GraphEditor->NotifyGraphChanged();
 		CustomizableObject->MarkPackageDirty();
-		//CustomizableObjectDetailsView->Invalidate(EInvalidateWidget::Layout);
 	}
 }
 
diff --git a/Engine/Plugins/Experimental/Mutable/Source/CustomizableObjectEditor/Private/MuCOE/CustomizableObjectEditorViewportClient.cpp b/Engine/Plugins/Experimental/Mutable/Source/CustomizableObjectEditor/Private/MuCOE/CustomizableObjectEditorViewportClient.cpp
--- a/Engine/Plugins/Experimental/Mutable/Source/CustomizableObjectEditor/Private/MuCOE/CustomizableObjectEditorViewportClient.cpp
+++ b/Engine/Plugins/Experimental/Mutable/Source/CustomizableObjectEditor/Private/MuCOE/CustomizableObjectEditorViewportClient.cpp
@@ -22,9 +22,12 @@
 #include "IContentBrowserSingleton.h"
 #include "InputKeyEventArgs.h"
 #include "ObjectTools.h"
+#include "MaterialDomain.h"
 #include "Materials/MaterialInstanceDynamic.h"
+#include "Materials/MaterialInstanceConstant.h"
 #include "Misc/MessageDialog.h"
 #include "Misc/PackageName.h"
+#include "Misc/ConfigCacheIni.h"
 #include "MuCO/CustomizableObjectInstance.h"
 #include "MuCO/CustomizableObjectSystem.h"
 #include "MuCOE/CustomizableObjectPreviewScene.h"
@@ -39,7 +42,6 @@
 #include "ScopedTransaction.h"
 #include "SkeletalDebugRendering.h"
 #include "UnrealWidget.h"
-#include "Materials/MaterialInstanceConstant.h"
 #include "MuCO/CustomizableObjectMipDataProvider.h"
 #include "MuCOE/CustomizableObjectEditorLogger.h"
 #include "MuCOE/UnrealBakeHelpers.h"
@@ -112,14 +114,48 @@ FCustomizableObjectEditorViewportClient::FCustomizableObjectEditorViewportClient
 	// now add the ClipMorph plane
 	ClipMorphNode = nullptr;
 	bClipMorphLocalStartOffset = true;
-	ClipMorphMaterial = LoadObject<UMaterial>(NULL, TEXT("Material'/Engine/EditorMaterials/LevelGridMaterial.LevelGridMaterial'"), NULL, LOAD_None, NULL);
-	check(ClipMorphMaterial);
 
-	// clip mesh preview
-	ClipMeshNode = nullptr;
-	ClipMeshComp = NewObject<UStaticMeshComponent>();
-	PreviewScene->AddComponent(ClipMeshComp, FTransform());
-	ClipMeshComp->SetVisibility(false);
+	// clip mesh StaticMesh preview
+	ClipMeshStaticMeshComp = NewObject<UStaticMeshComponent>(GetTransientPackage(), NAME_None, RF_Transactional);
+	PreviewScene->AddComponent(ClipMeshStaticMeshComp, FTransform());
+	ClipMeshStaticMeshComp->SetVisibility(false);
+
+	// clip mesh SkeletalMesh preview
+	ClipMeshSkeletalMeshComp = NewObject<USkeletalMeshComponent>(GetTransientPackage(), NAME_None, RF_Transactional);
+	PreviewScene->AddComponent(ClipMeshSkeletalMeshComp, FTransform());
+	ClipMeshStaticMeshComp->SetVisibility(false);
+
+	// Assign ClipMesh and ClipMorph Materials from the plugin config
+	FConfigFile* PluginConfig = GConfig->FindConfigFileWithBaseName("Mutable");
+	if (PluginConfig)
+	{
+		FString ClipMorphMaterialName;
+		PluginConfig->GetString(TEXT("EditorDefaults"), TEXT("ClipMorphMaterialName"), ClipMorphMaterialName);
+		if (!ClipMorphMaterialName.IsEmpty())
+		{
+			ClipMorphMaterial = LoadObject<UMaterial>(NULL, *ClipMorphMaterialName, NULL, LOAD_None, NULL);
+			ensure(ClipMorphMaterial);
+		}
+
+		// Clip mesh with mesh  material
+		FString ClipMeshMaterialName;
+		PluginConfig->GetString(TEXT("EditorDefaults"), TEXT("ClipMeshMaterialName"), ClipMeshMaterialName);
+		if (!ClipMeshMaterialName.IsEmpty())
+		{
+			ClipMeshMaterial = LoadObject<UMaterial>(nullptr, *ClipMeshMaterialName, nullptr, LOAD_None, nullptr);
+			ensure(ClipMeshMaterial);
+		}
+	}
+
+	if (!ClipMorphMaterial)
+	{
+		ClipMorphMaterial = UMaterial::GetDefaultMaterial(MD_Surface);
+	}
+
+	if (!ClipMeshMaterial)
+	{
+		ClipMeshMaterial = UMaterial::GetDefaultMaterial(MD_Surface);
+	}
 
 	BoundSphere.W = 100.f;
 
@@ -691,15 +727,33 @@ void FCustomizableObjectEditorViewportClient::HideGizmoClipMorph()
 }
 
 
-void FCustomizableObjectEditorViewportClient::ShowGizmoClipMesh(UCustomizableObjectNodeMeshClipWithMesh& InClipMeshNode, UStaticMesh& ClipMesh)
+void FCustomizableObjectEditorViewportClient::ShowGizmoClipMesh(UCustomizableObjectNodeMeshClipWithMesh& InClipMeshNode, UObject& ClipMesh, int32 LODIndex, int32 SectionIndex, int32 MaterialSlotIndex)
 {
+	HideGizmoClipMesh();
+
 	SetWidgetType(EWidgetType::ClipMesh);
 
 	ClipMeshNode = &InClipMeshNode;
 
-	ClipMeshComp->SetStaticMesh(&ClipMesh);
-	ClipMeshComp->SetVisibility(true);
-	ClipMeshComp->SetWorldTransform(InClipMeshNode.Transform);
+	if (UStaticMesh* StaticMesh = Cast<UStaticMesh>(&ClipMesh))
+	{
+		ClipMeshStaticMeshComp->SetStaticMesh(StaticMesh);
+		ClipMeshStaticMeshComp->SetVisibility(true);
+		ClipMeshStaticMeshComp->SetWorldTransform(InClipMeshNode.Transform);
+		ClipMeshStaticMeshComp->EmptyOverrideMaterials();
+		ClipMeshStaticMeshComp->SetMaterial(MaterialSlotIndex, ClipMeshMaterial);
+		ClipMeshStaticMeshComp->SetSectionPreview(SectionIndex);
+	}
+	else if (USkeletalMesh* SkeletalMesh = Cast<USkeletalMesh>(&ClipMesh))
+	{
+		ClipMeshSkeletalMeshComp->SetSkeletalMesh(SkeletalMesh);
+		ClipMeshSkeletalMeshComp->SetVisibility(true);
+		ClipMeshSkeletalMeshComp->SetWorldTransform(InClipMeshNode.Transform);
+		ClipMeshSkeletalMeshComp->SetForcedLOD(LODIndex);
+		ClipMeshSkeletalMeshComp->EmptyOverrideMaterials();
+		ClipMeshSkeletalMeshComp->SetMaterial(MaterialSlotIndex, ClipMeshMaterial);
+		ClipMeshSkeletalMeshComp->SetSectionPreview(SectionIndex);
+	}
 }
 
 
@@ -707,8 +761,9 @@ void FCustomizableObjectEditorViewportClient::HideGizmoClipMesh()
 {
 	if (WidgetType == EWidgetType::ClipMesh)
 	{
-		ClipMeshComp->SetVisibility(false);
-		
+		ClipMeshStaticMeshComp->SetVisibility(false);
+		ClipMeshSkeletalMeshComp->SetVisibility(false);
+
 		SetWidgetType(EWidgetType::Hidden);
 	}
 }
@@ -1148,7 +1203,10 @@ bool FCustomizableObjectEditorViewportClient::InputWidgetDelta(FViewport* InView
 				ClipMeshNode->Transform.SetScale3D(ClipMeshNode->Transform.GetScale3D() + Scale);
 			}
 
-			ClipMeshComp->SetWorldTransform(ClipMeshNode->Transform);
+			ClipMeshStaticMeshComp->Modify();
+			ClipMeshSkeletalMeshComp->Modify();
+			ClipMeshStaticMeshComp->SetWorldTransform(ClipMeshNode->Transform);
+			ClipMeshSkeletalMeshComp->SetWorldTransform(ClipMeshNode->Transform);
 
 			return true;
 		}
@@ -1739,6 +1797,7 @@ void FCustomizableObjectEditorViewportClient::DrawShadowedString(FCanvas* Canvas
 void FCustomizableObjectEditorViewportClient::AddReferencedObjects(FReferenceCollector& Collector)
 {
 	Collector.AddReferencedObject(ClipMorphMaterial);
+	Collector.AddReferencedObject(ClipMeshMaterial);
 	Collector.AddReferencedObject(TransparentPlaneMaterialXY);
 	Collector.AddReferencedObject(AnimationBeingPlayed);
 
diff --git a/Engine/Plugins/Experimental/Mutable/Source/CustomizableObjectEditor/Private/MuCOE/GenerateMutableSource/GenerateMutableSourceMesh.cpp b/Engine/Plugins/Experimental/Mutable/Source/CustomizableObjectEditor/Private/MuCOE/GenerateMutableSource/GenerateMutableSourceMesh.cpp
--- a/Engine/Plugins/Experimental/Mutable/Source/CustomizableObjectEditor/Private/MuCOE/GenerateMutableSource/GenerateMutableSourceMesh.cpp
+++ b/Engine/Plugins/Experimental/Mutable/Source/CustomizableObjectEditor/Private/MuCOE/GenerateMutableSource/GenerateMutableSourceMesh.cpp
@@ -4035,14 +4035,18 @@ mu::NodeMeshPtr GenerateMutableSourceMesh(const UEdGraphPin* Pin,
 
 					TypedNodeTable->GetPinLODAndSection(Pin, LODIndexConnected, SectionIndexConnected);
 
+					// Getting the mutable table mesh column name
+					FString MutableColumnName;
 					if (DefaultSkeletalMesh)
 					{
+						MutableColumnName = TypedNodeTable->GenerateSkeletalMeshMutableColumName(DataTableColumnName, LODIndex, SectionIndex);
 						GetLODAndSectionForAutomaticLODs(GenerationContext, *Node, *DefaultSkeletalMesh, LODIndexConnected, SectionIndexConnected, LODIndex, SectionIndex, bOnlyConnectedLOD);
 					}
+					else
+					{
+						MutableColumnName = TypedNodeTable->GenerateStaticMeshMutableColumName(DataTableColumnName, SectionIndexConnected);
+					}
 					
-					// Getting the mutable table mesh column name
-					FString MutableColumnName = TypedNodeTable->GenerateSkeletalMeshMutableColumName(DataTableColumnName, LODIndex, SectionIndex);
-
 					// Generating a new Mesh column if not exists
 					if (Table->FindColumn(MutableColumnName) == INDEX_NONE)
 					{
diff --git a/Engine/Plugins/Experimental/Mutable/Source/CustomizableObjectEditor/Private/MuCOE/GenerateMutableSource/GenerateMutableSourceModifier.cpp b/Engine/Plugins/Experimental/Mutable/Source/CustomizableObjectEditor/Private/MuCOE/GenerateMutableSource/GenerateMutableSourceModifier.cpp
--- a/Engine/Plugins/Experimental/Mutable/Source/CustomizableObjectEditor/Private/MuCOE/GenerateMutableSource/GenerateMutableSourceModifier.cpp
+++ b/Engine/Plugins/Experimental/Mutable/Source/CustomizableObjectEditor/Private/MuCOE/GenerateMutableSource/GenerateMutableSourceModifier.cpp
@@ -156,12 +156,11 @@ mu::NodeModifierPtr GenerateMutableSourceModifier(const UEdGraphPin * Pin, FMuta
 				}
 			}
 
-			if (Cast<UCustomizableObjectNodeStaticMesh>(ConnectedPin->GetOwningNode()))
+			if (FMatrix Matrix = TypedNodeClipMesh->Transform.ToMatrixWithScale(); Matrix != FMatrix::Identity)
 			{
 				mu::NodeMeshTransformPtr TransformMesh = new mu::NodeMeshTransform();
 				TransformMesh->SetSource(ClipMesh.get());
 
-				FMatrix Matrix = TypedNodeClipMesh->Transform.ToMatrixWithScale();
 				TransformMesh->SetTransform(FMatrix44f(Matrix));
 				ClipMesh = TransformMesh;
 			}
diff --git a/Engine/Plugins/Experimental/Mutable/Source/CustomizableObjectEditor/Private/MuCOE/Nodes/CustomizableObjectNodeMeshClipWithMesh.cpp b/Engine/Plugins/Experimental/Mutable/Source/CustomizableObjectEditor/Private/MuCOE/Nodes/CustomizableObjectNodeMeshClipWithMesh.cpp
--- a/Engine/Plugins/Experimental/Mutable/Source/CustomizableObjectEditor/Private/MuCOE/Nodes/CustomizableObjectNodeMeshClipWithMesh.cpp
+++ b/Engine/Plugins/Experimental/Mutable/Source/CustomizableObjectEditor/Private/MuCOE/Nodes/CustomizableObjectNodeMeshClipWithMesh.cpp
@@ -56,6 +56,17 @@ void UCustomizableObjectNodeMeshClipWithMesh::AllocateDefaultPins(UCustomizableO
 }
 
 
+void UCustomizableObjectNodeMeshClipWithMesh::ReconstructNode(UCustomizableObjectNodeRemapPins* RemapPinsMode)
+{
+	Super::ReconstructNode(RemapPinsMode);
+
+	if (!CustomizableObjectToClipWith)
+	{
+		CustomizableObjectToClipWith = Cast<UCustomizableObject>(GetOutermostObject());
+	}
+}
+
+
 FText UCustomizableObjectNodeMeshClipWithMesh::GetNodeTitle(ENodeTitleType::Type TitleType) const
 {
 	return LOCTEXT("Clip_Mesh_With_Mesh", "Clip Mesh With Mesh");
diff --git a/Engine/Plugins/Experimental/Mutable/Source/CustomizableObjectEditor/Private/MuCOE/Nodes/CustomizableObjectNodeSkeletalMesh.cpp b/Engine/Plugins/Experimental/Mutable/Source/CustomizableObjectEditor/Private/MuCOE/Nodes/CustomizableObjectNodeSkeletalMesh.cpp
--- a/Engine/Plugins/Experimental/Mutable/Source/CustomizableObjectEditor/Private/MuCOE/Nodes/CustomizableObjectNodeSkeletalMesh.cpp
+++ b/Engine/Plugins/Experimental/Mutable/Source/CustomizableObjectEditor/Private/MuCOE/Nodes/CustomizableObjectNodeSkeletalMesh.cpp
@@ -80,7 +80,7 @@ void UCustomizableObjectNodeSkeletalMesh::AllocateDefaultPins(UCustomizableObjec
 					continue;
 				}
 
-				UMaterialInterface* MaterialInterface = GetMaterialInterfaceFor(LODIndex, SectionIndex, ImportedModel);
+				UMaterialInterface* MaterialInterface = GetMaterialInterfaceFor(LODIndex, SectionIndex);
 				
 				FString SectionFriendlyName = MaterialInterface ? MaterialInterface->GetName() : FString::Printf(TEXT("Section %i"), SectionIndex);
 
@@ -183,21 +183,10 @@ UTexture2D* UCustomizableObjectNodeSkeletalMesh::FindTextureForPin(const UEdGrap
 	{
 		return nullptr;
 	}
-	
-	if (!SkeletalMesh)
-	{
-		return nullptr;
-	}
-
-	const FSkeletalMeshModel* ImportedModel = SkeletalMesh->GetImportedModel();
-	if (!ImportedModel)
-	{
-		return nullptr;
-	}	
 
 	if (const UCustomizableObjectNodeSkeletalMeshPinDataImage* PinData = Cast<UCustomizableObjectNodeSkeletalMeshPinDataImage>(GetPinData(*Pin)))
 	{
-		if (const UMaterialInterface* MaterialInterface = GetMaterialInterfaceFor(PinData->GetLODIndex(), PinData->GetSectionIndex(), ImportedModel))
+		if (const UMaterialInterface* MaterialInterface = GetMaterialInterfaceFor(PinData->GetLODIndex(), PinData->GetSectionIndex()))
 		{
 			const UMaterialInterface* Material = GetMaterialFor(Pin);
 			
@@ -341,12 +330,9 @@ void UCustomizableObjectNodeSkeletalMesh::GetPinSection(const UEdGraphPin& Pin,
 
 UMaterialInterface* UCustomizableObjectNodeSkeletalMesh::GetMaterialFor(const UEdGraphPin* Pin) const
 {
-	if (SkeletalMesh)
+	if (FSkeletalMaterial* SkeletalMaterial = GetSkeletalMaterialFor(*Pin))
 	{
-		if (FSkeletalMaterial* SkeletalMaterial = GetSkeletalMaterialFor(*Pin))
-		{
-			return SkeletalMaterial->MaterialInterface;
-		}
+		return SkeletalMaterial->MaterialInterface;
 	}
 
 	return nullptr;
@@ -364,6 +350,17 @@ FSkeletalMaterial* UCustomizableObjectNodeSkeletalMesh::GetSkeletalMaterialFor(c
 }
 
 
+int32 UCustomizableObjectNodeSkeletalMesh::GetSkeletalMaterialIndexFor(const UEdGraphPin& Pin) const
+{
+	int32 LODIndex;
+	int32 SectionIndex;
+	int32 LayoutIndex;
+	GetPinSection(Pin, LODIndex, SectionIndex, LayoutIndex);
+
+	return GetSkeletalMaterialIndexFor(LODIndex, SectionIndex);
+}
+
+
 bool UCustomizableObjectNodeSkeletalMesh::IsPinRelevant(const UEdGraphPin* Pin) const
 {
 	const UEdGraphSchema_CustomizableObject* Schema = GetDefault<UEdGraphSchema_CustomizableObject>();
@@ -430,7 +427,7 @@ bool UCustomizableObjectNodeSkeletalMesh::IsNodeOutDatedAndNeedsRefresh()
 			}
 			else if (const UCustomizableObjectNodeSkeletalMeshPinDataImage* ImagePinData = Cast<UCustomizableObjectNodeSkeletalMeshPinDataImage>(GetPinData(*Pin)))
 			{				
-				const UMaterialInterface* MaterialInterface = GetMaterialInterfaceFor(ImagePinData->GetLODIndex(), ImagePinData->GetSectionIndex(), ImportedModel);
+				const UMaterialInterface* MaterialInterface = GetMaterialInterfaceFor(ImagePinData->GetLODIndex(), ImagePinData->GetSectionIndex());
 				if (!MaterialInterface) // If we had an Image pin for sure we had a MaterialInstance.
 				{
 					return true;
@@ -689,9 +686,9 @@ bool UCustomizableObjectNodeSkeletalMesh::CheckIsValidLayout(const UEdGraphPin*
 	return !VisitedLayouts.Contains(false);
 }
 
-UMaterialInterface* UCustomizableObjectNodeSkeletalMesh::GetMaterialInterfaceFor(const int32 LODIndex, const int32 SectionIndex, const FSkeletalMeshModel* ImportedModel) const
+UMaterialInterface* UCustomizableObjectNodeSkeletalMesh::GetMaterialInterfaceFor(const int32 LODIndex, const int32 SectionIndex) const
 {
-	if (FSkeletalMaterial* SkeletalMaterial = GetSkeletalMaterialFor(LODIndex, SectionIndex, ImportedModel))
+	if (FSkeletalMaterial* SkeletalMaterial = GetSkeletalMaterialFor(LODIndex, SectionIndex))
 	{
 		return SkeletalMaterial->MaterialInterface;
 	}
@@ -700,16 +697,33 @@ UMaterialInterface* UCustomizableObjectNodeSkeletalMesh::GetMaterialInterfaceFor
 }
 
 
-FSkeletalMaterial* UCustomizableObjectNodeSkeletalMesh::GetSkeletalMaterialFor(const int32 LODIndex, const int32 SectionIndex, const FSkeletalMeshModel* ImportedModel) const
+FSkeletalMaterial* UCustomizableObjectNodeSkeletalMesh::GetSkeletalMaterialFor(const int32 LODIndex, const int32 SectionIndex) const
 {
 	if (!SkeletalMesh)
 	{
 		return nullptr;
 	}
 
+	const int32 SkeletalMeshMaterialIndex = GetSkeletalMaterialIndexFor(LODIndex, SectionIndex);
+	if (SkeletalMesh->GetMaterials().IsValidIndex(SkeletalMeshMaterialIndex))
+	{
+		return &SkeletalMesh->GetMaterials()[SkeletalMeshMaterialIndex];
+	}
+	
+	return nullptr;
+}
+
+
+int32 UCustomizableObjectNodeSkeletalMesh::GetSkeletalMaterialIndexFor(const int32 LODIndex, const int32 SectionIndex) const
+{
+	if (!SkeletalMesh)
+	{
+		return INDEX_NONE;
+	}
+
 	// We assume that LODIndex and MaterialIndex are valid for the imported model
 	int32 SkeletalMeshMaterialIndex = INDEX_NONE;
-	
+
 	// Check if we have lod info map to get the correct material index
 	if (const FSkeletalMeshLODInfo* LodInfo = SkeletalMesh->GetLODInfo(LODIndex))
 	{
@@ -722,30 +736,14 @@ FSkeletalMaterial* UCustomizableObjectNodeSkeletalMesh::GetSkeletalMaterialFor(c
 	// Only deduce index when the explicit mapping is not found or there is no remap
 	if (SkeletalMeshMaterialIndex == INDEX_NONE)
 	{
-	if (ImportedModel && ImportedModel->LODModels.IsValidIndex(LODIndex) && ImportedModel->LODModels[LODIndex].Sections.IsValidIndex(SectionIndex))
-	{
-		SkeletalMeshMaterialIndex = ImportedModel->LODModels[LODIndex].Sections[SectionIndex].MaterialIndex;
-	}
-	else
-	{
-		FSkeletalMeshModel* AuxImportedModel = SkeletalMesh->GetImportedModel();
-
-		if (AuxImportedModel)
+		FSkeletalMeshModel* ImportedModel = SkeletalMesh->GetImportedModel();
+		if (ImportedModel && ImportedModel->LODModels.IsValidIndex(LODIndex) && ImportedModel->LODModels[LODIndex].Sections.IsValidIndex(SectionIndex))
 		{
-			if (AuxImportedModel->LODModels.IsValidIndex(LODIndex) && AuxImportedModel->LODModels[LODIndex].Sections.IsValidIndex(SectionIndex))
-			{
-				SkeletalMeshMaterialIndex = AuxImportedModel->LODModels[LODIndex].Sections[SectionIndex].MaterialIndex;
-			}
+			SkeletalMeshMaterialIndex = ImportedModel->LODModels[LODIndex].Sections[SectionIndex].MaterialIndex;
 		}
 	}
-	}
-	
-	if (SkeletalMesh->GetMaterials().IsValidIndex(SkeletalMeshMaterialIndex))
-	{
-		return &SkeletalMesh->GetMaterials()[SkeletalMeshMaterialIndex];
-	}
-	
-	return nullptr;
+
+	return SkeletalMeshMaterialIndex;
 }
 
 
diff --git a/Engine/Plugins/Experimental/Mutable/Source/CustomizableObjectEditor/Private/MuCOE/Nodes/CustomizableObjectNodeTable.cpp b/Engine/Plugins/Experimental/Mutable/Source/CustomizableObjectEditor/Private/MuCOE/Nodes/CustomizableObjectNodeTable.cpp
--- a/Engine/Plugins/Experimental/Mutable/Source/CustomizableObjectEditor/Private/MuCOE/Nodes/CustomizableObjectNodeTable.cpp
+++ b/Engine/Plugins/Experimental/Mutable/Source/CustomizableObjectEditor/Private/MuCOE/Nodes/CustomizableObjectNodeTable.cpp
@@ -1459,6 +1459,24 @@ FSkeletalMaterial* UCustomizableObjectNodeTable::GetDefaultSkeletalMaterialFor(c
 		return nullptr;
 	}
 
+	const int32 SkeletalMeshMaterialIndex = GetDefaultSkeletalMaterialIndexFor(MeshPin);
+	if (SkeletalMesh->GetMaterials().IsValidIndex(SkeletalMeshMaterialIndex))
+	{
+		return &SkeletalMesh->GetMaterials()[SkeletalMeshMaterialIndex];
+	}
+
+	return nullptr;
+}
+
+int32 UCustomizableObjectNodeTable::GetDefaultSkeletalMaterialIndexFor(const UEdGraphPin& MeshPin) const
+{
+	USkeletalMesh* SkeletalMesh = GetColumnDefaultAssetByType<USkeletalMesh>(&MeshPin);
+
+	if (!SkeletalMesh)
+	{
+		return INDEX_NONE;
+	}
+
 	int32 LODIndex;
 	int32 SectionIndex;
 	GetPinLODAndSection(&MeshPin, LODIndex, SectionIndex);
@@ -1486,12 +1504,7 @@ FSkeletalMaterial* UCustomizableObjectNodeTable::GetDefaultSkeletalMaterialFor(c
 		}
 	}
 
-	if (SkeletalMesh->GetMaterials().IsValidIndex(SkeletalMeshMaterialIndex))
-	{
-		return &SkeletalMesh->GetMaterials()[SkeletalMeshMaterialIndex];
-	}
-
-	return nullptr;
+	return SkeletalMeshMaterialIndex;
 }
 
 
diff --git a/Engine/Plugins/Experimental/Mutable/Source/CustomizableObjectEditor/Private/MuCOE/SCustomizableObjectEditorViewport.cpp b/Engine/Plugins/Experimental/Mutable/Source/CustomizableObjectEditor/Private/MuCOE/SCustomizableObjectEditorViewport.cpp
--- a/Engine/Plugins/Experimental/Mutable/Source/CustomizableObjectEditor/Private/MuCOE/SCustomizableObjectEditorViewport.cpp
+++ b/Engine/Plugins/Experimental/Mutable/Source/CustomizableObjectEditor/Private/MuCOE/SCustomizableObjectEditorViewport.cpp
@@ -335,14 +335,14 @@ void SCustomizableObjectEditorViewportTabBody::HideGizmoClipMorph() const
 }
 
 
-void SCustomizableObjectEditorViewportTabBody::ShowGizmoClipMesh(UCustomizableObjectNodeMeshClipWithMesh& ClipMeshNode, UStaticMesh& ClipMesh) const
+void SCustomizableObjectEditorViewportTabBody::ShowGizmoClipMesh(UCustomizableObjectNodeMeshClipWithMesh& ClipMeshNode, UObject& ClipMesh, int32 LODIndex, int32 SectionIndex, int32 MaterialSlotIndex) const
 {
 	if (const TSharedPtr<SWidget> Toolbar = ViewportToolbarTransformWidget.Pin())
 	{
 		Toolbar->SetVisibility(EVisibility::Visible);		
 	}
 	
-	LevelViewportClient->ShowGizmoClipMesh(ClipMeshNode, ClipMesh);
+	LevelViewportClient->ShowGizmoClipMesh(ClipMeshNode, ClipMesh, LODIndex, SectionIndex, MaterialSlotIndex);
 }
 
 
diff --git a/Engine/Plugins/Experimental/Mutable/Source/CustomizableObjectEditor/Private/MuCOE/SCustomizableObjectEditorViewport.h b/Engine/Plugins/Experimental/Mutable/Source/CustomizableObjectEditor/Private/MuCOE/SCustomizableObjectEditorViewport.h
--- a/Engine/Plugins/Experimental/Mutable/Source/CustomizableObjectEditor/Private/MuCOE/SCustomizableObjectEditorViewport.h
+++ b/Engine/Plugins/Experimental/Mutable/Source/CustomizableObjectEditor/Private/MuCOE/SCustomizableObjectEditorViewport.h
@@ -121,7 +121,7 @@ public:
 	void HideGizmoClipMorph() const;
 
 	/** Do not call directly. Use ICustomizableObjectEditor functions instead. */
-	void ShowGizmoClipMesh(UCustomizableObjectNodeMeshClipWithMesh& ClipMeshNode, UStaticMesh& ClipMesh) const;
+	void ShowGizmoClipMesh(UCustomizableObjectNodeMeshClipWithMesh& ClipMeshNode, UObject& ClipMesh, int32 LODIndex, int32 SectionIndex, int32 MaterialSlotIndex) const;
 
 	/** Do not call directly. Use ICustomizableObjectEditor functions instead. */
 	void HideGizmoClipMesh() const;
diff --git a/Engine/Plugins/Experimental/Mutable/Source/CustomizableObjectEditor/Public/MuCOE/CustomizableObjectEditorViewportClient.h b/Engine/Plugins/Experimental/Mutable/Source/CustomizableObjectEditor/Public/MuCOE/CustomizableObjectEditorViewportClient.h
--- a/Engine/Plugins/Experimental/Mutable/Source/CustomizableObjectEditor/Public/MuCOE/CustomizableObjectEditorViewportClient.h
+++ b/Engine/Plugins/Experimental/Mutable/Source/CustomizableObjectEditor/Public/MuCOE/CustomizableObjectEditorViewportClient.h
@@ -185,7 +185,7 @@ public:
 	void HideGizmoClipMorph();
 	
 	/** Do not call directly. Use ICustomizableObjectEditor functions instead. */
-	void ShowGizmoClipMesh(UCustomizableObjectNodeMeshClipWithMesh& ClipMeshNode, UStaticMesh& ClipMesh);
+	void ShowGizmoClipMesh(UCustomizableObjectNodeMeshClipWithMesh& ClipMeshNode, UObject& ClipMesh, int32 LODIndex, int32 SectionIndex, int32 MaterialSlotIndex);
 
 	/** Do not call directly. Use ICustomizableObjectEditor functions instead. */
 	void HideGizmoClipMesh();
@@ -322,7 +322,9 @@ private:
 	float MorphLength;
 	FBoxSphereBounds MorphBounds;
 	UCustomizableObjectNodeMeshClipWithMesh* ClipMeshNode;
-	UStaticMeshComponent* ClipMeshComp;
+	TObjectPtr<UMaterial> ClipMeshMaterial;
+	UStaticMeshComponent* ClipMeshStaticMeshComp;
+	USkeletalMeshComponent* ClipMeshSkeletalMeshComp;
 	float Radius1;
 	float Radius2;
 	float RotationAngle;
diff --git a/Engine/Plugins/Experimental/Mutable/Source/CustomizableObjectEditor/Public/MuCOE/Nodes/CustomizableObjectNodeMeshClipWithMesh.h b/Engine/Plugins/Experimental/Mutable/Source/CustomizableObjectEditor/Public/MuCOE/Nodes/CustomizableObjectNodeMeshClipWithMesh.h
--- a/Engine/Plugins/Experimental/Mutable/Source/CustomizableObjectEditor/Public/MuCOE/Nodes/CustomizableObjectNodeMeshClipWithMesh.h
+++ b/Engine/Plugins/Experimental/Mutable/Source/CustomizableObjectEditor/Public/MuCOE/Nodes/CustomizableObjectNodeMeshClipWithMesh.h
@@ -56,6 +56,7 @@ public:
 
 	// UCustomizableObjectNode interface
 	virtual void AllocateDefaultPins(UCustomizableObjectNodeRemapPins* RemapPins) override;
+	virtual void ReconstructNode(UCustomizableObjectNodeRemapPins* RemapPinsMode) override;
 	virtual void PinConnectionListChanged(UEdGraphPin* Pin) override;
 	virtual void UpdateReferencedNodeId(const FGuid& NewGuid) override;
 	virtual void BeginPostDuplicate(bool bDuplicateForPIE) override;
diff --git a/Engine/Plugins/Experimental/Mutable/Source/CustomizableObjectEditor/Public/MuCOE/Nodes/CustomizableObjectNodeSkeletalMesh.h b/Engine/Plugins/Experimental/Mutable/Source/CustomizableObjectEditor/Public/MuCOE/Nodes/CustomizableObjectNodeSkeletalMesh.h
--- a/Engine/Plugins/Experimental/Mutable/Source/CustomizableObjectEditor/Public/MuCOE/Nodes/CustomizableObjectNodeSkeletalMesh.h
+++ b/Engine/Plugins/Experimental/Mutable/Source/CustomizableObjectEditor/Public/MuCOE/Nodes/CustomizableObjectNodeSkeletalMesh.h
@@ -15,7 +15,6 @@ namespace ENodeTitleType { enum Type : int; }
 class FArchive;
 class FAssetThumbnail;
 class FAssetThumbnailPool;
-class FSkeletalMeshModel;
 class ISinglePropertyView;
 class SOverlay;
 class SVerticalBox;
@@ -228,6 +227,7 @@ public:
 	/** Returns the material associated to the given output pin. */
 	UMaterialInterface* GetMaterialFor(const UEdGraphPin* Pin) const;
 	FSkeletalMaterial* GetSkeletalMaterialFor(const UEdGraphPin& Pin) const;
+	int32 GetSkeletalMaterialIndexFor(const UEdGraphPin& Pin) const;
 
 	virtual bool ProvidesCustomPinRelevancyTest() const override { return true; }
 	virtual bool IsPinRelevant(const UEdGraphPin* Pin) const override;
@@ -248,8 +248,9 @@ public:
 	TWeakPtr< SGraphNodeSkeletalMesh > GraphNodeSkeletalMesh;
 
 private:
-	UMaterialInterface* GetMaterialInterfaceFor(const int32 LODIndex, const int32 MaterialIndex, const FSkeletalMeshModel* ImportedModel = nullptr) const;
-	FSkeletalMaterial* GetSkeletalMaterialFor(const int32 LODIndex, const int32 MaterialIndex, const FSkeletalMeshModel* ImportedModel = nullptr) const;
+	UMaterialInterface* GetMaterialInterfaceFor(const int32 LODIndex, const int32 MaterialIndex) const;
+	FSkeletalMaterial* GetSkeletalMaterialFor(const int32 LODIndex, const int32 SectionIndex) const;
+	int32 GetSkeletalMaterialIndexFor(const int32 LODIndex, const int32 SectionIndex) const;
 
 	// Deprecated
 	
diff --git a/Engine/Plugins/Experimental/Mutable/Source/CustomizableObjectEditor/Public/MuCOE/Nodes/CustomizableObjectNodeTable.h b/Engine/Plugins/Experimental/Mutable/Source/CustomizableObjectEditor/Public/MuCOE/Nodes/CustomizableObjectNodeTable.h
--- a/Engine/Plugins/Experimental/Mutable/Source/CustomizableObjectEditor/Public/MuCOE/Nodes/CustomizableObjectNodeTable.h
+++ b/Engine/Plugins/Experimental/Mutable/Source/CustomizableObjectEditor/Public/MuCOE/Nodes/CustomizableObjectNodeTable.h
@@ -406,6 +406,9 @@ public:
 
 	/** Returns the skeletal material associated to the given skeletal mesh output pin. */
 	FSkeletalMaterial* GetDefaultSkeletalMaterialFor(const UEdGraphPin& MeshPin) const;
+
+	/** Returns the index of the skeletal material associated to the given skeletal mesh output pin. */
+	int32 GetDefaultSkeletalMaterialIndexFor(const UEdGraphPin& MeshPin) const;
 	
 private:
 
