[Mutable] COE and COIE previews now use full actors instead of components.

- Fixed issue were UCustomizableObjectInstanceUsage was setting the reference Skeletal Mesh even when bSkipSetReferenceSkeletalMesh was set to true.
- Removed unused StaticMeshComponent code.
- Removed unused floor offset code.
- Now the viewport level is the only one that has a strong pointer to the actor or components. All others are weak.
- Camera bounds now are calculated from the whole actor.
- Moved all camera related code from editors to the viewport client.
- Pressing F key now focuses the actor instead of refreshing the camera.
- Removed " No reference mesh is set in the Object Properties tab". This error is already shown in the compilation errors and had overlay issues.

#rnx

diff --git a/Engine/Plugins/Experimental/Mutable/Source/CustomizableObject/Private/MuCO/CustomizableObjectInstanceUsage.cpp b/Engine/Plugins/Experimental/Mutable/Source/CustomizableObject/Private/MuCO/CustomizableObjectInstanceUsage.cpp
--- a/Engine/Plugins/Experimental/Mutable/Source/CustomizableObject/Private/MuCO/CustomizableObjectInstanceUsage.cpp
+++ b/Engine/Plugins/Experimental/Mutable/Source/CustomizableObject/Private/MuCO/CustomizableObjectInstanceUsage.cpp
@@ -374,8 +374,11 @@ void UCustomizableObjectInstanceUsage::EditorUpdateComponent()
 			const int32 ComponentIndex = GetComponentIndex();
 
 			const bool bInstanceGenerated = CustomizableObjectInstance->GetPrivate()->SkeletalMeshStatus != ESkeletalMeshStatus::NotGenerated;
-			USkeletalMesh* GeneratedSkeletalMesh = bInstanceGenerated ? CustomizableObjectInstance->GetSkeletalMesh(ComponentIndex) :
-				CustomizableObjectInstance->GetCustomizableObject()->GetRefSkeletalMesh(ComponentIndex);
+			USkeletalMesh* GeneratedSkeletalMesh = bInstanceGenerated ?
+				CustomizableObjectInstance->GetSkeletalMesh(ComponentIndex) :
+				GetSkipSetReferenceSkeletalMesh() ?
+					nullptr :
+					CustomizableObjectInstance->GetCustomizableObject()->GetRefSkeletalMesh(ComponentIndex);
 
 			if (AttachedSkeletalMesh != GeneratedSkeletalMesh)
 			{
diff --git a/Engine/Plugins/Experimental/Mutable/Source/CustomizableObjectEditor/Private/MuCOE/CustomizableObjectEditor.cpp b/Engine/Plugins/Experimental/Mutable/Source/CustomizableObjectEditor/Private/MuCOE/CustomizableObjectEditor.cpp
--- a/Engine/Plugins/Experimental/Mutable/Source/CustomizableObjectEditor/Private/MuCOE/CustomizableObjectEditor.cpp
+++ b/Engine/Plugins/Experimental/Mutable/Source/CustomizableObjectEditor/Private/MuCOE/CustomizableObjectEditor.cpp
@@ -63,6 +63,8 @@
 #include "Widgets/Layout/SScrollBox.h"
 #include "Widgets/Notifications/SNotificationList.h"
 #include "SMutableObjectViewer.h"
+#include "MuCO/CustomizableSkeletalMeshActor.h"
+#include "MuCOE/CustomizableObjectInstanceFactory.h"
 
 class FAdvancedPreviewScene;
 class FWorkspaceItem;
@@ -153,12 +155,6 @@ void FCustomizableObjectEditor::UnregisterTabSpawners(const TSharedRef<class FTa
 
 FCustomizableObjectEditor::~FCustomizableObjectEditor()
 {
-	if (PreviewInstance && HelperCallback)
-	{
-		PreviewInstance->UpdatedDelegate.RemoveDynamic(HelperCallback, &UUpdateClassWrapper::DelegatedCallback);
-		HelperCallback->Delegate.Unbind();
-	}
-
 	if (PreviewInstance)
 	{
 		if (PreviewInstance->GetPrivate()->bSelectedProfileDirty && PreviewInstance->GetPrivate()->SelectedProfileIndex != INDEX_NONE)
@@ -167,12 +163,6 @@ FCustomizableObjectEditor::~FCustomizableObjectEditor()
 		}
 	}
 	
-	for (UCustomizableSkeletalComponent* Component : PreviewCustomizableSkeletalComponents)
-	{
-		Component->DestroyComponent();
-	}
-
-	PreviewCustomizableSkeletalComponents.Empty();
 	CustomizableObjectDetailsView.Reset();
 	GEditor->UnregisterForUndo(this);
 
@@ -195,9 +185,7 @@ void FCustomizableObjectEditor::InitCustomizableObjectEditor(const EToolkitMode:
 
 	CustomSettings = NewObject<UCustomSettings>();
 	CustomSettings->SetEditor(SharedThis(this));
-
-	HelperCallback = nullptr;
-
+	
 	// Support undo/redo
 	CustomizableObject->SetFlags(RF_Transactional);
 
@@ -351,48 +339,15 @@ void FCustomizableObjectEditor::SelectNode(const UEdGraphNode* Node)
 
 void FCustomizableObjectEditor::CreatePreviewInstance()
 {
-	OnUpdatePreviewInstanceWork.Add([this]()
-	{
-		ViewportClient->ResetCamera();
-	});
-
 	PreviewInstance = CustomizableObject->CreateInstance();
-
 	CustomizableInstanceDetailsView->SetObject(PreviewInstance, true);
-	
-	HelperCallback = NewObject<UUpdateClassWrapper>(GetTransientPackage());
-	PreviewInstance->UpdatedDelegate.AddDynamic(HelperCallback, &UUpdateClassWrapper::DelegatedCallback);
-
-	if (CustomizableObject->ReferenceSkeletalMeshes.Num())
-	{
-		CreatePreviewComponents();
 
-		if (!PreviewSkeletalMeshComponents.IsEmpty() && !PreviewCustomizableSkeletalComponents.IsEmpty())
-		{
-			HelperCallback->Delegate.BindSP(this, &FCustomizableObjectEditor::OnUpdatePreviewInstance);
-
-			PreviewInstance->SetBuildParameterRelevancy(true);
-			PreviewInstance->UpdateSkeletalMeshAsync(true, true);
-		}
-		else
-		{
-			TArray<UDebugSkelMeshComponent*> EmptyArray;
-			Viewport->SetPreviewComponents(EmptyArray);
-			ViewportClient->SetReferenceMeshMissingWarningMessage(true);
-		}
-	}
-	else
-	{
-		ViewportClient->SetReferenceMeshMissingWarningMessage(true);
-	}
+	CreatePreviewActor();
 
-	if (GraphEditor.IsValid())
-	{
-		for (TObjectPtr<UEdGraphNode>& Node : CustomizableObject->GetPrivate()->GetSource()->Nodes)
-		{
-			GraphEditor->RefreshNode(*Node);
-		}
-	}
+	PreviewInstance->UpdatedNativeDelegate.AddSP(SharedThis(this), &FCustomizableObjectEditor::OnUpdatePreviewInstance);
+	
+	PreviewInstance->SetBuildParameterRelevancy(true);
+	PreviewInstance->UpdateSkeletalMeshAsync(true, true);
 }
 
 
@@ -400,9 +355,6 @@ void FCustomizableObjectEditor::AddReferencedObjects( FReferenceCollector& Colle
 {
 	Collector.AddReferencedObject( CustomizableObject );
 	Collector.AddReferencedObject( PreviewInstance );
-	Collector.AddReferencedObjects( PreviewCustomizableSkeletalComponents );
-	Collector.AddReferencedObjects( PreviewSkeletalMeshComponents );
-	Collector.AddReferencedObject( HelperCallback );
 	Collector.AddReferencedObject( ProjectorParameter );
 	Collector.AddReferencedObject( CustomSettings );
 }
@@ -2073,7 +2025,7 @@ void FCustomizableObjectEditor::FindProperty(const FProperty* Property, const vo
 
 void FCustomizableObjectEditor::OnPostCompile()
 {
-	CreatePreviewComponents();
+	CreatePreviewActor();
 	PreviewInstance->UpdateSkeletalMeshAsync(true, true);
 }
 
@@ -2161,53 +2113,16 @@ void FCustomizableObjectEditor::LogSearchResult(const UObject& Context, const FS
 }
 
 
-void FCustomizableObjectEditor::UpdatePreviewVisibility()
+void FCustomizableObjectEditor::OnUpdatePreviewInstance(UCustomizableObjectInstance* Instance)
 {
-	for (UDebugSkelMeshComponent* PreviewSkeletalMeshComponent : PreviewSkeletalMeshComponents)
-	{
-		if (PreviewSkeletalMeshComponent)
-		{
-			if (PreviewInstance->GetPrivate()->SkeletalMeshStatus == ESkeletalMeshStatus::Success)
-			{
-				PreviewSkeletalMeshComponent->SetVisibility(true, true);
-			}
-			else
-			{
-				PreviewSkeletalMeshComponent->SetVisibility(false, true);
-			}
-		}
-	}
-}
-
-
-void FCustomizableObjectEditor::OnUpdatePreviewInstance()
-{
-	UpdatePreviewVisibility();
-		
-	for (UDebugSkelMeshComponent* PreviewSkeletalMeshComponent : PreviewSkeletalMeshComponents)
-	{
-		PreviewSkeletalMeshComponent->RegisterComponentWithWorld(PreviewSkeletalMeshComponent->GetWorld());
-		PreviewSkeletalMeshComponent->bComponentUseFixedSkelBounds = true; // First bounds computed would be using physics asset
-		PreviewSkeletalMeshComponent->UpdateBounds();
-	}
+	const bool bVisible = PreviewInstance->GetPrivate()->SkeletalMeshStatus == ESkeletalMeshStatus::Success;
+	Actor->GetRootComponent()->SetVisibility(bVisible, true);
+	Actor->GetRootComponent()->UpdateBounds();
 	
-	Viewport->SetPreviewComponents(PreviewSkeletalMeshComponents);
-	Viewport->SetDrawDefaultUVMaterial();
-	ViewportClient->Invalidate();
-	ViewportClient->ReSetAnimation();
-	ViewportClient->SetReferenceMeshMissingWarningMessage(false);
-
 	if (TextureAnalyzer.IsValid())
 	{
 		TextureAnalyzer->RefreshTextureAnalyzerTable(PreviewInstance);
 	}
-
-	// Do postponed work.
-	for (const TFunction<void()>& Work : OnUpdatePreviewInstanceWork)
-	{
-		Work();		
-	}
-	OnUpdatePreviewInstanceWork.Reset();
 }
 
 
@@ -2414,35 +2329,41 @@ void FCustomizableObjectEditor::GetExternalChildObjects(const UCustomizableObjec
 }
 
 
-void FCustomizableObjectEditor::CreatePreviewComponents()
+void FCustomizableObjectEditor::CreatePreviewActor()
 {
-	for (const TObjectPtr<UCustomizableSkeletalComponent>& Component : PreviewCustomizableSkeletalComponents)
+	if (Actor)
 	{
-		Component->DestroyComponent();
+		ViewportClient->GetWorld()->RemoveActor(Actor.Get(), false);
 	}
-	PreviewCustomizableSkeletalComponents.Empty();
-
-	for (const TObjectPtr<UDebugSkelMeshComponent>& Component : PreviewSkeletalMeshComponents)
-	{
-		Component->DestroyComponent();
-	}
-	PreviewSkeletalMeshComponents.Empty();
+	
+	Actor.Reset(ViewportClient->GetWorld()->SpawnActor<ASkeletalMeshActor>());
 
+	TArray<TWeakObjectPtr<UDebugSkelMeshComponent>> PreviewSkeletalMeshComponents;
+	
 	for (int32 ComponentIndex = 0; ComponentIndex < CustomizableObject->GetComponentCount(); ++ComponentIndex)
 	{
-		UCustomizableSkeletalComponent* PreviewCustomizableSkeletalComponent = NewObject<UCustomizableSkeletalComponent>(UCustomizableSkeletalComponent::StaticClass());
-		UDebugSkelMeshComponent* PreviewSkeletalMeshComponent = NewObject<UDebugSkelMeshComponent>(GetTransientPackage(), NAME_None, RF_Transient);
-
-		PreviewCustomizableSkeletalComponent->bSkipSetReferenceSkeletalMesh = true;
-		PreviewCustomizableSkeletalComponent->CustomizableObjectInstance = PreviewInstance;
-		PreviewCustomizableSkeletalComponent->ComponentIndex = ComponentIndex;
-		PreviewCustomizableSkeletalComponent->AttachToComponent(PreviewSkeletalMeshComponent, FAttachmentTransformRules::KeepRelativeTransform);
-
-		PreviewCustomizableSkeletalComponents.Add(PreviewCustomizableSkeletalComponent);
-		PreviewSkeletalMeshComponents.Add(PreviewSkeletalMeshComponent);
+		UDebugSkelMeshComponent* DebugComponent = NewObject<UDebugSkelMeshComponent>(Actor.Get(), NAME_None, RF_Transient);
+		DebugComponent->bCastInsetShadow = true; // For better quality shadows in the editor previews, more similar to the in-game ones
+		DebugComponent->bCanHighlightSelectedSections = false;
+		DebugComponent->bComponentUseFixedSkelBounds = true; // First bounds computed would be using physics asset
+		DebugComponent->MarkRenderStateDirty();
+		DebugComponent->AttachToComponent(Actor->GetRootComponent(), FAttachmentTransformRules::KeepRelativeTransform);
+		DebugComponent->RegisterComponent();
+		
+		PreviewSkeletalMeshComponents.Add(DebugComponent);
+		
+		UCustomizableSkeletalComponent* CustomizableComponent = NewObject<UCustomizableSkeletalComponent>(DebugComponent);
+		CustomizableComponent->bSkipSetReferenceSkeletalMesh = true;
+		CustomizableComponent->CustomizableObjectInstance = PreviewInstance;
+		CustomizableComponent->ComponentIndex = ComponentIndex;
+		CustomizableComponent->AttachToComponent(DebugComponent, FAttachmentTransformRules::KeepRelativeTransform);
+		CustomizableComponent->RegisterComponent();
 	}
+	
+	Viewport->SetPreviewActor(Actor.Get(), PreviewInstance, PreviewSkeletalMeshComponents);
 }
 
+
 void RemoveRestrictedChars(FString& String)
 {
 	// Remove restricted chars, according to FPaths::ValidatePath, RestrictedChars = "/?:&\\*\"<>|%#@^ ";
diff --git a/Engine/Plugins/Experimental/Mutable/Source/CustomizableObjectEditor/Private/MuCOE/CustomizableObjectEditor.h b/Engine/Plugins/Experimental/Mutable/Source/CustomizableObjectEditor/Private/MuCOE/CustomizableObjectEditor.h
--- a/Engine/Plugins/Experimental/Mutable/Source/CustomizableObjectEditor/Private/MuCOE/CustomizableObjectEditor.h
+++ b/Engine/Plugins/Experimental/Mutable/Source/CustomizableObjectEditor/Private/MuCOE/CustomizableObjectEditor.h
@@ -251,7 +251,7 @@ public:
 	virtual void UpdateGraphNodeProperties();
 
 	/** Callback to notify the editor when the PreviewInstance has been updated */
-	void OnUpdatePreviewInstance();
+	void OnUpdatePreviewInstance(UCustomizableObjectInstance* Instance);
 
 	/** Called when the Object Properties needs to be updated */
 	virtual void UpdateObjectProperties() override;
@@ -336,10 +336,7 @@ private:
 	bool CanRenameNodes() const;
 
 	virtual UEdGraphNode* CreateCommentBox(const FVector2D& NodePos) override;
-
-	/** Updates the visibility of PreviewSkeletalMeshComponent */
-	void UpdatePreviewVisibility();
-
+	
 	/** Searches a node that contains the inserted word */
 	void OnEnterText(const FText& NewText, ETextCommit::Type TextType);
 
@@ -357,8 +354,8 @@ private:
 	/** Open the Performance Report tab */
 	void OpenPerformanceReportTab();
 
-	/** Creates the necessary components for the preview of the CO instance */
-	void CreatePreviewComponents();
+	/** Creates the actor and necessary components for the preview of the CO instance */
+	void CreatePreviewActor();
 
 	/** Recursively find any property that its name or value contains the given string.
 	  * @param Property Root property.
@@ -394,10 +391,11 @@ public:
 
 private:
 	/** The currently viewed object. */
-	TObjectPtr<UCustomizableObject> CustomizableObject;
+	TObjectPtr<UCustomizableObject> CustomizableObject = nullptr;
 	TObjectPtr<UCustomizableObjectInstance> PreviewInstance = nullptr;
-	TArray<TObjectPtr<UCustomizableSkeletalComponent>> PreviewCustomizableSkeletalComponents;
-	TArray<TObjectPtr<UDebugSkelMeshComponent>> PreviewSkeletalMeshComponents;
+
+	/** Preview Actor. All preview components are attached to this actor. */
+	TStrongObjectPtr<AActor> Actor;
 
 	/** List of open tool panels; used to ensure only one exists at any one time */
 	TMap< FName, TWeakPtr<class SDockableTab> > SpawnedToolPanels;
@@ -446,10 +444,7 @@ private:
 
 	/** URL to open when pressing the documentation button generated by UE */
 	const FString DocumentationURL{ TEXT("[REDACTED_PROFILE_URL]/Mutable-Documentation/wiki") };
-
-	/** Postponed work to do when OnUpdatePreviewInstance is called. Emptied at the end on each OnUpdatePreviewInstance call. */
-	TArray<TFunction<void()>> OnUpdatePreviewInstanceWork;
-
+	
 	TObjectPtr<UProjectorParameter> ProjectorParameter = nullptr;
 
 	TObjectPtr<UCustomSettings> CustomSettings = nullptr;
diff --git a/Engine/Plugins/Experimental/Mutable/Source/CustomizableObjectEditor/Private/MuCOE/CustomizableObjectEditorViewportClient.cpp b/Engine/Plugins/Experimental/Mutable/Source/CustomizableObjectEditor/Private/MuCOE/CustomizableObjectEditorViewportClient.cpp
--- a/Engine/Plugins/Experimental/Mutable/Source/CustomizableObjectEditor/Private/MuCOE/CustomizableObjectEditorViewportClient.cpp
+++ b/Engine/Plugins/Experimental/Mutable/Source/CustomizableObjectEditor/Private/MuCOE/CustomizableObjectEditorViewportClient.cpp
@@ -73,16 +73,13 @@ FCustomizableObjectEditorViewportClient::FCustomizableObjectEditorViewportClient
 	bCameraLock = true;
 	bDrawSky = true;
 
-	bActivateOrbitalCamera = true;
-	bSetOrbitalOnPerspectiveMode = true;
+	bSetOrbitalOnPerspectiveMode = bCameraLock;
 
 	const int32 CameraSpeed = 3;
 	SetCameraSpeedSetting(CameraSpeed);
 
 	bShowBones = false;
-
-	bReferenceMeshMissingWarningMessageVisible = false;
-
+	
 	DrawHelper.bDrawPivot = false;
 	DrawHelper.bDrawWorldBox = false;
 	DrawHelper.bDrawKillZ = false;
@@ -106,8 +103,6 @@ FCustomizableObjectEditorViewportClient::FCustomizableObjectEditorViewportClient
 	EngineShowFlags.Grid = ConfigOption->bShowGrid;
 
 	OverrideNearClipPlane(1.0f);
-
-	SetPreviewComponent(nullptr);
 	
 	// now add the ClipMorph plane
 	ClipMorphNode = nullptr;
@@ -120,9 +115,7 @@ FCustomizableObjectEditorViewportClient::FCustomizableObjectEditorViewportClient
 	ClipMeshComp = NewObject<UStaticMeshComponent>();
 	PreviewScene->AddComponent(ClipMeshComp, FTransform());
 	ClipMeshComp->SetVisibility(false);
-
-	BoundSphere.W = 100.f;
-
+	
 	const float FOVMin = 5.f;
 	const float FOVMax = 170.f;
 	ViewFOV = FMath::Clamp<float>(53.43f, FOVMin, FOVMax);
@@ -158,92 +151,12 @@ FCustomizableObjectEditorViewportClient::FCustomizableObjectEditorViewportClient
 }
 
 
-void FCustomizableObjectEditorViewportClient::UpdateCameraSetup()
-{
-	// Look for any Skeletal Mesh Component that we can focus to.
-	bool bWasValidComponentFound = false;
-	for	(TWeakObjectPtr<UDebugSkelMeshComponent> SkeletalMeshComponent : SkeletalMeshComponents)
-	{
-		if (SkeletalMeshComponent.IsValid() && UE_MUTABLE_GETSKINNEDASSET(SkeletalMeshComponent))
-		{
-			bWasValidComponentFound = true;
-			break;
-		}
-	}
-	
-	static FRotator CustomOrbitRotation(-33.75, -135, 0);
-	if ( bWasValidComponentFound
-		||
-		(StaticMeshComponent.IsValid() && StaticMeshComponent->GetStaticMesh()) )
-	{
-		BoundSphere = GetCameraTarget();
-		FVector CustomOrbitZoom(0, BoundSphere.W / (75.0f * (float)PI / 360.0f), 0);
-		FVector CustomOrbitLookAt = BoundSphere.Center;
-
-		SetCameraSetup(CustomOrbitLookAt, CustomOrbitRotation, CustomOrbitZoom, CustomOrbitLookAt, GetViewLocation(), GetViewRotation() );
-
-		UpdateFloor();
-
-		EnableCameraLock(bActivateOrbitalCamera);
-		FBox Box( BoundSphere.Center - FVector(BoundSphere.W) / 2.0f, BoundSphere.Center + FVector(BoundSphere.W) / 2.0f );
-		FocusViewportOnBox( Box, false );
-	}
-}
-
-
-void FCustomizableObjectEditorViewportClient::UpdateFloor()
-{
-	// Move the floor to the bottom of the bounding box of the mesh, rather than on the origin
-	bool bFoundSkelMesh = false;
-
-	for (TWeakObjectPtr<USkeletalMeshComponent> SkeletalMeshComponent : SkeletalMeshComponents)
-	{
-		if (SkeletalMeshComponent.IsValid() )
-		{
-			SkeletalMeshComponent->bComponentUseFixedSkelBounds = true;
-			SkeletalMeshComponent->UpdateBounds();
-
-			bFoundSkelMesh = true;
-		}
-	}
-
-	// TODO: Optimize
-	if (bFoundSkelMesh)
-	{
-
-	}
-	else if (StaticMeshComponent.IsValid())
-	{
-		StaticMeshComponent->UpdateBounds();
-	}
-
-	FAdvancedPreviewScene* AdvancedScene = static_cast<FAdvancedPreviewScene*>(PreviewScene);
-	if (AdvancedScene != nullptr)
-	{
-		const UStaticMeshComponent* FloorMeshComponent = AdvancedScene->GetFloorMeshComponent();
-		if (FloorMeshComponent != nullptr)
-		{
-			UStaticMeshComponent* FloorMeshComponentCasted = const_cast<UStaticMeshComponent*>(FloorMeshComponent);
-			FloorMeshComponentCasted->SetWorldLocation(FVector(0.0f, 0.0f, -1.0f/* Does not seem to work Bottom.Z*/));
-		}
-	}
-}
-
-
 FCustomizableObjectEditorViewportClient::~FCustomizableObjectEditorViewportClient()
 {
 	UAssetViewerSettings::Get()->OnAssetViewerSettingsChanged().RemoveAll(this);
 }
 
 
-void FCustomizableObjectEditorViewportClient::Tick(float DeltaSeconds)
-{
-	FEditorViewportClient::Tick(DeltaSeconds);
-
-	UpdateFloor();
-}
-
-
 void DrawEllipse(FPrimitiveDrawInterface* PDI, const FVector& Base, const FVector& X, const FVector& Y, const FLinearColor& Color, float Radius1, float Radius2, int32 NumSides, uint8 DepthPriority, float Thickness, float DepthBias, bool bScreenSpace)
 {
 	const float	AngleDelta = 2.0f * PI / NumSides;
@@ -459,17 +372,6 @@ void FCustomizableObjectEditorViewportClient::Draw(FViewport* InViewport, FCanva
 		DrawUVs(InViewport, Canvas, YPos);
 	}
 
-	if (bReferenceMeshMissingWarningMessageVisible)
-	{
-		Canvas->DrawShadowedString(
-			6,
-			2,
-			*NSLOCTEXT("CustomizableObjectEditor", "NoReferenceMeshMutable", "Warning! No reference mesh is set in the Object Properties tab.").ToString(),
-			GEngine->GetSmallFont(),
-			FLinearColor::Red
-			);
-	}
-
 	if (StateChangeShowGeometryDataFlag)
 	{
 		ShowInstanceGeometryInformation(Canvas);
@@ -531,46 +433,7 @@ void FCustomizableObjectEditorViewportClient::DrawUVs(FViewport* InViewport, FCa
 	BatchedElements->AddLine( FVector( Box[ 2 ], 0.0f ), FVector( Box[ 3 ], 0.0f ), BorderColor, HitProxyId );
 	BatchedElements->AddLine( FVector( Box[ 3 ], 0.0f ), FVector( Box[ 0 ], 0.0f ), BorderColor, HitProxyId );
 
-	if (StaticMeshComponent.IsValid() &&
-		StaticMeshComponent->GetStaticMesh() &&
-		StaticMeshComponent->GetStaticMesh()->GetRenderData() &&
-		StaticMeshComponent->GetStaticMesh()->GetRenderData()->LODResources.IsValidIndex(LODLevel))
-	{
-		FStaticMeshLODResources* RenderData = &StaticMeshComponent->GetStaticMesh()->GetRenderData()->LODResources[LODLevel];
-		
-		if (RenderData && UVChannel >= 0 && UVChannel < static_cast<int32>(RenderData->VertexBuffers.StaticMeshVertexBuffer.GetNumTexCoords()))
-		{
-			//draw triangles
-			FIndexArrayView Indices = RenderData->IndexBuffer.GetArrayView();
-			uint32 NumIndices = Indices.Num();
-		
-			BatchedElements->AddReserveLines( NumIndices );
-
-			for (uint32 i = 0; i < NumIndices - 2; i += 3)
-			{
-				FVector2D UV1( RenderData->VertexBuffers.StaticMeshVertexBuffer.GetVertexUV( Indices[ i + 0 ], UVChannel ) );
-				FVector2D UV2( RenderData->VertexBuffers.StaticMeshVertexBuffer.GetVertexUV( Indices[ i + 1 ], UVChannel ) );
-				FVector2D UV3( RenderData->VertexBuffers.StaticMeshVertexBuffer.GetVertexUV( Indices[ i + 2 ], UVChannel ) );
-	
-				// Draw lines in black unless the UVs are outside of the 0.0 - 1.0 range.  For out-of-bounds
-				// UVs, we'll draw the line segment in red
-				
-				// If we are supporting a version lower than LWC get the right real type. 
-				using Vector2DRealType = TDecay<decltype( DeclVal<FVector2D>().X )>::Type;
-				
-				constexpr Vector2DRealType Zero = static_cast<Vector2DRealType>(0);
-
-				UV1 = ClampUVRange(UV1.X, UV1.Y) * UVBoxScale + UVBoxOrigin;
-				UV2 = ClampUVRange(UV2.X, UV2.Y) * UVBoxScale + UVBoxOrigin;
-				UV3 = ClampUVRange(UV3.X, UV3.Y) * UVBoxScale + UVBoxOrigin;
-
-				BatchedElements->AddLine(FVector(UV1, Zero), FVector(UV2, Zero), BorderColor, HitProxyId);
-				BatchedElements->AddLine(FVector(UV2, Zero), FVector(UV3, Zero), BorderColor, HitProxyId);
-				BatchedElements->AddLine(FVector(UV3, Zero), FVector(UV1, Zero), BorderColor, HitProxyId);
-			}
-		}
-	}
-	else if (SkeletalMeshComponents.Num())
+	if (SkeletalMeshComponents.Num())
 	{
 		if (SkeletalMeshComponents.IsValidIndex(ComponentIndex))
 		{
@@ -637,22 +500,6 @@ void FCustomizableObjectEditorViewportClient::DrawUVs(FViewport* InViewport, FCa
 }
 
 
-float FCustomizableObjectEditorViewportClient::GetFloorOffset() const
-{
-	FAdvancedPreviewScene* AdvancedScene = static_cast<FAdvancedPreviewScene*>(PreviewScene);
-	if (AdvancedScene != nullptr)
-	{
-		const UStaticMeshComponent* FloorMeshComponent = AdvancedScene->GetFloorMeshComponent();
-		if (FloorMeshComponent != nullptr)
-		{
-			return FloorMeshComponent->GetComponentLocation().Z;
-		}
-	}
-
-	return 0.0f;
-}
-
-
 void FCustomizableObjectEditorViewportClient::ShowGizmoClipMorph(UCustomizableObjectNodeMeshClipMorph& NodeMeshClipMorph)
 {
 	SetWidgetType(EWidgetType::ClipMorph);
@@ -768,89 +615,35 @@ void FCustomizableObjectEditorViewportClient::HideGizmoLight()
 }
 
 
-FSphere FCustomizableObjectEditorViewportClient::GetCameraTarget()
-{
-	bool bFoundTarget = false;
-	FSphere Sphere(FVector(0,0,0), 100.0f); // default
-
-	for (TWeakObjectPtr<USkeletalMeshComponent> SkeletalMeshComponent : SkeletalMeshComponents)
-	{
-		if (SkeletalMeshComponent.IsValid())
-		{
-			FBoxSphereBounds Bounds = SkeletalMeshComponent.Get()->CalcBounds(FTransform::Identity);
-
-			if (!bFoundTarget)
-			{
-				Sphere = Bounds.GetSphere();
-			}
-			else
-			{
-				Sphere += Bounds.GetSphere();
-			}
-
-			bFoundTarget = true;
-		}
-	}
-
-	if(!bFoundTarget && StaticMeshComponent.IsValid())
-	{
-		if( !bFoundTarget )
-		{
-			FBoxSphereBounds Bounds = StaticMeshComponent.Get()->CalcBounds(FTransform::Identity);
-			Sphere = Bounds.GetSphere();
-		}
-	}
-
-	return Sphere;
-}
-
-
-void FCustomizableObjectEditorViewportClient::SetPreviewComponent(UStaticMeshComponent* InStaticMeshComponent)
+void FCustomizableObjectEditorViewportClient::SetPreviewActor(const TWeakObjectPtr<AActor>& InActor, const TWeakObjectPtr<UCustomizableObjectInstance>& InInstance, const TArray<TWeakObjectPtr<UDebugSkelMeshComponent>>& InSkeletalMeshComponents)
 {
-	StaticMeshComponent = InStaticMeshComponent;
-	SkeletalMeshComponents.Reset();
+	SkeletalMeshComponents = InSkeletalMeshComponents;
+	Actor = InActor;
 
-	if (StaticMeshComponent.IsValid() && StaticMeshComponent->GetStaticMesh())
-	{
-		SetViewLocation( -FVector(0, StaticMeshComponent->GetStaticMesh()->GetBounds().SphereRadius / (75.0f * (float)PI / 360.0f), 0) );
-		SetViewRotation( FRotator(0, 90.f, 0) );
-		//LockLocation = FVector(0,StaticMeshComponent->StaticMesh->ThumbnailDistance,0);
-		//LockRot = StaticMeshComponent->StaticMesh->ThumbnailAngle;
-	}
-
-	UpdateCameraSetup();
+	InInstance->UpdatedNativeDelegate.AddSP(SharedThis(this), &FCustomizableObjectEditorViewportClient::OnInstanceUpdate);
 }
 
 
-void FCustomizableObjectEditorViewportClient::SetPreviewComponents(const TArray<UDebugSkelMeshComponent*>& InSkeletalMeshComponents)
+void FCustomizableObjectEditorViewportClient::OnInstanceUpdate(UCustomizableObjectInstance* Instance)
 {
-	SkeletalMeshComponents.Reset(InSkeletalMeshComponents.Num());
-	SkeletalMeshComponents.Append(InSkeletalMeshComponents);
+	Invalidate();
+	ReSetAnimation();
 	
-	StaticMeshComponent = nullptr;
-}
-
-void FCustomizableObjectEditorViewportClient::ResetCamera()
-{
-	float MaxSphereRadius = 0.0f;
-	for (const TWeakObjectPtr<UDebugSkelMeshComponent>& SkeletalMeshComponent : SkeletalMeshComponents)
+	// Configure the initial orbital position of the camera
+	if (!bIsCameraSetup)
 	{
-		if (UE_MUTABLE_GETSKINNEDASSET(SkeletalMeshComponent))
-		{
-			MaxSphereRadius = FMath::Max(MaxSphereRadius, UE_MUTABLE_GETSKINNEDASSET(SkeletalMeshComponent)->GetBounds().SphereRadius);
-		}
-	}
-	
-	SetViewLocation(-FVector(0, MaxSphereRadius / (75.0f * (float)PI / 360.0f), 0));
-	SetViewRotation(FRotator(0, 90.f, 0));
-
-	UpdateCameraSetup();
-}
+		bIsCameraSetup = true;
+		
+		FVector Center;
+		FVector Extents;
+		Actor.Get()->GetActorBounds(false, Center, Extents, true);
 
+		static FRotator CustomOrbitRotation(-33.75, -135, 0);
+		FVector CustomOrbitZoom(0, Extents.GetMax() * 2.5 / (75.0 * PI / 360.0), 0);
 
-void FCustomizableObjectEditorViewportClient::SetReferenceMeshMissingWarningMessage(bool bVisible)
-{
-	bReferenceMeshMissingWarningMessageVisible = bVisible;
+		SetCameraSetup(Center, CustomOrbitRotation, CustomOrbitZoom, Center, FVector::Zero(), {} /** Not used since orbital is enable just after. */ );	
+		EnableCameraLock(true);
+	}
 }
 
 
@@ -977,7 +770,12 @@ bool FCustomizableObjectEditorViewportClient::InputKey(const FInputKeyEventArgs&
 	{
 		if (EventArgs.Key == EKeys::F)
 		{
-			UpdateCameraSetup();
+			FVector Center;
+			FVector Extents;
+			Actor.Get()->GetActorBounds(false, Center, Extents, true);
+	
+			FocusViewportOnBox(FBox(Center - Extents, Center + Extents), true);
+			
 			return true;
 		}
 		else if (WidgetType != EWidgetType::Hidden) // Do not change the type when hidden.
@@ -1293,6 +1091,7 @@ bool FCustomizableObjectEditorViewportClient::HandleEndTransform()
 	return false;
 }
 
+
 FVector FCustomizableObjectEditorViewportClient::GetWidgetLocation() const
 {
 	switch (WidgetType)
@@ -1391,7 +1190,7 @@ void FCustomizableObjectEditorViewportClient::SetViewportType(ELevelViewportType
 	// Getting camera mode on perspective view
 	if (ViewportType == ELevelViewportType::LVT_Perspective)
 	{
-		bSetOrbitalOnPerspectiveMode = bActivateOrbitalCamera;
+		bSetOrbitalOnPerspectiveMode = bCameraLock;
 	}
 
 	// Set Camera mode
@@ -1520,26 +1319,6 @@ void FCustomizableObjectEditorViewportClient::RemoveAllLightsFromScene()
 }
 
 
-void FCustomizableObjectEditorViewportClient::SetFloorOffset(float NewValue)
-{
-	for (TWeakObjectPtr<USkeletalMeshComponent> SkeletalMeshComponent : SkeletalMeshComponents)
-	{
-		USkeletalMesh* Mesh = SkeletalMeshComponent.IsValid() ? Cast<USkeletalMesh>(UE_MUTABLE_GETSKINNEDASSET(SkeletalMeshComponent)) : nullptr;
-
-		if (Mesh)
-		{
-			// This value is saved in a UPROPERTY for the mesh, so changes are transactional
-			FScopedTransaction Transaction(LOCTEXT("SetFloorOffset", "Set Floor Offset"));
-			Mesh->Modify();
-
-			Mesh->SetFloorOffset(NewValue);
-			UpdateCameraSetup(); // This does the actual moving of the floor mesh
-			Invalidate();
-		}
-	}
-}
-
-
 //-------------------------------------------------------------------------------------------------
 
 class SMutableSelectFolderDlg : public SWindow
@@ -1937,13 +1716,12 @@ void FCustomizableObjectEditorViewportClient::SetEnvironmentMeshVisibility(uint3
 
 bool FCustomizableObjectEditorViewportClient::IsOrbitalCameraActive() const
 {
-	return bActivateOrbitalCamera;
+	return bCameraLock;
 }
 
 void FCustomizableObjectEditorViewportClient::SetCameraMode(bool Value)
 {
-	bActivateOrbitalCamera = Value;
-	UpdateCameraSetup();
+	EnableCameraLock(Value);
 }
 
 
diff --git a/Engine/Plugins/Experimental/Mutable/Source/CustomizableObjectEditor/Private/MuCOE/CustomizableObjectInstanceEditor.cpp b/Engine/Plugins/Experimental/Mutable/Source/CustomizableObjectEditor/Private/MuCOE/CustomizableObjectInstanceEditor.cpp
--- a/Engine/Plugins/Experimental/Mutable/Source/CustomizableObjectEditor/Private/MuCOE/CustomizableObjectInstanceEditor.cpp
+++ b/Engine/Plugins/Experimental/Mutable/Source/CustomizableObjectEditor/Private/MuCOE/CustomizableObjectInstanceEditor.cpp
@@ -33,6 +33,7 @@
 #include "Widgets/Notifications/SNotificationList.h"
 #include "WorkspaceMenuStructure.h"
 #include "WorkspaceMenuStructureModule.h"
+#include "Animation/SkeletalMeshActor.h"
 #include "MuCO/CustomizableObjectInstancePrivate.h"
 
 class FAdvancedPreviewScene;
@@ -210,26 +211,12 @@ FCustomizableObjectInstanceEditor::~FCustomizableObjectInstanceEditor()
 		HelperCallback = nullptr;
 	}
 
-	for (UCustomizableSkeletalComponent* PreviewCustomizableSkeletalComponent : PreviewCustomizableSkeletalComponents)
-	{
-		if (PreviewCustomizableSkeletalComponent)
-		{
-			PreviewCustomizableSkeletalComponent->DestroyComponent();
-		}
-	}
-
-	PreviewCustomizableSkeletalComponents.Reset();
 	PreviewSkeletalMeshComponents.Reset();
 
 	CustomizableInstanceDetailsView.Reset();
 	Viewport.Reset();
 
 	FCoreUObjectDelegates::OnObjectModified.Remove(OnObjectModifiedHandle);
-	
-	if (const UCustomizableObject* CustomizableObject = CustomizableObjectInstance->GetCustomizableObject())
-	{
-		CustomizableObject->GetPrivate()->Status.GetOnStateChangedDelegate().RemoveAll(this);		
-	}
 
 	FCoreUObjectDelegates::OnObjectPropertyChanged.RemoveAll(this);
 }
@@ -343,7 +330,7 @@ FText FCustomizableObjectInstanceEditor::GetBaseToolkitName() const
 }
 
 
-void FCustomizableObjectInstanceEditor::CreatePreviewInstance()
+void FCustomizableObjectInstanceEditor::SetupPreviewInstance()
 {
 	check(CustomizableObjectInstance);
 	if (!CustomizableObjectInstance->GetCustomizableObject())
@@ -351,25 +338,12 @@ void FCustomizableObjectInstanceEditor::CreatePreviewInstance()
 		return;
 	}
 
-	CreatePreviewComponents();
+	CreatePreviewActor();
 	
-	// Bind update delegate
-	HelperCallback = NewObject<UUpdateClassWrapperClass>();
-	HelperCallback->Delegate.BindSP(this, &FCustomizableObjectInstanceEditor::OnUpdatePreviewInstance);
-	CustomizableObjectInstance->UpdatedDelegate.AddDynamic(HelperCallback, &UUpdateClassWrapperClass::DelegatedCallback);
-
+	CustomizableObjectInstance->UpdatedNativeDelegate.AddSP(SharedThis(this), &FCustomizableObjectInstanceEditor::OnUpdatePreviewInstance);
+	
+	CustomizableObjectInstance->SetBuildParameterRelevancy(true);
 	CustomizableObjectInstance->UpdateSkeletalMeshAsync(true, true);
-	Viewport->SetPreviewComponents(PreviewSkeletalMeshComponents);
-}
-
-
-void FCustomizableObjectInstanceEditor::UpdatePreviewVisibility()
-{
-	const bool bEnableVisibility = CustomizableObjectInstance->GetPrivate()->SkeletalMeshStatus == ESkeletalMeshStatus::Success;
-	for (UDebugSkelMeshComponent* PreviewSkeletalMeshComponent : PreviewSkeletalMeshComponents)
-	{
-		PreviewSkeletalMeshComponent->SetVisibility(bEnableVisibility, true);
-	}
 }
 
 
@@ -418,17 +392,6 @@ void FCustomizableObjectInstanceEditor::OpenParentInEditor()
 void FCustomizableObjectInstanceEditor::AddReferencedObjects( FReferenceCollector& Collector )
 {
 	Collector.AddReferencedObject( CustomizableObjectInstance );
-	
-	for (auto& PreviewCustomizableSkeletalComponent : PreviewCustomizableSkeletalComponents)
-	{
-		Collector.AddReferencedObject(PreviewCustomizableSkeletalComponent);
-	}
-
-	for (auto& PreviewSkeletalMeshComponent : PreviewSkeletalMeshComponents)
-	{
-		Collector.AddReferencedObject(PreviewSkeletalMeshComponent);
-	}
-
 	Collector.AddReferencedObject( HelperCallback );
 	Collector.AddReferencedObject( ProjectorParameter );
 	Collector.AddReferencedObject( CustomSettings );
@@ -580,7 +543,7 @@ void FCustomizableObjectInstanceEditor::SetPoseAsset(class UPoseAsset* PoseAsset
 	{
 		Viewport->SetAnimation(nullptr, EAnimationMode::AnimationBlueprint);
 
-		for (UDebugSkelMeshComponent* PreviewSkeletalMeshComponent : PreviewSkeletalMeshComponents)
+		for (TWeakObjectPtr<UDebugSkelMeshComponent> PreviewSkeletalMeshComponent : PreviewSkeletalMeshComponents)
 		{
 			PreviewSkeletalMeshComponent->SetAnimationMode(EAnimationMode::AnimationBlueprint);
 			PreviewSkeletalMeshComponent->InitAnim(false);
@@ -638,32 +601,14 @@ void FCustomizableObjectInstanceEditor::OnInstancePropertySelectionChanged(FProp
 }
 
 
-void FCustomizableObjectInstanceEditor::OnUpdatePreviewInstance()
+void FCustomizableObjectInstanceEditor::OnUpdatePreviewInstance(UCustomizableObjectInstance* Instance)
 {
 	check(CustomizableObjectInstance)
 
-	Viewport->GetViewportClient()->Invalidate();
-
-	UpdatePreviewVisibility();
-
-	bool bNeedsToResetPreviewComponents = false;
-	for (UDebugSkelMeshComponent* PreviewSkeletalMeshComponent : PreviewSkeletalMeshComponents)
-	{
-		PreviewSkeletalMeshComponent->bComponentUseFixedSkelBounds = true; // First bounds computed would be using physics asset
-		PreviewSkeletalMeshComponent->UpdateBounds();
-
-		if (!PreviewSkeletalMeshComponent->IsRegistered())
-		{
-			bNeedsToResetPreviewComponents = true;
-			PreviewSkeletalMeshComponent->RegisterComponentWithWorld(PreviewSkeletalMeshComponent->GetWorld());
-		}
-	}
-
-	if(bNeedsToResetPreviewComponents)
-	{
-		Viewport->SetPreviewComponents(ObjectPtrDecay(PreviewSkeletalMeshComponents));
-	}
-
+	const bool bVisible = CustomizableObjectInstance->GetPrivate()->SkeletalMeshStatus == ESkeletalMeshStatus::Success;
+	Actor->GetRootComponent()->SetVisibility(bVisible, true);
+	Actor->GetRootComponent()->UpdateBounds();
+	
 	// Sets which UV material to draw
 	Viewport->SetDrawDefaultUVMaterial();
 
@@ -711,7 +656,7 @@ void FCustomizableObjectInstanceEditor::OnCustomizableObjectStatusChanged(FCusto
 	{
 	case FCustomizableObjectStatus::EState::ModelLoaded:
 		{
-			CreatePreviewInstance();
+			SetupPreviewInstance();
 			break;			
 		}
 		
@@ -741,7 +686,7 @@ void FCustomizableObjectInstanceEditor::OpenTextureAnalyzerTab()
 
 void FCustomizableObjectInstanceEditor::OnPostCompile()
 {
-	CreatePreviewComponents();
+	CreatePreviewActor();
 	CustomizableObjectInstance->UpdateSkeletalMeshAsync(true, true);
 }
 
@@ -1004,7 +949,7 @@ void FCustomizableObjectInstanceEditor::HideGizmoProjectorParameter(const TShare
 }
 
 
-void FCustomizableObjectInstanceEditor::CreatePreviewComponents()
+void FCustomizableObjectInstanceEditor::CreatePreviewActor()
 {
 	const UCustomizableObject* Object = CustomizableObjectInstance->GetCustomizableObject();
 	
@@ -1012,32 +957,37 @@ void FCustomizableObjectInstanceEditor::CreatePreviewComponents()
 	{
 		return;
 	}
-	
-	for (const TObjectPtr<UCustomizableSkeletalComponent>& Component : PreviewCustomizableSkeletalComponents)
-	{
-		Component->DestroyComponent();
-	}
-	PreviewCustomizableSkeletalComponents.Empty();
 
-	for (const TObjectPtr<UDebugSkelMeshComponent>& Component : PreviewSkeletalMeshComponents)
+	if (Actor)
 	{
-		Component->DestroyComponent();
+		Viewport->GetViewportClient()->GetWorld()->RemoveActor(Actor.Get(), false);
 	}
-	PreviewSkeletalMeshComponents.Empty();
+	
+	Actor.Reset(Viewport->GetViewportClient()->GetWorld()->SpawnActor<ASkeletalMeshActor>());
 
+	PreviewSkeletalMeshComponents.Reset();
+	
 	for (int32 ComponentIndex = 0; ComponentIndex < Object->GetComponentCount(); ++ComponentIndex)
 	{
-		UCustomizableSkeletalComponent* PreviewCustomizableSkeletalComponent = NewObject<UCustomizableSkeletalComponent>(UCustomizableSkeletalComponent::StaticClass());
-		UDebugSkelMeshComponent* PreviewSkeletalMeshComponent = NewObject<UDebugSkelMeshComponent>(GetTransientPackage(), NAME_None, RF_Transient);
-
-		PreviewCustomizableSkeletalComponent->bSkipSetReferenceSkeletalMesh = true;
-		PreviewCustomizableSkeletalComponent->CustomizableObjectInstance = CustomizableObjectInstance;
-		PreviewCustomizableSkeletalComponent->ComponentIndex = ComponentIndex;
-		PreviewCustomizableSkeletalComponent->AttachToComponent(PreviewSkeletalMeshComponent, FAttachmentTransformRules::KeepRelativeTransform);
-
-		PreviewCustomizableSkeletalComponents.Add(PreviewCustomizableSkeletalComponent);
-		PreviewSkeletalMeshComponents.Add(PreviewSkeletalMeshComponent);
+		UDebugSkelMeshComponent* DebugComponent = NewObject<UDebugSkelMeshComponent>(Actor.Get(), NAME_None, RF_Transient);
+		DebugComponent->bCastInsetShadow = true; // For better quality shadows in the editor previews, more similar to the in-game ones
+		DebugComponent->bCanHighlightSelectedSections = false;
+		DebugComponent->bComponentUseFixedSkelBounds = true; // First bounds computed would be using physics asset
+		DebugComponent->MarkRenderStateDirty();
+		DebugComponent->AttachToComponent(Actor->GetRootComponent(), FAttachmentTransformRules::KeepRelativeTransform);
+		DebugComponent->RegisterComponent();
+		
+		PreviewSkeletalMeshComponents.Add(DebugComponent);
+		
+		UCustomizableSkeletalComponent* CustomizableComponent = NewObject<UCustomizableSkeletalComponent>(DebugComponent);
+		CustomizableComponent->bSkipSetReferenceSkeletalMesh = true;
+		CustomizableComponent->CustomizableObjectInstance = CustomizableObjectInstance;
+		CustomizableComponent->ComponentIndex = ComponentIndex;
+		CustomizableComponent->AttachToComponent(DebugComponent, FAttachmentTransformRules::KeepRelativeTransform);
+		CustomizableComponent->RegisterComponent();
 	}
+	
+	Viewport->SetPreviewActor(Actor.Get(), CustomizableObjectInstance, PreviewSkeletalMeshComponents); // TODO GMT
 }
 
 
diff --git a/Engine/Plugins/Experimental/Mutable/Source/CustomizableObjectEditor/Private/MuCOE/CustomizableObjectInstanceEditor.h b/Engine/Plugins/Experimental/Mutable/Source/CustomizableObjectEditor/Private/MuCOE/CustomizableObjectInstanceEditor.h
--- a/Engine/Plugins/Experimental/Mutable/Source/CustomizableObjectEditor/Private/MuCOE/CustomizableObjectInstanceEditor.h
+++ b/Engine/Plugins/Experimental/Mutable/Source/CustomizableObjectEditor/Private/MuCOE/CustomizableObjectInstanceEditor.h
@@ -179,7 +179,7 @@ public:
 	virtual void HideGizmoProjectorParameter() override;
 	
 	/** Callback to notify the editor when the PreviewInstance has been updated */
-	void OnUpdatePreviewInstance();
+	void OnUpdatePreviewInstance(UCustomizableObjectInstance* Instance);
 
 	/** FTickableGameObject interface */
 	virtual bool IsTickable(void) const override;
@@ -195,7 +195,7 @@ private:
 	TSharedRef<SDockTab> SpawnTab_AdvancedPreviewSettings(const FSpawnTabArgs& Args);
 	TSharedRef<SDockTab> SpawnTab_TextureAnalyzer(const FSpawnTabArgs& Args);
 
-	void CreatePreviewInstance();
+	void SetupPreviewInstance();
 
 	/** Binds commands associated with the Static Mesh Editor. */
 	void BindCommands();
@@ -203,9 +203,6 @@ private:
 	/** Callback when selection changes in the Property Tree. */
 	void OnInstancePropertySelectionChanged(FProperty* InProperty);
 
-	/** Updates the visibility of PreviewSkeletalMeshComponent */
-	void UpdatePreviewVisibility();
-
 	/** Save Customizable Object Instance open in the editor */
 	void SaveAsset_Execute() override;
 
@@ -239,15 +236,17 @@ public:
 		const TSharedPtr<IDetailsView>& InstanceDetailsView);
 	
 private:
-	void CreatePreviewComponents();
+	void CreatePreviewActor();
 	
 	void OnPostCompile();
 	
 	/** The currently viewed object. */
 	TObjectPtr<UCustomizableObjectInstance> CustomizableObjectInstance;
-	TArray<TObjectPtr<UCustomizableSkeletalComponent>> PreviewCustomizableSkeletalComponents;
-	TObjectPtr<UStaticMeshComponent> PreviewStaticMeshComponent;
-	TArray<TObjectPtr<UDebugSkelMeshComponent>> PreviewSkeletalMeshComponents;
+
+	/** Preview Actor. All preview components are attached to this actor. */
+	TStrongObjectPtr<AActor> Actor;
+	
+	TArray<TWeakObjectPtr<UDebugSkelMeshComponent>> PreviewSkeletalMeshComponents;
 
 	/** List of open tool panels; used to ensure only one exists at any one time */
 	TMap<FName, TWeakPtr<SDockableTab>> SpawnedToolPanels;
diff --git a/Engine/Plugins/Experimental/Mutable/Source/CustomizableObjectEditor/Private/MuCOE/SCustomizableObjectEditorViewport.cpp b/Engine/Plugins/Experimental/Mutable/Source/CustomizableObjectEditor/Private/MuCOE/SCustomizableObjectEditorViewport.cpp
--- a/Engine/Plugins/Experimental/Mutable/Source/CustomizableObjectEditor/Private/MuCOE/SCustomizableObjectEditorViewport.cpp
+++ b/Engine/Plugins/Experimental/Mutable/Source/CustomizableObjectEditor/Private/MuCOE/SCustomizableObjectEditorViewport.cpp
@@ -260,26 +260,6 @@ void SCustomizableObjectEditorViewportTabBody::Construct(const FArguments& InArg
 }
 
 
-SCustomizableObjectEditorViewportTabBody::~SCustomizableObjectEditorViewportTabBody()
-{
-	if (LevelViewportClient.IsValid())
-	{
-		LevelViewportClient->Viewport = nullptr;
-	}
-
-	PreviewSkeletalMeshComponents.Reset();
-}
-
-
-void SCustomizableObjectEditorViewportTabBody::AddReferencedObjects( FReferenceCollector& Collector )
-{
-	for (auto& PreviewSkeletalMeshComponent : PreviewSkeletalMeshComponents)
-	{
-		Collector.AddReferencedObject(PreviewSkeletalMeshComponent);
-	}
-}
-
-
 void SCustomizableObjectEditorViewportTabBody::Tick(const FGeometry& AllottedGeometry, const double InCurrentTime, const float InDeltaTime)
 {
 	SCompoundWidget::Tick(AllottedGeometry, InCurrentTime, InDeltaTime);
@@ -292,7 +272,7 @@ void SCustomizableObjectEditorViewportTabBody::Tick(const FGeometry& AllottedGeo
 	{
 		MaterialNames.Empty();
 
-		for (UDebugSkelMeshComponent* PreviewSkeletalMeshComponent : PreviewSkeletalMeshComponents)
+		for (TWeakObjectPtr<UDebugSkelMeshComponent> PreviewSkeletalMeshComponent : PreviewSkeletalMeshComponents)
 		{
 			const TArray<UMaterialInterface*> Materials = PreviewSkeletalMeshComponent->GetMaterials();
 			for (UMaterialInterface* m : Materials)
@@ -416,35 +396,10 @@ void SCustomizableObjectEditorViewportTabBody::HideGizmoLight() const
 }
 
 
-void SCustomizableObjectEditorViewportTabBody::SetPreviewComponents(const TArray<UDebugSkelMeshComponent*>& InSkeletalMeshComponents)
+void SCustomizableObjectEditorViewportTabBody::SetPreviewActor(const TWeakObjectPtr<AActor>& InActor, const TWeakObjectPtr<UCustomizableObjectInstance>& InInstance, const TArray<TWeakObjectPtr<UDebugSkelMeshComponent>>& InSkeletalMeshComponents)
 {
-	FTransform Transform = FTransform::Identity;
-
-	for (UDebugSkelMeshComponent* PreviewSkeletalMeshComponent : PreviewSkeletalMeshComponents)
-	{
-		if (PreviewSkeletalMeshComponent)
-		{
-			Transform = PreviewSkeletalMeshComponent->GetComponentToWorld();
-			PreviewScenePtr->RemoveComponent(PreviewSkeletalMeshComponent);
-		}
-	}
-
-	PreviewSkeletalMeshComponents = ObjectPtrWrap(InSkeletalMeshComponents);
-
-	for (UDebugSkelMeshComponent* PreviewSkeletalMeshComponent : PreviewSkeletalMeshComponents)
-	{
-		// To support null components We check if the skeletal mesh is null in the viewport client
-		if (PreviewSkeletalMeshComponent)
-		{
-			PreviewSkeletalMeshComponent->bCastInsetShadow = true; // For better quality shadows in the editor previews, more similar to the in-game ones
-			PreviewSkeletalMeshComponent->bCanHighlightSelectedSections = false;
-			PreviewSkeletalMeshComponent->MarkRenderStateDirty();
-
-			PreviewScenePtr->AddComponent(PreviewSkeletalMeshComponent, Transform);
-		}
-	}
-
-	LevelViewportClient->SetPreviewComponents(ObjectPtrDecay(PreviewSkeletalMeshComponents));
+	PreviewSkeletalMeshComponents = InSkeletalMeshComponents;
+	LevelViewportClient->SetPreviewActor(InActor, InInstance, InSkeletalMeshComponents);
 }
 
 
@@ -454,12 +409,6 @@ bool SCustomizableObjectEditorViewportTabBody::IsVisible() const
 }
 
 
-const TArray<UDebugSkelMeshComponent*>& SCustomizableObjectEditorViewportTabBody::GetSkeletalMeshComponents() const
-{
-	return ObjectPtrDecay(PreviewSkeletalMeshComponents);
-}
-
-
 void SCustomizableObjectEditorViewportTabBody::BindCommands()
 {
 	FUICommandList& CommandList = *UICommandList;
@@ -636,7 +585,7 @@ void SCustomizableObjectEditorViewportTabBody::BindCommands()
 
 void SCustomizableObjectEditorViewportTabBody::OnTakeHighResScreenshot()
 {
-	CustomizableObjectHighresScreenshot = SCustomizableObjectHighresScreenshot::OpenDialog(ViewportWidget->GetSceneViewport(), LevelViewportClient, PreviewSkeletalMeshComponents[0], PreviewScenePtr);
+	CustomizableObjectHighresScreenshot = SCustomizableObjectHighresScreenshot::OpenDialog(ViewportWidget->GetSceneViewport(), LevelViewportClient, PreviewSkeletalMeshComponents[0].Get(), PreviewScenePtr);
 }
 
 
@@ -795,9 +744,9 @@ void SCustomizableObjectEditorViewportTabBody::GenerateUVSectionOptions()
 	
 	for (int32 ComponentIndex = 0; ComponentIndex < PreviewSkeletalMeshComponents.Num(); ++ComponentIndex)
 	{
-		const UDebugSkelMeshComponent* PreviewSkeletalMeshComponent = PreviewSkeletalMeshComponents[ComponentIndex];
+		const TWeakObjectPtr<UDebugSkelMeshComponent>& PreviewSkeletalMeshComponent = PreviewSkeletalMeshComponents[ComponentIndex];
 		
-		if (PreviewSkeletalMeshComponent != nullptr && UE_MUTABLE_GETSKINNEDASSET(PreviewSkeletalMeshComponent) != nullptr && UE_MUTABLE_GETSKINNEDASSET(PreviewSkeletalMeshComponent)->GetResourceForRendering() != nullptr)
+		if (PreviewSkeletalMeshComponent.IsValid() && UE_MUTABLE_GETSKINNEDASSET(PreviewSkeletalMeshComponent) != nullptr && UE_MUTABLE_GETSKINNEDASSET(PreviewSkeletalMeshComponent)->GetResourceForRendering() != nullptr)
 		{
 			const TArray<UMaterialInterface*> Materials = PreviewSkeletalMeshComponent->GetMaterials();
 			const FSkeletalMeshRenderData* MeshRes = UE_MUTABLE_GETSKINNEDASSET(PreviewSkeletalMeshComponent)->GetResourceForRendering();
@@ -903,9 +852,9 @@ void SCustomizableObjectEditorViewportTabBody::GenerateUVChannelOptions()
 	check(Index != INDEX_NONE);
 	const FSection& Section = UVSectionOption[Index];
 
-	const UDebugSkelMeshComponent* PreviewSkeletalMeshComponent = PreviewSkeletalMeshComponents[Section.ComponentIndex];
+	const TWeakObjectPtr<UDebugSkelMeshComponent>& PreviewSkeletalMeshComponent = PreviewSkeletalMeshComponents[Section.ComponentIndex];
 
-	if (PreviewSkeletalMeshComponent != nullptr && UE_MUTABLE_GETSKINNEDASSET(PreviewSkeletalMeshComponent) != nullptr
+	if (PreviewSkeletalMeshComponent.IsValid() && UE_MUTABLE_GETSKINNEDASSET(PreviewSkeletalMeshComponent) != nullptr
 		&& UE_MUTABLE_GETSKINNEDASSET(PreviewSkeletalMeshComponent)->GetResourceForRendering() != nullptr)
 	{
 		const FSkeletalMeshRenderData* MeshRes = UE_MUTABLE_GETSKINNEDASSET(PreviewSkeletalMeshComponent)->GetResourceForRendering();
@@ -988,9 +937,9 @@ int32 SCustomizableObjectEditorViewportTabBody::GetLODModelCount() const
 {
 	int32 LODModelCount = 0;
 
-	for (UDebugSkelMeshComponent* PreviewComponent : GetSkeletalMeshComponents())
+	for (const TWeakObjectPtr<UDebugSkelMeshComponent>& PreviewComponent : PreviewSkeletalMeshComponents)
 	{
-		if (PreviewComponent && UE_MUTABLE_GETSKINNEDASSET(PreviewComponent))
+		if (PreviewComponent.IsValid() && UE_MUTABLE_GETSKINNEDASSET(PreviewComponent))
 		{
 			const TIndirectArray<FSkeletalMeshLODRenderData>& LODModels = UE_MUTABLE_GETSKINNEDASSET(PreviewComponent)->GetResourceForRendering()->LODRenderData;
 			LODModelCount = FMath::Max(LODModelCount, LODModels.Num());
@@ -1047,9 +996,9 @@ void SCustomizableObjectEditorViewportTabBody::OnSetLODModel(int32 LODSelectionT
 {
 	LODSelection = LODSelectionType;
 
-	for (UDebugSkelMeshComponent* PreviewComponent : GetSkeletalMeshComponents())
+	for (const TWeakObjectPtr<UDebugSkelMeshComponent>& PreviewComponent : PreviewSkeletalMeshComponents)
 	{
-		if (PreviewComponent)
+		if (PreviewComponent.IsValid())
 		{
 			PreviewComponent->SetForcedLOD(LODSelection);
 		}
diff --git a/Engine/Plugins/Experimental/Mutable/Source/CustomizableObjectEditor/Private/MuCOE/SCustomizableObjectEditorViewport.h b/Engine/Plugins/Experimental/Mutable/Source/CustomizableObjectEditor/Private/MuCOE/SCustomizableObjectEditorViewport.h
--- a/Engine/Plugins/Experimental/Mutable/Source/CustomizableObjectEditor/Private/MuCOE/SCustomizableObjectEditorViewport.h
+++ b/Engine/Plugins/Experimental/Mutable/Source/CustomizableObjectEditor/Private/MuCOE/SCustomizableObjectEditorViewport.h
@@ -83,7 +83,7 @@ private:
 /**
  * CustomizableObject Editor Preview viewport widget with toolbars, etc.
  */
-class SCustomizableObjectEditorViewportTabBody : public SCompoundWidget, public FGCObject
+class SCustomizableObjectEditorViewportTabBody : public SCompoundWidget
 {
 public:
 	SLATE_BEGIN_ARGS(SCustomizableObjectEditorViewportTabBody){}
@@ -91,20 +91,11 @@ public:
 	SLATE_END_ARGS()
 
 	void Construct(const FArguments& InArgs);
-	virtual ~SCustomizableObjectEditorViewportTabBody() override;
 
 	void SetAnimation(class UAnimationAsset* Animation, EAnimationMode::Type AnimationType);
 
 	virtual FReply OnDrop(const FGeometry& MyGeometry, const FDragDropEvent& DragDropEvent) override;
 	
-	// FSerializableObject interface
-	virtual void AddReferencedObjects( FReferenceCollector& Collector ) override;
-	virtual FString GetReferencerName() const override
-	{
-		return TEXT("SCustomizableObjectEditorViewportTabBody");
-	}
-	// End of FSerializableObject interface
-
 	// SWidget interface
 	virtual void Tick(const FGeometry& AllottedGeometry, const double InCurrentTime, const float InDeltaTime) override;
 	// End of SWidget interface
@@ -154,11 +145,8 @@ public:
 
 	/** Callback to show / hide the state and runtime parameter test in the viewport */
 	TSharedRef<SWidget> ShowStateTestData();
-
-	/** Retrieves the skeletal mesh component. */
-	const TArray<UDebugSkelMeshComponent*>& GetSkeletalMeshComponents() const;
 	
-	void SetPreviewComponents(const TArray<UDebugSkelMeshComponent*>& InSkeletalMeshComponents);
+	void SetPreviewActor(const TWeakObjectPtr<AActor>& InActor, const TWeakObjectPtr<UCustomizableObjectInstance>& InInstance, const TArray<TWeakObjectPtr<UDebugSkelMeshComponent>>& InSkeletalMeshComponents);
 	
 	/** Function to get the number of LOD models associated with the preview skeletal mesh*/
 	int32 GetLODModelCount() const;
@@ -231,7 +219,7 @@ private:
 	bool IsVisible() const;
 	
 	// Components for the preview mesh.
-	TArray<TObjectPtr<UDebugSkelMeshComponent>> PreviewSkeletalMeshComponents;
+	TArray<TWeakObjectPtr<UDebugSkelMeshComponent>> PreviewSkeletalMeshComponents;
 
 	// The scene for this viewport.
 	TSharedPtr<FCustomizableObjectPreviewScene> PreviewScenePtr;
diff --git a/Engine/Plugins/Experimental/Mutable/Source/CustomizableObjectEditor/Private/MuCOE/SCustomizableObjectEditorViewportToolBar.cpp b/Engine/Plugins/Experimental/Mutable/Source/CustomizableObjectEditor/Private/MuCOE/SCustomizableObjectEditorViewportToolBar.cpp
--- a/Engine/Plugins/Experimental/Mutable/Source/CustomizableObjectEditor/Private/MuCOE/SCustomizableObjectEditorViewportToolBar.cpp
+++ b/Engine/Plugins/Experimental/Mutable/Source/CustomizableObjectEditor/Private/MuCOE/SCustomizableObjectEditorViewportToolBar.cpp
@@ -479,22 +479,6 @@ void SCustomizableObjectEditorViewportToolBar::OnFOVValueChanged(float NewValue)
 }
 
 
-TOptional<float> SCustomizableObjectEditorViewportToolBar::OnGetFloorOffset() const
-{
-	TSharedPtr<FCustomizableObjectEditorViewportClient> ViewportClient = Viewport.Pin()->GetViewportClient();
-
-	return ViewportClient->GetFloorOffset();
-}
-
-
-void SCustomizableObjectEditorViewportToolBar::OnFloorOffsetChanged(float NewValue)
-{
-	TSharedPtr<FCustomizableObjectEditorViewportClient> ViewportClient = Viewport.Pin()->GetViewportClient();
-
-	ViewportClient->SetFloorOffset(NewValue);
-}
-
-
 ECheckBoxState SCustomizableObjectEditorViewportToolBar::IsRotationGridSnapChecked() const
 {
 	return GetDefault<ULevelEditorViewportSettings>()->RotGridEnabled ? ECheckBoxState::Checked : ECheckBoxState::Unchecked;
diff --git a/Engine/Plugins/Experimental/Mutable/Source/CustomizableObjectEditor/Public/MuCOE/CustomizableObjectEditorViewportClient.h b/Engine/Plugins/Experimental/Mutable/Source/CustomizableObjectEditor/Public/MuCOE/CustomizableObjectEditorViewportClient.h
--- a/Engine/Plugins/Experimental/Mutable/Source/CustomizableObjectEditor/Public/MuCOE/CustomizableObjectEditorViewportClient.h
+++ b/Engine/Plugins/Experimental/Mutable/Source/CustomizableObjectEditor/Public/MuCOE/CustomizableObjectEditorViewportClient.h
@@ -75,7 +75,6 @@ public:
 	class UPersonaOptions* ConfigOption;
 
 	// FEditorViewportClient interface
-	virtual void Tick(float DeltaSeconds) override;
 	virtual void Draw(const FSceneView* View,FPrimitiveDrawInterface* PDI) override;
 	virtual void Draw(FViewport* Viewport,FCanvas* Canvas) override;
 
@@ -92,17 +91,9 @@ public:
 	virtual void SetWidgetCoordSystemSpace(ECoordSystem NewCoordSystem) override;
 	virtual void SetViewportType(ELevelViewportType InViewportType) override;
 	// End of FEditorViewportClient
-
-	void UpdateCameraSetup();
-	void UpdateFloor();
-
-	/** */
-	void SetPreviewComponent(UStaticMeshComponent* InPreviewStaticMeshComponent);
-	void SetPreviewComponents(const TArray<UDebugSkelMeshComponent*>& InPreviewSkeletalMeshComponents);
-
-	/** Sets an informative message in the viewport warning the user that the CustomizableObject has no reference mesh */
-	void SetReferenceMeshMissingWarningMessage(bool bVisible);
-
+	
+	void SetPreviewActor(const TWeakObjectPtr<AActor>& InActor, const TWeakObjectPtr<UCustomizableObjectInstance>& InInstance, const TArray<TWeakObjectPtr<UDebugSkelMeshComponent>>& InSkeletalMeshComponents);
+	
 	/**
 	 *	Draws the UV overlay for the current LOD.
 	 *
@@ -165,19 +156,7 @@ public:
 
 	/** Callback for toggling the bounds show flag. */
 	void SetShowBounds();
-		
-	/** Returns the desired target of the camera */
-	FSphere GetCameraTarget();
-
-	/** Point the camera to the Skeletal Mesh Components. */
-	void ResetCamera();
-	
-	/* Returns the floor height offset */	
-	float GetFloorOffset() const;
 
-	/* Sets the floor height offset, saves it to config and invalidates the viewport so it shows up immediately */
-	void SetFloorOffset(float NewValue);
-	
 	/** Do not call directly. Use ICustomizableObjectEditor functions instead. */
 	void ShowGizmoClipMorph(UCustomizableObjectNodeMeshClipMorph& ClipPlainNode);
 
@@ -287,9 +266,10 @@ private:
 
 	bool HandleBeginTransform();
 	bool HandleEndTransform();
-	
-	/** Component for the static/skeletal mesh. */
-	TWeakObjectPtr<UStaticMeshComponent> StaticMeshComponent;
+
+	void OnInstanceUpdate(UCustomizableObjectInstance* Instance);
+
+	TWeakObjectPtr<AActor> Actor;
 	TArray<TWeakObjectPtr<UDebugSkelMeshComponent>> SkeletalMeshComponents;
 
 	/** True if the widget is being dragged. */
@@ -307,9 +287,7 @@ private:
 	int32 UVDrawSectionIndex = 0;
 	int32 UVDrawLODIndex = 0;
 	int32 UVDrawUVIndex = 0;
-
-	bool bReferenceMeshMissingWarningMessageVisible;
-
+	
 	UCustomizableObjectNodeMeshClipMorph* ClipMorphNode;
 	TObjectPtr<UMaterial> ClipMorphMaterial;
 	bool bClipMorphLocalStartOffset;
@@ -327,8 +305,6 @@ private:
 	float Radius2;
 	float RotationAngle;
 
-	FSphere BoundSphere;
-
 	/** Light being edited */
 	ULightComponent* SelectedLightComponent;
 
@@ -338,6 +314,9 @@ private:
 	/** To know if an animation ia being played by the Customizable Object and restre it after compilation */
 	bool IsPlayingAnimation;
 
+	/** true if the camera has already being setup. */
+	bool bIsCameraSetup = false;
+
 	/** Animation being played by the Customizable Object, if any */
 	TObjectPtr<UAnimationAsset> AnimationBeingPlayed;
 
@@ -346,10 +325,7 @@ private:
 
 	/** Flag to control whether to show / hide the instance geometry information data */
 	bool StateChangeShowGeometryDataFlag;
-
-	/** To know if the orbital camera is being used or not */
-	bool bActivateOrbitalCamera;
-
+	
 	/** Material for cylinder arc solid render */
 	TObjectPtr<UMaterialInterface> TransparentPlaneMaterialXY;
 
diff --git a/Engine/Plugins/Experimental/Mutable/Source/CustomizableObjectEditor/Public/MuCOE/SCustomizableObjectEditorViewportToolBar.h b/Engine/Plugins/Experimental/Mutable/Source/CustomizableObjectEditor/Public/MuCOE/SCustomizableObjectEditorViewportToolBar.h
--- a/Engine/Plugins/Experimental/Mutable/Source/CustomizableObjectEditor/Public/MuCOE/SCustomizableObjectEditorViewportToolBar.h
+++ b/Engine/Plugins/Experimental/Mutable/Source/CustomizableObjectEditor/Public/MuCOE/SCustomizableObjectEditorViewportToolBar.h
@@ -93,11 +93,6 @@ private:
 	float OnGetFOVValue() const;
 	/** Called when the FOV slider is adjusted in the perspective viewport */
 	void OnFOVValueChanged(float NewValue) const;
-
-	/** Called by the floor offset slider in the perspective viewport to get the offset value */
-	TOptional<float> OnGetFloorOffset() const;
-	/** Called when the floor offset slider is adjusted in the perspective viewport */
-	void OnFloorOffsetChanged(float NewValue);
 	
 	ECheckBoxState IsRotationGridSnapChecked() const;
 	void HandleToggleRotationGridSnap(ECheckBoxState InState);
