[mutable] Fixed UV editor scaling.
+ Now the UV previewer scales correctly, even if OS scale is different than 100%.

#rnx

diff --git a/Engine/Plugins/Experimental/Mutable/Source/CustomizableObjectEditor/Private/MuCOE/CustomizableObjectEditorViewportClient.cpp b/Engine/Plugins/Experimental/Mutable/Source/CustomizableObjectEditor/Private/MuCOE/CustomizableObjectEditorViewportClient.cpp
--- a/Engine/Plugins/Experimental/Mutable/Source/CustomizableObjectEditor/Private/MuCOE/CustomizableObjectEditorViewportClient.cpp
+++ b/Engine/Plugins/Experimental/Mutable/Source/CustomizableObjectEditor/Private/MuCOE/CustomizableObjectEditorViewportClient.cpp
@@ -410,7 +410,8 @@ void FCustomizableObjectEditorViewportClient::DrawUVs(FViewport* InViewport, FCa
 	const uint32 MinX = BorderWidth;
 	const FVector2D UVBoxOrigin(MinX, MinY);
 	const FVector2D BoxOrigin( MinX - 1, MinY - 1 );
-	const uint32 UVBoxScale = FMath::Min(InViewport->GetSizeXY().X - MinX, InViewport->GetSizeXY().Y - MinY) - BorderWidth;
+	const FVector2D ViewportSize = FVector2D(InViewport->GetSizeXY())/InCanvas->GetDPIScale(); // Remove Window (OS) scale.
+	const uint32 UVBoxScale = FMath::Min(ViewportSize.X - MinX, ViewportSize.Y - MinY) - BorderWidth;
 	const uint32 BoxSize = UVBoxScale + 2;
 	const FVector2D Box[ 4 ] = {
 		BoxOrigin,									// topleft
