[mutable] Improved pass-through meshes at the UE level. Support for static/other meshes and async loading of mesh references.
#tests PIE

#rnx

diff --git a/Engine/Plugins/Experimental/Mutable/Source/CustomizableObject/Internal/MuCO/CustomizableObjectInstancePrivate.h b/Engine/Plugins/Experimental/Mutable/Source/CustomizableObject/Internal/MuCO/CustomizableObjectInstancePrivate.h
--- a/Engine/Plugins/Experimental/Mutable/Source/CustomizableObject/Internal/MuCO/CustomizableObjectInstancePrivate.h
+++ b/Engine/Plugins/Experimental/Mutable/Source/CustomizableObject/Internal/MuCO/CustomizableObjectInstancePrivate.h
@@ -352,13 +352,18 @@ public:
 
 	// The pass-through assets that will be loaded during an update
 	TArray<TSoftObjectPtr<UTexture>> PassThroughTexturesToLoad;
-	TArray<TSoftObjectPtr<USkeletalMesh>> PassThroughMeshesToLoad;
+	TArray<TSoftObjectPtr<UStreamableRenderAsset>> PassThroughMeshesToLoad;
 
 	// Used during an update to prevent the pass-through textures loaded by LoadAdditionalAssetsAsync() from being unloaded by GC
 	// between AdditionalAssetsAsyncLoaded() and their setting into the generated materials in BuildMaterials()
 	UPROPERTY(Transient)
 	TArray<TObjectPtr<UTexture>> LoadedPassThroughTexturesPendingSetMaterial;
 
+	// Used during an update to prevent the pass-through meshes loaded by LoadAdditionalAssetsAsync() from being unloaded by GC
+	// between AdditionalAssetsAsyncLoaded() and their setting into the generated materials in BuildMaterials()
+	UPROPERTY(Transient)
+	TArray<TObjectPtr<UStreamableRenderAsset>> LoadedPassThroughMeshesPendingSetMaterial;
+
 private:
 	ECOInstanceFlags InstanceFlagsPrivate = ECOInstanceFlags::ECONone;
 
diff --git a/Engine/Plugins/Experimental/Mutable/Source/CustomizableObject/Private/MuCO/BusyWaits_Deprecated.cpp b/Engine/Plugins/Experimental/Mutable/Source/CustomizableObject/Private/MuCO/BusyWaits_Deprecated.cpp
--- a/Engine/Plugins/Experimental/Mutable/Source/CustomizableObject/Private/MuCO/BusyWaits_Deprecated.cpp
+++ b/Engine/Plugins/Experimental/Mutable/Source/CustomizableObject/Private/MuCO/BusyWaits_Deprecated.cpp
@@ -176,6 +176,7 @@ namespace CustomizableObjectSystem::ImplDeprecated
 		UCustomizableInstancePrivate* CustomizableObjectInstancePrivateData = OperationData->Instance->GetPrivate();
 
 		CustomizableObjectInstancePrivateData->PassThroughTexturesToLoad.Empty();
+		CustomizableObjectInstancePrivateData->PassThroughMeshesToLoad.Empty();
 
 		if (OperationData->PixelFormatOverride)
 		{
diff --git a/Engine/Plugins/Experimental/Mutable/Source/CustomizableObject/Private/MuCO/CustomizableObjectInstance.cpp b/Engine/Plugins/Experimental/Mutable/Source/CustomizableObject/Private/MuCO/CustomizableObjectInstance.cpp
--- a/Engine/Plugins/Experimental/Mutable/Source/CustomizableObject/Private/MuCO/CustomizableObjectInstance.cpp
+++ b/Engine/Plugins/Experimental/Mutable/Source/CustomizableObject/Private/MuCO/CustomizableObjectInstance.cpp
@@ -2075,11 +2075,19 @@ bool UCustomizableInstancePrivate::UpdateSkeletalMesh_PostBeginUpdate0(UCustomiz
 
 		if (Component.Mesh && Component.Mesh->IsReference())
 		{
-			// TODO: This shouldn't happen here synchronosuly. It should have been requested as an async load.
 			int32 ReferenceID = Component.Mesh->GetReferencedMesh();
 			TSoftObjectPtr<USkeletalMesh> Ref = ModelResources.PassThroughMeshes[ReferenceID];
 
-			SkeletalMeshes[Component.Id] = Ref.LoadSynchronous();
+			if (Ref.IsValid())
+			{
+				SkeletalMeshes[Component.Id] = Ref.Get();
+			}
+			else
+			{
+				// This shouldn't happen here synchronosuly. It should have been requested as an async load.
+				SkeletalMeshes[Component.Id] = Ref.LoadSynchronous();
+			}
+
 			break;
 		}		
 
@@ -5364,6 +5372,11 @@ UE::Tasks::FTask UCustomizableInstancePrivate::LoadAdditionalAssetsAndData(
 		AssetsToStream.Add(TextureRef.ToSoftObjectPath());
 	}
 
+	for (TSoftObjectPtr<UStreamableRenderAsset>& MeshRef : PassThroughMeshesToLoad)
+	{
+		AssetsToStream.Add(MeshRef.ToSoftObjectPath());
+	}
+
 	TArray<UE::Tasks::FTaskEvent> StreamingCompletionEvents;
 	if (AssetsToStream.Num() > 0)
 	{	
@@ -5807,6 +5820,16 @@ void UCustomizableInstancePrivate::AdditionalAssetsAsyncLoaded(UCustomizableObje
 
 	PassThroughTexturesToLoad.Empty();
 
+	LoadedPassThroughMeshesPendingSetMaterial.Empty(PassThroughMeshesToLoad.Num());
+
+	for (TSoftObjectPtr<UStreamableRenderAsset>& MeshRef : PassThroughMeshesToLoad)
+	{
+		ensure(MeshRef.IsValid());
+		LoadedPassThroughMeshesPendingSetMaterial.Add(MeshRef.Get());
+	}
+
+	PassThroughMeshesToLoad.Empty();
+
 	// Only Unload in cooked builds. Unloading them when in the editor will trigger an assert. 
 	if (FPlatformProperties::RequiresCookedData())
 	{
@@ -5966,8 +5989,11 @@ void UCustomizableInstancePrivate::BuildMaterials(const TSharedRef<FUpdateContex
 
 		const bool bReuseMaterials = !OperationData->MeshChanged[ComponentIndex];
 
+		// If the mesh is not transient, it means it's pass-through so it should use material overrides and not be modified in any way
+		const bool bIsTransientMesh = static_cast<bool>(SkeletalMesh->HasAllFlags(EObjectFlags::RF_Transient));
+
 		// It is not safe to replace the materials of a SkeletalMesh whose resources are initialized. Use overrides instead.
-		const bool bUseOverrideMaterialsOnly = OperationData->bUseMeshCache && SkeletalMesh->GetResourceForRendering()->IsInitialized();
+		const bool bUseOverrideMaterialsOnly = !bIsTransientMesh || (OperationData->bUseMeshCache && SkeletalMesh->GetResourceForRendering()->IsInitialized());
 
 		ComponentsData[ComponentIndex].OverrideMaterials.Reset();
 
@@ -5992,8 +6018,10 @@ void UCustomizableInstancePrivate::BuildMaterials(const TSharedRef<FUpdateContex
 				continue;
 			}
 
-			TArray<int32>& LODMaterialMap = SkeletalMesh->GetLODInfo(LODIndex)->LODMaterialMap;
-			LODMaterialMap.Reset();
+			if (!bUseOverrideMaterialsOnly && LODIndex < SkeletalMesh->GetLODNum())
+			{
+				SkeletalMesh->GetLODInfo(LODIndex)->LODMaterialMap.Reset();
+			}
 
 			const FMutableRefSkeletalMeshData& RefSkeletalMeshData = ModelResources.ReferenceSkeletalMeshesData[Component.Id];
 
@@ -6004,8 +6032,12 @@ void UCustomizableInstancePrivate::BuildMaterials(const TSharedRef<FUpdateContex
 				// Reuse MaterialSlot from the previous LOD.
 				if (const int32 MaterialIndex = SurfaceIdToMaterialIndex.Find(Surface.SurfaceId); MaterialIndex != INDEX_NONE)
 				{
-					const int32 LODMaterialIndex = LODMaterialMap.Add(MaterialIndex);
-					SkeletalMesh->GetResourceForRendering()->LODRenderData[LODIndex].RenderSections[SurfaceIndex].MaterialIndex = LODMaterialIndex;
+					if (!bUseOverrideMaterialsOnly)
+					{
+						const int32 LODMaterialIndex = SkeletalMesh->GetLODInfo(LODIndex)->LODMaterialMap.Add(MaterialIndex);
+						SkeletalMesh->GetResourceForRendering()->LODRenderData[LODIndex].RenderSections[SurfaceIndex].MaterialIndex = LODMaterialIndex;
+					}
+
 					continue;
 				}
 
@@ -6034,8 +6066,11 @@ void UCustomizableInstancePrivate::BuildMaterials(const TSharedRef<FUpdateContex
 
 				SetMeshUVChannelDensity(MaterialSlot.UVChannelData, RefSkeletalMeshData.Settings.DefaultUVChannelDensity);
 
-				const int32 LODMaterialIndex = LODMaterialMap.Add(MaterialSlotIndex);
-				SkeletalMesh->GetResourceForRendering()->LODRenderData[LODIndex].RenderSections[SurfaceIndex].MaterialIndex = LODMaterialIndex;
+				if (!bUseOverrideMaterialsOnly)
+				{
+					const int32 LODMaterialIndex = SkeletalMesh->GetLODInfo(LODIndex)->LODMaterialMap.Add(MaterialSlotIndex);
+					SkeletalMesh->GetResourceForRendering()->LODRenderData[LODIndex].RenderSections[SurfaceIndex].MaterialIndex = LODMaterialIndex;
+				}
 
 				FMutableMaterialPlaceholder MutableMaterialPlaceholder;
 				MutableMaterialPlaceholder.ParentMaterialID = MaterialTemplate->GetUniqueID();
@@ -6488,6 +6523,7 @@ void UCustomizableInstancePrivate::BuildMaterials(const TSharedRef<FUpdateContex
 			}
 		}
 
+		if(!bUseOverrideMaterialsOnly)
 		{
 			// Copy data from the FirstGeneratedLOD into the LODs below.
 			for (int32 LODIndex = OperationData->FirstLODAvailable; LODIndex < FirstGeneratedLOD; ++LODIndex)
@@ -6508,10 +6544,7 @@ void UCustomizableInstancePrivate::BuildMaterials(const TSharedRef<FUpdateContex
 					}
 				}
 			}
-		}
 
-		if (!bUseOverrideMaterialsOnly)
-		{
 			// Force recreate render state after replacing the materials to avoid a crash in the render pipeline if the old materials are GCed while in use.
 			RecreateRenderStateOnComponent[ComponentIndex] = SkeletalMesh->GetResourceForRendering()->IsInitialized() && SkeletalMesh->GetMaterials() != Materials;
 
@@ -6592,8 +6625,9 @@ void UCustomizableInstancePrivate::BuildMaterials(const TSharedRef<FUpdateContex
 
 		Exchange(GeneratedTextures, NewGeneratedTextures);
 
-		// All pass-through textures have been set, no need to keep referencing them from the instance
+		// All pass-through textures and meshes have been set, no need to keep referencing them from the instance
 		LoadedPassThroughTexturesPendingSetMaterial.Empty();
+		LoadedPassThroughMeshesPendingSetMaterial.Empty();
 	}
 }
 
@@ -7144,6 +7178,14 @@ void UCustomizableInstancePrivate::RegenerateImportedModels()
 			continue;
 		}
 
+		const bool bIsTransientMesh = static_cast<bool>(SkeletalMesh->HasAllFlags(EObjectFlags::RF_Transient));
+
+		if (!bIsTransientMesh)
+		{
+			// This must be a pass-through referenced mesh so don't do anything to it
+			continue;
+		}
+
 		FSkeletalMeshRenderData* RenderData = SkeletalMesh->GetResourceForRendering();
 		if (!RenderData || RenderData->IsInitialized())
 		{
diff --git a/Engine/Plugins/Experimental/Mutable/Source/CustomizableObject/Private/MuCO/CustomizableObjectSystem.cpp b/Engine/Plugins/Experimental/Mutable/Source/CustomizableObject/Private/MuCO/CustomizableObjectSystem.cpp
--- a/Engine/Plugins/Experimental/Mutable/Source/CustomizableObject/Private/MuCO/CustomizableObjectSystem.cpp
+++ b/Engine/Plugins/Experimental/Mutable/Source/CustomizableObject/Private/MuCO/CustomizableObjectSystem.cpp
@@ -2291,6 +2291,26 @@ namespace impl
 		FInstanceUpdateData::FComponent& Component = OperationData->InstanceUpdateData.Components[ComponentIndex];
 
 		Component.Mesh = GetMeshTask.GetResult();
+
+		if (Component.Mesh->IsReference())
+		{
+			const UCustomizableObjectInstance* Instance = OperationData->Instance.Get();
+			const UCustomizableObject* CustomizableObject = Instance->GetCustomizableObject();
+			const FModelResources& ModelResources = CustomizableObject->GetPrivate()->GetModelResources();
+
+			uint32 ReferenceID = Component.Mesh->GetReferencedMesh();
+
+			if (ModelResources.PassThroughMeshes.IsValidIndex(ReferenceID))
+			{
+				TSoftObjectPtr<USkeletalMesh> Ref = ModelResources.PassThroughMeshes[ReferenceID];
+				Instance->GetPrivate()->PassThroughMeshesToLoad.Add(Ref);
+			}
+			else
+			{
+				// internal error.
+				UE_LOG(LogMutable, Error, TEXT("Referenced mesh [%d] was not stored in the resource array."), ReferenceID);
+			}
+		}
 			
 		Task_Mutable_GetMeshes_GetMesh_Loop(OperationData, StartTime, StartCycles, GetMeshesData, ++GetMeshIndex);
 	}
@@ -2344,6 +2364,7 @@ namespace impl
 			UCustomizableInstancePrivate* CustomizableObjectInstancePrivateData = OperationData->Instance->GetPrivate();
 
 			CustomizableObjectInstancePrivateData->PassThroughTexturesToLoad.Empty();
+			CustomizableObjectInstancePrivateData->PassThroughMeshesToLoad.Empty();
 
 			if (OperationData->PixelFormatOverride)
 			{
diff --git a/Engine/Plugins/Experimental/Mutable/Source/CustomizableObjectEditor/Private/MuCOE/GenerateMutableSource/GenerateMutableSourceMesh.cpp b/Engine/Plugins/Experimental/Mutable/Source/CustomizableObjectEditor/Private/MuCOE/GenerateMutableSource/GenerateMutableSourceMesh.cpp
--- a/Engine/Plugins/Experimental/Mutable/Source/CustomizableObjectEditor/Private/MuCOE/GenerateMutableSource/GenerateMutableSourceMesh.cpp
+++ b/Engine/Plugins/Experimental/Mutable/Source/CustomizableObjectEditor/Private/MuCOE/GenerateMutableSource/GenerateMutableSourceMesh.cpp
@@ -3480,7 +3480,7 @@ mu::NodeMeshPtr GenerateMutableSourceMesh(const UEdGraphPin* Pin,
 
 	else if (const UCustomizableObjectNodePassThroughMesh* TypedNodePassThroughMesh = Cast<UCustomizableObjectNodePassThroughMesh>(Node))
 	{
-		USkeletalMesh* BaseMesh = TypedNodePassThroughMesh->Mesh;
+		UStreamableRenderAsset* BaseMesh = TypedNodePassThroughMesh->GetMesh();
 		if (BaseMesh)
 		{
 			mu::Ptr<mu::NodeMeshConstant> MeshNode = new mu::NodeMeshConstant();
diff --git a/Engine/Plugins/Experimental/Mutable/Source/CustomizableObjectEditor/Private/MuCOE/Nodes/CustomizableObjectNodePassThroughMesh.cpp b/Engine/Plugins/Experimental/Mutable/Source/CustomizableObjectEditor/Private/MuCOE/Nodes/CustomizableObjectNodePassThroughMesh.cpp
--- a/Engine/Plugins/Experimental/Mutable/Source/CustomizableObjectEditor/Private/MuCOE/Nodes/CustomizableObjectNodePassThroughMesh.cpp
+++ b/Engine/Plugins/Experimental/Mutable/Source/CustomizableObjectEditor/Private/MuCOE/Nodes/CustomizableObjectNodePassThroughMesh.cpp
@@ -21,10 +21,10 @@ void UCustomizableObjectNodePassThroughMesh::AllocateDefaultPins(UCustomizableOb
 
 FText UCustomizableObjectNodePassThroughMesh::GetNodeTitle(ENodeTitleType::Type TitleType) const
 {
-	if (Mesh)
+	if (Mesh.IsValid())
 	{
 		FFormatNamedArguments Args;
-		Args.Add(TEXT("MeshName"), FText::FromString(Mesh->GetName()));
+		Args.Add(TEXT("MeshName"), FText::FromString(Mesh.GetAssetName()));
 
 		return FText::Format(LOCTEXT("PassThrough Mesh_Title", "{MeshName}\nPassThrough Mesh"), Args);
 	}
diff --git a/Engine/Plugins/Experimental/Mutable/Source/CustomizableObjectEditor/Public/MuCOE/Nodes/CustomizableObjectNodePassThroughMesh.h b/Engine/Plugins/Experimental/Mutable/Source/CustomizableObjectEditor/Public/MuCOE/Nodes/CustomizableObjectNodePassThroughMesh.h
--- a/Engine/Plugins/Experimental/Mutable/Source/CustomizableObjectEditor/Public/MuCOE/Nodes/CustomizableObjectNodePassThroughMesh.h
+++ b/Engine/Plugins/Experimental/Mutable/Source/CustomizableObjectEditor/Public/MuCOE/Nodes/CustomizableObjectNodePassThroughMesh.h
@@ -14,14 +14,14 @@ class CUSTOMIZABLEOBJECTEDITOR_API UCustomizableObjectNodePassThroughMesh : publ
 public:
 	GENERATED_BODY()
 
-	UPROPERTY(EditAnywhere, Category = Mesh, Meta = (DisplayName = Mesh))
-	TObjectPtr<class USkeletalMesh> Mesh = nullptr;
+	UPROPERTY(EditAnywhere, Category = Mesh, Meta = (DisplayName = Mesh, AllowedClasses = "/Script/Engine.StaticMesh, /Script/Engine.SkeletalMesh"))
+	FSoftObjectPath Mesh;
 
 	// UCustomizableObjectNodeMesh interface
 	virtual TArray<UCustomizableObjectLayout*> GetLayouts(const UEdGraphPin& OutPin) const override { return {}; }
 	virtual UTexture2D* FindTextureForPin(const UEdGraphPin* Pin) const override { return nullptr; }
 	virtual void GetUVChannelForPin(const UEdGraphPin* Pin, TArray<FVector2f>& OutSegments, int32 UVIndex) const override {}
-	virtual UObject* GetMesh() const override { return Mesh.Get(); }
+	virtual UStreamableRenderAsset* GetMesh() const override { return Cast<UStreamableRenderAsset>(Mesh.TryLoad()); }
 	virtual UEdGraphPin* GetMeshPin(int32 LOD, int32 SectionIndex) const override { return {}; }
 	virtual UEdGraphPin* GetLayoutPin(int32 LODIndex, int32 SectionIndex, int32 LayoutIndex) const override { return {}; }
 	virtual void GetPinSection(const UEdGraphPin& Pin, int32& OutLODIndex, int32& OutSectionIndex, int32& OutLayoutIndex) const override {}
