[Mutable] Fixed OnPostCompile delegate not being called when compilations failed.

- Fixed previews not recreating the actors when compilations failed.
- COE now create its instance on editor initialization. Will not recreate the instance each time.
- Fixed double creation of previw actors (OnCustomizableObjectStatusChanged + OnPostCompile).
- Fixed SCustomizableObjectEditorViewportTabBody crash with invalid PreviewSkeletalMeshComponents.
- SetModel to nullptr if the compilation had errors.

#rnx

diff --git a/Engine/Plugins/Experimental/Mutable/Source/CustomizableObjectEditor/Private/MuCOE/CustomizableObjectCompiler.cpp b/Engine/Plugins/Experimental/Mutable/Source/CustomizableObjectEditor/Private/MuCOE/CustomizableObjectCompiler.cpp
--- a/Engine/Plugins/Experimental/Mutable/Source/CustomizableObjectEditor/Private/MuCOE/CustomizableObjectCompiler.cpp
+++ b/Engine/Plugins/Experimental/Mutable/Source/CustomizableObjectEditor/Private/MuCOE/CustomizableObjectCompiler.cpp
@@ -1199,7 +1199,7 @@ void FCustomizableObjectCompiler::CompleteRequest(ECompilationStatePrivate State
 		// Unlock the object so that instances can be updated
 		UCustomizableObjectSystem* System = UCustomizableObjectSystem::IsCreated() ? UCustomizableObjectSystem::GetInstance() : nullptr;
 		if (System && !System->HasAnyFlags(EObjectFlags::RF_BeginDestroyed))
-		{
+		{	
 			System->UnlockObject(CurrentObject);
 		}
 
@@ -1208,13 +1208,17 @@ void FCustomizableObjectCompiler::CompleteRequest(ECompilationStatePrivate State
 			CurrentModel->GetPrivate()->UnloadRoms();
 		}
 
-		CurrentObject->GetPrivate()->SetModel(CurrentModel, GenerateIdentifier(*CurrentObject));
-
 		if (Result == ECompilationResultPrivate::Success || Result == ECompilationResultPrivate::Warnings)
 		{
-			CurrentObject->GetPrivate()->PostCompile();
+			CurrentObject->GetPrivate()->SetModel(CurrentModel, GenerateIdentifier(*CurrentObject));
+		}
+		else
+		{
+			CurrentObject->GetPrivate()->SetModel(nullptr, {});
 		}
 
+		CurrentObject->GetPrivate()->PostCompile();
+
 		UE_LOG(LogMutable, Display, TEXT("Finished compiling Customizable Object %s. Compilation took %5.3f seconds to complete."),
 			*CurrentObject->GetName(), FPlatformTime::Seconds() - CompilationStartTime);
 	}
diff --git a/Engine/Plugins/Experimental/Mutable/Source/CustomizableObjectEditor/Private/MuCOE/CustomizableObjectEditor.cpp b/Engine/Plugins/Experimental/Mutable/Source/CustomizableObjectEditor/Private/MuCOE/CustomizableObjectEditor.cpp
--- a/Engine/Plugins/Experimental/Mutable/Source/CustomizableObjectEditor/Private/MuCOE/CustomizableObjectEditor.cpp
+++ b/Engine/Plugins/Experimental/Mutable/Source/CustomizableObjectEditor/Private/MuCOE/CustomizableObjectEditor.cpp
@@ -310,12 +310,17 @@ void FCustomizableObjectEditor::InitCustomizableObjectEditor(const EToolkitMode:
 	FCoreUObjectDelegates::OnObjectModified.AddRaw(this, &FCustomizableObjectEditor::OnObjectModified);
 	
 	UCustomizableObjectPrivate* CustomizableObjectPrivate = CustomizableObject->GetPrivate();
+	
+	PreviewInstance = CustomizableObject->CreateInstance();	
+	PreviewInstance->UpdatedNativeDelegate.AddSP(SharedThis(this), &FCustomizableObjectEditor::OnUpdatePreviewInstance);
+	PreviewInstance->SetBuildParameterRelevancy(true);
 
+	CustomizableInstanceDetailsView->SetObject(PreviewInstance, true);
+	
 	CustomizableObjectPrivate->Status.GetOnStateChangedDelegate().AddRaw(this, &FCustomizableObjectEditor::OnCustomizableObjectStatusChanged);
-	const FCustomizableObjectStatusTypes::EState CurrentStatus = CustomizableObjectPrivate->Status.Get();
-	OnCustomizableObjectStatusChanged(CurrentStatus, CurrentStatus);
-
-	CustomizableObject->GetPostCompileDelegate().AddSP(this, &FCustomizableObjectEditor::OnPostCompile);
+	OnCustomizableObjectStatusChanged(FCustomizableObjectStatusTypes::EState::Loading, CustomizableObjectPrivate->Status.Get());  // Fake we are still in the loading phase.
+	
+	CustomizableObject->GetPostCompileDelegate().AddSP(this, &FCustomizableObjectEditor::OnPostCompile); // Must be attached after creating the Instance since the Instance also does some work in this delegate.
 }
 
 
@@ -337,20 +342,6 @@ void FCustomizableObjectEditor::SelectNode(const UEdGraphNode* Node)
 }
 
 
-void FCustomizableObjectEditor::CreatePreviewInstance()
-{
-	PreviewInstance = CustomizableObject->CreateInstance();
-	CustomizableInstanceDetailsView->SetObject(PreviewInstance, true);
-
-	CreatePreviewActor();
-
-	PreviewInstance->UpdatedNativeDelegate.AddSP(SharedThis(this), &FCustomizableObjectEditor::OnUpdatePreviewInstance);
-	
-	PreviewInstance->SetBuildParameterRelevancy(true);
-	PreviewInstance->UpdateSkeletalMeshAsync(true, true);
-}
-
-
 void FCustomizableObjectEditor::AddReferencedObjects( FReferenceCollector& Collector )
 {
 	Collector.AddReferencedObject( CustomizableObject );
@@ -2420,23 +2411,19 @@ void RemoveRestrictedChars(FString& String)
 }
 
 
-void FCustomizableObjectEditor::OnCustomizableObjectStatusChanged(FCustomizableObjectStatus::EState, const FCustomizableObjectStatus::EState CurrentState)
+void FCustomizableObjectEditor::OnCustomizableObjectStatusChanged(FCustomizableObjectStatus::EState PreviousState, const FCustomizableObjectStatus::EState CurrentState)
 {
-	switch (CurrentState)
+	if (PreviousState == FCustomizableObjectStatusTypes::EState::Loading)
 	{
-	case FCustomizableObjectStatus::EState::ModelLoaded:
-		if (!PreviewInstance)
+		if (CurrentState == FCustomizableObjectStatusTypes::EState::ModelLoaded)
 		{
-			CreatePreviewInstance();
+			CreatePreviewActor();
+			PreviewInstance->UpdateSkeletalMeshAsync(true, true);
+		}
+		else if (CurrentState == FCustomizableObjectStatusTypes::EState::NoModel)
+		{
+			CustomizableObject->ConditionalAutoCompile();
 		}
-		break;
-		
-	case FCustomizableObjectStatus::EState::NoModel:
-		CustomizableObject->ConditionalAutoCompile();
-		break;
-
-	default:
-		break;
 	}
 }
 
diff --git a/Engine/Plugins/Experimental/Mutable/Source/CustomizableObjectEditor/Private/MuCOE/CustomizableObjectEditor.h b/Engine/Plugins/Experimental/Mutable/Source/CustomizableObjectEditor/Private/MuCOE/CustomizableObjectEditor.h
--- a/Engine/Plugins/Experimental/Mutable/Source/CustomizableObjectEditor/Private/MuCOE/CustomizableObjectEditor.h
+++ b/Engine/Plugins/Experimental/Mutable/Source/CustomizableObjectEditor/Private/MuCOE/CustomizableObjectEditor.h
@@ -275,8 +275,6 @@ private:
 	TSharedRef<SDockTab> SpawnTab_PerformanceReport(const FSpawnTabArgs& Args);
 	TSharedRef<SDockTab> SpawnTab_TagExplorer(const FSpawnTabArgs& Args);
 	
-	void CreatePreviewInstance();
-
 	/** Binds commands associated with the Static Mesh Editor. */
 	void BindCommands();
 	
diff --git a/Engine/Plugins/Experimental/Mutable/Source/CustomizableObjectEditor/Private/MuCOE/CustomizableObjectEditorViewportClient.cpp b/Engine/Plugins/Experimental/Mutable/Source/CustomizableObjectEditor/Private/MuCOE/CustomizableObjectEditorViewportClient.cpp
--- a/Engine/Plugins/Experimental/Mutable/Source/CustomizableObjectEditor/Private/MuCOE/CustomizableObjectEditorViewportClient.cpp
+++ b/Engine/Plugins/Experimental/Mutable/Source/CustomizableObjectEditor/Private/MuCOE/CustomizableObjectEditorViewportClient.cpp
@@ -376,7 +376,7 @@ void FCustomizableObjectEditorViewportClient::Draw(const FSceneView* View, FPrim
 
 	if (bShowBones)
 	{
-		for (TWeakObjectPtr<UDebugSkelMeshComponent> SkeletalMeshComponent : SkeletalMeshComponents)
+		for (const TWeakObjectPtr<UDebugSkelMeshComponent>& SkeletalMeshComponent : SkeletalMeshComponents)
 		{
 			if (SkeletalMeshComponent.IsValid())
 			{
diff --git a/Engine/Plugins/Experimental/Mutable/Source/CustomizableObjectEditor/Private/MuCOE/CustomizableObjectInstanceEditor.cpp b/Engine/Plugins/Experimental/Mutable/Source/CustomizableObjectEditor/Private/MuCOE/CustomizableObjectInstanceEditor.cpp
--- a/Engine/Plugins/Experimental/Mutable/Source/CustomizableObjectEditor/Private/MuCOE/CustomizableObjectInstanceEditor.cpp
+++ b/Engine/Plugins/Experimental/Mutable/Source/CustomizableObjectEditor/Private/MuCOE/CustomizableObjectInstanceEditor.cpp
@@ -252,9 +252,12 @@ void FCustomizableObjectInstanceEditor::InitCustomizableObjectInstanceEditor( co
 	// Set the instance
 	check(InCustomizableObjectInstance);
 	CustomizableObjectInstance = InCustomizableObjectInstance;
+	CustomizableObjectInstance->UpdatedNativeDelegate.AddSP(SharedThis(this), &FCustomizableObjectInstanceEditor::OnUpdatePreviewInstance);
+	CustomizableObjectInstance->SetBuildParameterRelevancy(true);
+	
 	bOnlyRelevantParameters = InCustomizableObjectInstance->GetPrivate()->bShowOnlyRelevantParameters;
 	bOnlyRuntimeParameters = InCustomizableObjectInstance->GetPrivate()->bShowOnlyRuntimeParameters;
-
+	
 	FAdvancedPreviewSceneModule& AdvancedPreviewSceneModule = FModuleManager::LoadModuleChecked<FAdvancedPreviewSceneModule>("AdvancedPreviewScene");
 
 	TSharedPtr<FAdvancedPreviewScene> AdvancedPreviewScene = StaticCastSharedPtr<FAdvancedPreviewScene>(Viewport->GetPreviewScene());
@@ -301,8 +304,7 @@ void FCustomizableObjectInstanceEditor::InitCustomizableObjectInstanceEditor( co
 		UCustomizableObjectPrivate* CustomizableObjectPrivate = CustomizableObject->GetPrivate();
 
 		CustomizableObjectPrivate->Status.GetOnStateChangedDelegate().AddRaw(this, &FCustomizableObjectInstanceEditor::OnCustomizableObjectStatusChanged);
-		const FCustomizableObjectStatusTypes::EState CurrentStatus = CustomizableObjectPrivate->Status.Get();
-		OnCustomizableObjectStatusChanged(CurrentStatus, CurrentStatus);
+		OnCustomizableObjectStatusChanged(FCustomizableObjectStatusTypes::EState::Loading, CustomizableObjectPrivate->Status.Get()); // Fake we are still in the loading phase.
 
 		int32 StateParameterCount = CustomizableObject->GetStateParameterCount(CustomizableObjectInstance->GetCurrentState());
 		int32 ParameterCount = CustomizableObject->GetParameterCount();
@@ -330,23 +332,6 @@ FText FCustomizableObjectInstanceEditor::GetBaseToolkitName() const
 }
 
 
-void FCustomizableObjectInstanceEditor::SetupPreviewInstance()
-{
-	check(CustomizableObjectInstance);
-	if (!CustomizableObjectInstance->GetCustomizableObject())
-	{
-		return;
-	}
-
-	CreatePreviewActor();
-	
-	CustomizableObjectInstance->UpdatedNativeDelegate.AddSP(SharedThis(this), &FCustomizableObjectInstanceEditor::OnUpdatePreviewInstance);
-	
-	CustomizableObjectInstance->SetBuildParameterRelevancy(true);
-	CustomizableObjectInstance->UpdateSkeletalMeshAsync(true, true);
-}
-
-
 void FCustomizableObjectInstanceEditor::SaveAsset_Execute()
 {
 	if (!CustomizableObjectInstance)
@@ -650,31 +635,23 @@ TStatId FCustomizableObjectInstanceEditor::GetStatId() const
 }
 
 
-void FCustomizableObjectInstanceEditor::OnCustomizableObjectStatusChanged(FCustomizableObjectStatus::EState, const FCustomizableObjectStatus::EState CurrentState)
+void FCustomizableObjectInstanceEditor::OnCustomizableObjectStatusChanged(FCustomizableObjectStatus::EState PreviousState, const FCustomizableObjectStatus::EState CurrentState)
 {
-	switch (CurrentState)
+	if (PreviousState == FCustomizableObjectStatusTypes::EState::Loading)
 	{
-	case FCustomizableObjectStatus::EState::ModelLoaded:
+		if (CurrentState == FCustomizableObjectStatusTypes::EState::ModelLoaded)
 		{
-			SetupPreviewInstance();
-			break;			
+			CreatePreviewActor();
+			CustomizableObjectInstance->UpdateSkeletalMeshAsync(true, true);
 		}
-		
-	case FCustomizableObjectStatus::EState::NoModel:
+		else if (CurrentState == FCustomizableObjectStatusTypes::EState::NoModel)
 		{
-			UCustomizableObject* CustomizableObject = CustomizableObjectInstance->GetCustomizableObject();
-			
-			if (CustomizableObject)
+			if (UCustomizableObject* CustomizableObject = CustomizableObjectInstance->GetCustomizableObject())
 			{
 				CustomizableObject->ConditionalAutoCompile();
 			}
-
-			break;			
 		}
-
-	default:
-		break;
-	}	
+	}
 }
 
 
diff --git a/Engine/Plugins/Experimental/Mutable/Source/CustomizableObjectEditor/Private/MuCOE/CustomizableObjectInstanceEditor.h b/Engine/Plugins/Experimental/Mutable/Source/CustomizableObjectEditor/Private/MuCOE/CustomizableObjectInstanceEditor.h
--- a/Engine/Plugins/Experimental/Mutable/Source/CustomizableObjectEditor/Private/MuCOE/CustomizableObjectInstanceEditor.h
+++ b/Engine/Plugins/Experimental/Mutable/Source/CustomizableObjectEditor/Private/MuCOE/CustomizableObjectInstanceEditor.h
@@ -195,8 +195,6 @@ private:
 	TSharedRef<SDockTab> SpawnTab_AdvancedPreviewSettings(const FSpawnTabArgs& Args);
 	TSharedRef<SDockTab> SpawnTab_TextureAnalyzer(const FSpawnTabArgs& Args);
 
-	void SetupPreviewInstance();
-
 	/** Binds commands associated with the Static Mesh Editor. */
 	void BindCommands();
 		
diff --git a/Engine/Plugins/Experimental/Mutable/Source/CustomizableObjectEditor/Private/MuCOE/SCustomizableObjectEditorViewport.cpp b/Engine/Plugins/Experimental/Mutable/Source/CustomizableObjectEditor/Private/MuCOE/SCustomizableObjectEditorViewport.cpp
--- a/Engine/Plugins/Experimental/Mutable/Source/CustomizableObjectEditor/Private/MuCOE/SCustomizableObjectEditorViewport.cpp
+++ b/Engine/Plugins/Experimental/Mutable/Source/CustomizableObjectEditor/Private/MuCOE/SCustomizableObjectEditorViewport.cpp
@@ -274,6 +274,11 @@ void SCustomizableObjectEditorViewportTabBody::Tick(const FGeometry& AllottedGeo
 
 		for (TWeakObjectPtr<UDebugSkelMeshComponent> PreviewSkeletalMeshComponent : PreviewSkeletalMeshComponents)
 		{
+			if (!PreviewSkeletalMeshComponent.IsValid())
+			{
+				continue;
+			}
+			
 			const TArray<UMaterialInterface*> Materials = PreviewSkeletalMeshComponent->GetMaterials();
 			for (UMaterialInterface* m : Materials)
 			{
