[Mutable] Added Pin Viewer to the context menu.

- Added new Hide Pin action.
- Now all Node Details inherit from FCustomizableObjectNodeDetails.
- FCustomizableObjectNodeBase has a new virtual function which specifies the FCustomizableObjectNodeDetails if it has to create the PinViwer.

#rnx

diff --git a/Engine/Plugins/Experimental/Mutable/Source/CustomizableObjectEditor/Private/MuCOE/CustomizableObjectEditorModule.cpp b/Engine/Plugins/Experimental/Mutable/Source/CustomizableObjectEditor/Private/MuCOE/CustomizableObjectEditorModule.cpp
--- a/Engine/Plugins/Experimental/Mutable/Source/CustomizableObjectEditor/Private/MuCOE/CustomizableObjectEditorModule.cpp
+++ b/Engine/Plugins/Experimental/Mutable/Source/CustomizableObjectEditor/Private/MuCOE/CustomizableObjectEditorModule.cpp
@@ -50,7 +50,7 @@
 #include "MuCOE/Nodes/CustomizableObjectNodeMorphMaterialDetails.h"
 #include "MuCOE/Nodes/CustomizableObjectNodeObjectDetails.h"
 #include "MuCOE/Nodes/CustomizableObjectNodeObjectGroup.h"
-#include "MuCOE/Nodes/CustomizableObjectNodePinViewerDetails.h"
+#include "MuCOE/Nodes/CustomizableObjectNodeDetails.h"
 #include "MuCOE/Nodes/CustomizableObjectNodeProjectorConstant.h"
 #include "MuCOE/Nodes/CustomizableObjectNodeProjectorParameterDetails.h"
 #include "MuCOE/Nodes/CustomizableObjectNodeRemoveMeshBlocks.h"
@@ -233,9 +233,8 @@ void FCustomizableObjectEditorModule::StartupModule()
 	//UThumbnailManager::Get().RegisterCustomRenderer(UCustomizableObjectInstance::StaticClass(), UCustomizableObjectInstanceThumbnailRenderer::StaticClass());
 	
 	// Property views
+	// Nodes
 	FPropertyEditorModule& PropertyModule = FModuleManager::LoadModuleChecked<FPropertyEditorModule>("PropertyEditor");
-	RegisterCustomDetails(PropertyModule, UCustomizableObject::StaticClass(), FOnGetDetailCustomizationInstance::CreateStatic(&FCustomizableObjectDetails::MakeInstance));
-	RegisterCustomDetails(PropertyModule, UCustomizableObjectInstance::StaticClass(), FOnGetDetailCustomizationInstance::CreateStatic(&FCustomizableInstanceDetails::MakeInstance));
 	RegisterCustomDetails(PropertyModule, UCustomizableObjectNodeLayoutBlocks::StaticClass(), FOnGetDetailCustomizationInstance::CreateStatic(&FCustomizableObjectNodeLayoutBlocksDetails::MakeInstance));
 	RegisterCustomDetails(PropertyModule, UCustomizableObjectNodeEditMaterial::StaticClass(), FOnGetDetailCustomizationInstance::CreateStatic(&FCustomizableObjectNodeEditMaterialDetails::MakeInstance));
 	RegisterCustomDetails(PropertyModule, UCustomizableObjectNodeRemoveMeshBlocks::StaticClass(), FOnGetDetailCustomizationInstance::CreateStatic(&FCustomizableObjectNodeRemoveMeshBlocksDetails::MakeInstance));
@@ -250,13 +249,17 @@ void FCustomizableObjectEditorModule::StartupModule()
 	RegisterCustomDetails(PropertyModule, UCustomizableObjectNodeMeshClipMorph::StaticClass(), FOnGetDetailCustomizationInstance::CreateStatic(&FCustomizableObjectNodeMeshClipMorphDetails::MakeInstance));
 	RegisterCustomDetails(PropertyModule, UCustomizableObjectNodeMeshClipWithMesh::StaticClass(), FOnGetDetailCustomizationInstance::CreateStatic(&FCustomizableObjectNodeMeshClipWithMeshDetails::MakeInstance));
 	RegisterCustomDetails(PropertyModule, UCustomizableObjectNodeExternalPin::StaticClass(), FOnGetDetailCustomizationInstance::CreateStatic(&FCustomizableObjectNodeExternalPinDetails::MakeInstance));
-	RegisterCustomDetails(PropertyModule, UCustomSettings::StaticClass(), FOnGetDetailCustomizationInstance::CreateStatic(&FCustomizableObjectCustomSettingsDetails::MakeInstance));
-	RegisterCustomDetails(PropertyModule, UCustomizableObjectNodeMaterial::StaticClass(), FOnGetDetailCustomizationInstance::CreateStatic(&FCustomizableObjectNodePinViewerDetails::MakeInstance));
+	RegisterCustomDetails(PropertyModule, UCustomizableObjectNodeMaterial::StaticClass(), FOnGetDetailCustomizationInstance::CreateStatic(&FCustomizableObjectNodeDetails::MakeInstance));
 	RegisterCustomDetails(PropertyModule, UCustomizableObjectNodeCopyMaterial::StaticClass(), FOnGetDetailCustomizationInstance::CreateStatic(&FCustomizableObjectNodeCopyMaterialDetails::MakeInstance));
 	RegisterCustomDetails(PropertyModule, UCustomizableObjectNodeSkeletalMesh::StaticClass(), FOnGetDetailCustomizationInstance::CreateStatic(&FCustomizableObjectNodeSkeletalMeshDetails::MakeInstance));
-	RegisterCustomDetails(PropertyModule, UCustomizableObjectNodeStaticMesh::StaticClass(), FOnGetDetailCustomizationInstance::CreateStatic(&FCustomizableObjectNodePinViewerDetails::MakeInstance));
+	RegisterCustomDetails(PropertyModule, UCustomizableObjectNodeStaticMesh::StaticClass(), FOnGetDetailCustomizationInstance::CreateStatic(&FCustomizableObjectNodeDetails::MakeInstance));
 	RegisterCustomDetails(PropertyModule, UCustomizableObjectNodeTable::StaticClass(), FOnGetDetailCustomizationInstance::CreateStatic(&FCustomizableObjectNodeTableDetails::MakeInstance));
 
+	// Other Objects
+	RegisterCustomDetails(PropertyModule, UCustomizableObject::StaticClass(), FOnGetDetailCustomizationInstance::CreateStatic(&FCustomizableObjectDetails::MakeInstance));
+	RegisterCustomDetails(PropertyModule, UCustomizableObjectInstance::StaticClass(), FOnGetDetailCustomizationInstance::CreateStatic(&FCustomizableInstanceDetails::MakeInstance));
+	RegisterCustomDetails(PropertyModule, UCustomSettings::StaticClass(), FOnGetDetailCustomizationInstance::CreateStatic(&FCustomizableObjectCustomSettingsDetails::MakeInstance));
+
 	//Custom properties
 	PropertyModule.RegisterCustomPropertyTypeLayout("CustomizableObjectIdentifier", FOnGetPropertyTypeCustomizationInstance::CreateStatic(&FCustomizableObjectIdentifierCustomization::MakeInstance));
 	PropertyModule.RegisterCustomPropertyTypeLayout(FMeshReshapeBoneReference::StaticStruct()->GetFName(), FOnGetPropertyTypeCustomizationInstance::CreateStatic(&FMeshReshapeBonesReferenceCustomization::MakeInstance));
diff --git a/Engine/Plugins/Experimental/Mutable/Source/CustomizableObjectEditor/Private/MuCOE/CustomizableObjectNodeObjectGroupDetails.cpp b/Engine/Plugins/Experimental/Mutable/Source/CustomizableObjectEditor/Private/MuCOE/CustomizableObjectNodeObjectGroupDetails.cpp
--- a/Engine/Plugins/Experimental/Mutable/Source/CustomizableObjectEditor/Private/MuCOE/CustomizableObjectNodeObjectGroupDetails.cpp
+++ b/Engine/Plugins/Experimental/Mutable/Source/CustomizableObjectEditor/Private/MuCOE/CustomizableObjectNodeObjectGroupDetails.cpp
@@ -25,6 +25,8 @@ TSharedRef<IDetailCustomization> FCustomizableObjectNodeObjectGroupDetails::Make
 
 void FCustomizableObjectNodeObjectGroupDetails::CustomizeDetails(IDetailLayoutBuilder& DetailBuilder)
 {
+	FCustomizableObjectNodeDetails::CustomizeDetails(DetailBuilder);
+
 	const IDetailsView* DetailsView = DetailBuilder.GetDetailsView();
 
 	if (DetailsView->GetSelectedObjects().Num())
diff --git a/Engine/Plugins/Experimental/Mutable/Source/CustomizableObjectEditor/Private/MuCOE/CustomizableObjectNodeObjectGroupDetails.h b/Engine/Plugins/Experimental/Mutable/Source/CustomizableObjectEditor/Private/MuCOE/CustomizableObjectNodeObjectGroupDetails.h
--- a/Engine/Plugins/Experimental/Mutable/Source/CustomizableObjectEditor/Private/MuCOE/CustomizableObjectNodeObjectGroupDetails.h
+++ b/Engine/Plugins/Experimental/Mutable/Source/CustomizableObjectEditor/Private/MuCOE/CustomizableObjectNodeObjectGroupDetails.h
@@ -3,6 +3,7 @@
 #pragma once
 
 #include "IDetailCustomization.h"
+#include "Nodes/CustomizableObjectNodeDetails.h"
 
 struct EVisibility;
 namespace ESelectInfo { enum Type : int; }
@@ -10,7 +11,7 @@ class IDetailLayoutBuilder;
 class STextComboBox;
 class UCustomizableObjectNodeObjectGroup;
 
-class FCustomizableObjectNodeObjectGroupDetails : public IDetailCustomization
+class FCustomizableObjectNodeObjectGroupDetails : public FCustomizableObjectNodeDetails
 {
 public:
 	/** Makes a new instance of this detail layout class for a specific detail view requesting it */
diff --git a/Engine/Plugins/Experimental/Mutable/Source/CustomizableObjectEditor/Private/MuCOE/EdGraphSchema_CustomizableObject.cpp b/Engine/Plugins/Experimental/Mutable/Source/CustomizableObjectEditor/Private/MuCOE/EdGraphSchema_CustomizableObject.cpp
--- a/Engine/Plugins/Experimental/Mutable/Source/CustomizableObjectEditor/Private/MuCOE/EdGraphSchema_CustomizableObject.cpp
+++ b/Engine/Plugins/Experimental/Mutable/Source/CustomizableObjectEditor/Private/MuCOE/EdGraphSchema_CustomizableObject.cpp
@@ -81,6 +81,7 @@
 #include "ToolMenuSection.h"
 #include "MuCOE/Nodes/CustomizableObjectNodeReroute.h"
 #include "Toolkits/ToolkitManager.h"
+#include "PropertyEditorModule.h"
 
 class IToolkit;
 
@@ -869,6 +870,35 @@ void UEdGraphSchema_CustomizableObject::GetContextMenuActions(UToolMenu* Menu, U
 	}
 	else // On Pin right click
 	{
+		if (const UCustomizableObjectNode* Node = Cast<UCustomizableObjectNode>(Context->Node))
+		{
+			if (const UEdGraphPin* Pin = Context->Pin)
+			{
+				FPropertyEditorModule& PropPlugin = FModuleManager::LoadModuleChecked<FPropertyEditorModule>("PropertyEditor");
+				
+				if (Node->CanPinBeHidden(*Pin))
+				{
+					FToolMenuSection& Section = Menu->FindOrAddSection("EdGraphSchemaPinActions");
+					Section.AddMenuEntry("HidePin",
+						LOCTEXT("HidePin", "Hide Pin"),
+						LOCTEXT("SetTexturePinModeDefault_Tooltip", "Hides the selected pin."),
+						FSlateIcon(),
+						FUIAction(FExecuteAction::CreateLambda([Node = const_cast<UCustomizableObjectNode*>(Node), Pin]()
+						{
+							Node->SetPinHidden(*const_cast<UEdGraphPin*>(Pin), true);
+						})));
+				}
+				
+				
+				if (TSharedPtr<SWidget> Widget = Node->CustomizePinDetails(*Pin))
+				{
+					FToolMenuSection& Section = Menu->FindOrAddSection("EdGraphSchemaPinActions");
+					Section.AddSeparator("Pin Viewer");
+					Section.AddEntry(FToolMenuEntry::InitWidget("Pin Viewer", Widget.ToSharedRef(), {}));
+				}
+			}
+		}
+		
 		const UCustomizableObjectNodeTable* TableNode = Cast<UCustomizableObjectNodeTable>(Context->Node);
 		UEdGraphPin* TexturePin = (UEdGraphPin*)Context->Pin;
 
diff --git a/Engine/Plugins/Experimental/Mutable/Source/CustomizableObjectEditor/Private/MuCOE/Nodes/CustomizableObjectNode.cpp b/Engine/Plugins/Experimental/Mutable/Source/CustomizableObjectEditor/Private/MuCOE/Nodes/CustomizableObjectNode.cpp
--- a/Engine/Plugins/Experimental/Mutable/Source/CustomizableObjectEditor/Private/MuCOE/Nodes/CustomizableObjectNode.cpp
+++ b/Engine/Plugins/Experimental/Mutable/Source/CustomizableObjectEditor/Private/MuCOE/Nodes/CustomizableObjectNode.cpp
@@ -648,11 +648,17 @@ void UCustomizableObjectNode::SetPinHidden(const TArray<UEdGraphPin*>& PinsToHid
 
 bool UCustomizableObjectNode::CanPinBeHidden(const UEdGraphPin& Pin) const
 {
-	return !Pin.LinkedTo.Num() && !Pin.bOrphanedPin;
+	return !Pin.LinkedTo.Num() && !Pin.bOrphanedPin && HasPinViewer();
 }
 
 
-TSharedPtr<SWidget> UCustomizableObjectNode::CustomizePinDetails(UEdGraphPin& Pin)
+bool UCustomizableObjectNode::HasPinViewer() const
+{
+	return false;
+}
+
+
+TSharedPtr<SWidget> UCustomizableObjectNode::CustomizePinDetails(const UEdGraphPin& Pin) const
 {
 	return nullptr;
 }
diff --git a/Engine/Plugins/Experimental/Mutable/Source/CustomizableObjectEditor/Private/MuCOE/Nodes/CustomizableObjectNodeCopyMaterialDetails.cpp b/Engine/Plugins/Experimental/Mutable/Source/CustomizableObjectEditor/Private/MuCOE/Nodes/CustomizableObjectNodeCopyMaterialDetails.cpp
--- a/Engine/Plugins/Experimental/Mutable/Source/CustomizableObjectEditor/Private/MuCOE/Nodes/CustomizableObjectNodeCopyMaterialDetails.cpp
+++ b/Engine/Plugins/Experimental/Mutable/Source/CustomizableObjectEditor/Private/MuCOE/Nodes/CustomizableObjectNodeCopyMaterialDetails.cpp
@@ -11,6 +11,8 @@ TSharedRef<IDetailCustomization> FCustomizableObjectNodeCopyMaterialDetails::Mak
 
 void FCustomizableObjectNodeCopyMaterialDetails::CustomizeDetails(IDetailLayoutBuilder& DetailBuilder)
 {
+	FCustomizableObjectNodeDetails::CustomizeDetails(DetailBuilder);
+
 	TArray<FName> Names;
 	DetailBuilder.GetCategoryNames(Names);
 
diff --git a/Engine/Plugins/Experimental/Mutable/Source/CustomizableObjectEditor/Private/MuCOE/Nodes/CustomizableObjectNodeCopyMaterialDetails.h b/Engine/Plugins/Experimental/Mutable/Source/CustomizableObjectEditor/Private/MuCOE/Nodes/CustomizableObjectNodeCopyMaterialDetails.h
--- a/Engine/Plugins/Experimental/Mutable/Source/CustomizableObjectEditor/Private/MuCOE/Nodes/CustomizableObjectNodeCopyMaterialDetails.h
+++ b/Engine/Plugins/Experimental/Mutable/Source/CustomizableObjectEditor/Private/MuCOE/Nodes/CustomizableObjectNodeCopyMaterialDetails.h
@@ -2,12 +2,13 @@
 
 #pragma once
 
+#include "CustomizableObjectNodeDetails.h"
 #include "IDetailCustomization.h"
 
 class IDetailLayoutBuilder;
 
 /** Copy Material node details panel. Hides all properties from the inheret Material node. */
-class FCustomizableObjectNodeCopyMaterialDetails : public IDetailCustomization
+class FCustomizableObjectNodeCopyMaterialDetails : public FCustomizableObjectNodeDetails
 {
 public:
 	static TSharedRef<IDetailCustomization> MakeInstance();
diff --git a/Engine/Plugins/Experimental/Mutable/Source/CustomizableObjectEditor/Private/MuCOE/Nodes/CustomizableObjectNodeDetails.cpp b/Engine/Plugins/Experimental/Mutable/Source/CustomizableObjectEditor/Private/MuCOE/Nodes/CustomizableObjectNodeDetails.cpp
new file mode 100644
--- /dev/null
+++ b/Engine/Plugins/Experimental/Mutable/Source/CustomizableObjectEditor/Private/MuCOE/Nodes/CustomizableObjectNodeDetails.cpp
@@ -0,0 +1,51 @@
+// Copyright Epic Games, Inc. All Rights Reserved.
+
+#include "MuCOE/Nodes/CustomizableObjectNodeDetails.h"
+
+#include "DetailCategoryBuilder.h"
+#include "DetailLayoutBuilder.h"
+#include "DetailWidgetRow.h"
+#include "MuCOE/Nodes/CustomizableObjectNode.h"
+#include "MuCOE/PinViewer/SPinViewer.h"
+
+class IDetailLayoutBuilder;
+
+#define LOCTEXT_NAMESPACE "CustomizableObjectEditor"
+
+
+TSharedRef<IDetailCustomization> FCustomizableObjectNodeDetails::MakeInstance()
+{
+	return MakeShareable(new FCustomizableObjectNodeDetails);
+}
+
+
+void FCustomizableObjectNodeDetails::CustomizeDetails(IDetailLayoutBuilder& DetailBuilder)
+{
+	const IDetailsView* DetailsView = DetailBuilder.GetDetailsView();
+	if (DetailsView->GetSelectedObjects().Num())
+	{
+		if (UCustomizableObjectNode* Node = Cast<UCustomizableObjectNode>(DetailsView->GetSelectedObjects()[0]))
+		{
+			if (!Node->HasPinViewer())
+			{
+				return;
+			}
+			
+			IDetailCategoryBuilder& PinViewerCategoryBuilder = DetailBuilder.EditCategory("PinViewer", LOCTEXT("PinViewer", "Pins"), ECategoryPriority::Uncommon);
+			PinViewerCategoryBuilder.AddCustomRow(LOCTEXT("PinViewerDetailsCategory", "PinViewer")).ShouldAutoExpand(true)
+			[
+				SNew(SPinViewer).Node(Node)
+			];
+		}
+	}
+}
+
+
+// Only implemented to silence -Woverloaded-virtual warning.
+void FCustomizableObjectNodeDetails::CustomizeDetails(const TSharedPtr<IDetailLayoutBuilder>& DetailBuilder)
+{
+	IDetailCustomization::CustomizeDetails(DetailBuilder);
+}
+
+
+#undef LOCTEXT_NAMESPACE
diff --git a/Engine/Plugins/Experimental/Mutable/Source/CustomizableObjectEditor/Private/MuCOE/Nodes/CustomizableObjectNodeDetails.h b/Engine/Plugins/Experimental/Mutable/Source/CustomizableObjectEditor/Private/MuCOE/Nodes/CustomizableObjectNodeDetails.h
new file mode 100644
--- /dev/null
+++ b/Engine/Plugins/Experimental/Mutable/Source/CustomizableObjectEditor/Private/MuCOE/Nodes/CustomizableObjectNodeDetails.h
@@ -0,0 +1,20 @@
+// Copyright Epic Games, Inc. All Rights Reserved.
+
+#pragma once
+
+#include "IDetailCustomization.h"
+
+class IDetailLayoutBuilder;
+
+
+/** Base of all UCustomizableObjectNode details. */
+class FCustomizableObjectNodeDetails : public IDetailCustomization
+{
+public:
+	/** Makes a new instance of this detail layout class for a specific detail view requesting it */
+	static TSharedRef<IDetailCustomization> MakeInstance();
+	
+	// IDetailCustomization interface
+	virtual void CustomizeDetails(IDetailLayoutBuilder& DetailBuilder) override;
+	virtual void CustomizeDetails(const TSharedPtr<IDetailLayoutBuilder>& DetailBuilder) override;
+};
diff --git a/Engine/Plugins/Experimental/Mutable/Source/CustomizableObjectEditor/Private/MuCOE/Nodes/CustomizableObjectNodeEditMaterial.cpp b/Engine/Plugins/Experimental/Mutable/Source/CustomizableObjectEditor/Private/MuCOE/Nodes/CustomizableObjectNodeEditMaterial.cpp
--- a/Engine/Plugins/Experimental/Mutable/Source/CustomizableObjectEditor/Private/MuCOE/Nodes/CustomizableObjectNodeEditMaterial.cpp
+++ b/Engine/Plugins/Experimental/Mutable/Source/CustomizableObjectEditor/Private/MuCOE/Nodes/CustomizableObjectNodeEditMaterial.cpp
@@ -76,6 +76,12 @@ bool UCustomizableObjectNodeEditMaterial::CustomRemovePin(UEdGraphPin& Pin)
 }
 
 
+bool UCustomizableObjectNodeEditMaterial::HasPinViewer() const
+{
+	return true;
+}
+
+
 UEdGraphPin* UCustomizableObjectNodeEditMaterial::OutputPin() const
 {
 	return Super::OutputPin();
diff --git a/Engine/Plugins/Experimental/Mutable/Source/CustomizableObjectEditor/Private/MuCOE/Nodes/CustomizableObjectNodeEditMaterialDetails.cpp b/Engine/Plugins/Experimental/Mutable/Source/CustomizableObjectEditor/Private/MuCOE/Nodes/CustomizableObjectNodeEditMaterialDetails.cpp
--- a/Engine/Plugins/Experimental/Mutable/Source/CustomizableObjectEditor/Private/MuCOE/Nodes/CustomizableObjectNodeEditMaterialDetails.cpp
+++ b/Engine/Plugins/Experimental/Mutable/Source/CustomizableObjectEditor/Private/MuCOE/Nodes/CustomizableObjectNodeEditMaterialDetails.cpp
@@ -24,6 +24,8 @@ TSharedRef<IDetailCustomization> FCustomizableObjectNodeEditMaterialDetails::Mak
 
 void FCustomizableObjectNodeEditMaterialDetails::CustomizeDetails( IDetailLayoutBuilder& DetailBuilder )
 {
+	FCustomizableObjectNodeDetails::CustomizeDetails(DetailBuilder);
+
 	Node = nullptr;
 	const IDetailsView* DetailsView = DetailBuilder.GetDetailsView();
 	if (DetailsView->GetSelectedObjects().Num())
@@ -62,8 +64,6 @@ void FCustomizableObjectNodeEditMaterialDetails::CustomizeDetails( IDetailLayout
 			.Text( LOCTEXT( "Node not found", "Node not found" ) )
 		];
 	}
-
-	PinViewerAttachToDetailCustomization(DetailBuilder);
 }
 
 #undef LOCTEXT_NAMESPACE
diff --git a/Engine/Plugins/Experimental/Mutable/Source/CustomizableObjectEditor/Private/MuCOE/Nodes/CustomizableObjectNodeEditMaterialDetails.h b/Engine/Plugins/Experimental/Mutable/Source/CustomizableObjectEditor/Private/MuCOE/Nodes/CustomizableObjectNodeEditMaterialDetails.h
--- a/Engine/Plugins/Experimental/Mutable/Source/CustomizableObjectEditor/Private/MuCOE/Nodes/CustomizableObjectNodeEditMaterialDetails.h
+++ b/Engine/Plugins/Experimental/Mutable/Source/CustomizableObjectEditor/Private/MuCOE/Nodes/CustomizableObjectNodeEditMaterialDetails.h
@@ -3,6 +3,7 @@
 #pragma once
 
 
+#include "CustomizableObjectNodeDetails.h"
 #include "IDetailCustomization.h"
 
 class FString;
@@ -11,7 +12,7 @@ class UCustomizableObjectNodeEditMaterial;
 class SCustomizableObjectNodeLayoutBlocksSelector;
 
 
-class FCustomizableObjectNodeEditMaterialDetails : public IDetailCustomization
+class FCustomizableObjectNodeEditMaterialDetails : public FCustomizableObjectNodeDetails
 {
 public:
 	/** Makes a new instance of this detail layout class for a specific detail view requesting it */
diff --git a/Engine/Plugins/Experimental/Mutable/Source/CustomizableObjectEditor/Private/MuCOE/Nodes/CustomizableObjectNodeExternalPinDetails.cpp b/Engine/Plugins/Experimental/Mutable/Source/CustomizableObjectEditor/Private/MuCOE/Nodes/CustomizableObjectNodeExternalPinDetails.cpp
--- a/Engine/Plugins/Experimental/Mutable/Source/CustomizableObjectEditor/Private/MuCOE/Nodes/CustomizableObjectNodeExternalPinDetails.cpp
+++ b/Engine/Plugins/Experimental/Mutable/Source/CustomizableObjectEditor/Private/MuCOE/Nodes/CustomizableObjectNodeExternalPinDetails.cpp
@@ -24,6 +24,8 @@ TSharedRef<IDetailCustomization> FCustomizableObjectNodeExternalPinDetails::Make
 
 void FCustomizableObjectNodeExternalPinDetails::CustomizeDetails( IDetailLayoutBuilder& DetailBuilder )
 {
+	FCustomizableObjectNodeDetails::CustomizeDetails(DetailBuilder);
+
 	Node = nullptr;
 	DetailBuilderPtr = &DetailBuilder;
 	
diff --git a/Engine/Plugins/Experimental/Mutable/Source/CustomizableObjectEditor/Private/MuCOE/Nodes/CustomizableObjectNodeExternalPinDetails.h b/Engine/Plugins/Experimental/Mutable/Source/CustomizableObjectEditor/Private/MuCOE/Nodes/CustomizableObjectNodeExternalPinDetails.h
--- a/Engine/Plugins/Experimental/Mutable/Source/CustomizableObjectEditor/Private/MuCOE/Nodes/CustomizableObjectNodeExternalPinDetails.h
+++ b/Engine/Plugins/Experimental/Mutable/Source/CustomizableObjectEditor/Private/MuCOE/Nodes/CustomizableObjectNodeExternalPinDetails.h
@@ -2,6 +2,7 @@
 
 #pragma once
 
+#include "CustomizableObjectNodeDetails.h"
 #include "IDetailCustomization.h"
 
 namespace ESelectInfo { enum Type : int; }
@@ -12,7 +13,7 @@ class UCustomizableObjectNodeExternalPin;
 struct FAssetData;
 
 
-class FCustomizableObjectNodeExternalPinDetails : public IDetailCustomization
+class FCustomizableObjectNodeExternalPinDetails : public FCustomizableObjectNodeDetails
 {
 public:
 	/** Makes a new instance of this detail layout class for a specific detail view requesting it */
diff --git a/Engine/Plugins/Experimental/Mutable/Source/CustomizableObjectEditor/Private/MuCOE/Nodes/CustomizableObjectNodeLayoutBlocksDetails.cpp b/Engine/Plugins/Experimental/Mutable/Source/CustomizableObjectEditor/Private/MuCOE/Nodes/CustomizableObjectNodeLayoutBlocksDetails.cpp
--- a/Engine/Plugins/Experimental/Mutable/Source/CustomizableObjectEditor/Private/MuCOE/Nodes/CustomizableObjectNodeLayoutBlocksDetails.cpp
+++ b/Engine/Plugins/Experimental/Mutable/Source/CustomizableObjectEditor/Private/MuCOE/Nodes/CustomizableObjectNodeLayoutBlocksDetails.cpp
@@ -34,6 +34,8 @@ TSharedRef<IDetailCustomization> FCustomizableObjectNodeLayoutBlocksDetails::Mak
 
 void FCustomizableObjectNodeLayoutBlocksDetails::CustomizeDetails( IDetailLayoutBuilder& DetailBuilder )
 {
+	FCustomizableObjectNodeDetails::CustomizeDetails(DetailBuilder);
+
 	Node = nullptr;
 	const IDetailsView* DetailsView = DetailBuilder.GetDetailsView();
 	if (DetailsView->GetSelectedObjects().Num())
diff --git a/Engine/Plugins/Experimental/Mutable/Source/CustomizableObjectEditor/Private/MuCOE/Nodes/CustomizableObjectNodeLayoutBlocksDetails.h b/Engine/Plugins/Experimental/Mutable/Source/CustomizableObjectEditor/Private/MuCOE/Nodes/CustomizableObjectNodeLayoutBlocksDetails.h
--- a/Engine/Plugins/Experimental/Mutable/Source/CustomizableObjectEditor/Private/MuCOE/Nodes/CustomizableObjectNodeLayoutBlocksDetails.h
+++ b/Engine/Plugins/Experimental/Mutable/Source/CustomizableObjectEditor/Private/MuCOE/Nodes/CustomizableObjectNodeLayoutBlocksDetails.h
@@ -3,6 +3,7 @@
 #pragma once
 
 
+#include "CustomizableObjectNodeDetails.h"
 #include "IDetailCustomization.h"
 #include "UObject/WeakObjectPtr.h"
 
@@ -16,7 +17,7 @@ class SWidget;
 class UCustomizableObjectNodeLayoutBlocks;
 struct EVisibility;
 
-class FCustomizableObjectNodeLayoutBlocksDetails : public IDetailCustomization
+class FCustomizableObjectNodeLayoutBlocksDetails : public FCustomizableObjectNodeDetails
 {
 public:
 	/** Makes a new instance of this detail layout class for a specific detail view requesting it */
diff --git a/Engine/Plugins/Experimental/Mutable/Source/CustomizableObjectEditor/Private/MuCOE/Nodes/CustomizableObjectNodeMaterial.cpp b/Engine/Plugins/Experimental/Mutable/Source/CustomizableObjectEditor/Private/MuCOE/Nodes/CustomizableObjectNodeMaterial.cpp
--- a/Engine/Plugins/Experimental/Mutable/Source/CustomizableObjectEditor/Private/MuCOE/Nodes/CustomizableObjectNodeMaterial.cpp
+++ b/Engine/Plugins/Experimental/Mutable/Source/CustomizableObjectEditor/Private/MuCOE/Nodes/CustomizableObjectNodeMaterial.cpp
@@ -366,6 +366,12 @@ bool UCustomizableObjectNodeMaterial::CanPinBeHidden(const UEdGraphPin& Pin) con
 }
 
 
+bool UCustomizableObjectNodeMaterial::HasPinViewer() const
+{
+	return true;
+}
+
+
 UCustomizableObjectNodeRemapPinsByName* UCustomizableObjectNodeMaterial::CreateRemapPinsDefault() const
 {
 	return NewObject<UCustomizableObjectNodeMaterialRemapPinsByName>();
@@ -1071,7 +1077,7 @@ FString UCustomizableObjectNodeMaterial::GetRefreshMessage() const
 }
 
 
-TSharedPtr<SWidget> UCustomizableObjectNodeMaterial::CustomizePinDetails(UEdGraphPin& Pin)
+TSharedPtr<SWidget> UCustomizableObjectNodeMaterial::CustomizePinDetails(const UEdGraphPin& Pin) const
 {
 	if (UCustomizableObjectNodeMaterialPinDataImage* PinData = Cast<UCustomizableObjectNodeMaterialPinDataImage>(GetPinData(Pin)))
 	{
diff --git a/Engine/Plugins/Experimental/Mutable/Source/CustomizableObjectEditor/Private/MuCOE/Nodes/CustomizableObjectNodeMeshClipMorphDetails.cpp b/Engine/Plugins/Experimental/Mutable/Source/CustomizableObjectEditor/Private/MuCOE/Nodes/CustomizableObjectNodeMeshClipMorphDetails.cpp
--- a/Engine/Plugins/Experimental/Mutable/Source/CustomizableObjectEditor/Private/MuCOE/Nodes/CustomizableObjectNodeMeshClipMorphDetails.cpp
+++ b/Engine/Plugins/Experimental/Mutable/Source/CustomizableObjectEditor/Private/MuCOE/Nodes/CustomizableObjectNodeMeshClipMorphDetails.cpp
@@ -29,6 +29,8 @@ TSharedRef<IDetailCustomization> FCustomizableObjectNodeMeshClipMorphDetails::Ma
 
 void FCustomizableObjectNodeMeshClipMorphDetails::CustomizeDetails(IDetailLayoutBuilder& DetailBuilder)
 {
+	FCustomizableObjectNodeDetails::CustomizeDetails(DetailBuilder);
+
 	Node = nullptr;
 	DetailBuilderPtr = &DetailBuilder;
 
diff --git a/Engine/Plugins/Experimental/Mutable/Source/CustomizableObjectEditor/Private/MuCOE/Nodes/CustomizableObjectNodeMeshClipMorphDetails.h b/Engine/Plugins/Experimental/Mutable/Source/CustomizableObjectEditor/Private/MuCOE/Nodes/CustomizableObjectNodeMeshClipMorphDetails.h
--- a/Engine/Plugins/Experimental/Mutable/Source/CustomizableObjectEditor/Private/MuCOE/Nodes/CustomizableObjectNodeMeshClipMorphDetails.h
+++ b/Engine/Plugins/Experimental/Mutable/Source/CustomizableObjectEditor/Private/MuCOE/Nodes/CustomizableObjectNodeMeshClipMorphDetails.h
@@ -2,6 +2,7 @@
 
 #pragma once
 
+#include "CustomizableObjectNodeDetails.h"
 #include "IDetailCustomization.h"
 #include "UObject/ObjectMacros.h"
 
@@ -13,7 +14,7 @@ class IDetailLayoutBuilder;
 class IPropertyHandle;
 
 
-class FCustomizableObjectNodeMeshClipMorphDetails : public IDetailCustomization
+class FCustomizableObjectNodeMeshClipMorphDetails : public FCustomizableObjectNodeDetails
 {
 public:
 	/** Makes a new instance of this detail layout class for a specific detail view requesting it */
diff --git a/Engine/Plugins/Experimental/Mutable/Source/CustomizableObjectEditor/Private/MuCOE/Nodes/CustomizableObjectNodeMeshClipWithMeshDetails.cpp b/Engine/Plugins/Experimental/Mutable/Source/CustomizableObjectEditor/Private/MuCOE/Nodes/CustomizableObjectNodeMeshClipWithMeshDetails.cpp
--- a/Engine/Plugins/Experimental/Mutable/Source/CustomizableObjectEditor/Private/MuCOE/Nodes/CustomizableObjectNodeMeshClipWithMeshDetails.cpp
+++ b/Engine/Plugins/Experimental/Mutable/Source/CustomizableObjectEditor/Private/MuCOE/Nodes/CustomizableObjectNodeMeshClipWithMeshDetails.cpp
@@ -31,6 +31,8 @@ TSharedRef<IDetailCustomization> FCustomizableObjectNodeMeshClipWithMeshDetails:
 
 void FCustomizableObjectNodeMeshClipWithMeshDetails::CustomizeDetails(IDetailLayoutBuilder& DetailBuilder)
 {
+	FCustomizableObjectNodeDetails::CustomizeDetails(DetailBuilder);
+
 	Node = nullptr;
 	DetailBuilderPtr = &DetailBuilder;
 
diff --git a/Engine/Plugins/Experimental/Mutable/Source/CustomizableObjectEditor/Private/MuCOE/Nodes/CustomizableObjectNodeMeshClipWithMeshDetails.h b/Engine/Plugins/Experimental/Mutable/Source/CustomizableObjectEditor/Private/MuCOE/Nodes/CustomizableObjectNodeMeshClipWithMeshDetails.h
--- a/Engine/Plugins/Experimental/Mutable/Source/CustomizableObjectEditor/Private/MuCOE/Nodes/CustomizableObjectNodeMeshClipWithMeshDetails.h
+++ b/Engine/Plugins/Experimental/Mutable/Source/CustomizableObjectEditor/Private/MuCOE/Nodes/CustomizableObjectNodeMeshClipWithMeshDetails.h
@@ -2,6 +2,7 @@
 
 #pragma once
 
+#include "CustomizableObjectNodeDetails.h"
 #include "IDetailCustomization.h"
 #include "Widgets/SCompoundWidget.h"
 
@@ -16,7 +17,7 @@ class UCustomizableObjectNodeMeshClipWithMesh;
 struct FAssetData;
 
 
-class FCustomizableObjectNodeMeshClipWithMeshDetails : public IDetailCustomization
+class FCustomizableObjectNodeMeshClipWithMeshDetails : public FCustomizableObjectNodeDetails
 {
 public:
 	/** Makes a new instance of this detail layout class for a specific detail view requesting it */
diff --git a/Engine/Plugins/Experimental/Mutable/Source/CustomizableObjectEditor/Private/MuCOE/Nodes/CustomizableObjectNodeMeshMorphDetails.cpp b/Engine/Plugins/Experimental/Mutable/Source/CustomizableObjectEditor/Private/MuCOE/Nodes/CustomizableObjectNodeMeshMorphDetails.cpp
--- a/Engine/Plugins/Experimental/Mutable/Source/CustomizableObjectEditor/Private/MuCOE/Nodes/CustomizableObjectNodeMeshMorphDetails.cpp
+++ b/Engine/Plugins/Experimental/Mutable/Source/CustomizableObjectEditor/Private/MuCOE/Nodes/CustomizableObjectNodeMeshMorphDetails.cpp
@@ -24,6 +24,8 @@ TSharedRef<IDetailCustomization> FCustomizableObjectNodeMeshMorphDetails::MakeIn
 
 void FCustomizableObjectNodeMeshMorphDetails::CustomizeDetails( IDetailLayoutBuilder& DetailBuilder )
 {
+	FCustomizableObjectNodeDetails::CustomizeDetails(DetailBuilder);
+
 	const IDetailsView* DetailsView = DetailBuilder.GetDetailsView();
 	if ( DetailsView->GetSelectedObjects().Num() )
 	{
diff --git a/Engine/Plugins/Experimental/Mutable/Source/CustomizableObjectEditor/Private/MuCOE/Nodes/CustomizableObjectNodeMeshMorphDetails.h b/Engine/Plugins/Experimental/Mutable/Source/CustomizableObjectEditor/Private/MuCOE/Nodes/CustomizableObjectNodeMeshMorphDetails.h
--- a/Engine/Plugins/Experimental/Mutable/Source/CustomizableObjectEditor/Private/MuCOE/Nodes/CustomizableObjectNodeMeshMorphDetails.h
+++ b/Engine/Plugins/Experimental/Mutable/Source/CustomizableObjectEditor/Private/MuCOE/Nodes/CustomizableObjectNodeMeshMorphDetails.h
@@ -2,6 +2,7 @@
 
 #pragma once
 
+#include "CustomizableObjectNodeDetails.h"
 #include "IDetailCustomization.h"
 
 namespace ESelectInfo { enum Type : int; }
@@ -10,7 +11,7 @@ class FString;
 class IDetailLayoutBuilder;
 class IPropertyHandle;
 
-class FCustomizableObjectNodeMeshMorphDetails : public IDetailCustomization
+class FCustomizableObjectNodeMeshMorphDetails : public FCustomizableObjectNodeDetails
 {
 public:
 	/** Makes a new instance of this detail layout class for a specific detail view requesting it */
diff --git a/Engine/Plugins/Experimental/Mutable/Source/CustomizableObjectEditor/Private/MuCOE/Nodes/CustomizableObjectNodeMorphMaterialDetails.cpp b/Engine/Plugins/Experimental/Mutable/Source/CustomizableObjectEditor/Private/MuCOE/Nodes/CustomizableObjectNodeMorphMaterialDetails.cpp
--- a/Engine/Plugins/Experimental/Mutable/Source/CustomizableObjectEditor/Private/MuCOE/Nodes/CustomizableObjectNodeMorphMaterialDetails.cpp
+++ b/Engine/Plugins/Experimental/Mutable/Source/CustomizableObjectEditor/Private/MuCOE/Nodes/CustomizableObjectNodeMorphMaterialDetails.cpp
@@ -26,6 +26,7 @@ TSharedRef<IDetailCustomization> FCustomizableObjectNodeMorphMaterialDetails::Ma
 
 void FCustomizableObjectNodeMorphMaterialDetails::CustomizeDetails( IDetailLayoutBuilder& DetailBuilder )
 {
+	FCustomizableObjectNodeDetails::CustomizeDetails(DetailBuilder);
 
 	const IDetailsView* DetailsView = DetailBuilder.GetDetailsView();
 	if ( DetailsView->GetSelectedObjects().Num() )
diff --git a/Engine/Plugins/Experimental/Mutable/Source/CustomizableObjectEditor/Private/MuCOE/Nodes/CustomizableObjectNodeMorphMaterialDetails.h b/Engine/Plugins/Experimental/Mutable/Source/CustomizableObjectEditor/Private/MuCOE/Nodes/CustomizableObjectNodeMorphMaterialDetails.h
--- a/Engine/Plugins/Experimental/Mutable/Source/CustomizableObjectEditor/Private/MuCOE/Nodes/CustomizableObjectNodeMorphMaterialDetails.h
+++ b/Engine/Plugins/Experimental/Mutable/Source/CustomizableObjectEditor/Private/MuCOE/Nodes/CustomizableObjectNodeMorphMaterialDetails.h
@@ -2,6 +2,7 @@
 
 #pragma once
 
+#include "CustomizableObjectNodeDetails.h"
 #include "IDetailCustomization.h"
 
 namespace ESelectInfo { enum Type : int; }
@@ -11,7 +12,7 @@ class IDetailLayoutBuilder;
 class IPropertyHandle;
 class UCustomizableObjectNodeObject;
 
-class FCustomizableObjectNodeMorphMaterialDetails : public IDetailCustomization
+class FCustomizableObjectNodeMorphMaterialDetails : public FCustomizableObjectNodeDetails
 {
 public:
 	// Makes a new instance of this detail layout class for a specific detail view requesting it 
diff --git a/Engine/Plugins/Experimental/Mutable/Source/CustomizableObjectEditor/Private/MuCOE/Nodes/CustomizableObjectNodeObjectDetails.cpp b/Engine/Plugins/Experimental/Mutable/Source/CustomizableObjectEditor/Private/MuCOE/Nodes/CustomizableObjectNodeObjectDetails.cpp
--- a/Engine/Plugins/Experimental/Mutable/Source/CustomizableObjectEditor/Private/MuCOE/Nodes/CustomizableObjectNodeObjectDetails.cpp
+++ b/Engine/Plugins/Experimental/Mutable/Source/CustomizableObjectEditor/Private/MuCOE/Nodes/CustomizableObjectNodeObjectDetails.cpp
@@ -156,6 +156,8 @@ TSharedRef<IDetailCustomization> FCustomizableObjectNodeObjectDetails::MakeInsta
 
 void FCustomizableObjectNodeObjectDetails::CustomizeDetails( IDetailLayoutBuilder& DetailBuilder )
 {
+	FCustomizableObjectNodeDetails::CustomizeDetails(DetailBuilder);
+
 	BaseObjectNode = nullptr;
 	DetailBuilderPtr = &DetailBuilder;
 
diff --git a/Engine/Plugins/Experimental/Mutable/Source/CustomizableObjectEditor/Private/MuCOE/Nodes/CustomizableObjectNodeObjectDetails.h b/Engine/Plugins/Experimental/Mutable/Source/CustomizableObjectEditor/Private/MuCOE/Nodes/CustomizableObjectNodeObjectDetails.h
--- a/Engine/Plugins/Experimental/Mutable/Source/CustomizableObjectEditor/Private/MuCOE/Nodes/CustomizableObjectNodeObjectDetails.h
+++ b/Engine/Plugins/Experimental/Mutable/Source/CustomizableObjectEditor/Private/MuCOE/Nodes/CustomizableObjectNodeObjectDetails.h
@@ -2,6 +2,7 @@
 
 #pragma once
 
+#include "CustomizableObjectNodeDetails.h"
 #include "IDetailCustomization.h"
 #include "IPropertyTypeCustomization.h"
 #include "PropertyEditorModule.h"
@@ -79,7 +80,7 @@ private:
 };
 
 
-class FCustomizableObjectNodeObjectDetails : public IDetailCustomization
+class FCustomizableObjectNodeObjectDetails : public FCustomizableObjectNodeDetails
 {
 public:
 	/** Makes a new instance of this detail layout class for a specific detail view requesting it */
diff --git a/Engine/Plugins/Experimental/Mutable/Source/CustomizableObjectEditor/Private/MuCOE/Nodes/CustomizableObjectNodeParentedMaterialDetails.cpp b/Engine/Plugins/Experimental/Mutable/Source/CustomizableObjectEditor/Private/MuCOE/Nodes/CustomizableObjectNodeParentedMaterialDetails.cpp
--- a/Engine/Plugins/Experimental/Mutable/Source/CustomizableObjectEditor/Private/MuCOE/Nodes/CustomizableObjectNodeParentedMaterialDetails.cpp
+++ b/Engine/Plugins/Experimental/Mutable/Source/CustomizableObjectEditor/Private/MuCOE/Nodes/CustomizableObjectNodeParentedMaterialDetails.cpp
@@ -30,6 +30,8 @@ TSharedRef<IDetailCustomization> FCustomizableObjectNodeParentedMaterialDetails:
 
 void FCustomizableObjectNodeParentedMaterialDetails::CustomizeDetails(IDetailLayoutBuilder& DetailBuilder)
 {
+	FCustomizableObjectNodeDetails::CustomizeDetails(DetailBuilder);
+
 	Node = nullptr;
 	NodeParentedMaterial = nullptr;
 
diff --git a/Engine/Plugins/Experimental/Mutable/Source/CustomizableObjectEditor/Private/MuCOE/Nodes/CustomizableObjectNodeParentedMaterialDetails.h b/Engine/Plugins/Experimental/Mutable/Source/CustomizableObjectEditor/Private/MuCOE/Nodes/CustomizableObjectNodeParentedMaterialDetails.h
--- a/Engine/Plugins/Experimental/Mutable/Source/CustomizableObjectEditor/Private/MuCOE/Nodes/CustomizableObjectNodeParentedMaterialDetails.h
+++ b/Engine/Plugins/Experimental/Mutable/Source/CustomizableObjectEditor/Private/MuCOE/Nodes/CustomizableObjectNodeParentedMaterialDetails.h
@@ -2,6 +2,7 @@
 
 #pragma once
 
+#include "CustomizableObjectNodeDetails.h"
 #include "IDetailCustomization.h"
 #include "Misc/Guid.h"
 #include "UObject/WeakObjectPtr.h"
@@ -15,7 +16,7 @@ class UCustomizableObject;
 class UCustomizableObjectNode;
 class UCustomizableObjectNodeMaterial;
 
-class FCustomizableObjectNodeParentedMaterialDetails : public IDetailCustomization
+class FCustomizableObjectNodeParentedMaterialDetails : public FCustomizableObjectNodeDetails
 {
 public:
 	/** Makes a new instance of this detail layout class for a specific detail view requesting it */
diff --git a/Engine/Plugins/Experimental/Mutable/Source/CustomizableObjectEditor/Private/MuCOE/Nodes/CustomizableObjectNodePinViewerDetails.cpp b/Engine/Plugins/Experimental/Mutable/Source/CustomizableObjectEditor/Private/MuCOE/Nodes/CustomizableObjectNodePinViewerDetails.cpp
deleted file mode 100644
diff --git a/Engine/Plugins/Experimental/Mutable/Source/CustomizableObjectEditor/Private/MuCOE/Nodes/CustomizableObjectNodePinViewerDetails.h b/Engine/Plugins/Experimental/Mutable/Source/CustomizableObjectEditor/Private/MuCOE/Nodes/CustomizableObjectNodePinViewerDetails.h
deleted file mode 100644
diff --git a/Engine/Plugins/Experimental/Mutable/Source/CustomizableObjectEditor/Private/MuCOE/Nodes/CustomizableObjectNodeProjectorParameterDetails.cpp b/Engine/Plugins/Experimental/Mutable/Source/CustomizableObjectEditor/Private/MuCOE/Nodes/CustomizableObjectNodeProjectorParameterDetails.cpp
--- a/Engine/Plugins/Experimental/Mutable/Source/CustomizableObjectEditor/Private/MuCOE/Nodes/CustomizableObjectNodeProjectorParameterDetails.cpp
+++ b/Engine/Plugins/Experimental/Mutable/Source/CustomizableObjectEditor/Private/MuCOE/Nodes/CustomizableObjectNodeProjectorParameterDetails.cpp
@@ -29,6 +29,8 @@ TSharedRef<IDetailCustomization> FCustomizableObjectNodeProjectorParameterDetail
 
 void FCustomizableObjectNodeProjectorParameterDetails::CustomizeDetails( IDetailLayoutBuilder& DetailBuilder )
 {
+	FCustomizableObjectNodeDetails::CustomizeDetails(DetailBuilder);
+
 	NodeConstant = nullptr;
 	NodeParameter = nullptr;
 	DetailBuilderPtr = &DetailBuilder;
diff --git a/Engine/Plugins/Experimental/Mutable/Source/CustomizableObjectEditor/Private/MuCOE/Nodes/CustomizableObjectNodeProjectorParameterDetails.h b/Engine/Plugins/Experimental/Mutable/Source/CustomizableObjectEditor/Private/MuCOE/Nodes/CustomizableObjectNodeProjectorParameterDetails.h
--- a/Engine/Plugins/Experimental/Mutable/Source/CustomizableObjectEditor/Private/MuCOE/Nodes/CustomizableObjectNodeProjectorParameterDetails.h
+++ b/Engine/Plugins/Experimental/Mutable/Source/CustomizableObjectEditor/Private/MuCOE/Nodes/CustomizableObjectNodeProjectorParameterDetails.h
@@ -2,6 +2,7 @@
 
 #pragma once
 
+#include "CustomizableObjectNodeDetails.h"
 #include "IDetailCustomization.h"
 
 class FReply;
@@ -12,7 +13,7 @@ class IDetailLayoutBuilder;
 class IPropertyHandle;
 
 
-class FCustomizableObjectNodeProjectorParameterDetails : public IDetailCustomization
+class FCustomizableObjectNodeProjectorParameterDetails : public FCustomizableObjectNodeDetails
 {
 public:
 	/** Makes a new instance of this detail layout class for a specific detail view requesting it */
diff --git a/Engine/Plugins/Experimental/Mutable/Source/CustomizableObjectEditor/Private/MuCOE/Nodes/CustomizableObjectNodeRemoveMeshBlocksDetails.cpp b/Engine/Plugins/Experimental/Mutable/Source/CustomizableObjectEditor/Private/MuCOE/Nodes/CustomizableObjectNodeRemoveMeshBlocksDetails.cpp
--- a/Engine/Plugins/Experimental/Mutable/Source/CustomizableObjectEditor/Private/MuCOE/Nodes/CustomizableObjectNodeRemoveMeshBlocksDetails.cpp
+++ b/Engine/Plugins/Experimental/Mutable/Source/CustomizableObjectEditor/Private/MuCOE/Nodes/CustomizableObjectNodeRemoveMeshBlocksDetails.cpp
@@ -24,6 +24,8 @@ TSharedRef<IDetailCustomization> FCustomizableObjectNodeRemoveMeshBlocksDetails:
 
 void FCustomizableObjectNodeRemoveMeshBlocksDetails::CustomizeDetails( IDetailLayoutBuilder& DetailBuilder )
 {
+	FCustomizableObjectNodeDetails::CustomizeDetails(DetailBuilder);
+
 	Node = nullptr;
 	const IDetailsView* DetailsView = DetailBuilder.GetDetailsView();
 	if (DetailsView->GetSelectedObjects().Num())
diff --git a/Engine/Plugins/Experimental/Mutable/Source/CustomizableObjectEditor/Private/MuCOE/Nodes/CustomizableObjectNodeRemoveMeshBlocksDetails.h b/Engine/Plugins/Experimental/Mutable/Source/CustomizableObjectEditor/Private/MuCOE/Nodes/CustomizableObjectNodeRemoveMeshBlocksDetails.h
--- a/Engine/Plugins/Experimental/Mutable/Source/CustomizableObjectEditor/Private/MuCOE/Nodes/CustomizableObjectNodeRemoveMeshBlocksDetails.h
+++ b/Engine/Plugins/Experimental/Mutable/Source/CustomizableObjectEditor/Private/MuCOE/Nodes/CustomizableObjectNodeRemoveMeshBlocksDetails.h
@@ -3,6 +3,7 @@
 #pragma once
 
 
+#include "CustomizableObjectNodeDetails.h"
 #include "IDetailCustomization.h"
 
 class FString;
@@ -11,7 +12,7 @@ class UCustomizableObjectNodeRemoveMeshBlocks;
 class SCustomizableObjectNodeLayoutBlocksSelector;
 
 
-class FCustomizableObjectNodeRemoveMeshBlocksDetails : public IDetailCustomization
+class FCustomizableObjectNodeRemoveMeshBlocksDetails : public FCustomizableObjectNodeDetails
 {
 public:
 	/** Makes a new instance of this detail layout class for a specific detail view requesting it */
diff --git a/Engine/Plugins/Experimental/Mutable/Source/CustomizableObjectEditor/Private/MuCOE/Nodes/CustomizableObjectNodeSkeletalMesh.cpp b/Engine/Plugins/Experimental/Mutable/Source/CustomizableObjectEditor/Private/MuCOE/Nodes/CustomizableObjectNodeSkeletalMesh.cpp
--- a/Engine/Plugins/Experimental/Mutable/Source/CustomizableObjectEditor/Private/MuCOE/Nodes/CustomizableObjectNodeSkeletalMesh.cpp
+++ b/Engine/Plugins/Experimental/Mutable/Source/CustomizableObjectEditor/Private/MuCOE/Nodes/CustomizableObjectNodeSkeletalMesh.cpp
@@ -154,6 +154,12 @@ UCustomizableObjectNodeRemapPins* UCustomizableObjectNodeSkeletalMesh::CreateRem
 }
 
 
+bool UCustomizableObjectNodeSkeletalMesh::HasPinViewer() const
+{
+	return true;
+}
+
+
 FText UCustomizableObjectNodeSkeletalMesh::GetNodeTitle(ENodeTitleType::Type TitleType) const
 {
 	if (SkeletalMesh)
diff --git a/Engine/Plugins/Experimental/Mutable/Source/CustomizableObjectEditor/Private/MuCOE/Nodes/CustomizableObjectNodeSkeletalMeshDetails.cpp b/Engine/Plugins/Experimental/Mutable/Source/CustomizableObjectEditor/Private/MuCOE/Nodes/CustomizableObjectNodeSkeletalMeshDetails.cpp
--- a/Engine/Plugins/Experimental/Mutable/Source/CustomizableObjectEditor/Private/MuCOE/Nodes/CustomizableObjectNodeSkeletalMeshDetails.cpp
+++ b/Engine/Plugins/Experimental/Mutable/Source/CustomizableObjectEditor/Private/MuCOE/Nodes/CustomizableObjectNodeSkeletalMeshDetails.cpp
@@ -20,6 +20,8 @@ TSharedRef<IDetailCustomization> FCustomizableObjectNodeSkeletalMeshDetails::Mak
 
 void FCustomizableObjectNodeSkeletalMeshDetails::CustomizeDetails(IDetailLayoutBuilder& DetailBuilder)
 {
+	FCustomizableObjectNodeDetails::CustomizeDetails(DetailBuilder);
+
 	const IDetailsView* DetailsView = DetailBuilder.GetDetailsView();
 
     if (DetailsView->GetSelectedObjects().Num())
@@ -49,8 +51,6 @@ void FCustomizableObjectNodeSkeletalMeshDetails::CustomizeDetails(IDetailLayoutB
 
 	TSharedRef<IPropertyHandle> SkeletalMeshProperty = DetailBuilder.GetProperty(GET_MEMBER_NAME_CHECKED(UCustomizableObjectNodeSkeletalMesh, SkeletalMesh));
 	SkeletalMeshProperty->SetOnPropertyValueChanged(FSimpleDelegate::CreateSP(MorphSelector.Get(), &SCustomizableObjectNodeSkeletalMeshRTMorphSelector::UpdateWidget));
-
-	PinViewerAttachToDetailCustomization(DetailBuilder);
 }
 
 
diff --git a/Engine/Plugins/Experimental/Mutable/Source/CustomizableObjectEditor/Private/MuCOE/Nodes/CustomizableObjectNodeSkeletalMeshDetails.h b/Engine/Plugins/Experimental/Mutable/Source/CustomizableObjectEditor/Private/MuCOE/Nodes/CustomizableObjectNodeSkeletalMeshDetails.h
--- a/Engine/Plugins/Experimental/Mutable/Source/CustomizableObjectEditor/Private/MuCOE/Nodes/CustomizableObjectNodeSkeletalMeshDetails.h
+++ b/Engine/Plugins/Experimental/Mutable/Source/CustomizableObjectEditor/Private/MuCOE/Nodes/CustomizableObjectNodeSkeletalMeshDetails.h
@@ -2,12 +2,13 @@
 
 #pragma once
 
+#include "CustomizableObjectNodeDetails.h"
 #include "IDetailCustomization.h"
 
 class IDetailLayoutBuilder;
 
 
-class FCustomizableObjectNodeSkeletalMeshDetails : public IDetailCustomization
+class FCustomizableObjectNodeSkeletalMeshDetails : public FCustomizableObjectNodeDetails
 {
 public:
 	/** Makes a new instance of this detail layout class for a specific detail view requesting it */
diff --git a/Engine/Plugins/Experimental/Mutable/Source/CustomizableObjectEditor/Private/MuCOE/Nodes/CustomizableObjectNodeStaticMesh.cpp b/Engine/Plugins/Experimental/Mutable/Source/CustomizableObjectEditor/Private/MuCOE/Nodes/CustomizableObjectNodeStaticMesh.cpp
--- a/Engine/Plugins/Experimental/Mutable/Source/CustomizableObjectEditor/Private/MuCOE/Nodes/CustomizableObjectNodeStaticMesh.cpp
+++ b/Engine/Plugins/Experimental/Mutable/Source/CustomizableObjectEditor/Private/MuCOE/Nodes/CustomizableObjectNodeStaticMesh.cpp
@@ -341,6 +341,12 @@ UCustomizableObjectNodeRemapPinsByName* UCustomizableObjectNodeStaticMesh::Creat
 }
 
 
+bool UCustomizableObjectNodeStaticMesh::HasPinViewer() const
+{
+	return true;
+}
+
+
 bool UCustomizableObjectNodeStaticMesh::IsNodeOutDatedAndNeedsRefresh()
 {
 	bool Result = false;
diff --git a/Engine/Plugins/Experimental/Mutable/Source/CustomizableObjectEditor/Private/MuCOE/Nodes/CustomizableObjectNodeTable.cpp b/Engine/Plugins/Experimental/Mutable/Source/CustomizableObjectEditor/Private/MuCOE/Nodes/CustomizableObjectNodeTable.cpp
--- a/Engine/Plugins/Experimental/Mutable/Source/CustomizableObjectEditor/Private/MuCOE/Nodes/CustomizableObjectNodeTable.cpp
+++ b/Engine/Plugins/Experimental/Mutable/Source/CustomizableObjectEditor/Private/MuCOE/Nodes/CustomizableObjectNodeTable.cpp
@@ -1240,6 +1240,12 @@ UCustomizableObjectNodeTableRemapPins* UCustomizableObjectNodeTable::CreateRemap
 }
 
 
+bool UCustomizableObjectNodeTable::HasPinViewer() const
+{
+	return true;
+}
+
+
 bool UCustomizableObjectNodeTable::IsImagePinDefault(const UEdGraphPin* Pin) const
 {
 	if (UCustomizableObjectNodeTableImagePinData* PinData = Cast<UCustomizableObjectNodeTableImagePinData>(GetPinData(*(Pin))))
diff --git a/Engine/Plugins/Experimental/Mutable/Source/CustomizableObjectEditor/Private/MuCOE/Nodes/CustomizableObjectNodeTableDetails.cpp b/Engine/Plugins/Experimental/Mutable/Source/CustomizableObjectEditor/Private/MuCOE/Nodes/CustomizableObjectNodeTableDetails.cpp
--- a/Engine/Plugins/Experimental/Mutable/Source/CustomizableObjectEditor/Private/MuCOE/Nodes/CustomizableObjectNodeTableDetails.cpp
+++ b/Engine/Plugins/Experimental/Mutable/Source/CustomizableObjectEditor/Private/MuCOE/Nodes/CustomizableObjectNodeTableDetails.cpp
@@ -36,6 +36,8 @@ TSharedRef<IDetailCustomization> FCustomizableObjectNodeTableDetails::MakeInstan
 
 void FCustomizableObjectNodeTableDetails::CustomizeDetails(const TSharedPtr<IDetailLayoutBuilder>& DetailBuilder)
 {
+	IDetailCustomization::CustomizeDetails(DetailBuilder);
+	
 	Node = 0;
 	DetailBuilderPtr = DetailBuilder;
 
diff --git a/Engine/Plugins/Experimental/Mutable/Source/CustomizableObjectEditor/Private/MuCOE/Nodes/CustomizableObjectNodeTableDetails.h b/Engine/Plugins/Experimental/Mutable/Source/CustomizableObjectEditor/Private/MuCOE/Nodes/CustomizableObjectNodeTableDetails.h
--- a/Engine/Plugins/Experimental/Mutable/Source/CustomizableObjectEditor/Private/MuCOE/Nodes/CustomizableObjectNodeTableDetails.h
+++ b/Engine/Plugins/Experimental/Mutable/Source/CustomizableObjectEditor/Private/MuCOE/Nodes/CustomizableObjectNodeTableDetails.h
@@ -2,6 +2,7 @@
 
 #pragma once
 
+#include "CustomizableObjectNodeDetails.h"
 #include "IDetailCustomization.h"
 
 #include "UObject/WeakObjectPtr.h"
@@ -32,16 +33,13 @@ enum class EAnimColumnType
 
 
 /** Copy Material node details panel. Hides all properties from the inheret Material node. */
-class FCustomizableObjectNodeTableDetails : public IDetailCustomization
+class FCustomizableObjectNodeTableDetails : public FCustomizableObjectNodeDetails
 {
 public:
 	
 	static TSharedRef<IDetailCustomization> MakeInstance();
 
 	// ILayoutDetails interface
-	/** Do not use. Add details customization in the other CustomizeDetails signature. */
-	virtual void CustomizeDetails(IDetailLayoutBuilder& DetailBuilder) override {};
-
 	/** Hides details copied from CustomizableObjectNodeMaterial. */
 	virtual void CustomizeDetails(const TSharedPtr<IDetailLayoutBuilder>& DetailBuilder) override;
 
diff --git a/Engine/Plugins/Experimental/Mutable/Source/CustomizableObjectEditor/Private/MuCOE/PinViewer/SPinViewer.cpp b/Engine/Plugins/Experimental/Mutable/Source/CustomizableObjectEditor/Private/MuCOE/PinViewer/SPinViewer.cpp
--- a/Engine/Plugins/Experimental/Mutable/Source/CustomizableObjectEditor/Private/MuCOE/PinViewer/SPinViewer.cpp
+++ b/Engine/Plugins/Experimental/Mutable/Source/CustomizableObjectEditor/Private/MuCOE/PinViewer/SPinViewer.cpp
@@ -331,19 +331,4 @@ EColumnSortMode::Type SPinViewer::GetColumnSortMode(const FName ColumnId) const
 }
 
 
-void PinViewerAttachToDetailCustomization(IDetailLayoutBuilder& DetailBuilder)
-{
-	const TArray<TWeakObjectPtr<UObject>>& SelectedObjects = DetailBuilder.GetDetailsView()->GetSelectedObjects();
-	if (SelectedObjects.Num())
-	{
-		UCustomizableObjectNode* Node = Cast<UCustomizableObjectNode>(SelectedObjects[0].Get());
-
-		IDetailCategoryBuilder& PinViewerCategoryBuilder = DetailBuilder.EditCategory("PinViewer", LOCTEXT("PinViewer", "Pins"), ECategoryPriority::Uncommon);
-		PinViewerCategoryBuilder.AddCustomRow(LOCTEXT("PinViewerDetailsCategory", "PinViwer")).ShouldAutoExpand(true)
-		[
-			SNew(SPinViewer).Node(Node)
-		];
-	}
-}
-
 #undef LOCTEXT_NAMESPACE
diff --git a/Engine/Plugins/Experimental/Mutable/Source/CustomizableObjectEditor/Private/MuCOE/PinViewer/SPinViewer.h b/Engine/Plugins/Experimental/Mutable/Source/CustomizableObjectEditor/Private/MuCOE/PinViewer/SPinViewer.h
--- a/Engine/Plugins/Experimental/Mutable/Source/CustomizableObjectEditor/Private/MuCOE/PinViewer/SPinViewer.h
+++ b/Engine/Plugins/Experimental/Mutable/Source/CustomizableObjectEditor/Private/MuCOE/PinViewer/SPinViewer.h
@@ -94,7 +94,3 @@ private:
 	TSharedPtr<class SScrollBox> Scrollbox;
 };
 
-
-/** Create and attach a PinViewer Widget to the given custom details DetailBuilder.
- * The pin viewer will be shown as a new category with the lowest priority. */
-void PinViewerAttachToDetailCustomization(IDetailLayoutBuilder& DetailBuilder);
diff --git a/Engine/Plugins/Experimental/Mutable/Source/CustomizableObjectEditor/Public/MuCOE/Nodes/CustomizableObjectNode.h b/Engine/Plugins/Experimental/Mutable/Source/CustomizableObjectEditor/Public/MuCOE/Nodes/CustomizableObjectNode.h
--- a/Engine/Plugins/Experimental/Mutable/Source/CustomizableObjectEditor/Public/MuCOE/Nodes/CustomizableObjectNode.h
+++ b/Engine/Plugins/Experimental/Mutable/Source/CustomizableObjectEditor/Public/MuCOE/Nodes/CustomizableObjectNode.h
@@ -230,8 +230,11 @@ public:
 	/** Return true if a pin can be hidden. Override it to have a custom behaivour. */
 	virtual bool CanPinBeHidden(const UEdGraphPin& Pin) const;
 
+	/** @return true if the details should spawn a Pin Viewer. */
+	virtual bool HasPinViewer() const;
+
 	/** Returns pin custom details. Override if required. */
-	virtual TSharedPtr<SWidget> CustomizePinDetails(UEdGraphPin& Pin);
+	virtual TSharedPtr<SWidget> CustomizePinDetails(const UEdGraphPin& Pin) const;
 	
 	/** See GetPinData(const FEdGraphPinReference&). */
 	UCustomizableObjectNodePinData* GetPinData(const UEdGraphPin &Pin) const;
diff --git a/Engine/Plugins/Experimental/Mutable/Source/CustomizableObjectEditor/Public/MuCOE/Nodes/CustomizableObjectNodeEditMaterial.h b/Engine/Plugins/Experimental/Mutable/Source/CustomizableObjectEditor/Public/MuCOE/Nodes/CustomizableObjectNodeEditMaterial.h
--- a/Engine/Plugins/Experimental/Mutable/Source/CustomizableObjectEditor/Public/MuCOE/Nodes/CustomizableObjectNodeEditMaterial.h
+++ b/Engine/Plugins/Experimental/Mutable/Source/CustomizableObjectEditor/Public/MuCOE/Nodes/CustomizableObjectNodeEditMaterial.h
@@ -57,6 +57,7 @@ public:
 	virtual FString GetRefreshMessage() const override;
 	virtual bool IsSingleOutputNode() const override;
 	virtual bool CustomRemovePin(UEdGraphPin& Pin) override;
+	virtual bool HasPinViewer() const override;
 
 	// UCustomizableObjectNodeEditMaterialBase interface
 	virtual UEdGraphPin* OutputPin() const override;
diff --git a/Engine/Plugins/Experimental/Mutable/Source/CustomizableObjectEditor/Public/MuCOE/Nodes/CustomizableObjectNodeMaterial.h b/Engine/Plugins/Experimental/Mutable/Source/CustomizableObjectEditor/Public/MuCOE/Nodes/CustomizableObjectNodeMaterial.h
--- a/Engine/Plugins/Experimental/Mutable/Source/CustomizableObjectEditor/Public/MuCOE/Nodes/CustomizableObjectNodeMaterial.h
+++ b/Engine/Plugins/Experimental/Mutable/Source/CustomizableObjectEditor/Public/MuCOE/Nodes/CustomizableObjectNodeMaterial.h
@@ -142,6 +142,7 @@ public:
 	virtual void PostBackwardsCompatibleFixup() override;
 	void AllocateDefaultPins(UCustomizableObjectNodeRemapPins* RemapPins) override;
 	bool CanPinBeHidden(const UEdGraphPin& Pin) const override;
+	virtual bool HasPinViewer() const override;
 	virtual UCustomizableObjectNodeRemapPinsByName* CreateRemapPinsDefault() const override;
 	bool ProvidesCustomPinRelevancyTest() const override { return true; }
 	bool IsPinRelevant(const UEdGraphPin* Pin) const override;
@@ -157,7 +158,7 @@ public:
 
 	bool IsNodeOutDatedAndNeedsRefresh() override;
 	FString GetRefreshMessage() const override;
-	virtual TSharedPtr<SWidget> CustomizePinDetails(UEdGraphPin& Pin) override;
+	virtual TSharedPtr<SWidget> CustomizePinDetails(const UEdGraphPin& Pin) const override;
 
 	
 	// UCustomizableObjectNodeMaterialBase interface
diff --git a/Engine/Plugins/Experimental/Mutable/Source/CustomizableObjectEditor/Public/MuCOE/Nodes/CustomizableObjectNodeSkeletalMesh.h b/Engine/Plugins/Experimental/Mutable/Source/CustomizableObjectEditor/Public/MuCOE/Nodes/CustomizableObjectNodeSkeletalMesh.h
--- a/Engine/Plugins/Experimental/Mutable/Source/CustomizableObjectEditor/Public/MuCOE/Nodes/CustomizableObjectNodeSkeletalMesh.h
+++ b/Engine/Plugins/Experimental/Mutable/Source/CustomizableObjectEditor/Public/MuCOE/Nodes/CustomizableObjectNodeSkeletalMesh.h
@@ -212,7 +212,8 @@ public:
 	virtual void BackwardsCompatibleFixup() override;
 	virtual void AllocateDefaultPins(UCustomizableObjectNodeRemapPins* RemapPins) override;
 	virtual UCustomizableObjectNodeRemapPins* CreateRemapPinsDefault() const override;
-
+	virtual bool HasPinViewer() const override;
+	
 	// UCustomizableObjectNodeMesh interface
 	virtual UTexture2D* FindTextureForPin(const UEdGraphPin* Pin) const override;
 	virtual void GetUVChannelForPin(const UEdGraphPin* Pin, TArray<FVector2f>& OutSegments, int32 UVIndex) const override;
diff --git a/Engine/Plugins/Experimental/Mutable/Source/CustomizableObjectEditor/Public/MuCOE/Nodes/CustomizableObjectNodeStaticMesh.h b/Engine/Plugins/Experimental/Mutable/Source/CustomizableObjectEditor/Public/MuCOE/Nodes/CustomizableObjectNodeStaticMesh.h
--- a/Engine/Plugins/Experimental/Mutable/Source/CustomizableObjectEditor/Public/MuCOE/Nodes/CustomizableObjectNodeStaticMesh.h
+++ b/Engine/Plugins/Experimental/Mutable/Source/CustomizableObjectEditor/Public/MuCOE/Nodes/CustomizableObjectNodeStaticMesh.h
@@ -140,10 +140,11 @@ public:
 	FText GetTooltipText() const override;
 
 	// UCustomizableObjectNode interface
-	void AllocateDefaultPins(UCustomizableObjectNodeRemapPins* RemapPins) override;
-	bool ProvidesCustomPinRelevancyTest() const override { return true; }
-	bool IsPinRelevant(const UEdGraphPin* Pin) const override;
-	UCustomizableObjectNodeRemapPinsByName* CreateRemapPinsByName() const override;
+	virtual void AllocateDefaultPins(UCustomizableObjectNodeRemapPins* RemapPins) override;
+	virtual bool ProvidesCustomPinRelevancyTest() const override { return true; }
+	virtual bool IsPinRelevant(const UEdGraphPin* Pin) const override;
+	virtual UCustomizableObjectNodeRemapPinsByName* CreateRemapPinsByName() const override;
+	virtual bool HasPinViewer() const override;
 
 	bool IsNodeOutDatedAndNeedsRefresh() override;
 	FString GetRefreshMessage() const override;
diff --git a/Engine/Plugins/Experimental/Mutable/Source/CustomizableObjectEditor/Public/MuCOE/Nodes/CustomizableObjectNodeTable.h b/Engine/Plugins/Experimental/Mutable/Source/CustomizableObjectEditor/Public/MuCOE/Nodes/CustomizableObjectNodeTable.h
--- a/Engine/Plugins/Experimental/Mutable/Source/CustomizableObjectEditor/Public/MuCOE/Nodes/CustomizableObjectNodeTable.h
+++ b/Engine/Plugins/Experimental/Mutable/Source/CustomizableObjectEditor/Public/MuCOE/Nodes/CustomizableObjectNodeTable.h
@@ -260,7 +260,8 @@ public:
 	virtual bool ProvidesCustomPinRelevancyTest() const override;
 	virtual bool IsPinRelevant(const UEdGraphPin* Pin) const override;
 	UCustomizableObjectNodeTableRemapPins* CreateRemapPinsDefault() const;
-
+	virtual bool HasPinViewer() const override;
+	
 	/*** Allows to perform work when remapping the pin data. */
 	virtual void RemapPinsData(const TMap<UEdGraphPin*, UEdGraphPin*>& PinsToRemap) override;
 	
