[mutable] Lower log message severity for old data because some branches want to keep the old data.
#rnx

diff --git a/Engine/Plugins/Experimental/Mutable/Source/CustomizableObjectEditor/Private/MuCOE/Nodes/CustomizableObjectNodeObject.cpp b/Engine/Plugins/Experimental/Mutable/Source/CustomizableObjectEditor/Private/MuCOE/Nodes/CustomizableObjectNodeObject.cpp
--- a/Engine/Plugins/Experimental/Mutable/Source/CustomizableObjectEditor/Private/MuCOE/Nodes/CustomizableObjectNodeObject.cpp
+++ b/Engine/Plugins/Experimental/Mutable/Source/CustomizableObjectEditor/Private/MuCOE/Nodes/CustomizableObjectNodeObject.cpp
@@ -65,7 +65,7 @@ void UCustomizableObjectNodeObject::BackwardsCompatibleFixup()
 		UCustomizableObject* CustomizableObject = Cast<UCustomizableObject>(GetCustomizableObjectGraph()->GetOuter());
 		if (CustomizableObject)
 		{
-			FMessageLog("Mutable").Message(EMessageSeverity::Warning)
+			FMessageLog("Mutable").Message(EMessageSeverity::Info)
 				->AddToken(FTextToken::Create(LOCTEXT("Indeterministic Warning", "The object was saved with an old version and it may generate indeterministic packages. Resave it to fix the problem.")))
 				->AddToken(FUObjectToken::Create(CustomizableObject));
 		}
