[Mutable] Better CopyMaterial implementation.

Now NodeMaterial and CopyMaterial inherit from NodeMaterialBase. No other nodes inherit from NodeMaterialBase. This allows NodeMaterials and CopyMaterials to be completly interchangeable.

#rnx

diff --git a/Engine/Plugins/Experimental/Mutable/Source/CustomizableObjectEditor/Private/MuCOE/CustomizableObjectEditorModule.cpp b/Engine/Plugins/Experimental/Mutable/Source/CustomizableObjectEditor/Private/MuCOE/CustomizableObjectEditorModule.cpp
--- a/Engine/Plugins/Experimental/Mutable/Source/CustomizableObjectEditor/Private/MuCOE/CustomizableObjectEditorModule.cpp
+++ b/Engine/Plugins/Experimental/Mutable/Source/CustomizableObjectEditor/Private/MuCOE/CustomizableObjectEditorModule.cpp
@@ -29,7 +29,6 @@
 #include "MuCOE/CustomizableObjectInstanceFactory.h"
 #include "MuCOE/CustomizableObjectNodeObjectGroupDetails.h"
 #include "MuCOE/Nodes/CustomizableObjectNodeCopyMaterial.h"
-#include "MuCOE/Nodes/CustomizableObjectNodeCopyMaterialDetails.h"
 #include "MuCOE/Nodes/CustomizableObjectNodeEditMaterial.h"
 #include "MuCOE/Nodes/CustomizableObjectNodeEditMaterialDetails.h"
 #include "MuCOE/Nodes/CustomizableObjectNodeEditMaterialBaseDetails.h"
@@ -250,7 +249,6 @@ void FCustomizableObjectEditorModule::StartupModule()
 	RegisterCustomDetails(PropertyModule, UCustomizableObjectNodeMeshClipWithMesh::StaticClass(), FOnGetDetailCustomizationInstance::CreateStatic(&FCustomizableObjectNodeMeshClipWithMeshDetails::MakeInstance));
 	RegisterCustomDetails(PropertyModule, UCustomizableObjectNodeExternalPin::StaticClass(), FOnGetDetailCustomizationInstance::CreateStatic(&FCustomizableObjectNodeExternalPinDetails::MakeInstance));
 	RegisterCustomDetails(PropertyModule, UCustomizableObjectNodeMaterial::StaticClass(), FOnGetDetailCustomizationInstance::CreateStatic(&FCustomizableObjectNodeDetails::MakeInstance));
-	RegisterCustomDetails(PropertyModule, UCustomizableObjectNodeCopyMaterial::StaticClass(), FOnGetDetailCustomizationInstance::CreateStatic(&FCustomizableObjectNodeCopyMaterialDetails::MakeInstance));
 	RegisterCustomDetails(PropertyModule, UCustomizableObjectNodeSkeletalMesh::StaticClass(), FOnGetDetailCustomizationInstance::CreateStatic(&FCustomizableObjectNodeSkeletalMeshDetails::MakeInstance));
 	RegisterCustomDetails(PropertyModule, UCustomizableObjectNodeStaticMesh::StaticClass(), FOnGetDetailCustomizationInstance::CreateStatic(&FCustomizableObjectNodeDetails::MakeInstance));
 	RegisterCustomDetails(PropertyModule, UCustomizableObjectNodeTable::StaticClass(), FOnGetDetailCustomizationInstance::CreateStatic(&FCustomizableObjectNodeTableDetails::MakeInstance));
diff --git a/Engine/Plugins/Experimental/Mutable/Source/CustomizableObjectEditor/Private/MuCOE/EdGraphSchema_CustomizableObject.cpp b/Engine/Plugins/Experimental/Mutable/Source/CustomizableObjectEditor/Private/MuCOE/EdGraphSchema_CustomizableObject.cpp
--- a/Engine/Plugins/Experimental/Mutable/Source/CustomizableObjectEditor/Private/MuCOE/EdGraphSchema_CustomizableObject.cpp
+++ b/Engine/Plugins/Experimental/Mutable/Source/CustomizableObjectEditor/Private/MuCOE/EdGraphSchema_CustomizableObject.cpp
@@ -1070,7 +1070,7 @@ void UEdGraphSchema_CustomizableObject::DroppedAssetsOnGraph(const TArray<FAsset
 			{
 				UMaterialInterface* Material = Cast<UMaterialInterface>(Object);
 				UCustomizableObjectNodeMaterial* Node = NewObject<UCustomizableObjectNodeMaterial>(Graph);
-				Node->Material = Material;
+				Node->SetMaterial(Material);
 				GraphNode = Node;
 				break;
 			}
diff --git a/Engine/Plugins/Experimental/Mutable/Source/CustomizableObjectEditor/Private/MuCOE/GenerateMutableSource/GenerateMutableSource.cpp b/Engine/Plugins/Experimental/Mutable/Source/CustomizableObjectEditor/Private/MuCOE/GenerateMutableSource/GenerateMutableSource.cpp
--- a/Engine/Plugins/Experimental/Mutable/Source/CustomizableObjectEditor/Private/MuCOE/GenerateMutableSource/GenerateMutableSource.cpp
+++ b/Engine/Plugins/Experimental/Mutable/Source/CustomizableObjectEditor/Private/MuCOE/GenerateMutableSource/GenerateMutableSource.cpp
@@ -54,6 +54,7 @@
 #include "PlatformInfo.h"
 #include "Math/NumericLimits.h"
 #include "Hash/CityHash.h"
+#include "MuCOE/Nodes/CustomizableObjectNodeCopyMaterial.h"
 
 #define LOCTEXT_NAMESPACE "CustomizableObjectEditor"
 
@@ -1544,19 +1545,19 @@ bool AffectsCurrentComponent(const UEdGraphPin* Pin, FMutableGraphGenerationCont
 		return bAffectsCurrentComponent;
 	}
 
-	else if (const UCustomizableObjectNodeMaterial* TypedNodeMat = Cast<UCustomizableObjectNodeMaterial>(Node))
+	else if (const UCustomizableObjectNodeMaterialBase* TypedNodeMat = Cast<UCustomizableObjectNodeMaterialBase>(Node))
 	{
-		ComponentIndex = TypedNodeMat->MeshComponentIndex;
+		ComponentIndex = TypedNodeMat->GetMeshComponentIndex();
 	}
 	else if (const UCustomizableObjectNodeExtendMaterial* TypedNodeExt = Cast<UCustomizableObjectNodeExtendMaterial>(Node))
 	{
-		UCustomizableObjectNodeMaterial* OriginalParentMaterialNode = TypedNodeExt->GetParentMaterialNode();
-		ComponentIndex = OriginalParentMaterialNode ? OriginalParentMaterialNode->MeshComponentIndex : GenerationContext.CurrentMeshComponent;
+		UCustomizableObjectNodeMaterialBase* OriginalParentMaterialNode = TypedNodeExt->GetParentMaterialNode();
+		ComponentIndex = OriginalParentMaterialNode ? OriginalParentMaterialNode->GetMeshComponentIndex() : GenerationContext.CurrentMeshComponent;
 	}
 	else if (const UCustomizableObjectNodeEditMaterialBase* TypedNodeEdit = Cast<UCustomizableObjectNodeEditMaterialBase>(Node))
 	{
-		UCustomizableObjectNodeMaterial* OriginalParentMaterialNode = TypedNodeEdit->GetParentMaterialNode();
-		ComponentIndex = OriginalParentMaterialNode ? OriginalParentMaterialNode->MeshComponentIndex : GenerationContext.CurrentMeshComponent;
+		UCustomizableObjectNodeMaterialBase* OriginalParentMaterialNode = TypedNodeEdit->GetParentMaterialNode();
+		ComponentIndex = OriginalParentMaterialNode ? OriginalParentMaterialNode->GetMeshComponentIndex() : GenerationContext.CurrentMeshComponent;
 	}
 	else if (const UCustomizableObjectNodeModifierBase* TypedNodeModifier = Cast<UCustomizableObjectNodeModifierBase>(Node))
 	{
@@ -1746,7 +1747,7 @@ void PopulateReferenceSkeletalMeshesData(FMutableGraphGenerationContext& Generat
 }
 
 
-uint32 GetBaseTextureSize(const FMutableGraphGenerationContext& GenerationContext, const UCustomizableObjectNodeMaterial* Material, uint32 ImageIndex)
+uint32 GetBaseTextureSize(const FMutableGraphGenerationContext& GenerationContext, const UCustomizableObjectNodeMaterialBase* Material, uint32 ImageIndex)
 {
 	const FGeneratedImageProperties* ImageProperties = GenerationContext.ImageProperties.Find({ Material, ImageIndex });
 	return ImageProperties ? ImageProperties->TextureSize : 0;
diff --git a/Engine/Plugins/Experimental/Mutable/Source/CustomizableObjectEditor/Private/MuCOE/GenerateMutableSource/GenerateMutableSource.h b/Engine/Plugins/Experimental/Mutable/Source/CustomizableObjectEditor/Private/MuCOE/GenerateMutableSource/GenerateMutableSource.h
--- a/Engine/Plugins/Experimental/Mutable/Source/CustomizableObjectEditor/Private/MuCOE/GenerateMutableSource/GenerateMutableSource.h
+++ b/Engine/Plugins/Experimental/Mutable/Source/CustomizableObjectEditor/Private/MuCOE/GenerateMutableSource/GenerateMutableSource.h
@@ -27,6 +27,7 @@
 #include "MuT/Table.h"
 #include "UObject/Package.h"
 
+class UCustomizableObjectNodeMaterialBase;
 class UCustomizableObjectNodeMeshClipWithMesh;
 class UCustomizableObjectNodeTable;
 struct FCustomizableObjectClothingAssetData;
@@ -175,7 +176,7 @@ struct FGeneratedImageKey
 
 struct FGeneratedImagePropertiesKey
 {
-	FGeneratedImagePropertiesKey(const UCustomizableObjectNodeMaterial* InMaterial, uint32 InImageIndex)
+	FGeneratedImagePropertiesKey(const UCustomizableObjectNodeMaterialBase* InMaterial, uint32 InImageIndex)
 	{
 		MaterialReferenceId = (PTRINT)InMaterial;
 		ImageIndex = InImageIndex;
@@ -330,7 +331,7 @@ struct FGroupProjectorImageInfo
 	mu::NodeImagePtr ImageNode;
 	mu::NodeImagePtr ImageResizeNode;
 	mu::NodeSurfaceNewPtr SurfNode;
-	UCustomizableObjectNodeMaterial* TypedNodeMat;
+	UCustomizableObjectNodeMaterialBase* TypedNodeMat;
 	FString TextureName;
 	FString RealTextureName;
 	FString AlternateResStateName;
@@ -338,7 +339,7 @@ struct FGroupProjectorImageInfo
 	bool bIsAlternateResolutionResized = false;
 	int32 UVLayout = 0;
 
-	FGroupProjectorImageInfo(mu::NodeImagePtr InImageNode, const FString& InTextureName, const FString& InRealTextureName, UCustomizableObjectNodeMaterial* InTypedNodeMat,
+	FGroupProjectorImageInfo(mu::NodeImagePtr InImageNode, const FString& InTextureName, const FString& InRealTextureName, UCustomizableObjectNodeMaterialBase* InTypedNodeMat,
 		float InAlternateProjectionResolutionFactor, const FString& InAlternateResStateName, mu::NodeSurfaceNewPtr InSurfNode, int32 InUVLayout)
 		: TypedNodeMat(InTypedNodeMat), TextureName(InTextureName), RealTextureName(InRealTextureName),
 		AlternateResStateName(InAlternateResStateName), AlternateProjectionResolutionFactor(InAlternateProjectionResolutionFactor), 
@@ -348,7 +349,7 @@ struct FGroupProjectorImageInfo
 		SurfNode = InSurfNode;
 	}
 
-	static FString GenerateId(const UCustomizableObjectNodeMaterialBase* TypedNodeMat, int32 ImageIndex)
+	static FString GenerateId(const UCustomizableObjectNode* TypedNodeMat, int32 ImageIndex)
 	{
 		return TypedNodeMat->GetOutermost()->GetPathName() + TypedNodeMat->NodeGuid.ToString() + FString("-") + FString::FromInt(ImageIndex);
 	}
@@ -840,8 +841,8 @@ struct FMutableGraphGenerationContext
 		TArray<SIZE_T> NodeModifierIDs;
 	};
 
-	// UCustomizableObjectNodeMaterial material to SharedSurfaceId
-	TMap<UCustomizableObjectNodeMaterial*, TArray<FSharedSurface>> SharedSurfaceIds;
+	// Material to SharedSurfaceId
+	TMap<UCustomizableObjectNodeMaterialBase*, TArray<FSharedSurface>> SharedSurfaceIds;
 
 	/** Resource Data constants */
 	TMap<uint32, int32> StreamedResourceIndices;
@@ -903,7 +904,7 @@ FString GenerateAnimationInstanceTag(const int32 AnimInstanceIndex, const FName&
 
 FString GenerateGameplayTag(const FString& GameplayTag);
 
-uint32 GetBaseTextureSize(const FMutableGraphGenerationContext& GenerationContext, const UCustomizableObjectNodeMaterial* Material, uint32 ImageIndex);
+uint32 GetBaseTextureSize(const FMutableGraphGenerationContext& GenerationContext, const UCustomizableObjectNodeMaterialBase* Material, uint32 ImageIndex);
 
 // Computes the LOD bias for a texture given the current mesh LOD and automatic LOD settings, the reference texture settings
 // and whether it's being built for a server or not
diff --git a/Engine/Plugins/Experimental/Mutable/Source/CustomizableObjectEditor/Private/MuCOE/GenerateMutableSource/GenerateMutableSourceGroupProjector.cpp b/Engine/Plugins/Experimental/Mutable/Source/CustomizableObjectEditor/Private/MuCOE/GenerateMutableSource/GenerateMutableSourceGroupProjector.cpp
--- a/Engine/Plugins/Experimental/Mutable/Source/CustomizableObjectEditor/Private/MuCOE/GenerateMutableSource/GenerateMutableSourceGroupProjector.cpp
+++ b/Engine/Plugins/Experimental/Mutable/Source/CustomizableObjectEditor/Private/MuCOE/GenerateMutableSource/GenerateMutableSourceGroupProjector.cpp
@@ -32,23 +32,18 @@ class UPoseAsset;
 
 
 mu::NodeImagePtr GenerateMutableGroupProjection(const int32 NodeLOD, const int32 ImageIndex, mu::NodeMeshPtr MeshNode, FMutableGraphGenerationContext& GenerationContext,
-	const UCustomizableObjectNodeMaterialBase* NodeMaterialBase, bool& bShareProjectionTexturesBetweenLODs, bool& bIsGroupProjectorImage,
-	UTexture2D*& GroupProjectionReferenceTexture, TMap<FString, float>& TextureNameToProjectionResFactor, FString& AlternateResStateName,
-	UCustomizableObjectNodeMaterial* ParentMaterial)
+	UCustomizableObjectNodeMaterialBase* TypedNodeMat, UCustomizableObjectNodeExtendMaterial* TypedNodeExt, bool& bShareProjectionTexturesBetweenLODs, bool& bIsGroupProjectorImage,
+	UTexture2D*& GroupProjectionReferenceTexture, TMap<FString, float>& TextureNameToProjectionResFactor, FString& AlternateResStateName)
 {
+	check(static_cast<bool>(TypedNodeMat) != static_cast<bool>(TypedNodeExt)); // XOr. TypedNodeMat valid or TypedNodeExt valid. At least one valid.
+
 	if (!GenerationContext.ProjectorGroupMap.Num() || !MeshNode.get())
 	{
 		return mu::NodeImagePtr();
 	}
-	
-	const UCustomizableObjectNodeMaterial* TypedNodeMat = Cast<UCustomizableObjectNodeMaterial>(NodeMaterialBase);
-	check(!TypedNodeMat || !ParentMaterial); // Logical implication
 
-	const UCustomizableObjectNodeExtendMaterial* TypedNodeExt = Cast<UCustomizableObjectNodeExtendMaterial>(NodeMaterialBase);
-	check(!TypedNodeExt || ParentMaterial); // Logical implication
+	const UCustomizableObjectNodeMaterialBase* NodeMaterial = TypedNodeMat ? TypedNodeMat : TypedNodeExt->GetParentMaterialNode();
 
-	check(TypedNodeMat || TypedNodeExt);
-	
 	TArray<mu::Ptr<mu::NodeImageProject>> ImageNodes;
 	TArray<FGroupProjectorTempData> ImageNodes_ProjectorTempData;
 
@@ -78,8 +73,8 @@ mu::NodeImagePtr GenerateMutableGroupProjection(const int32 NodeLOD, const int32
 
 		const bool bProjectToImage = [&]
 		{
-			const UCustomizableObjectNodeMaterial* NodeMaterial = TypedNodeMat ? TypedNodeMat : ParentMaterial;
 			const FString ParameterName = NodeMaterial->GetParameterName(EMaterialParameterType::Texture, ImageIndex).ToString();
+
 			return ParameterName == ProjectorTempData.CustomizableObjectNodeGroupProjectorParameter->MaterialChannelNameToConnect;
 		}();
 		
@@ -93,7 +88,7 @@ mu::NodeImagePtr GenerateMutableGroupProjection(const int32 NodeLOD, const int32
 				}
 				else
 				{
-					const FGuid ImageId = ParentMaterial->GetParameterId(EMaterialParameterType::Texture, ImageIndex);
+					const FGuid ImageId = TypedNodeExt->GetParentMaterialNode()->GetParameterId(EMaterialParameterType::Texture, ImageIndex);
 					return TypedNodeExt->UsesImage(ImageId);
 				}
 			}();
@@ -102,12 +97,18 @@ mu::NodeImagePtr GenerateMutableGroupProjection(const int32 NodeLOD, const int32
 			{
 				const FString ImageName = [&]
 				{
-					const UCustomizableObjectNodeMaterial* NodeMaterial = TypedNodeMat ? TypedNodeMat : ParentMaterial;
-					return NodeMaterial->GetParameterName(EMaterialParameterType::Texture, ImageIndex).ToString();
+					if (TypedNodeMat)
+					{
+						return TypedNodeMat->GetParameterName(EMaterialParameterType::Texture, ImageIndex).ToString();
+					}
+					else
+					{
+						return TypedNodeExt->GetParentMaterialNode()->GetParameterName(EMaterialParameterType::Texture, ImageIndex).ToString();
+					}
 				}();
 				
 				FString msg = FString::Printf(TEXT("Material image [%s] is connected to an image but will be replaced by a Group Projector."), *ImageName);
-				GenerationContext.Compiler->CompilerLog(FText::FromString(msg), TypedNodeMat);
+				GenerationContext.Compiler->CompilerLog(FText::FromString(msg), TypedNodeMat ? TypedNodeMat : TypedNodeExt->GetParentMaterialNode());
 				continue;
 			}
 			
@@ -164,7 +165,7 @@ mu::NodeImagePtr GenerateMutableGroupProjection(const int32 NodeLOD, const int32
 					if (!NodeMeshApplyPose)
 					{
 						FString msg = FString::Printf(TEXT("Couldn't get bone transform information from a Pose Asset."));
-						GenerationContext.Compiler->CompilerLog(FText::FromString(msg), NodeMaterialBase);
+						GenerationContext.Compiler->CompilerLog(FText::FromString(msg), TypedNodeMat ? TypedNodeMat : TypedNodeExt->GetParentMaterialNode());
 					}
 
 					MeshSwitchNode->SetOption(SelectorIndex + 1, NodeMeshApplyPose);
diff --git a/Engine/Plugins/Experimental/Mutable/Source/CustomizableObjectEditor/Private/MuCOE/GenerateMutableSource/GenerateMutableSourceGroupProjector.h b/Engine/Plugins/Experimental/Mutable/Source/CustomizableObjectEditor/Private/MuCOE/GenerateMutableSource/GenerateMutableSourceGroupProjector.h
--- a/Engine/Plugins/Experimental/Mutable/Source/CustomizableObjectEditor/Private/MuCOE/GenerateMutableSource/GenerateMutableSourceGroupProjector.h
+++ b/Engine/Plugins/Experimental/Mutable/Source/CustomizableObjectEditor/Private/MuCOE/GenerateMutableSource/GenerateMutableSourceGroupProjector.h
@@ -5,6 +5,7 @@
 #include "Containers/Map.h"
 #include "MuT/NodeSurfaceNew.h"
 
+class UCustomizableObjectNodeExtendMaterial;
 class FString;
 class UCustomizableObjectNodeMaterial;
 class UCustomizableObjectNodeMaterialBase;
@@ -13,10 +14,10 @@ class UEdGraphPin;
 class UTexture2D;
 struct FMutableGraphGenerationContext;
 
+
 mu::NodeImagePtr GenerateMutableGroupProjection(const int32 NodeLOD, const int32 ImageIndex, mu::NodeMeshPtr MeshNode, FMutableGraphGenerationContext& GenerationContext,
-	const UCustomizableObjectNodeMaterialBase* NodeMaterialBase, bool& bShareProjectionTexturesBetweenLODs, bool& bIsGroupProjectorImage,
-	UTexture2D*& GroupProjectionReferenceTexture, TMap<FString, float>& TextureNameToProjectionResFactor, FString& AlternateResStateName,
-	UCustomizableObjectNodeMaterial* ParentMaterial);
+	UCustomizableObjectNodeMaterialBase* TypedNodeMat, UCustomizableObjectNodeExtendMaterial* TypedNodeExt, bool& bShareProjectionTexturesBetweenLODs, bool& bIsGroupProjectorImage,
+	UTexture2D*& GroupProjectionReferenceTexture, TMap<FString, float>& TextureNameToProjectionResFactor, FString& AlternateResStateName);
 
 
 /** Convert a CustomizableObject Source Graph into a mutable source graph. */
diff --git a/Engine/Plugins/Experimental/Mutable/Source/CustomizableObjectEditor/Private/MuCOE/GenerateMutableSource/GenerateMutableSourceSurface.cpp b/Engine/Plugins/Experimental/Mutable/Source/CustomizableObjectEditor/Private/MuCOE/GenerateMutableSource/GenerateMutableSourceSurface.cpp
--- a/Engine/Plugins/Experimental/Mutable/Source/CustomizableObjectEditor/Private/MuCOE/GenerateMutableSource/GenerateMutableSourceSurface.cpp
+++ b/Engine/Plugins/Experimental/Mutable/Source/CustomizableObjectEditor/Private/MuCOE/GenerateMutableSource/GenerateMutableSourceSurface.cpp
@@ -204,9 +204,9 @@ void SetSurfaceFormat( FMutableGraphGenerationContext& GenerationContext,
 }
 
 
-void AddModifierToSharedSurface(FMutableGraphGenerationContext& GenerationContext, UCustomizableObjectNodeMaterial* NodeMaterial, const UCustomizableObjectNode& NodeModifier)
+void AddModifierToSharedSurface(FMutableGraphGenerationContext& GenerationContext, UCustomizableObjectNodeMaterialBase* NodeMaterial, const UCustomizableObjectNode& NodeModifier)
 {
-	if (!NodeMaterial || !NodeMaterial->bReuseMaterialBetweenLODs)
+	if (!NodeMaterial || !NodeMaterial->IsReuseMaterialBetweenLODs())
 	{
 		return;
 	}
@@ -255,37 +255,13 @@ mu::NodeSurfacePtr GenerateMutableSourceSurface(const UEdGraphPin * Pin, FMutabl
 		}
 	}
 
-	if (UCustomizableObjectNodeMaterial* TypedNodeMat = Cast<UCustomizableObjectNodeMaterial>(Node))
+	if (UCustomizableObjectNodeMaterialBase* TypedNodeMat = Cast<UCustomizableObjectNodeMaterialBase>(Node))
 	{
-		if (TypedNodeMat->MeshComponentIndex != GenerationContext.CurrentMeshComponent)
+		if (TypedNodeMat->GetMeshComponentIndex() != GenerationContext.CurrentMeshComponent)
 		{
 			return Result;
 		}
 
-		// NodeCopyMaterial. Special case when the TypedNodeMat is a NodeCopyMaterial. The TypedNodeMat pointer now points to the parent NodeMaterial except when reading the mesh pin, which comes from the NodeCopyMaterial.
-		UCustomizableObjectNodeMaterial* TypedNodeMaterial = TypedNodeMat;
-		UCustomizableObjectNodeMaterial* TypedNodeCopyMaterial = TypedNodeMat;
-
-		if (UCustomizableObjectNodeCopyMaterial* TypedDerivedNodeCopyMaterial = Cast<UCustomizableObjectNodeCopyMaterial>(TypedNodeMat))
-		{
-			UEdGraphPin* MaterialPin = TypedDerivedNodeCopyMaterial->GetMaterialPin();
-			if (const UEdGraphPin* ConnectedPin = FollowInputPin(*MaterialPin))
-			{
-				if (UCustomizableObjectNodeMaterial* ConnectedNodeMaterial = Cast<UCustomizableObjectNodeMaterial>(ConnectedPin->GetOwningNode()))
-				{
-					if (!ConnectedNodeMaterial->IsA(UCustomizableObjectNodeCopyMaterial::StaticClass()))
-					{
-						TypedNodeMaterial = ConnectedNodeMaterial;
-						TypedNodeMat = ConnectedNodeMaterial;
-					}
-					else
-					{
-						GenerationContext.Compiler->CompilerLog(LOCTEXT("CopyMaterialInput", "Copy Material Node can only have a Material Node as an input."), TypedDerivedNodeCopyMaterial);
-					}
-				}
-			}
-		}
-
 		const UEdGraphPin* ConnectedMaterialPin = FollowInputPin(*TypedNodeMat->GetMeshPin());
 		// Warn when texture connections are improperly used by connecting them directly to material inputs when no layout is used
 		// TODO: delete the if clause and the warning when static meshes are operational again
@@ -304,15 +280,15 @@ mu::NodeSurfacePtr GenerateMutableSourceSurface(const UEdGraphPin * Pin, FMutabl
 		// Add to the list of surfaces that could be reused between LODs for this NodeMaterial.
 		TArray<FMutableGraphGenerationContext::FSharedSurface>& SharedSurfaces = GenerationContext.SharedSurfaceIds.FindOrAdd(TypedNodeMat, {});
 		FMutableGraphGenerationContext::FSharedSurface& SharedSurface = SharedSurfaces.Add_GetRef(FMutableGraphGenerationContext::FSharedSurface(GenerationContext.CurrentLOD, SurfNode));
-		SharedSurface.bMakeUnique = !TypedNodeMat->bReuseMaterialBetweenLODs;
+		SharedSurface.bMakeUnique = !TypedNodeMat->IsReuseMaterialBetweenLODs();
 
 		int32 ReferencedMaterialsIndex = -1;
-		if (TypedNodeMat->Material)
+		if (TypedNodeMat->GetMaterial())
 		{
-			GenerationContext.AddParticipatingObject(*TypedNodeMat->Material);
+			GenerationContext.AddParticipatingObject(*TypedNodeMat->GetMaterial());
 
 			const int32 lastMaterialAmount = GenerationContext.ReferencedMaterials.Num();
-			ReferencedMaterialsIndex = GenerationContext.ReferencedMaterials.AddUnique(TypedNodeMat->Material);
+			ReferencedMaterialsIndex = GenerationContext.ReferencedMaterials.AddUnique(TypedNodeMat->GetMaterial());
 			// Used ReferencedMaterialsIndex instead of TypedNodeMat->Material->GetName() to prevent material name collisions
 
 			// Take slot name from skeletal mesh if one can be found, else leave empty.
@@ -364,9 +340,7 @@ mu::NodeSurfacePtr GenerateMutableSourceSurface(const UEdGraphPin * Pin, FMutabl
 		}
 
 		mu::NodeMeshPtr MeshNode;
-
-		TypedNodeMat = TypedNodeCopyMaterial; // NodeCopyMaterial. Start reading mesh pin. Set TypedNodeMat pointer to NodeCopyMaterial
-
+		
 		if (const UEdGraphPin* ConnectedPin = FollowInputPin(*TypedNodeMat->GetMeshPin()))
 		{
 			FMutableGraphMeshGenerationData MeshData;
@@ -395,9 +369,7 @@ mu::NodeSurfacePtr GenerateMutableSourceSurface(const UEdGraphPin * Pin, FMutabl
 				GenerationContext.Compiler->CompilerLog(LOCTEXT("MeshFailed", "Mesh generation failed."), Node);
 			}
 		}
-
-		TypedNodeMat = TypedNodeMaterial; // NodeCopyMaterial. End reading mesh pin. Set TypedNodeMat pointer back to the parent NodeMaterial
-
+		
 		TMap<FString, float> TextureNameToProjectionResFactor;
 		FString AlternateResStateName;
 
@@ -420,11 +392,11 @@ mu::NodeSurfacePtr GenerateMutableSourceSurface(const UEdGraphPin * Pin, FMutabl
 				if (UMaterialInstance * TableMaterial = TypedNodeTable->GetColumnDefaultAssetByType<UMaterialInstance>(TempConnectedPin))
 				{
 					// Checking if the reference material of the Table Node has the same parent as the material of the Material Node 
-					if (!TypedNodeMat->Material || TableMaterial->GetMaterial() != TypedNodeMat->Material->GetMaterial())
+					if (!TypedNodeMat->GetMaterial() || TableMaterial->GetMaterial() != TypedNodeMat->GetMaterial()->GetMaterial())
 					{
 						bTableMaterialPinLinked = false;
 
-						GenerationContext.Compiler->CompilerLog(LOCTEXT("DifferentParentMaterial","The Deafult Material Instance of the Data Table must have the same Parent Material."), TypedNodeMat);
+						GenerationContext.Compiler->CompilerLog(LOCTEXT("DifferentParentMaterial", "The Deafult Material Instance of the Data Table must have the same Parent Material."), TypedNodeMat->GetMaterialNode());
 					}
 				}
 				else
@@ -452,7 +424,7 @@ mu::NodeSurfacePtr GenerateMutableSourceSurface(const UEdGraphPin * Pin, FMutabl
 
 			const bool bIsImagePinLinked = ImagePin && FollowInputPin(*ImagePin);
 
-			if (bIsImagePinLinked && !TypedNodeMaterial->IsImageMutableMode(ImageIndex))
+			if (bIsImagePinLinked && !TypedNodeMat->IsImageMutableMode(ImageIndex))
 			{
 				if (const UEdGraphPin* ConnectedPin = FollowInputPin(*ImagePin))
 				{
@@ -477,7 +449,7 @@ mu::NodeSurfacePtr GenerateMutableSourceSurface(const UEdGraphPin * Pin, FMutabl
 					const FString SurfNodeImageName = FString::Printf(TEXT("%d"), Props.ImagePropertiesIndex);
 					SurfNode->SetImageName(ImageIndex, SurfNodeImageName);
 					SurfNode->SetImageLayoutIndex(ImageIndex, -1);
-					SurfNode->SetImageAdditionalNames(ImageIndex, TypedNodeMat->Material->GetName(), Props.TextureParameterName);
+					SurfNode->SetImageAdditionalNames(ImageIndex, TypedNodeMat->GetMaterial()->GetName(), Props.TextureParameterName);
 
 				}
 			}
@@ -497,11 +469,10 @@ mu::NodeSurfacePtr GenerateMutableSourceSurface(const UEdGraphPin * Pin, FMutabl
 					bool bIsGroupProjectorImage = false;
 
 					GroupProjectionImg = GenerateMutableGroupProjection(LOD, ImageIndex, MeshNode, GenerationContext,
-						TypedNodeMat, bShareProjectionTexturesBetweenLODs, bIsGroupProjectorImage,
-						GroupProjectionReferenceTexture, TextureNameToProjectionResFactor, AlternateResStateName,
-						nullptr);
+						TypedNodeMat, nullptr, bShareProjectionTexturesBetweenLODs, bIsGroupProjectorImage,
+						GroupProjectionReferenceTexture, TextureNameToProjectionResFactor, AlternateResStateName);
 
-					if (GroupProjectionImg.get() || TypedNodeMaterial->IsImageMutableMode(ImageIndex))
+					if (GroupProjectionImg.get() || TypedNodeMat->IsImageMutableMode(ImageIndex))
 					{
 						// Get the reference texture
 						UTexture2D* ReferenceTexture = nullptr;
@@ -596,7 +567,7 @@ mu::NodeSurfacePtr GenerateMutableSourceSurface(const UEdGraphPin * Pin, FMutabl
 						// Generate the texture nodes
 						mu::NodeImagePtr ImageNode = [&]()
 						{
-							if (TypedNodeMaterial->IsImageMutableMode(ImageIndex))
+							if (TypedNodeMat->IsImageMutableMode(ImageIndex))
 							{
 								if (ImagePin)
 								{
@@ -841,7 +812,7 @@ mu::NodeSurfacePtr GenerateMutableSourceSurface(const UEdGraphPin * Pin, FMutabl
 						SurfNode->SetImageName(ImageIndex, SurfNodeImageName + LayerEncoding);
 						const int32 UVLayout = TypedNodeMat->GetImageUVLayout(ImageIndex);
 						SurfNode->SetImageLayoutIndex(ImageIndex, UVLayout);
-						SurfNode->SetImageAdditionalNames(ImageIndex, TypedNodeMat->Material->GetName(), ImageName);
+						SurfNode->SetImageAdditionalNames(ImageIndex, TypedNodeMat->GetMaterial()->GetName(), ImageName);
 
 						if (bShareProjectionTexturesBetweenLODs && bIsGroupProjectorImage)
 						{
@@ -956,17 +927,17 @@ mu::NodeSurfacePtr GenerateMutableSourceSurface(const UEdGraphPin * Pin, FMutabl
 		}
 		
 
-		for (const FString& Tag : TypedNodeMat->Tags)
+		for (const FString& Tag : TypedNodeMat->GetTags())
 		{
 			SurfNode->AddTag(Tag);
 		}
 
-		TArray<mu::NodeSurfaceNewPtr>* ArraySurfaceNodePtr = GenerationContext.MapMaterialNodeToMutableSurfaceNodeArray.Find(TypedNodeMat);
+		TArray<mu::NodeSurfaceNewPtr>* ArraySurfaceNodePtr = GenerationContext.MapMaterialNodeToMutableSurfaceNodeArray.Find(TypedNodeMat->GetMaterialNode());
 		if (ArraySurfaceNodePtr == nullptr)
 		{
 			TArray<mu::NodeSurfaceNewPtr> ArraySurfaceNode;
 			ArraySurfaceNode.Add(SurfNode);
-			ArraySurfaceNodePtr = &GenerationContext.MapMaterialNodeToMutableSurfaceNodeArray.Add(TypedNodeMat, ArraySurfaceNode);
+			ArraySurfaceNodePtr = &GenerationContext.MapMaterialNodeToMutableSurfaceNodeArray.Add(TypedNodeMat->GetMaterialNode(), ArraySurfaceNode);
 		}
 		else
 		{
@@ -999,7 +970,7 @@ mu::NodeSurfacePtr GenerateMutableSourceSurface(const UEdGraphPin * Pin, FMutabl
 				SurfNode2->SetScalarName(ScalarParamIndex, SurfNode->GetScalarName(ScalarParamIndex));
 			}
 
-			for (const FString& Tag : TypedNodeMat->Tags)
+			for (const FString& Tag : TypedNodeMat->GetTags())
 			{
 				SurfNode2->AddTag(Tag);
 			}
@@ -1067,22 +1038,7 @@ mu::NodeSurfacePtr GenerateMutableSourceSurface(const UEdGraphPin * Pin, FMutabl
 
 		[&] // Using a lambda so control flow is easier to manage.
 		{
-			UCustomizableObjectNodeMaterial* OriginalParentMaterialNode = TypedNodeExt->GetParentMaterialNode();
-			UCustomizableObjectNodeMaterial* ParentMaterialNode = OriginalParentMaterialNode;
-
-			// Copy material
-			UCustomizableObjectNodeMaterial* CopyMaterialParentMaterialNode = nullptr;
-			const UCustomizableObjectNodeCopyMaterial* ParentMaterialCopyNodeCast = Cast<UCustomizableObjectNodeCopyMaterial>(OriginalParentMaterialNode);
-			if (ParentMaterialCopyNodeCast)
-			{
-				CopyMaterialParentMaterialNode = ParentMaterialCopyNodeCast->GetMaterialNode();
-				if (!CopyMaterialParentMaterialNode)
-				{
-					GenerationContext.Compiler->CompilerLog(LOCTEXT("BaseMissing", "Base Material not set (or not found)."), ParentMaterialCopyNodeCast);
-					return;
-				}
-			}
-		
+			UCustomizableObjectNodeMaterialBase* ParentMaterialNode = TypedNodeExt->GetParentMaterialNode();
 			if (!ParentMaterialNode)
 			{
 				GenerationContext.Compiler->CompilerLog(LOCTEXT("ParentMissing", "Parent node not set (or not found)."), Node);
@@ -1148,12 +1104,7 @@ mu::NodeSurfacePtr GenerateMutableSourceSurface(const UEdGraphPin * Pin, FMutabl
 				SurfNode->SetMesh(MeshPatch.get());
 				MeshPatch->SetMessageContext(Node);
 			}
-
-			if (ParentMaterialCopyNodeCast) // CopyMaterial swap.
-			{
-				ParentMaterialNode = CopyMaterialParentMaterialNode; 
-			}
-		
+			
 			const int32 NumImages = ParentMaterialNode->GetNumParameters(EMaterialParameterType::Texture);
 			SurfNode->SetImageCount(NumImages);
 			for (int32 ImageIndex = 0; ImageIndex < NumImages; ++ImageIndex)
@@ -1185,9 +1136,8 @@ mu::NodeSurfacePtr GenerateMutableSourceSurface(const UEdGraphPin * Pin, FMutabl
 					FString AlternateResStateName;
 					
 					ImageNode = GenerateMutableGroupProjection(LOD, ImageIndex, AddMeshNode, GenerationContext,
-						TypedNodeExt, bShareProjectionTexturesBetweenLODs, bIsGroupProjectorImage,
-						GroupProjectionReferenceTexture, TextureNameToProjectionResFactor, AlternateResStateName,
-						ParentMaterialNode);
+						nullptr, TypedNodeExt, bShareProjectionTexturesBetweenLODs, bIsGroupProjectorImage,
+						GroupProjectionReferenceTexture, TextureNameToProjectionResFactor, AlternateResStateName);
 				}
 				
 				if (!ImageNode) // Else if
@@ -1228,7 +1178,7 @@ mu::NodeSurfacePtr GenerateMutableSourceSurface(const UEdGraphPin * Pin, FMutabl
 		mu::NodeSurfaceEditPtr SurfNode = new mu::NodeSurfaceEdit();
 		Result = SurfNode;
 
-		UCustomizableObjectNodeMaterial* ParentMaterialNode = TypedNodeRem->GetParentMaterialNode();
+		UCustomizableObjectNodeMaterialBase* ParentMaterialNode = TypedNodeRem->GetParentMaterialNode();
 		if (!ParentMaterialNode)
 		{
 			GenerationContext.Compiler->CompilerLog(LOCTEXT("ParentMissing", "Parent node not set (or not found)."), Node);
@@ -1264,8 +1214,7 @@ mu::NodeSurfacePtr GenerateMutableSourceSurface(const UEdGraphPin * Pin, FMutabl
 		mu::NodeSurfaceEditPtr SurfNode = new mu::NodeSurfaceEdit();
 		Result = SurfNode;
 
-		UCustomizableObjectNodeMaterial* ParentMaterialNode = TypedNodeRemBlocks->GetParentMaterialNode();
-
+		UCustomizableObjectNodeMaterialBase* ParentMaterialNode = TypedNodeRemBlocks->GetParentMaterialNode();
 		if (!ParentMaterialNode)
 		{
 			GenerationContext.Compiler->CompilerLog(LOCTEXT("ParentMissing", "Parent node not set (or not found)."), Node);
@@ -1347,7 +1296,7 @@ mu::NodeSurfacePtr GenerateMutableSourceSurface(const UEdGraphPin * Pin, FMutabl
 		mu::NodeSurfaceEditPtr SurfNode = new mu::NodeSurfaceEdit();
 		Result = SurfNode;
 
-		UCustomizableObjectNodeMaterial* ParentMaterialNode = TypedNodeEdit->GetParentMaterialNode();
+		UCustomizableObjectNodeMaterialBase* ParentMaterialNode = TypedNodeEdit->GetParentMaterialNode();
 		if (!ParentMaterialNode)
 		{
 			GenerationContext.Compiler->CompilerLog(LOCTEXT("ParentMissing", "Parent node not set (or not found)."), Node);
@@ -1452,7 +1401,7 @@ mu::NodeSurfacePtr GenerateMutableSourceSurface(const UEdGraphPin * Pin, FMutabl
 		mu::NodeSurfaceEditPtr SurfNode = new mu::NodeSurfaceEdit();
 		Result = SurfNode;
 
-		UCustomizableObjectNodeMaterial* ParentMaterialNode = TypedNodeMorph->GetParentMaterialNode();
+		UCustomizableObjectNodeMaterialBase* ParentMaterialNode = TypedNodeMorph->GetParentMaterialNode();
 		if (!ParentMaterialNode)
 		{
 			GenerationContext.Compiler->CompilerLog(LOCTEXT("ParentMissing", "Parent node not set (or not found)."), Node);
@@ -1464,7 +1413,6 @@ mu::NodeSurfacePtr GenerateMutableSourceSurface(const UEdGraphPin * Pin, FMutabl
 			SurfNode->SetParent(ParentNode.get());
 
 			const UEdGraphPin* BaseSourcePin = FindMeshBaseSource(*ParentMaterialNode->OutputPin(), false);
-
 			if (!BaseSourcePin)
 			{
 				GenerationContext.Compiler->CompilerLog(LOCTEXT("ParentMissing", "Parent node not set (or not found)."), Node);
diff --git a/Engine/Plugins/Experimental/Mutable/Source/CustomizableObjectEditor/Private/MuCOE/GraphTraversal.cpp b/Engine/Plugins/Experimental/Mutable/Source/CustomizableObjectEditor/Private/MuCOE/GraphTraversal.cpp
--- a/Engine/Plugins/Experimental/Mutable/Source/CustomizableObjectEditor/Private/MuCOE/GraphTraversal.cpp
+++ b/Engine/Plugins/Experimental/Mutable/Source/CustomizableObjectEditor/Private/MuCOE/GraphTraversal.cpp
@@ -452,7 +452,7 @@ const UEdGraphPin* FindMeshBaseSource(const UEdGraphPin& Pin, const bool bOnlyLo
 		}
 	}
 
-	else if (const UCustomizableObjectNodeMaterial* TypedNodeMat = Cast<UCustomizableObjectNodeMaterial>(Node))
+	else if (const UCustomizableObjectNodeMaterialBase* TypedNodeMat = Cast<UCustomizableObjectNodeMaterialBase>(Node))
 	{
 		if (const UEdGraphPin* ConnectedPin = FollowInputPin(*TypedNodeMat->GetMeshPin()))
 		{
diff --git a/Engine/Plugins/Experimental/Mutable/Source/CustomizableObjectEditor/Private/MuCOE/Nodes/CustomizableObjectNodeCopyMaterial.cpp b/Engine/Plugins/Experimental/Mutable/Source/CustomizableObjectEditor/Private/MuCOE/Nodes/CustomizableObjectNodeCopyMaterial.cpp
--- a/Engine/Plugins/Experimental/Mutable/Source/CustomizableObjectEditor/Private/MuCOE/Nodes/CustomizableObjectNodeCopyMaterial.cpp
+++ b/Engine/Plugins/Experimental/Mutable/Source/CustomizableObjectEditor/Private/MuCOE/Nodes/CustomizableObjectNodeCopyMaterial.cpp
@@ -8,6 +8,7 @@
 #include "MuCOE/Nodes/CustomizableObjectNodeExternalPin.h"
 #include "MuCOE/Nodes/CustomizableObjectNodeObject.h"
 #include "MuCOE/Nodes/CustomizableObjectNodeSkeletalMesh.h"
+#include "MuCOE/Nodes/CustomizableObjectNodeMaterial.h"
 
 #define LOCTEXT_NAMESPACE "CustomizableObjectEditor"
 
@@ -18,43 +19,92 @@ const static FString MeshPinName = TEXT("Mesh_Input_Pin");
 /** Material input pin key */
 const static FString MaterialPinName = TEXT("Material_Input_Pin");
 
+/** Material output pin key*/
+const static FString OutputPinName = TEXT("Material_Output_Pin");
+
+
 FText UCustomizableObjectNodeCopyMaterial::GetNodeTitle(ENodeTitleType::Type TitleType) const
 {
 	return LOCTEXT("Copy_Material", "Copy Material");
 }
 
+
 void UCustomizableObjectNodeCopyMaterial::AllocateDefaultPins(UCustomizableObjectNodeRemapPins* RemapPins)
 {
 	const UEdGraphSchema_CustomizableObject* Schema = GetDefault<UEdGraphSchema_CustomizableObject>();
 
 	// Input pins
-	FString PinFriendlyName = TEXT("Mesh");
 	UEdGraphPin* Pin = CustomCreatePin(EGPD_Input, Schema->PC_Mesh, FName(MeshPinName));
-	Pin->PinFriendlyName = FText::FromString(PinFriendlyName);
+	Pin->PinFriendlyName = LOCTEXT("MeshPin", "Mesh");
 	Pin->bDefaultValueIsIgnored = true;
 
-	PinFriendlyName = TEXT("Base Material");
 	Pin = CustomCreatePin(EGPD_Input, Schema->PC_Material, FName(MaterialPinName));
-	Pin->PinFriendlyName = FText::FromString(PinFriendlyName);
+	Pin->PinFriendlyName = LOCTEXT("BaseMaterialPin", "Base Material");;
 	Pin->bDefaultValueIsIgnored = true;
 
 	// Output pins
-	PinFriendlyName = TEXT("Material");
-	FString PinName = PinFriendlyName + FString(TEXT("_Output_Pin"));
-	Pin = CustomCreatePin(EGPD_Output, Schema->PC_Material, FName(*PinName));
-	Pin->PinFriendlyName = FText::FromString(PinFriendlyName);
+	Pin = CustomCreatePin(EGPD_Output, Schema->PC_Material, FName(OutputPinName));
+	Pin->PinFriendlyName = LOCTEXT("MaterialPin", "Material");
 }
 
+
 UEdGraphPin* UCustomizableObjectNodeCopyMaterial::GetMeshPin() const
 {
 	return FindPin(MeshPinName);
 }
 
+
+FPostImagePinModeChangedDelegate* UCustomizableObjectNodeCopyMaterial::GetPostImagePinModeChangedDelegate()
+{
+	if (UCustomizableObjectNodeMaterial* NodeMaterial = GetMaterialNode())
+	{
+		return NodeMaterial->GetPostImagePinModeChangedDelegate();
+	}
+	else
+	{
+		return nullptr;
+	}
+}
+
+
+TArray<UCustomizableObjectLayout*> UCustomizableObjectNodeCopyMaterial::GetLayouts() const
+{
+	if (UCustomizableObjectNodeMaterial* NodeMaterial = GetMaterialNode())
+	{
+		return NodeMaterial->GetLayouts();
+	}
+	else
+	{
+		return {};
+	}
+}
+
+
+UEdGraphPin* UCustomizableObjectNodeCopyMaterial::OutputPin() const
+{
+	return FindPin(OutputPinName);
+}
+
+
+bool UCustomizableObjectNodeCopyMaterial::RealMaterialDataHasChanged() const
+{
+	if (UCustomizableObjectNodeMaterial* NodeMaterial = GetMaterialNode())
+	{
+		return NodeMaterial->RealMaterialDataHasChanged();
+	}
+	else
+	{
+		return false;
+	}
+}
+
+
 UEdGraphPin* UCustomizableObjectNodeCopyMaterial::GetMaterialPin() const
 {
 	return FindPin(MaterialPinName);
 }
 
+
 UCustomizableObjectNodeSkeletalMesh* UCustomizableObjectNodeCopyMaterial::GetMeshNode() const
 {
 	UCustomizableObjectNodeSkeletalMesh* Result = nullptr;
@@ -86,6 +136,7 @@ UCustomizableObjectNodeMaterial* UCustomizableObjectNodeCopyMaterial::GetMateria
 	return Result;
 }
 
+
 bool UCustomizableObjectNodeCopyMaterial::CanConnect(const UEdGraphPin* InOwnedInputPin, const UEdGraphPin* InOutputPin, bool& bOutIsOtherNodeBlocklisted, bool& bOutArePinsCompatible) const
 {
 	if (!Super::CanConnect(InOwnedInputPin, InOutputPin, bOutIsOtherNodeBlocklisted, bOutIsOtherNodeBlocklisted))
@@ -103,21 +154,25 @@ bool UCustomizableObjectNodeCopyMaterial::CanConnect(const UEdGraphPin* InOwnedI
 	return true;
 }
 
+
 bool UCustomizableObjectNodeCopyMaterial::ShouldBreakExistingConnections(const UEdGraphPin* InputPin, const UEdGraphPin* OutputPin) const
 {
 	return true;
 }
 
+
 bool UCustomizableObjectNodeCopyMaterial::IsNodeOutDatedAndNeedsRefresh()
 {
 	return false;
 }
 
+
 bool UCustomizableObjectNodeCopyMaterial::ProvidesCustomPinRelevancyTest() const
 {
 	return true;
 }
 
+
 bool UCustomizableObjectNodeCopyMaterial::IsPinRelevant(const UEdGraphPin* Pin) const
 {
 	const UEdGraphNode* Node = Pin->GetOwningNode();
@@ -136,9 +191,241 @@ bool UCustomizableObjectNodeCopyMaterial::IsPinRelevant(const UEdGraphPin* Pin)
 	}
 }
 
+
+UMaterialInterface* UCustomizableObjectNodeCopyMaterial::GetMaterial() const
+{
+	if (UCustomizableObjectNodeMaterial* NodeMaterial = GetMaterialNode())
+	{
+		return NodeMaterial->GetMaterial();
+	}
+	else
+	{
+		return nullptr;
+	}
+}
+
+
+bool UCustomizableObjectNodeCopyMaterial::IsReuseMaterialBetweenLODs() const
+{
+	if (UCustomizableObjectNodeMaterial* NodeMaterial = GetMaterialNode())
+	{
+		return NodeMaterial->IsReuseMaterialBetweenLODs();
+	}
+	else
+	{
+		return false;
+	}
+}
+
+
+int32 UCustomizableObjectNodeCopyMaterial::GetMeshComponentIndex() const
+{
+	if (UCustomizableObjectNodeMaterial* NodeMaterial = GetMaterialNode())
+	{
+		return NodeMaterial->GetMeshComponentIndex();
+	}
+	else
+	{
+		return INDEX_NONE;
+	}
+}
+
+
+TArray<FString> UCustomizableObjectNodeCopyMaterial::GetTags() const
+{
+	if (UCustomizableObjectNodeMaterial* NodeMaterial = GetMaterialNode())
+	{
+		return NodeMaterial->GetTags();
+	}
+	else
+	{
+		return {};
+	}
+}
+
+
+UEdGraphPin* UCustomizableObjectNodeCopyMaterial::GetMaterialAssetPin() const
+{
+	if (UCustomizableObjectNodeMaterial* NodeMaterial = GetMaterialNode())
+	{
+		return NodeMaterial->GetMaterialAssetPin();
+	}
+	else
+	{
+		return nullptr;
+	}
+}
+
+
+int32 UCustomizableObjectNodeCopyMaterial::GetNumParameters(EMaterialParameterType Type) const
+{
+	if (UCustomizableObjectNodeMaterial* NodeMaterial = GetMaterialNode())
+	{
+		return NodeMaterial->GetNumParameters(Type);
+	}
+	else
+	{
+		return 0;
+	}
+}
+
+
+FGuid UCustomizableObjectNodeCopyMaterial::GetParameterId(EMaterialParameterType Type, int32 ParameterIndex) const
+{
+	UCustomizableObjectNodeMaterial* NodeMaterial = GetMaterialNode();
+	check(NodeMaterial);
+
+	return NodeMaterial->GetParameterId(Type, ParameterIndex);
+}
+
+
+FName UCustomizableObjectNodeCopyMaterial::GetParameterName(EMaterialParameterType Type, int32 ParameterIndex) const
+{
+	UCustomizableObjectNodeMaterial* NodeMaterial = GetMaterialNode();
+	check(NodeMaterial);
+	
+	return NodeMaterial->GetParameterName(Type, ParameterIndex);
+}
+
+
+int32 UCustomizableObjectNodeCopyMaterial::GetParameterLayerIndex(EMaterialParameterType Type, int32 ParameterIndex) const
+{
+	UCustomizableObjectNodeMaterial* NodeMaterial = GetMaterialNode();
+	check(NodeMaterial);
+	
+	return NodeMaterial->GetParameterLayerIndex(Type, ParameterIndex);
+}
+
+
+FText UCustomizableObjectNodeCopyMaterial::GetParameterLayerName(EMaterialParameterType Type, int32 ParameterIndex) const
+{
+	UCustomizableObjectNodeMaterial* NodeMaterial = GetMaterialNode();
+	check(NodeMaterial);
+	
+	return NodeMaterial->GetParameterLayerName(Type, ParameterIndex);
+}
+
+
+bool UCustomizableObjectNodeCopyMaterial::HasParameter(const FGuid& ParameterId) const
+{
+	if (UCustomizableObjectNodeMaterial* NodeMaterial = GetMaterialNode())
+	{
+		return NodeMaterial->HasParameter(ParameterId);
+	}
+	else
+	{
+		return false;	
+	}
+}
+
+
+const UEdGraphPin* UCustomizableObjectNodeCopyMaterial::GetParameterPin(EMaterialParameterType Type, int32 ParameterIndex) const
+{
+	if (UCustomizableObjectNodeMaterial* NodeMaterial = GetMaterialNode())
+	{
+		return NodeMaterial->GetParameterPin(Type, ParameterIndex);
+	}
+	else
+	{
+		return nullptr;	
+	}
+}
+
+
+bool UCustomizableObjectNodeCopyMaterial::IsImageMutableMode(int32 ImageIndex) const
+{
+	UCustomizableObjectNodeMaterial* NodeMaterial = GetMaterialNode();
+	check(NodeMaterial);
+	
+	return NodeMaterial->IsImageMutableMode(ImageIndex);
+}
+
+
+bool UCustomizableObjectNodeCopyMaterial::IsImageMutableMode(const UEdGraphPin& Pin) const
+{
+	UCustomizableObjectNodeMaterial* NodeMaterial = GetMaterialNode();
+	check(NodeMaterial);
+
+	return NodeMaterial->IsImageMutableMode(Pin);
+}
+
+
+void UCustomizableObjectNodeCopyMaterial::UpdateImagePinMode(const FGuid ParameterId)
+{
+	if (UCustomizableObjectNodeMaterial* NodeMaterial = GetMaterialNode())
+	{
+		NodeMaterial->UpdateImagePinMode(ParameterId);
+	}
+}
+
+
+void UCustomizableObjectNodeCopyMaterial::UpdateImagePinMode(const UEdGraphPin& Pin)
+{
+	if (UCustomizableObjectNodeMaterial* NodeMaterial = GetMaterialNode())
+	{
+		NodeMaterial->UpdateImagePinMode(Pin);
+	}
+}
+
+
+void UCustomizableObjectNodeCopyMaterial::UpdateAllImagesPinMode()
+{
+	if (UCustomizableObjectNodeMaterial* NodeMaterial = GetMaterialNode())
+	{
+		NodeMaterial->UpdateAllImagesPinMode();
+	}
+}
+
+
+UTexture2D* UCustomizableObjectNodeCopyMaterial::GetImageReferenceTexture(int32 ImageIndex) const
+{
+	if (UCustomizableObjectNodeMaterial* NodeMaterial = GetMaterialNode())
+	{
+		return NodeMaterial->GetImageReferenceTexture(ImageIndex);
+	}
+	else
+	{
+		return nullptr;
+	}
+}
+
+
+UTexture2D* UCustomizableObjectNodeCopyMaterial::GetImageValue(int32 ImageIndex) const
+{
+	if (UCustomizableObjectNodeMaterial* NodeMaterial = GetMaterialNode())
+	{
+		return NodeMaterial->GetImageValue(ImageIndex);
+	}
+	else
+	{
+		return nullptr;
+	}
+}
+
+
+int32 UCustomizableObjectNodeCopyMaterial::GetImageUVLayout(int32 ImageIndex) const
+{
+	if (UCustomizableObjectNodeMaterial* NodeMaterial = GetMaterialNode())
+	{
+		return NodeMaterial->GetImageUVLayout(ImageIndex);
+	}
+	else
+	{
+		return UCustomizableObjectNodeMaterialPinDataImage::UV_LAYOUT_IGNORE;
+	}
+}
+
+
 FText UCustomizableObjectNodeCopyMaterial::GetTooltipText() const
 {
 	return LOCTEXT("CopyMaterial_Tooltip", "Copies a Customizable Object material.\nDuplicates all Material node input pins and properties except for the Mesh input pin.");
 }
 
+
+FLinearColor UCustomizableObjectNodeCopyMaterial::GetNodeTitleColor() const
+{
+	const UEdGraphSchema_CustomizableObject* Schema = GetDefault<UEdGraphSchema_CustomizableObject>();
+	return Schema->GetPinTypeColor(Schema->PC_Material);
+}
+
 #undef LOCTEXT_NAMESPACE
diff --git a/Engine/Plugins/Experimental/Mutable/Source/CustomizableObjectEditor/Private/MuCOE/Nodes/CustomizableObjectNodeCopyMaterialDetails.cpp b/Engine/Plugins/Experimental/Mutable/Source/CustomizableObjectEditor/Private/MuCOE/Nodes/CustomizableObjectNodeCopyMaterialDetails.cpp
deleted file mode 100644
diff --git a/Engine/Plugins/Experimental/Mutable/Source/CustomizableObjectEditor/Private/MuCOE/Nodes/CustomizableObjectNodeCopyMaterialDetails.h b/Engine/Plugins/Experimental/Mutable/Source/CustomizableObjectEditor/Private/MuCOE/Nodes/CustomizableObjectNodeCopyMaterialDetails.h
deleted file mode 100644
diff --git a/Engine/Plugins/Experimental/Mutable/Source/CustomizableObjectEditor/Private/MuCOE/Nodes/CustomizableObjectNodeEditMaterial.cpp b/Engine/Plugins/Experimental/Mutable/Source/CustomizableObjectEditor/Private/MuCOE/Nodes/CustomizableObjectNodeEditMaterial.cpp
--- a/Engine/Plugins/Experimental/Mutable/Source/CustomizableObjectEditor/Private/MuCOE/Nodes/CustomizableObjectNodeEditMaterial.cpp
+++ b/Engine/Plugins/Experimental/Mutable/Source/CustomizableObjectEditor/Private/MuCOE/Nodes/CustomizableObjectNodeEditMaterial.cpp
@@ -19,7 +19,7 @@ void UCustomizableObjectNodeEditMaterial::AllocateDefaultPins(UCustomizableObjec
 {
 	const UEdGraphSchema_CustomizableObject* Schema = GetDefault<UEdGraphSchema_CustomizableObject>();
 	
-	if (const UCustomizableObjectNodeMaterial* ParentMaterialNode = GetParentMaterialNodeIfPath())
+	if (const UCustomizableObjectNodeMaterialBase* ParentMaterialNode = GetParentMaterialNodeIfPath())
 	{
 		const int32 NumImages = ParentMaterialNode->GetNumParameters(EMaterialParameterType::Texture);
 		for (int32 ImageIndex = 0; ImageIndex < NumImages; ++ImageIndex)
@@ -153,7 +153,7 @@ void UCustomizableObjectNodeEditMaterial::BackwardsCompatibleFixup()
 	
 	if (CustomizableObjectCustomVersion < FCustomizableObjectCustomVersion::AutomaticNodeMaterial)
 	{
-		if (const UCustomizableObjectNodeMaterial* ParentMaterial = GetParentMaterialNode())
+		if (const UCustomizableObjectNodeMaterialBase* ParentMaterial = GetParentMaterialNode())
 		{
 			for (const FCustomizableObjectNodeEditMaterialImage& Image : Images_DEPRECATED)
 			{
diff --git a/Engine/Plugins/Experimental/Mutable/Source/CustomizableObjectEditor/Private/MuCOE/Nodes/CustomizableObjectNodeExtendMaterial.cpp b/Engine/Plugins/Experimental/Mutable/Source/CustomizableObjectEditor/Private/MuCOE/Nodes/CustomizableObjectNodeExtendMaterial.cpp
--- a/Engine/Plugins/Experimental/Mutable/Source/CustomizableObjectEditor/Private/MuCOE/Nodes/CustomizableObjectNodeExtendMaterial.cpp
+++ b/Engine/Plugins/Experimental/Mutable/Source/CustomizableObjectEditor/Private/MuCOE/Nodes/CustomizableObjectNodeExtendMaterial.cpp
@@ -48,7 +48,7 @@ void UCustomizableObjectNodeExtendMaterial::BackwardsCompatibleFixup()
 	
 	if (CustomizableObjectCustomVersion < FCustomizableObjectCustomVersion::ExtendMaterialRemoveImages)
 	{
-		if (const UCustomizableObjectNodeMaterial* ParentMaterial = GetParentMaterialNode())
+		if (const UCustomizableObjectNodeMaterialBase* ParentMaterial = GetParentMaterialNode())
 		{
 			for (const FCustomizableObjectNodeExtendMaterialImage& Image : Images_DEPRECATED)
 			{
@@ -106,15 +106,7 @@ void UCustomizableObjectNodeExtendMaterial::AllocateDefaultPins(UCustomizableObj
 	AddMeshPin->bDefaultValueIsIgnored = true;
 
 	// Begin texture pins
-	UCustomizableObjectNodeMaterial* ParentMaterialNode = GetParentMaterialNodeIfPath();
-	
-	// In case of a NodeCopyMaterial get the real NodeMaterial
-	if (const UCustomizableObjectNodeCopyMaterial* ParentMaterialCopyNode = Cast<UCustomizableObjectNodeCopyMaterial>(ParentMaterialNode))
-	{
-		ParentMaterialNode = ParentMaterialCopyNode->GetMaterialNode();
-	}
-
-	if (ParentMaterialNode)
+	if (UCustomizableObjectNodeMaterialBase* ParentMaterialNode = GetParentMaterialNodeIfPath())
 	{
 		const int32 NumImages = ParentMaterialNode->GetNumParameters(EMaterialParameterType::Texture);
 		for (int32 ImageIndex = 0; ImageIndex < NumImages; ++ImageIndex)
@@ -199,7 +191,7 @@ void UCustomizableObjectNodeExtendMaterial::SetParentNode(UCustomizableObject* O
 }
 
 
-TArray<UCustomizableObjectLayout*> UCustomizableObjectNodeExtendMaterial::GetLayouts()
+TArray<UCustomizableObjectLayout*> UCustomizableObjectNodeExtendMaterial::GetLayouts() const
 {
 	TArray<UCustomizableObjectLayout*> Result;
 
diff --git a/Engine/Plugins/Experimental/Mutable/Source/CustomizableObjectEditor/Private/MuCOE/Nodes/CustomizableObjectNodeMaterial.cpp b/Engine/Plugins/Experimental/Mutable/Source/CustomizableObjectEditor/Private/MuCOE/Nodes/CustomizableObjectNodeMaterial.cpp
--- a/Engine/Plugins/Experimental/Mutable/Source/CustomizableObjectEditor/Private/MuCOE/Nodes/CustomizableObjectNodeMaterial.cpp
+++ b/Engine/Plugins/Experimental/Mutable/Source/CustomizableObjectEditor/Private/MuCOE/Nodes/CustomizableObjectNodeMaterial.cpp
@@ -762,6 +762,12 @@ UEdGraphPin* UCustomizableObjectNodeMaterial::GetMaterialAssetPin() const
 }
 
 
+const UCustomizableObjectNodeMaterial* UCustomizableObjectNodeMaterial::GetMaterialNode() const
+{
+	return this;
+}
+
+
 bool UCustomizableObjectNodeMaterial::IsImageMutableMode(const int32 ImageIndex) const
 {
 	if (const UEdGraphPin* Pin = GetParameterPin(EMaterialParameterType::Texture, ImageIndex))
@@ -850,7 +856,7 @@ UTexture2D* UCustomizableObjectNodeMaterial::GetImageValue(const int32 ImageInde
 }
 
 
-TArray<UCustomizableObjectLayout*> UCustomizableObjectNodeMaterial::GetLayouts()
+TArray<UCustomizableObjectLayout*> UCustomizableObjectNodeMaterial::GetLayouts() const
 {
 	TArray<UCustomizableObjectLayout*> Result;
 
@@ -940,11 +946,8 @@ FGuid UCustomizableObjectNodeMaterial::GetParameterId(const EMaterialParameterTy
 
 FName UCustomizableObjectNodeMaterial::GetParameterName(const EMaterialParameterType Type, const int32 ParameterIndex) const
 {
-	if (!Material)
-	{
-		return FName();
-	}
-
+	check(Material);
+	
 	const FMaterialCachedParameterEntry& Entry = Material->GetCachedExpressionData().GetParameterTypeEntry(Type);
 
 	for (TSet<FMaterialParameterInfo>::TConstIterator It(Entry.ParameterInfoSet); It; ++It)
@@ -966,10 +969,7 @@ FName UCustomizableObjectNodeMaterial::GetParameterName(const EMaterialParameter
 
 int32 UCustomizableObjectNodeMaterial::GetParameterLayerIndex(const EMaterialParameterType Type, const int32 ParameterIndex) const
 {
-	if (!Material)
-	{
-		return -1;
-	}
+	check(Material);
 
 	const FMaterialCachedParameterEntry& Entry = Material->GetCachedExpressionData().GetParameterTypeEntry(Type);
 
@@ -992,11 +992,7 @@ int32 UCustomizableObjectNodeMaterial::GetParameterLayerIndex(const EMaterialPar
 
 FText UCustomizableObjectNodeMaterial::GetParameterLayerName(const EMaterialParameterType Type, const int32 ParameterIndex) const
 {
-	ensure(Material);
-	if (!Material)
-	{
-		return FText();
-	}
+	check(Material)
 
 	int32 LayerIndex = GetParameterLayerIndex(Type,ParameterIndex);
 
@@ -1124,6 +1120,36 @@ void UCustomizableObjectNodeMaterial::ReconstructNode(UCustomizableObjectNodeRem
 }
 
 
+void UCustomizableObjectNodeMaterial::SetMaterial(UMaterialInterface* InMaterial)
+{
+	Material = InMaterial;
+}
+
+
+UMaterialInterface* UCustomizableObjectNodeMaterial::GetMaterial() const
+{
+	return Material;	
+}
+
+
+bool UCustomizableObjectNodeMaterial::IsReuseMaterialBetweenLODs() const
+{
+	return bReuseMaterialBetweenLODs;	
+}
+
+
+int32 UCustomizableObjectNodeMaterial::GetMeshComponentIndex() const
+{
+	return MeshComponentIndex;	
+}
+
+
+TArray<FString> UCustomizableObjectNodeMaterial::GetTags() const
+{
+	return Tags;
+}
+
+
 bool UCustomizableObjectNodeMaterial::RealMaterialDataHasChanged() const
 {
 	for (const UEdGraphPin* Pin : GetAllNonOrphanPins())
@@ -1142,6 +1168,12 @@ bool UCustomizableObjectNodeMaterial::RealMaterialDataHasChanged() const
 }
 
 
+FPostImagePinModeChangedDelegate* UCustomizableObjectNodeMaterial::GetPostImagePinModeChangedDelegate()
+{
+	return &PostImagePinModeChangedDelegate;
+}
+
+
 bool UCustomizableObjectNodeMaterial::IsPinRelevant(const UEdGraphPin* Pin) const
 {
 	const UEdGraphSchema_CustomizableObject* Schema = GetDefault<UEdGraphSchema_CustomizableObject>();
diff --git a/Engine/Plugins/Experimental/Mutable/Source/CustomizableObjectEditor/Private/MuCOE/Nodes/CustomizableObjectNodeMaterialBase.cpp b/Engine/Plugins/Experimental/Mutable/Source/CustomizableObjectEditor/Private/MuCOE/Nodes/CustomizableObjectNodeMaterialBase.cpp
--- a/Engine/Plugins/Experimental/Mutable/Source/CustomizableObjectEditor/Private/MuCOE/Nodes/CustomizableObjectNodeMaterialBase.cpp
+++ b/Engine/Plugins/Experimental/Mutable/Source/CustomizableObjectEditor/Private/MuCOE/Nodes/CustomizableObjectNodeMaterialBase.cpp
@@ -2,10 +2,3 @@
 
 #include "MuCOE/Nodes/CustomizableObjectNodeMaterialBase.h"
 
-class UEdGraphPin;
-
-
-bool UCustomizableObjectNodeMaterialBase::ShouldBreakExistingConnections(const UEdGraphPin* InputPin, const UEdGraphPin* OutputPin) const
-{
-	return false;
-}
diff --git a/Engine/Plugins/Experimental/Mutable/Source/CustomizableObjectEditor/Private/MuCOE/Nodes/CustomizableObjectNodeMeshClipWithMeshDetails.cpp b/Engine/Plugins/Experimental/Mutable/Source/CustomizableObjectEditor/Private/MuCOE/Nodes/CustomizableObjectNodeMeshClipWithMeshDetails.cpp
--- a/Engine/Plugins/Experimental/Mutable/Source/CustomizableObjectEditor/Private/MuCOE/Nodes/CustomizableObjectNodeMeshClipWithMeshDetails.cpp
+++ b/Engine/Plugins/Experimental/Mutable/Source/CustomizableObjectEditor/Private/MuCOE/Nodes/CustomizableObjectNodeMeshClipWithMeshDetails.cpp
@@ -139,10 +139,10 @@ void FCustomizableObjectNodeMeshClipWithMeshDetails::CustomizeDetails(IDetailLay
 
 		for (const UCustomizableObjectNodeMaterial* MaterialNode : ArrayMaterialNodes)
 		{
-			if (MaterialNode->Material)
+			if (MaterialNode->GetMaterial())
 			{
 				NumElementBeforeAdd = ArrayMaterialNodeName.Num();
-				ArrayMaterialNodeName.AddUnique(MaterialNode->Material->GetName());
+				ArrayMaterialNodeName.AddUnique(MaterialNode->GetMaterial()->GetName());
 				NumElementAfterAdd = ArrayMaterialNodeName.Num();
 
 				if (NumElementBeforeAdd != NumElementAfterAdd)
@@ -310,7 +310,7 @@ void FCustomizableObjectNodeMeshClipWithMeshDetails::OnMeshClipWithMeshNodeCombo
 		Node->ArrayMaterialNodeToClipWithID.Empty();
 		for (const UCustomizableObjectNodeMaterial* MaterialNode : ArrayMaterialNode)
 		{
-			if (MaterialNode && MaterialNode->Material && *Selection == MaterialNode->Material->GetName())
+			if (MaterialNode && MaterialNode->GetMaterial() && *Selection == MaterialNode->GetMaterial()->GetName())
 			{
 				Node->ArrayMaterialNodeToClipWithID.Add(MaterialNode->NodeGuid);
 			}
diff --git a/Engine/Plugins/Experimental/Mutable/Source/CustomizableObjectEditor/Private/MuCOE/Nodes/CustomizableObjectNodeObject.cpp b/Engine/Plugins/Experimental/Mutable/Source/CustomizableObjectEditor/Private/MuCOE/Nodes/CustomizableObjectNodeObject.cpp
--- a/Engine/Plugins/Experimental/Mutable/Source/CustomizableObjectEditor/Private/MuCOE/Nodes/CustomizableObjectNodeObject.cpp
+++ b/Engine/Plugins/Experimental/Mutable/Source/CustomizableObjectEditor/Private/MuCOE/Nodes/CustomizableObjectNodeObject.cpp
@@ -16,6 +16,7 @@
 #include "Logging/MessageLog.h"
 #include "Containers/Queue.h"
 #include "MuCO/CustomizableObjectCustomVersion.h"
+#include "MuCOE/Nodes/CustomizableObjectNodeCopyMaterial.h"
 
 class UCustomizableObjectNodeRemapPins;
 
@@ -246,9 +247,9 @@ bool UCustomizableObjectNodeObject::CanDuplicateNode() const
 }
 
 
-TArray<UCustomizableObjectNodeMaterial*> UCustomizableObjectNodeObject::GetMaterialNodes(const int LOD) const
+TArray<UCustomizableObjectNodeMaterialBase*> UCustomizableObjectNodeObject::GetMaterialNodes(const int LOD) const
 {
-	TArray<UCustomizableObjectNodeMaterial*> Result;
+	TArray<UCustomizableObjectNodeMaterialBase*> Result;
 
 	TQueue<UEdGraphNode*> PotentialCustomizableNodeObjects;
 
@@ -260,10 +261,10 @@ TArray<UCustomizableObjectNodeMaterial*> UCustomizableObjectNodeObject::GetMater
 	UEdGraphNode* CurrentElement;
 	while (PotentialCustomizableNodeObjects.Dequeue(CurrentElement))
 	{
-		if (UCustomizableObjectNodeMaterial* CurrentMaterialNode = Cast<UCustomizableObjectNodeMaterial>(CurrentElement))
+		if (UCustomizableObjectNodeMaterialBase* CurrentMaterialNode = Cast<UCustomizableObjectNodeMaterialBase>(CurrentElement))
 		{
 			Result.Add(CurrentMaterialNode);
-		} 
+		}
 		else if (UCustomizableObjectNodeMaterialVariation* CurrentMaterialVariationNode = Cast<UCustomizableObjectNodeMaterialVariation>(CurrentElement))
 		{
 			// Case of material variation. It's not a material, but a node that further references any material, add all its inputs that could be a material
diff --git a/Engine/Plugins/Experimental/Mutable/Source/CustomizableObjectEditor/Private/MuCOE/Nodes/CustomizableObjectNodeParentedMaterial.cpp b/Engine/Plugins/Experimental/Mutable/Source/CustomizableObjectEditor/Private/MuCOE/Nodes/CustomizableObjectNodeParentedMaterial.cpp
--- a/Engine/Plugins/Experimental/Mutable/Source/CustomizableObjectEditor/Private/MuCOE/Nodes/CustomizableObjectNodeParentedMaterial.cpp
+++ b/Engine/Plugins/Experimental/Mutable/Source/CustomizableObjectEditor/Private/MuCOE/Nodes/CustomizableObjectNodeParentedMaterial.cpp
@@ -6,21 +6,15 @@
 #include "MuCOE/Nodes/CustomizableObjectNodeObject.h"
 
 
-UCustomizableObjectNodeMaterial* FCustomizableObjectNodeParentedMaterial::GetParentMaterialNode() const
+UCustomizableObjectNodeMaterialBase* FCustomizableObjectNodeParentedMaterial::GetParentMaterialNode() const
 {
-	if (UCustomizableObjectNode* Node = GetParentNode())
-	{
-		return CastChecked<UCustomizableObjectNodeMaterial>(Node);
-	}
-	else
-	{
-		return nullptr;	
-	}
+	return Cast<UCustomizableObjectNodeMaterialBase>(GetParentNode());
 }
 
-TArray<UCustomizableObjectNodeMaterial*> FCustomizableObjectNodeParentedMaterial::GetPossibleParentMaterialNodes() const
+
+TArray<UCustomizableObjectNodeMaterialBase*> FCustomizableObjectNodeParentedMaterial::GetPossibleParentMaterialNodes() const
 {
-	TArray<UCustomizableObjectNodeMaterial*> Result;
+	TArray<UCustomizableObjectNodeMaterialBase*> Result;
 	
 	const UCustomizableObjectNode& Node = GetNode();
 	const int32 LOD = Node.GetLOD();
@@ -50,22 +44,8 @@ TArray<UCustomizableObjectNodeMaterial*> FCustomizableObjectNodeParentedMaterial
 		// If LODStrategy is set to AutomaticFromMesh, find MaterialNodes belonging to lower LODs.
 		for (; LODIndex <= LOD; ++LODIndex)
 		{
-			const TArray<UCustomizableObjectNodeMaterial*> MaterialNodes = ParentObjectNode->GetMaterialNodes(LODIndex);
-
-			for (UCustomizableObjectNodeMaterial* MaterialNode : MaterialNodes)
-			{
-				if (UCustomizableObjectNodeCopyMaterial* TypedCopyMaterialNode = Cast<UCustomizableObjectNodeCopyMaterial>(MaterialNode))
-				{
-					if (TypedCopyMaterialNode->GetMaterialNode())
-					{
-						Result.Add(TypedCopyMaterialNode);
-					}
-				}
-				else if (UCustomizableObjectNodeMaterial* TypedMaterialNode = Cast<UCustomizableObjectNodeMaterial>(MaterialNode))
-				{
-					Result.Add(TypedMaterialNode);
-				}
-			}
+			const TArray<UCustomizableObjectNodeMaterialBase*> MaterialNodes = ParentObjectNode->GetMaterialNodes(LODIndex);
+			Result.Append(MaterialNodes);
 		}
 	}
 	
@@ -73,9 +53,9 @@ TArray<UCustomizableObjectNodeMaterial*> FCustomizableObjectNodeParentedMaterial
 }
 
 
-UCustomizableObjectNodeMaterial* FCustomizableObjectNodeParentedMaterial::GetParentMaterialNodeIfPath() const
+UCustomizableObjectNodeMaterialBase* FCustomizableObjectNodeParentedMaterial::GetParentMaterialNodeIfPath() const
 {
-	UCustomizableObjectNodeMaterial* ParentMaterialNode = GetParentMaterialNode();
+	UCustomizableObjectNodeMaterialBase* ParentMaterialNode = GetParentMaterialNode();
 	
 	if (GetPossibleParentMaterialNodes().Contains(ParentMaterialNode)) // There is a path to the parent material
 	{
diff --git a/Engine/Plugins/Experimental/Mutable/Source/CustomizableObjectEditor/Private/MuCOE/Nodes/CustomizableObjectNodeParentedMaterialDetails.cpp b/Engine/Plugins/Experimental/Mutable/Source/CustomizableObjectEditor/Private/MuCOE/Nodes/CustomizableObjectNodeParentedMaterialDetails.cpp
--- a/Engine/Plugins/Experimental/Mutable/Source/CustomizableObjectEditor/Private/MuCOE/Nodes/CustomizableObjectNodeParentedMaterialDetails.cpp
+++ b/Engine/Plugins/Experimental/Mutable/Source/CustomizableObjectEditor/Private/MuCOE/Nodes/CustomizableObjectNodeParentedMaterialDetails.cpp
@@ -56,10 +56,10 @@ void FCustomizableObjectNodeParentedMaterialDetails::CustomizeDetails(IDetailLay
 	if (NodeParentedMaterial)
 	{
 		// Get parent materials
-		TArray<UCustomizableObjectNodeMaterial*> ParentMaterialBaseNodes = NodeParentedMaterial->GetPossibleParentMaterialNodes();
+		TArray<UCustomizableObjectNodeMaterialBase*> ParentMaterialBaseNodes = NodeParentedMaterial->GetPossibleParentMaterialNodes();
 
 		ParentMaterialOptionReferences.Empty();
-		for (UCustomizableObjectNodeMaterial* ParentMaterial : ParentMaterialBaseNodes)
+		for (UCustomizableObjectNodeMaterialBase* ParentMaterial : ParentMaterialBaseNodes)
 		{
 			UCustomizableObject* Object = Cast<UCustomizableObject>(ParentMaterial->GetGraph()->GetOuter());
 			FMaterialReference MaterialReference = { Object, ParentMaterial->NodeGuid };
@@ -158,19 +158,28 @@ void FCustomizableObjectNodeParentedMaterialDetails::OnParentComboBoxSelectionCh
 }
 
 
-TArray<TSharedPtr<FString>> FCustomizableObjectNodeParentedMaterialDetails::GetComboBoxNames(const TArray<UCustomizableObjectNodeMaterial*>& ParentMaterialNodes) const
+TArray<TSharedPtr<FString>> FCustomizableObjectNodeParentedMaterialDetails::GetComboBoxNames(const TArray<UCustomizableObjectNodeMaterialBase*>& ParentMaterialNodes) const
 {
 	TArray<TSharedPtr<FString>> Result;
 
-	for (UCustomizableObjectNodeMaterial* ParentMaterialNode : ParentMaterialNodes)
+	for (UCustomizableObjectNodeMaterialBase* ParentMaterialNode : ParentMaterialNodes)
 	{
 		const FString ObjectName = ParentMaterialNode->GetParentObjectNodes(ParentMaterialNode->GetLOD())[0]->ObjectName; // Since ParentMaterialNode comes from ParentMaterialNodes, it will always have at least a ParentObjectNode.
 		FString Name = ObjectName + TEXT(" - ");
 
 		if (UCustomizableObjectNodeCopyMaterial* CopyMaterialNode = Cast<UCustomizableObjectNodeCopyMaterial>(ParentMaterialNode))
 		{
-			Name += FString(" [Copy Material] (") + GetComboBoxParentMaterialName(CopyMaterialNode) + FString(" + ");
+			Name += LOCTEXT("CopyMaterialParentHeader", "Copy Material - ").ToString();
 
+			if (UCustomizableObjectNodeMaterial* NodeMaterial = CopyMaterialNode->GetMaterialNode())
+			{
+				Name += GetComboBoxParentMaterialName(NodeMaterial) + FString(" + ");
+			}
+			else
+			{
+				Name += LOCTEXT("MissingParentMaterial", "[Missing Parent Material]").ToString();
+			}
+			
 			const UCustomizableObjectNodeSkeletalMesh* NodeSkeletalMesh = CopyMaterialNode->GetMeshNode();
 			if (NodeSkeletalMesh && NodeSkeletalMesh->SkeletalMesh)
 			{
@@ -178,12 +187,10 @@ TArray<TSharedPtr<FString>> FCustomizableObjectNodeParentedMaterialDetails::GetC
 			}
 			else
 			{
-				Name += FString("[Missing Skeletal Mesh]");
+				Name += LOCTEXT("MissingSkeletalMesh", "[Missing Skeletal Mesh]").ToString();
 			}
-
-			Name += FString(")");
 		}
-		if (UCustomizableObjectNodeMaterial* MaterialNode = Cast<UCustomizableObjectNodeMaterial>(ParentMaterialNode))
+		else if (UCustomizableObjectNodeMaterial* MaterialNode = Cast<UCustomizableObjectNodeMaterial>(ParentMaterialNode))
 		{
 			Name += GetComboBoxParentMaterialName(MaterialNode);
 		}
@@ -201,7 +208,7 @@ TArray<TSharedPtr<FString>> FCustomizableObjectNodeParentedMaterialDetails::GetC
 
 FString FCustomizableObjectNodeParentedMaterialDetails::GetComboBoxParentMaterialName(const UCustomizableObjectNodeMaterial* ParentMaterial) const
 {
-	return ParentMaterial->Material ? ParentMaterial->Material->GetName() : FString("[Missing material]");
+	return ParentMaterial->GetMaterial() ? ParentMaterial->GetMaterial()->GetName() : LOCTEXT("MissingMaterial", "[Missing Material]").ToString();
 }
 
 
diff --git a/Engine/Plugins/Experimental/Mutable/Source/CustomizableObjectEditor/Private/MuCOE/Nodes/CustomizableObjectNodeParentedMaterialDetails.h b/Engine/Plugins/Experimental/Mutable/Source/CustomizableObjectEditor/Private/MuCOE/Nodes/CustomizableObjectNodeParentedMaterialDetails.h
--- a/Engine/Plugins/Experimental/Mutable/Source/CustomizableObjectEditor/Private/MuCOE/Nodes/CustomizableObjectNodeParentedMaterialDetails.h
+++ b/Engine/Plugins/Experimental/Mutable/Source/CustomizableObjectEditor/Private/MuCOE/Nodes/CustomizableObjectNodeParentedMaterialDetails.h
@@ -7,6 +7,8 @@
 #include "Misc/Guid.h"
 #include "UObject/WeakObjectPtr.h"
 
+class UCustomizableObjectNodeMaterialBase;
+
 namespace ESelectInfo { enum Type : int; }
 
 class FCustomizableObjectNodeParentedMaterial;
@@ -41,7 +43,7 @@ protected:
 	 *
 	 * All ParentMaterialNodes pointers must be valid.
 	 */
-	TArray<TSharedPtr<FString>> GetComboBoxNames(const TArray<UCustomizableObjectNodeMaterial*>& ParentMaterialNodes) const;
+	TArray<TSharedPtr<FString>> GetComboBoxNames(const TArray<UCustomizableObjectNodeMaterialBase*>& ParentMaterialNodes) const;
 
 	/** Return a fornamted name for the parent Material. */
 	FString GetComboBoxParentMaterialName(const UCustomizableObjectNodeMaterial* ParentMaterial) const;
diff --git a/Engine/Plugins/Experimental/Mutable/Source/CustomizableObjectEditor/Private/MuCOE/Nodes/CustomizableObjectNodeReroute.cpp b/Engine/Plugins/Experimental/Mutable/Source/CustomizableObjectEditor/Private/MuCOE/Nodes/CustomizableObjectNodeReroute.cpp
--- a/Engine/Plugins/Experimental/Mutable/Source/CustomizableObjectEditor/Private/MuCOE/Nodes/CustomizableObjectNodeReroute.cpp
+++ b/Engine/Plugins/Experimental/Mutable/Source/CustomizableObjectEditor/Private/MuCOE/Nodes/CustomizableObjectNodeReroute.cpp
@@ -11,13 +11,10 @@
 
 void UCustomizableObjectNodeReroute::AllocateDefaultPins(UCustomizableObjectNodeRemapPins* RemapPins)
 {
-	const FName InputPinName(TEXT("InputPin"));
-	const FName OutputPinName(TEXT("OutputPin"));
-
-	UEdGraphPin* MyInputPin = CreatePin(EGPD_Input, UEdGraphSchema_CustomizableObject::PC_Wildcard, InputPinName);
+	UEdGraphPin* MyInputPin = CreatePin(EGPD_Input, UEdGraphSchema_CustomizableObject::PC_Wildcard, TEXT("InputPin"));
 	MyInputPin->bDefaultValueIsIgnored = true;
 
-	CreatePin(EGPD_Output, UEdGraphSchema_CustomizableObject::PC_Wildcard, OutputPinName);
+	CreatePin(EGPD_Output, UEdGraphSchema_CustomizableObject::PC_Wildcard, TEXT("OutputPin"));
 }
 
 
diff --git a/Engine/Plugins/Experimental/Mutable/Source/CustomizableObjectEditor/Private/MuCOE/Nodes/CustomizableObjectNodeUseMaterial.cpp b/Engine/Plugins/Experimental/Mutable/Source/CustomizableObjectEditor/Private/MuCOE/Nodes/CustomizableObjectNodeUseMaterial.cpp
--- a/Engine/Plugins/Experimental/Mutable/Source/CustomizableObjectEditor/Private/MuCOE/Nodes/CustomizableObjectNodeUseMaterial.cpp
+++ b/Engine/Plugins/Experimental/Mutable/Source/CustomizableObjectEditor/Private/MuCOE/Nodes/CustomizableObjectNodeUseMaterial.cpp
@@ -20,7 +20,7 @@ bool FCustomizableObjectNodeUseMaterial::IsNodeOutDatedAndNeedsRefreshWork()
 	
 	const bool bOutdated = [&]()
 	{
-		if (const UCustomizableObjectNodeMaterial* ParentMaterialNode = NodeParentedMaterial.GetParentMaterialNode())
+		if (const UCustomizableObjectNodeMaterialBase* ParentMaterialNode = NodeParentedMaterial.GetParentMaterialNode())
 		{
 			if (ParentMaterialNode->RealMaterialDataHasChanged())
 			{
@@ -58,10 +58,14 @@ void FCustomizableObjectNodeUseMaterial::PreSetParentNodeWork(UCustomizableObjec
 {
 	const FCustomizableObjectNodeParentedMaterial& NodeParentedMaterial = GetNodeParentedMaterial();
 
-	if (UCustomizableObjectNodeMaterial* ParentMaterialNode = NodeParentedMaterial.GetParentMaterialNode())
+	if (UCustomizableObjectNodeMaterialBase* ParentMaterialNode = NodeParentedMaterial.GetParentMaterialNode())
 	{
 		ParentMaterialNode->PostReconstructNodeDelegate.Remove(PostReconstructNodeDelegateHandler);
-		ParentMaterialNode->PostImagePinModeChangedDelegate.Remove(PostTextureParameterModeChangedDelegateHandle);
+		
+		if (FPostImagePinModeChangedDelegate* PostImagePinModeChangedDelegate = ParentMaterialNode->GetPostImagePinModeChangedDelegate())
+		{
+			PostImagePinModeChangedDelegate->Remove(PostTextureParameterModeChangedDelegateHandle);
+		}
 	}
 }
 
@@ -71,10 +75,14 @@ void FCustomizableObjectNodeUseMaterial::PostSetParentNodeWork(UCustomizableObje
 	UCustomizableObjectNode& Node = GetNode();
 	const FCustomizableObjectNodeParentedMaterial& NodeParentedMaterial = GetNodeParentedMaterial();
 
-	if (UCustomizableObjectNodeMaterial* ParentMaterialNode = NodeParentedMaterial.GetParentMaterialNode())
+	if (UCustomizableObjectNodeMaterialBase* ParentMaterialNode = NodeParentedMaterial.GetParentMaterialNode())
 	{
 		PostReconstructNodeDelegateHandler = ParentMaterialNode->PostReconstructNodeDelegate.AddUObject(&Node, &UCustomizableObjectNode::ReconstructNode);
-		PostTextureParameterModeChangedDelegateHandle = ParentMaterialNode->PostImagePinModeChangedDelegate.AddUObject(&Node, &UCustomizableObjectNode::ReconstructNode);
+
+		if (FPostImagePinModeChangedDelegate* PostImagePinModeChangedDelegate = ParentMaterialNode->GetPostImagePinModeChangedDelegate())
+		{
+			PostTextureParameterModeChangedDelegateHandle = PostImagePinModeChangedDelegate->AddUObject(&Node, &UCustomizableObjectNode::ReconstructNode);
+		}
 	}
 	
 	GetNode().ReconstructNode();
@@ -162,10 +170,14 @@ void FCustomizableObjectNodeUseMaterial::PostBackwardsCompatibleFixupWork()
 	UCustomizableObjectNode& Node = GetNode();
 	const FCustomizableObjectNodeParentedMaterial& NodeParentedMaterial = GetNodeParentedMaterial();
 
-	if (UCustomizableObjectNodeMaterial* ParentMaterialNode = NodeParentedMaterial.GetParentMaterialNode())
+	if (UCustomizableObjectNodeMaterialBase* ParentMaterialNode = NodeParentedMaterial.GetParentMaterialNode())
 	{
 		PostReconstructNodeDelegateHandler = ParentMaterialNode->PostReconstructNodeDelegate.AddUObject(&Node, &UCustomizableObjectNode::ReconstructNode);
-		PostTextureParameterModeChangedDelegateHandle = ParentMaterialNode->PostImagePinModeChangedDelegate.AddUObject(&Node, &UCustomizableObjectNode::ReconstructNode);
+
+		if (FPostImagePinModeChangedDelegate* PostImagePinModeChangedDelegate = ParentMaterialNode->GetPostImagePinModeChangedDelegate())
+		{
+			PostTextureParameterModeChangedDelegateHandle = PostImagePinModeChangedDelegate->AddUObject(&Node, &UCustomizableObjectNode::ReconstructNode);
+		}
 	}
 
 	Node.ReconstructNode(); // Reconstruct the node since the parent could have changed while not loaded. 
diff --git a/Engine/Plugins/Experimental/Mutable/Source/CustomizableObjectEditor/Private/MuCOE/SCustomizableObjectEditorTagExplorer.cpp b/Engine/Plugins/Experimental/Mutable/Source/CustomizableObjectEditor/Private/MuCOE/SCustomizableObjectEditorTagExplorer.cpp
--- a/Engine/Plugins/Experimental/Mutable/Source/CustomizableObjectEditor/Private/MuCOE/SCustomizableObjectEditorTagExplorer.cpp
+++ b/Engine/Plugins/Experimental/Mutable/Source/CustomizableObjectEditor/Private/MuCOE/SCustomizableObjectEditorTagExplorer.cpp
@@ -256,14 +256,16 @@ void SCustomizableObjectEditorTagExplorer::FillTagInformation(UCustomizableObjec
 	{
 		for (const TObjectPtr<UEdGraphNode>& Node : Object->GetPrivate()->GetSource()->Nodes)
 		{
-			if (UCustomizableObjectNodeMaterial* TypedNodeMat = Cast<UCustomizableObjectNodeMaterial>(Node))
+			if (UCustomizableObjectNodeMaterialBase* TypedNodeMat = Cast<UCustomizableObjectNodeMaterialBase>(Node))
 			{
-				for (int32 i = 0; i < TypedNodeMat->Tags.Num(); ++i)
+				const TArray<FString>& MaterialTags = TypedNodeMat->GetTags();
+				
+				for (int32 i = 0; i < MaterialTags.Num(); ++i)
 				{
-					NodeTags.Add(TypedNodeMat->Tags[i], TypedNodeMat);
-					if (Tags.Find(TypedNodeMat->Tags[i]) == INDEX_NONE)
+					NodeTags.Add(MaterialTags[i], TypedNodeMat);
+					if (Tags.Find(MaterialTags[i]) == INDEX_NONE)
 					{
-						Tags.Add(TypedNodeMat->Tags[i]);
+						Tags.Add(MaterialTags[i]);
 					}
 				}
 			}
@@ -387,7 +389,7 @@ void SCustomizableObjectEditorTagExplorer::OnComboBoxSelectionChanged(FString Ne
 
 		for (UCustomizableObjectNode* node : auxNodes)
 		{
-			if (Cast<UCustomizableObjectNodeMaterial>(node) && MaterialNodes.Find(node) == INDEX_NONE)
+			if (Cast<UCustomizableObjectNodeMaterialBase>(node) && MaterialNodes.Find(node) == INDEX_NONE)
 			{
 				MaterialNodes.Add(node);
 			}
diff --git a/Engine/Plugins/Experimental/Mutable/Source/CustomizableObjectEditor/Public/MuCOE/Nodes/CustomizableObjectNodeCopyMaterial.h b/Engine/Plugins/Experimental/Mutable/Source/CustomizableObjectEditor/Public/MuCOE/Nodes/CustomizableObjectNodeCopyMaterial.h
--- a/Engine/Plugins/Experimental/Mutable/Source/CustomizableObjectEditor/Public/MuCOE/Nodes/CustomizableObjectNodeCopyMaterial.h
+++ b/Engine/Plugins/Experimental/Mutable/Source/CustomizableObjectEditor/Public/MuCOE/Nodes/CustomizableObjectNodeCopyMaterial.h
@@ -11,61 +11,57 @@ class UCustomizableObjectNodeSkeletalMesh;
 class UEdGraphPin;
 class UObject;
 
-/**
- * Copy Material node. Duplicates a Material Node. Duplicates all Material node input pins and properties except for the Mesh input pin.
- * A new Mesh has to be defined through the new Mesh input pin.
- *
- * Input pins:
- * - Mesh: New mesh.
- * - Material: Material to duplicate.
- * 
- * Output pins:
- * - Material: New material.
- * 
- * Properties:
- *   [NONE]
- */
+
+/** Generates a new Surface with the same connections as the Parent Material Surface, but with a different Mesh. */
 UCLASS()
-class CUSTOMIZABLEOBJECTEDITOR_API UCustomizableObjectNodeCopyMaterial : public UCustomizableObjectNodeMaterial
+class CUSTOMIZABLEOBJECTEDITOR_API UCustomizableObjectNodeCopyMaterial : public UCustomizableObjectNodeMaterialBase
 {
 public:
 	GENERATED_BODY()
 
-	FText GetNodeTitle(ENodeTitleType::Type TitleType) const override;
+	// UEdGraphNode interface
+	virtual FText GetNodeTitle(ENodeTitleType::Type TitleType) const override;
+	virtual FText GetTooltipText() const override;
+	virtual FLinearColor GetNodeTitleColor() const override;
 
 	// UCustomizableObjectNode interface
-	void AllocateDefaultPins(UCustomizableObjectNodeRemapPins* RemapPins) override;
-
-	/** Get Mesh input pin. */
-	UEdGraphPin* GetMeshPin() const override;
-
-	/** Get Material input pin. */
-	UEdGraphPin* GetMaterialPin() const;
-
-	/** Get Mesh node. */
-	UCustomizableObjectNodeSkeletalMesh* GetMeshNode() const;
-
-	/** Get the Material node. */
-	UCustomizableObjectNodeMaterial* GetMaterialNode() const;
-
-	/** Input material pin can only be connect to a NodeMaterial output pin or a NodeExternalPin. A connection to a NodeMaterialCopy output pin is not allowed. */
+	virtual void AllocateDefaultPins(UCustomizableObjectNodeRemapPins* RemapPins) override;
 	virtual bool CanConnect(const UEdGraphPin* InOwnedInputPin, const UEdGraphPin* InOutputPin, bool& bOutIsOtherNodeBlocklisted, bool& bOutArePinsCompatible) const override;
-	
-	bool ShouldBreakExistingConnections(const UEdGraphPin* InputPin, const UEdGraphPin* OutputPin) const override;
-
-	/** Override the NodeMaterial method. Does not break any connection. */
-	void BreakExistingConnectionsPostConnection(UEdGraphPin* InputPin, UEdGraphPin* OutputPin) override {};
-
+	virtual bool ShouldBreakExistingConnections(const UEdGraphPin* InputPin, const UEdGraphPin* OutputPin) const override;
 	virtual bool IsNodeOutDatedAndNeedsRefresh() override;
-
-	bool ProvidesCustomPinRelevancyTest() const override;
-
-	/** 
-	 * Only are relevant pins:
-	 * - From: NodeMaterial, excluding NodeCopyMaterial 
-	 * - To: NodeBaseObject
-	 */
-	bool IsPinRelevant(const UEdGraphPin* Pin) const override;
-
-	FText GetTooltipText() const override;
+	virtual bool ProvidesCustomPinRelevancyTest() const override;
+	virtual bool IsPinRelevant(const UEdGraphPin* Pin) const override;
+
+	// UCustomizableObjectNodeMaterialBase interface
+	virtual UMaterialInterface* GetMaterial() const override;
+	virtual bool IsReuseMaterialBetweenLODs() const override;
+	virtual int32 GetMeshComponentIndex() const override;
+	virtual TArray<FString> GetTags() const override;
+	virtual UEdGraphPin* GetMaterialAssetPin() const override;
+	virtual int32 GetNumParameters(EMaterialParameterType Type) const override;
+	virtual FGuid GetParameterId(EMaterialParameterType Type, int32 ParameterIndex) const override;
+	virtual FName GetParameterName(EMaterialParameterType Type, int32 ParameterIndex) const override;
+	virtual int32 GetParameterLayerIndex(EMaterialParameterType Type, int32 ParameterIndex) const override;
+	virtual FText GetParameterLayerName(EMaterialParameterType Type, int32 ParameterIndex) const override;
+	virtual bool HasParameter(const FGuid& ParameterId) const override;
+	virtual const UEdGraphPin* GetParameterPin(EMaterialParameterType Type, int32 ParameterIndex) const override;
+	virtual bool IsImageMutableMode(int32 ImageIndex) const override;
+	virtual bool IsImageMutableMode(const UEdGraphPin& Pin) const override;
+	virtual void UpdateImagePinMode(const FGuid ParameterId) override;
+	virtual void UpdateImagePinMode(const UEdGraphPin& Pin) override;
+	virtual void UpdateAllImagesPinMode() override;
+	virtual UTexture2D* GetImageReferenceTexture(int32 ImageIndex) const override;
+	virtual UTexture2D* GetImageValue(int32 ImageIndex) const override;
+	virtual int32 GetImageUVLayout(int32 ImageIndex) const override;
+	virtual UCustomizableObjectNodeMaterial* GetMaterialNode() const override;
+	virtual UEdGraphPin* GetMeshPin() const override;
+	virtual FPostImagePinModeChangedDelegate* GetPostImagePinModeChangedDelegate() override;
+	virtual TArray<UCustomizableObjectLayout*> GetLayouts() const override;
+	virtual UEdGraphPin* OutputPin() const override;
+	virtual bool RealMaterialDataHasChanged() const override;
+	
+	// Own interface
+	UCustomizableObjectNodeSkeletalMesh* GetMeshNode() const;
+	
+	UEdGraphPin* GetMaterialPin() const;
 };
diff --git a/Engine/Plugins/Experimental/Mutable/Source/CustomizableObjectEditor/Public/MuCOE/Nodes/CustomizableObjectNodeExtendMaterial.h b/Engine/Plugins/Experimental/Mutable/Source/CustomizableObjectEditor/Public/MuCOE/Nodes/CustomizableObjectNodeExtendMaterial.h
--- a/Engine/Plugins/Experimental/Mutable/Source/CustomizableObjectEditor/Public/MuCOE/Nodes/CustomizableObjectNodeExtendMaterial.h
+++ b/Engine/Plugins/Experimental/Mutable/Source/CustomizableObjectEditor/Public/MuCOE/Nodes/CustomizableObjectNodeExtendMaterial.h
@@ -22,7 +22,7 @@ struct FEdGraphPinReference;
 
 UCLASS()
 class CUSTOMIZABLEOBJECTEDITOR_API UCustomizableObjectNodeExtendMaterial :
-	public UCustomizableObjectNodeMaterialBase,
+	public UCustomizableObjectNode,
 	public FCustomizableObjectNodeParentedMaterial,
 	public FCustomizableObjectNodeUseMaterial
 {
@@ -53,25 +53,24 @@ public:
 	virtual FString GetRefreshMessage() const override;
 	virtual bool IsSingleOutputNode() const override;
 	virtual bool CustomRemovePin(UEdGraphPin& Pin) override;
-
-	// UCustomizableObjectNodeMaterialBase interface
-	virtual UEdGraphPin* OutputPin() const override;
-	virtual void SetParentNode(UCustomizableObject* Object, FGuid NodeId) override;
-	TArray<UCustomizableObjectLayout*> GetLayouts() override;
 	
 	// FCustomizableObjectNodeParentMaterial interface
 	virtual void SaveParentNode(UCustomizableObject* Object, FGuid NodeId) override;
 	virtual UCustomizableObjectNode& GetNode() override;
 	virtual FGuid GetParentNodeId() const override;
 	virtual UCustomizableObject* GetParentObject() const override;
-	
+	virtual void SetParentNode(UCustomizableObject* Object, FGuid NodeId) override;
+
 	// FCustomizableObjectNodeUseMaterial interface
 	virtual FCustomizableObjectNodeParentedMaterial& GetNodeParentedMaterial() override;
 	virtual TMap<FGuid, FEdGraphPinReference>& GetPinsParameter() override;
-	
+	virtual UEdGraphPin* OutputPin() const override;
+
 	// Own interface
 	UEdGraphPin* AddMeshPin() const;
 	
+	TArray<UCustomizableObjectLayout*> GetLayouts() const;
+	
 private:
 	UPROPERTY(EditAnywhere, Category = CustomizableObject)
 	TObjectPtr<UCustomizableObject> ParentMaterialObject = nullptr;
diff --git a/Engine/Plugins/Experimental/Mutable/Source/CustomizableObjectEditor/Public/MuCOE/Nodes/CustomizableObjectNodeMaterial.h b/Engine/Plugins/Experimental/Mutable/Source/CustomizableObjectEditor/Public/MuCOE/Nodes/CustomizableObjectNodeMaterial.h
--- a/Engine/Plugins/Experimental/Mutable/Source/CustomizableObjectEditor/Public/MuCOE/Nodes/CustomizableObjectNodeMaterial.h
+++ b/Engine/Plugins/Experimental/Mutable/Source/CustomizableObjectEditor/Public/MuCOE/Nodes/CustomizableObjectNodeMaterial.h
@@ -27,9 +27,6 @@ struct FFrame;
 struct FPropertyChangedEvent;
 
 
-DECLARE_MULTICAST_DELEGATE(FPostImagePinModeChangedDelegate)
-
-
 /** Custom remap pins by name action.
  *
  * Remap pins by Texture Parameter Id. */
@@ -105,34 +102,17 @@ class CUSTOMIZABLEOBJECTEDITOR_API UCustomizableObjectNodeMaterial : public UCus
 {
 public:
 	GENERATED_BODY()
-
-	UPROPERTY(EditAnywhere, Category=CustomizableObject)
-	TObjectPtr<UMaterialInterface> Material = nullptr;
-
-	UPROPERTY(EditAnywhere, Category=CustomizableObject, Meta = (ToolTip = "Set all Mateiral Texture Parameters to the specified mode. Each Texture Parameter Pin can override this mode."))
-	ENodePinMode TextureParametersMode = ENodePinMode::Passthrough;
-
-	UPROPERTY(EditAnywhere, BlueprintReadWrite, Category = CustomizableObject)
-	TArray<FString> Tags;
-
-	/** Selects which Mesh component of the Instance this material belongs to */
-	UPROPERTY(EditAnywhere, BlueprintReadWrite, Category = CustomizableObject, meta = (ClampMin = "0"))
-	int32 MeshComponentIndex = 0;
-
-	/** Materials will be reused between LODs when possible unless explicitly disabled. */
-	UPROPERTY(EditAnywhere, Category = CustomizableObject)
-	bool bReuseMaterialBetweenLODs = true;
 	
 	// UObject interface.
 	virtual void Serialize(FArchive& Ar) override;
 	virtual void PostLoad() override;
-	void PostEditChangeProperty(FPropertyChangedEvent& PropertyChangedEvent) override;
+	virtual void PostEditChangeProperty(FPropertyChangedEvent& PropertyChangedEvent) override;
 
 	// EdGraphNode interface
-	FText GetNodeTitle(ENodeTitleType::Type TitleType) const override;
-	FLinearColor GetNodeTitleColor() const override;
-	FText GetTooltipText() const override;
-	TSharedPtr<SGraphNode> CreateVisualWidget() override;
+	virtual FText GetNodeTitle(ENodeTitleType::Type TitleType) const override;
+	virtual FLinearColor GetNodeTitleColor() const override;
+	virtual FText GetTooltipText() const override;
+	virtual TSharedPtr<SGraphNode> CreateVisualWidget() override;
 	virtual void PinConnectionListChanged(UEdGraphPin* Pin) override;
 	virtual void PostPasteNode() override;
 	virtual bool CanConnect(const UEdGraphPin* InOwnedInputPin, const UEdGraphPin* InOutputPin, bool& bOutIsOtherNodeBlocklisted, bool& bOutArePinsCompatible) const override;
@@ -140,114 +120,75 @@ public:
 	// UCustomizableObjectNode interface
 	virtual void BackwardsCompatibleFixup() override;
 	virtual void PostBackwardsCompatibleFixup() override;
-	void AllocateDefaultPins(UCustomizableObjectNodeRemapPins* RemapPins) override;
-	bool CanPinBeHidden(const UEdGraphPin& Pin) const override;
+	virtual void AllocateDefaultPins(UCustomizableObjectNodeRemapPins* RemapPins) override;
+	virtual bool CanPinBeHidden(const UEdGraphPin& Pin) const override;
 	virtual bool HasPinViewer() const override;
 	virtual UCustomizableObjectNodeRemapPinsByName* CreateRemapPinsDefault() const override;
-	bool ProvidesCustomPinRelevancyTest() const override { return true; }
-	bool IsPinRelevant(const UEdGraphPin* Pin) const override;
+	virtual bool ProvidesCustomPinRelevancyTest() const override { return true; }
+	virtual bool IsPinRelevant(const UEdGraphPin* Pin) const override;
 	virtual bool CustomRemovePin(UEdGraphPin& Pin) override;
 	virtual void ReconstructNode(UCustomizableObjectNodeRemapPins* RemapPinsMode) override;
-
-	// Own Interface
-	virtual UEdGraphPin* GetMeshPin() const;
-	UEdGraphPin* GetMaterialAssetPin() const;
-
-	/** Allow only one connection to a single BaseObjec node. There can be multiple conections to NodeCopyMaterial nodes. */
-	void BreakExistingConnectionsPostConnection(UEdGraphPin* InputPin, UEdGraphPin* OutputPin) override;
-
-	bool IsNodeOutDatedAndNeedsRefresh() override;
-	FString GetRefreshMessage() const override;
+	virtual void BreakExistingConnectionsPostConnection(UEdGraphPin* InputPin, UEdGraphPin* OutputPin) override;
+	virtual bool IsNodeOutDatedAndNeedsRefresh() override;
+	virtual FString GetRefreshMessage() const override;
 	virtual TSharedPtr<SWidget> CustomizePinDetails(const UEdGraphPin& Pin) const override;
 
 	
 	// UCustomizableObjectNodeMaterialBase interface
-	TArray<class UCustomizableObjectLayout*> GetLayouts() override;
-	UEdGraphPin* OutputPin() const override;
-	
-	/** Return true if a Material Parameter has changed on which we had a pin connected or data saved. */
-	bool RealMaterialDataHasChanged() const;
+	virtual TArray<UCustomizableObjectLayout*> GetLayouts() const override;
+	virtual UEdGraphPin* OutputPin() const override;
+	virtual UMaterialInterface* GetMaterial() const override;
+	virtual bool IsReuseMaterialBetweenLODs() const override;
+	virtual int32 GetMeshComponentIndex() const override;
+	virtual TArray<FString> GetTags() const override;
+	virtual int32 GetNumParameters(EMaterialParameterType Type) const override;
+	virtual FGuid GetParameterId(EMaterialParameterType Type, int32 ParameterIndex) const override;
+	virtual FName GetParameterName(EMaterialParameterType Type, int32 ParameterIndex) const override;
+	virtual int32 GetParameterLayerIndex(EMaterialParameterType Type, int32 ParameterIndex) const override;
+	virtual FText GetParameterLayerName(EMaterialParameterType Type, int32 ParameterIndex) const override;
+	virtual bool HasParameter(const FGuid& ParameterId) const override;
+	virtual const UEdGraphPin* GetParameterPin(EMaterialParameterType Type, int32 ParameterIndex) const override;
+	virtual bool IsImageMutableMode(int32 ImageIndex) const override;
+	virtual bool IsImageMutableMode(const UEdGraphPin& Pin) const override;
+	virtual void UpdateImagePinMode(const FGuid ParameterId) override;
+	virtual void UpdateImagePinMode(const UEdGraphPin& Pin) override;
+	virtual void UpdateAllImagesPinMode() override;
+	virtual UTexture2D* GetImageReferenceTexture(int32 ImageIndex) const override;
+	virtual UTexture2D* GetImageValue(int32 ImageIndex) const override;
+	virtual int32 GetImageUVLayout(int32 ImageIndex) const override;
+	virtual UEdGraphPin* GetMeshPin() const override;
+	virtual UEdGraphPin* GetMaterialAssetPin() const override;
+	virtual const UCustomizableObjectNodeMaterial* GetMaterialNode() const override;
+	virtual bool RealMaterialDataHasChanged() const override;
+	virtual FPostImagePinModeChangedDelegate* GetPostImagePinModeChangedDelegate() override;
 	
-	// --------------------
-	// ALL PARAMETERS
-	// --------------------
-
-	/** Returns the number of Material Parameters. */
-	int32 GetNumParameters(EMaterialParameterType Type) const;
+	// Own Interface
 
-	/** Returns the Material Parameter id.
-	 *
-	 * @param ParameterIndex Have to be valid. */
-	FGuid GetParameterId(EMaterialParameterType Type, int32 ParameterIndex) const;
-	
-	/** Returns the Material Parameter name.
-	 *
-	 * @param ParameterIndex Have to be valid. */
-	FName GetParameterName(EMaterialParameterType Type, int32 ParameterIndex) const;
+	void SetMaterial(UMaterialInterface* InMaterial);
 	
-	/** Get the Material Parameter layer index.
-	 *
-	 * @param ParameterIndex Have to be valid.
-	 * @returns INDEX_NONE for global parameters. */
-	int32 GetParameterLayerIndex(EMaterialParameterType Type, int32 ParameterIndex) const;
-	
-	/** Get the Material Parameter layer name.
-	 *
-	 * @param ParameterIndex Have to be valid. */
-	FText GetParameterLayerName(EMaterialParameterType Type, int32 ParameterIndex) const;
+	/** Return true if a Material Parameter has changed on which we had a pin connected or data saved. */
 
-	/** Returns true if the Material contains the given Material Parameter. */
-	bool HasParameter(const FGuid& ParameterId) const;
-	
-	/** Get the Vector pin for the given Material Vector Parameter.
-	 * Not all parameters have pin.
-	 *
-	 * @param ParameterIndex Has to be valid.
-	 * @return Can return nullptr. */
-	const UEdGraphPin* GetParameterPin(EMaterialParameterType Type, int32 ParameterIndex) const;
-	
-	// --------------------
-	// IMAGES PARAMETERS
-	// --------------------
+private:
+	/** Delegate called when a Texture Parameter Pin Mode changes. */
+	FPostImagePinModeChangedDelegate PostImagePinModeChangedDelegate;
 
-	/** Returns true if the Material Texture Parameter goes through Mutable.
-	 *
-	 * @param ImageIndex Have to be valid. */
-	bool IsImageMutableMode(int32 ImageIndex) const;
+	UPROPERTY(EditAnywhere, Category=CustomizableObject)
+	TObjectPtr<UMaterialInterface> Material = nullptr;
 
-	/** Given an Image pin, returns true if the Material Texture Parameter goes through Mutable. */
-	bool IsImageMutableMode(const UEdGraphPin& Pin) const;
+	UPROPERTY(EditAnywhere, Category=CustomizableObject, Meta = (ToolTip = "Set all Mateiral Texture Parameters to the specified mode. Each Texture Parameter Pin can override this mode."))
+	ENodePinMode TextureParametersMode = ENodePinMode::Passthrough;
 
-	/** Update a Material Texture Parameter Mode. */
-	void UpdateImagePinMode(const FGuid ParameterId);
-	
-	/** Update a Material Texture Parameter Mode. */
-	void UpdateImagePinMode(const UEdGraphPin& Pin);
+	UPROPERTY(EditAnywhere, Category = CustomizableObject)
+	TArray<FString> Tags;
 
-	/** Update all Material Texture Parameters Mode. */
-	void UpdateAllImagesPinMode();
-	
-	/** Returns the reference texture assigned to a Material Texture Parameter.
-	 *
-	 * @param ImageIndex Have to be valid.
-	 * @return nullptr if it does not have one assigned. */
-	UTexture2D* GetImageReferenceTexture(int32 ImageIndex) const;
-
-	/** Returns the Texture set in the Material Texture Parameter.
-	 *
-	 * @param ImageIndex Have to be valid. */
-	UTexture2D* GetImageValue(int32 ImageIndex) const;
-	
-	/** Get the Material Texture Parameter UV Index.
-	 *
-	 * @param ImageIndex Have to be valid.
-	 * @return Return -1 if the UV Index is set to Ignore. Return >= 0 for a valid UV Index.  */
-	int32 GetImageUVLayout(int32 ImageIndex) const;
+	/** Selects which Mesh component of the Instance this material belongs to */
+	UPROPERTY(EditAnywhere, Category = CustomizableObject, meta = (ClampMin = "0"))
+	int32 MeshComponentIndex = 0;
+
+	/** Materials will be reused between LODs when possible unless explicitly disabled. */
+	UPROPERTY(EditAnywhere, Category = CustomizableObject)
+	bool bReuseMaterialBetweenLODs = true;
 
-	/** Delegate called when a Texture Parameter Pin Mode changes. */
-	FPostImagePinModeChangedDelegate PostImagePinModeChangedDelegate;
-	
-private:
 	/** Last static or skeletal mesh connected. Used to remove the callback once disconnected. */
 	TWeakObjectPtr<UCustomizableObjectNode> LastMeshNodeConnected;
 
@@ -347,7 +288,7 @@ public:
 	UPROPERTY(EditAnywhere, Category = NoCategory)
 	EUVLayoutMode UVLayoutMode = EUVLayoutMode::FromMaterial;
 	
-	/** Index of the UV channel that will be used with this image.It is necessary to apply the proper layout transformations to it. */
+	/** Index of the UV channel that will be used with this image. It is necessary to apply the proper layout transformations to it. */
 	UPROPERTY(EditAnywhere, Category = NoCategory, meta = (EditCondition = "UVLayoutMode == EUVLayoutMode::Index", EditConditionHides))
 	int32 UVLayout = -2;
 
diff --git a/Engine/Plugins/Experimental/Mutable/Source/CustomizableObjectEditor/Public/MuCOE/Nodes/CustomizableObjectNodeMaterialBase.h b/Engine/Plugins/Experimental/Mutable/Source/CustomizableObjectEditor/Public/MuCOE/Nodes/CustomizableObjectNodeMaterialBase.h
--- a/Engine/Plugins/Experimental/Mutable/Source/CustomizableObjectEditor/Public/MuCOE/Nodes/CustomizableObjectNodeMaterialBase.h
+++ b/Engine/Plugins/Experimental/Mutable/Source/CustomizableObjectEditor/Public/MuCOE/Nodes/CustomizableObjectNodeMaterialBase.h
@@ -7,29 +7,125 @@
 #include "CustomizableObjectNodeMaterialBase.generated.h"
 
 class UCustomizableObjectLayout;
-
+class UCustomizableObjectNodeMaterial;
 class UEdGraphPin;
 class UObject;
 
 
+DECLARE_MULTICAST_DELEGATE(FPostImagePinModeChangedDelegate)
+
 
+/** Equivalent to mu::NodeSurface but with limitations. Currently only nodes that generate a mu::NodeSurfaceNew inherit this (NodeMaterial and NodeCopyMaterial).
+ * Nodes that generate mu::NodeSurfaceEdit (NodeEditMaterial, NodeExtendMaterial), mu::NodeSurfaceSwitch (NodeSwitchMaterial)... are excluded.
+ *
+ * WARNING: All methods in this class should be PURE_VIRTUAL!
+ */
 UCLASS(abstract)
 class CUSTOMIZABLEOBJECTEDITOR_API UCustomizableObjectNodeMaterialBase : public UCustomizableObjectNode
 {
 public:
 	GENERATED_BODY()
-
+	
 	// Own interface
-	virtual TArray<class UCustomizableObjectLayout*> GetLayouts() { return TArray<class UCustomizableObjectLayout*>(); }
+	virtual const UCustomizableObjectNodeMaterial* GetMaterialNode() const PURE_VIRTUAL(UCustomizableObjectNodeMaterial::GetMaterialNode, return {}; );
 
-	// This method should be overidden in all derived classes
-	virtual UEdGraphPin* OutputPin() const
+	UCustomizableObjectNodeMaterial* GetMaterialNode()
 	{
-		check(false); 
-		return nullptr;
+		const UCustomizableObjectNodeMaterialBase* ConstThis = this;
+		return const_cast<UCustomizableObjectNodeMaterial*>(ConstThis->GetMaterialNode());
 	}
 
-	/** Allow multiple connections to NodeCopyMaterial but only one connection to a NodeObject. */
-	bool ShouldBreakExistingConnections(const UEdGraphPin* InputPin, const UEdGraphPin* OutputPin) const override;
+	virtual TArray<UCustomizableObjectLayout*> GetLayouts() const PURE_VIRTUAL(UCustomizableObjectNodeMaterial::GetLayouts, return {}; );
+	
+	virtual UMaterialInterface* GetMaterial() const PURE_VIRTUAL(UCustomizableObjectNodeMaterial::GetMaterial, return {}; );
+	
+	virtual bool IsReuseMaterialBetweenLODs() const PURE_VIRTUAL(UCustomizableObjectNodeMaterial::IsReuseMaterialBetweenLODs, return {}; );
+
+	virtual int32 GetMeshComponentIndex() const PURE_VIRTUAL(UCustomizableObjectNodeMaterial::GetMeshComponentIndex, return {}; );
+	
+	virtual TArray<FString> GetTags() const PURE_VIRTUAL(UCustomizableObjectNodeMaterial::GetTags, return {}; );
+	
+	virtual UEdGraphPin* GetMeshPin() const PURE_VIRTUAL(UCustomizableObjectNodeMaterial::GetMeshPin, return {}; );
+
+	virtual UEdGraphPin* GetMaterialAssetPin() const PURE_VIRTUAL(UCustomizableObjectNodeMaterial::GetMaterialAssetPin, return {}; );
+	
+	/** Returns the number of Material Parameters. */
+	virtual int32 GetNumParameters(EMaterialParameterType Type) const PURE_VIRTUAL(UCustomizableObjectNodeMaterial::GetNumParameters, return {}; );
+
+	/** Returns the Material Parameter id.
+	 *
+	 * @param ParameterIndex Have to be valid. */
+	virtual FGuid GetParameterId(EMaterialParameterType Type, int32 ParameterIndex) const PURE_VIRTUAL(UCustomizableObjectNodeMaterial::GetParameterId, return {}; );
+	
+	/** Returns the Material Parameter name.
+	 *
+	 * @param ParameterIndex Have to be valid. */
+	virtual FName GetParameterName(EMaterialParameterType Type, int32 ParameterIndex) const PURE_VIRTUAL(UCustomizableObjectNodeMaterial::GetParameterName, return {}; );
+	
+	/** Get the Material Parameter layer index.
+	 *
+	 * @param ParameterIndex Have to be valid.
+	 * @returns INDEX_NONE for global parameters. */
+	virtual int32 GetParameterLayerIndex(EMaterialParameterType Type, int32 ParameterIndex) const PURE_VIRTUAL(UCustomizableObjectNodeMaterial::GetParameterLayerIndex, return {}; );
+	
+	/** Get the Material Parameter layer name.
+	 *
+	 * @param ParameterIndex Have to be valid. */
+	virtual FText GetParameterLayerName(EMaterialParameterType Type, int32 ParameterIndex) const PURE_VIRTUAL(UCustomizableObjectNodeMaterial::GetParameterLayerName, return {}; );
+
+	/** Returns true if the Material contains the given Material Parameter. */
+	virtual bool HasParameter(const FGuid& ParameterId) const PURE_VIRTUAL(UCustomizableObjectNodeMaterial::HasParameter, return {}; );
+	
+	/** Get the Vector pin for the given Material Vector Parameter.
+	 * Not all parameters have pin.
+	 *
+	 * @param ParameterIndex Has to be valid.
+	 * @return Can return nullptr. */
+	virtual const UEdGraphPin* GetParameterPin(EMaterialParameterType Type, int32 ParameterIndex) const PURE_VIRTUAL(UCustomizableObjectNodeMaterial::GetMaterialNode, return {}; );
+	
+	// --------------------
+	// IMAGES PARAMETERS
+	// --------------------
+
+	/** Returns true if the Material Texture Parameter goes through Mutable.
+	 *
+	 * @param ImageIndex Have to be valid. */
+	virtual bool IsImageMutableMode(int32 ImageIndex) const PURE_VIRTUAL(UCustomizableObjectNodeMaterial::IsImageMutableMode, return {}; );
+
+	/** Given an Image pin, returns true if the Material Texture Parameter goes through Mutable. */
+	virtual bool IsImageMutableMode(const UEdGraphPin& Pin) const PURE_VIRTUAL(UCustomizableObjectNodeMaterial::IsImageMutableMode, return {}; );
+
+	/** Update a Material Texture Parameter Mode. */
+	virtual void UpdateImagePinMode(const FGuid ParameterId) PURE_VIRTUAL(UCustomizableObjectNodeMaterial::UpdateImagePinMode, );
+	
+	/** Update a Material Texture Parameter Mode. */
+	virtual void UpdateImagePinMode(const UEdGraphPin& Pin) PURE_VIRTUAL(UCustomizableObjectNodeMaterial::UpdateImagePinMode, );
+
+	/** Update all Material Texture Parameters Mode. */
+	virtual void UpdateAllImagesPinMode() PURE_VIRTUAL(UCustomizableObjectNodeMaterial::UpdateAllImagesPinMode, );
+	
+	/** Returns the reference texture assigned to a Material Texture Parameter.
+	 *
+	 * @param ImageIndex Have to be valid.
+	 * @return nullptr if it does not have one assigned. */
+	virtual UTexture2D* GetImageReferenceTexture(int32 ImageIndex) const PURE_VIRTUAL(UCustomizableObjectNodeMaterial::GetImageReferenceTexture, return {}; );
+
+	/** Returns the Texture set in the Material Texture Parameter.
+	 *
+	 * @param ImageIndex Have to be valid. */
+	virtual UTexture2D* GetImageValue(int32 ImageIndex) const PURE_VIRTUAL(UCustomizableObjectNodeMaterial::GetImageValue, return {}; );
+	
+	/** Get the Material Texture Parameter UV Index.
+	 *
+	 * @param ImageIndex Have to be valid.
+	 * @return Return -1 if the UV Index is set to Ignore. Return >= 0 for a valid UV Index.  */ 
+	virtual int32 GetImageUVLayout(int32 ImageIndex) const PURE_VIRTUAL(UCustomizableObjectNodeMaterial::GetImageUVLayout, return {}; );
+
+	// This method should be overidden in all derived classes
+	virtual UEdGraphPin* OutputPin() const PURE_VIRTUAL(UCustomizableObjectNodeMaterial::OutputPin, return {}; );
+	
+	virtual FPostImagePinModeChangedDelegate* GetPostImagePinModeChangedDelegate() PURE_VIRTUAL(UCustomizableObjectNodeMaterial::OutputPin, return {}; );
+	
+	virtual bool RealMaterialDataHasChanged() const PURE_VIRTUAL(UCustomizableObjectNodeMaterial::OutputPin, return {}; );
 };
 
diff --git a/Engine/Plugins/Experimental/Mutable/Source/CustomizableObjectEditor/Public/MuCOE/Nodes/CustomizableObjectNodeMeshClipDeform.h b/Engine/Plugins/Experimental/Mutable/Source/CustomizableObjectEditor/Public/MuCOE/Nodes/CustomizableObjectNodeMeshClipDeform.h
--- a/Engine/Plugins/Experimental/Mutable/Source/CustomizableObjectEditor/Public/MuCOE/Nodes/CustomizableObjectNodeMeshClipDeform.h
+++ b/Engine/Plugins/Experimental/Mutable/Source/CustomizableObjectEditor/Public/MuCOE/Nodes/CustomizableObjectNodeMeshClipDeform.h
@@ -42,9 +42,10 @@ public:
 
 	// UCustomizableObjectNode interface
 	virtual void AllocateDefaultPins(UCustomizableObjectNodeRemapPins* RemapPins) override;
-	virtual UEdGraphPin* OutputPin() const override;
 
 	// Own interface
+	UEdGraphPin* OutputPin() const;
+	
 	UEdGraphPin* ClipShapePin() const;
 };
 
diff --git a/Engine/Plugins/Experimental/Mutable/Source/CustomizableObjectEditor/Public/MuCOE/Nodes/CustomizableObjectNodeMeshClipMorph.h b/Engine/Plugins/Experimental/Mutable/Source/CustomizableObjectEditor/Public/MuCOE/Nodes/CustomizableObjectNodeMeshClipMorph.h
--- a/Engine/Plugins/Experimental/Mutable/Source/CustomizableObjectEditor/Public/MuCOE/Nodes/CustomizableObjectNodeMeshClipMorph.h
+++ b/Engine/Plugins/Experimental/Mutable/Source/CustomizableObjectEditor/Public/MuCOE/Nodes/CustomizableObjectNodeMeshClipMorph.h
@@ -89,10 +89,9 @@ public:
 	virtual void AllocateDefaultPins(UCustomizableObjectNodeRemapPins* RemapPins) override;
 	virtual void PinConnectionListChanged(UEdGraphPin* Pin) override;
 	
-	// UCustomizableObjectNodeMaterialBase interface
-	virtual UEdGraphPin* OutputPin() const override;
-
 	// Own interface
+	UEdGraphPin* OutputPin() const;
+
 	FVector GetOriginWithOffset() const;
 	void FindLocalAxes(FVector& Right, FVector& Up, FVector& Forward) const;
 
diff --git a/Engine/Plugins/Experimental/Mutable/Source/CustomizableObjectEditor/Public/MuCOE/Nodes/CustomizableObjectNodeMeshClipWithMesh.h b/Engine/Plugins/Experimental/Mutable/Source/CustomizableObjectEditor/Public/MuCOE/Nodes/CustomizableObjectNodeMeshClipWithMesh.h
--- a/Engine/Plugins/Experimental/Mutable/Source/CustomizableObjectEditor/Public/MuCOE/Nodes/CustomizableObjectNodeMeshClipWithMesh.h
+++ b/Engine/Plugins/Experimental/Mutable/Source/CustomizableObjectEditor/Public/MuCOE/Nodes/CustomizableObjectNodeMeshClipWithMesh.h
@@ -60,11 +60,10 @@ public:
 	virtual void PinConnectionListChanged(UEdGraphPin* Pin) override;
 	virtual void UpdateReferencedNodeId(const FGuid& NewGuid) override;
 	virtual void BeginPostDuplicate(bool bDuplicateForPIE) override;
-
-	// UCustomizableObjectNodeMaterialBase interface
-	virtual UEdGraphPin* OutputPin() const override;
-
+	
 	// Own interface
+	UEdGraphPin* OutputPin() const;
+
 	UEdGraphPin* ClipMeshPin() const;
 };
 
diff --git a/Engine/Plugins/Experimental/Mutable/Source/CustomizableObjectEditor/Public/MuCOE/Nodes/CustomizableObjectNodeModifierBase.h b/Engine/Plugins/Experimental/Mutable/Source/CustomizableObjectEditor/Public/MuCOE/Nodes/CustomizableObjectNodeModifierBase.h
--- a/Engine/Plugins/Experimental/Mutable/Source/CustomizableObjectEditor/Public/MuCOE/Nodes/CustomizableObjectNodeModifierBase.h
+++ b/Engine/Plugins/Experimental/Mutable/Source/CustomizableObjectEditor/Public/MuCOE/Nodes/CustomizableObjectNodeModifierBase.h
@@ -10,12 +10,11 @@ class UObject;
 
 
 UCLASS(abstract)
-class CUSTOMIZABLEOBJECTEDITOR_API UCustomizableObjectNodeModifierBase : public UCustomizableObjectNodeMaterialBase
+class CUSTOMIZABLEOBJECTEDITOR_API UCustomizableObjectNodeModifierBase : public UCustomizableObjectNode
 {
 public:
 	GENERATED_BODY()
 
 	UCustomizableObjectNodeModifierBase();
-
 };
 
diff --git a/Engine/Plugins/Experimental/Mutable/Source/CustomizableObjectEditor/Public/MuCOE/Nodes/CustomizableObjectNodeModifierClipWithUVMask.h b/Engine/Plugins/Experimental/Mutable/Source/CustomizableObjectEditor/Public/MuCOE/Nodes/CustomizableObjectNodeModifierClipWithUVMask.h
--- a/Engine/Plugins/Experimental/Mutable/Source/CustomizableObjectEditor/Public/MuCOE/Nodes/CustomizableObjectNodeModifierClipWithUVMask.h
+++ b/Engine/Plugins/Experimental/Mutable/Source/CustomizableObjectEditor/Public/MuCOE/Nodes/CustomizableObjectNodeModifierClipWithUVMask.h
@@ -40,10 +40,9 @@ public:
 	virtual void AllocateDefaultPins(UCustomizableObjectNodeRemapPins* RemapPins) override;
 	virtual void PinConnectionListChanged(UEdGraphPin* Pin) override;
 
-	// UCustomizableObjectNodeMaterialBase interface
-	virtual UEdGraphPin* OutputPin() const override;
-
 	// Own interface
+	UEdGraphPin* OutputPin() const;
+
 	UEdGraphPin* ClipMaskPin() const;
 };
 
diff --git a/Engine/Plugins/Experimental/Mutable/Source/CustomizableObjectEditor/Public/MuCOE/Nodes/CustomizableObjectNodeObject.h b/Engine/Plugins/Experimental/Mutable/Source/CustomizableObjectEditor/Public/MuCOE/Nodes/CustomizableObjectNodeObject.h
--- a/Engine/Plugins/Experimental/Mutable/Source/CustomizableObjectEditor/Public/MuCOE/Nodes/CustomizableObjectNodeObject.h
+++ b/Engine/Plugins/Experimental/Mutable/Source/CustomizableObjectEditor/Public/MuCOE/Nodes/CustomizableObjectNodeObject.h
@@ -10,7 +10,7 @@
 namespace ENodeTitleType { enum Type : int; }
 
 class UCustomizableObject;
-class UCustomizableObjectNodeMaterial;
+class UCustomizableObjectNodeMaterialBase;
 class UCustomizableObjectNodeRemapPins;
 class UObject;
 struct FPropertyChangedEvent;
@@ -270,7 +270,7 @@ public:
 	 *
 	 * @param LOD LOD which materials have to belong to.
 	 */
-	TArray<UCustomizableObjectNodeMaterial*> GetMaterialNodes(int LOD) const;
+	TArray<UCustomizableObjectNodeMaterialBase*> GetMaterialNodes(int LOD) const;
 
 	bool IsSingleOutputNode() const override;
 
diff --git a/Engine/Plugins/Experimental/Mutable/Source/CustomizableObjectEditor/Public/MuCOE/Nodes/CustomizableObjectNodeParentedMaterial.h b/Engine/Plugins/Experimental/Mutable/Source/CustomizableObjectEditor/Public/MuCOE/Nodes/CustomizableObjectNodeParentedMaterial.h
--- a/Engine/Plugins/Experimental/Mutable/Source/CustomizableObjectEditor/Public/MuCOE/Nodes/CustomizableObjectNodeParentedMaterial.h
+++ b/Engine/Plugins/Experimental/Mutable/Source/CustomizableObjectEditor/Public/MuCOE/Nodes/CustomizableObjectNodeParentedMaterial.h
@@ -6,7 +6,7 @@
 #include "MuCOE/Nodes/CustomizableObjectNodeParentedNode.h"
 
 class UCustomizableObjectNode;
-class UCustomizableObjectNodeMaterial;
+class UCustomizableObjectNodeMaterialBase;
 
 
 /** Node Material specialization of ICustomizableObjectNodeParentedNode. */
@@ -15,13 +15,13 @@ class FCustomizableObjectNodeParentedMaterial : public ICustomizableObjectNodePa
 public:
 	// Own interface
 	/** Return the parent material node. */
-	UCustomizableObjectNodeMaterial* GetParentMaterialNode() const;
+	UCustomizableObjectNodeMaterialBase* GetParentMaterialNode() const;
 	
 	/** Return all possible parent material nodes of the node. */
-	TArray<UCustomizableObjectNodeMaterial*> GetPossibleParentMaterialNodes() const;
+	TArray<UCustomizableObjectNodeMaterialBase*> GetPossibleParentMaterialNodes() const;
 
 	/** Returns the parent material node if there exist a path to it. */
-	UCustomizableObjectNodeMaterial* GetParentMaterialNodeIfPath() const;
+	UCustomizableObjectNodeMaterialBase* GetParentMaterialNodeIfPath() const;
 	
 protected:
 	/** Return the node which this interface belongs to. */
