[mutable] Fixed reference to null object
- Crash during bake of instance when trying to display a referencing null text.

[RN] #rnx

diff --git a/Engine/Plugins/Experimental/Mutable/Source/CustomizableObjectEditor/Private/MuCOE/CustomizableObjectInstanceBakingUtils.cpp b/Engine/Plugins/Experimental/Mutable/Source/CustomizableObjectEditor/Private/MuCOE/CustomizableObjectInstanceBakingUtils.cpp
--- a/Engine/Plugins/Experimental/Mutable/Source/CustomizableObjectEditor/Private/MuCOE/CustomizableObjectInstanceBakingUtils.cpp
+++ b/Engine/Plugins/Experimental/Mutable/Source/CustomizableObjectEditor/Private/MuCOE/CustomizableObjectInstanceBakingUtils.cpp
@@ -20,7 +20,6 @@
 #include "MuCO/CustomizableObjectInstanceAssetUserData.h"
 #include "MuCO/CustomizableObjectInstancePrivate.h"
 #include "MuCO/CustomizableObjectMipDataProvider.h"
-#include "MuCO/ICustomizableObjectEditorModule.h"
 #include "MuT/UnrealPixelFormatOverride.h"
 #include "Rendering/SkeletalMeshModel.h"
 
@@ -32,22 +31,15 @@
  * @param InMessage The message to display
  * @param InTitle The title to be used for the popup or the log generated
  */
-void ShowErrorNotification(const FText& InMessage,  const FText* InTitle = nullptr)
+void ShowErrorNotification(const FText& InMessage, const FText& InTitle = LOCTEXT("CustomizableObjecInstanceBakingUtils_GenericBakingError","Baking Error") )
 {
 	if (!FApp::IsUnattended())
 	{
-		FMessageDialog::Open(EAppMsgType::Ok, InMessage, *InTitle);
+		FMessageDialog::Open(EAppMsgType::Ok, InMessage, InTitle);
 	}
 	else
 	{
-		if (InTitle)
-		{
-			UE_LOG(LogMutable, Error, TEXT("%s - %s"), *InTitle->ToString(), *InMessage.ToString());
-		}
-		else
-		{
-			UE_LOG(LogMutable, Error, TEXT("%s"), *InMessage.ToString());
-		}
+		UE_LOG(LogMutable, Error, TEXT("%s - %s"), *InTitle.ToString(), *InMessage.ToString());
 	}
 }
 
@@ -88,7 +80,7 @@ bool ValidateProvidedFileName(const FString& FileName)
 		if (InvalidCharacter != '0')
 		{
 			const FText InvalidCharacterText = FText::FromString(FString::Chr(InvalidCharacter));
-			const FText ErrorText = FText::Format(LOCTEXT("FCustomizableObjectEditorViewportClient_BakeInstance_InvalidCharacter", "The selected contains an invalid character ({0})."), InvalidCharacterText);
+			const FText ErrorText = FText::Format(LOCTEXT("CustomizableObjecInstanceBakingUtils_InvalidCharacter", "The selected contains an invalid character ({0})."), InvalidCharacterText);
 
 			ShowErrorNotification(ErrorText);
 		
@@ -120,7 +112,7 @@ bool ValidateProvidedAssetPath(const FString& FileName, const FString& AssetPath
 	const bool bWouldOverrideParentCO = InstanceCO->GetPathName() == FullAssetPath;
 	if (bWouldOverrideParentCO)
 	{
-		const FText ErrorText = LOCTEXT("FCustomizableObjectEditorViewportClient_BakeInstance_OverwriteCO", "The selected path would overwrite the instance's parent Customizable Object.");
+		const FText ErrorText = LOCTEXT("CustomizableObjecInstanceBakingUtils_OverwriteCO", "The selected path would overwrite the instance's parent Customizable Object.");
 
 		ShowErrorNotification(ErrorText);
 		
@@ -216,10 +208,10 @@ bool ManageBakingAction(const FString& InAssetSavePath, const FString& InObjName
 				// Close the editors that contains this asset
 				if (!ObjectEditorInstance->CloseWindow(EAssetEditorCloseReason::AssetEditorHostClosed))
 				{
-					const FText Caption = LOCTEXT("OpenExisitngFile", "Open File");
-					const FText Message = FText::Format(LOCTEXT("CantCloseAsset", "This Obejct \"{0}\" is open in an editor and can't be closed automatically. Please close the editor and try to bake it again"), FText::FromString(InObjName));
+					const FText Caption = LOCTEXT("CustomizableObjecInstanceBakingUtils_OpenExisitngFile", "Open File");
+					const FText Message = FText::Format(LOCTEXT("CustomizableObjecInstanceBakingUtils_CantCloseAsset", "This Obejct \"{0}\" is open in an editor and can't be closed automatically. Please close the editor and try to bake it again"), FText::FromString(InObjName));
 
-					ShowErrorNotification(Message, &Caption);
+					ShowErrorNotification(Message, Caption);
 					
 					return false;
 				}
@@ -230,8 +222,8 @@ bool ManageBakingAction(const FString& InAssetSavePath, const FString& InObjName
 		if (!bIsUnattended && !bOverridePermissionGranted)
 		{
 			check (!FApp::IsUnattended())
-			const FText Caption = LOCTEXT("Already existing baked files", "Already existing baked files");
-			const FText Message = FText::Format(LOCTEXT("OverwriteBakedInstance", "Instance baked files already exist in selected destination \"{0}\", this action will overwrite them."), FText::AsCultureInvariant(InAssetSavePath));
+			const FText Caption = LOCTEXT("CustomizableObjecInstanceBakingUtils_AlreadyExistingBakedFiles", "Already existing baked files");
+			const FText Message = FText::Format(LOCTEXT("CustomizableObjecInstanceBakingUtils_OverwriteBakedInstance", "Instance baked files already exist in selected destination \"{0}\", this action will overwrite them."), FText::AsCultureInvariant(InAssetSavePath));
 
 			if (FMessageDialog::Open(EAppMsgType::OkCancel, Message, Caption) == EAppReturnType::Cancel)
 			{
@@ -258,7 +250,7 @@ bool ManageBakingAction(const FString& InAssetSavePath, const FString& InObjName
 				TArray<UObject*> ObjectsToDelete;
 				ObjectsToDelete.Add(ExistingObject);
 				
-				const FText Message = FText::Format(LOCTEXT("AssetOverriden", "The COI asset \"{0}\" already exists and will be overriden due to user demand."), FText::FromString(ExistingPackage->GetName()));
+				const FText Message = FText::Format(LOCTEXT("CustomizableObjecInstanceBakingUtils_AssetOverriden", "The COI asset \"{0}\" already exists and will be overriden due to user demand."), FText::FromString(ExistingPackage->GetName()));
 				UE_LOG(LogMutable,Warning,TEXT("%s"), *Message.ToString());
 
 				// Notify the caller we did proceed with the override (performed later)
@@ -274,7 +266,7 @@ bool ManageBakingAction(const FString& InAssetSavePath, const FString& InObjName
 				OutSaveResolution = EPackageSaveResolutionType::UnableToOverride;
 				
 				// Report that the file will not get overriden since we have no permission to do so
-				const FText UnableToOverrideMessage = FText::Format(LOCTEXT("AssetCanNotBeOverriden", "Could not replace the COI asset \"{0}\" as it already exists."), FText::FromString(ExistingPackage->GetName()));
+				const FText UnableToOverrideMessage = FText::Format(LOCTEXT("CustomizableObjecInstanceBakingUtils_AssetCanNotBeOverriden", "Could not replace the COI asset \"{0}\" as it already exists."), FText::FromString(ExistingPackage->GetName()));
 				UE_LOG(LogMutable,Error,TEXT("%s"), *UnableToOverrideMessage.ToString());
 
 				return false;
@@ -378,7 +370,7 @@ bool BakeCustomizableObjectInstance(
 	if (!InstanceCO || InstanceCO->GetPrivate()->IsLocked())
 	{
 		FCustomizableObjectEditorLogger::CreateLog(
-		LOCTEXT("CustomizableObjectCompilingTryLater_Baking", "Please wait until the Customizable Object is compiled"))
+		LOCTEXT("CustomizableObjecInstanceBakingUtils_LockedObject", "Please wait until the Customizable Object is compiled"))
 		.Category(ELoggerCategory::COInstanceBaking)
 		.CustomNotification()
 		.Notification(true)
@@ -390,7 +382,7 @@ bool BakeCustomizableObjectInstance(
 	if (InstanceCO->GetPrivate()->Status.Get() == FCustomizableObjectStatus::EState::Loading)
 	{
 		FCustomizableObjectEditorLogger::CreateLog(
-			LOCTEXT("CustomizableObjectCompileTryLater_BakeInstance","Please wait unitl Customizable Object is loaded"))
+			LOCTEXT("CustomizableObjecInstanceBakingUtils_LoadingObject","Please wait unitl Customizable Object is loaded"))
 		.Category(ELoggerCategory::COInstanceBaking)
 		.CustomNotification()
 		.Notification(true)
@@ -424,7 +416,7 @@ bool BakeCustomizableObjectInstance(
 	if (InstanceCO->GetPrivate()->GetCompileOptions().TextureCompression != ECustomizableObjectTextureCompression::HighQuality)
 	{
 		FCustomizableObjectEditorLogger::CreateLog(
-		LOCTEXT("CustomizableObjectBakeLowQuality", "The Customizable Object wasn't compiled with high quality textures. For the best baking results, change the Texture Compression setting and recompile it."))
+		LOCTEXT("CustomizableObjecInstanceBakingUtils_LowQualityTextures", "The Customizable Object wasn't compiled with high quality textures. For the best baking results, change the Texture Compression setting and recompile it."))
 		.Category(ELoggerCategory::COInstanceBaking)
 		.CustomNotification()
 		.Notification(true)
