[Mutable] Fixed crash due to to a discrepancy between LODInfo.Num() and LODRenderData.Num().

diff --git a/Engine/Plugins/Experimental/Mutable/Source/CustomizableObjectEditor/Private/MuCOE/CustomizableObjectInstanceBakingUtils.cpp b/Engine/Plugins/Experimental/Mutable/Source/CustomizableObjectEditor/Private/MuCOE/CustomizableObjectInstanceBakingUtils.cpp
--- a/Engine/Plugins/Experimental/Mutable/Source/CustomizableObjectEditor/Private/MuCOE/CustomizableObjectInstanceBakingUtils.cpp
+++ b/Engine/Plugins/Experimental/Mutable/Source/CustomizableObjectEditor/Private/MuCOE/CustomizableObjectInstanceBakingUtils.cpp
@@ -811,11 +811,12 @@ bool BakeCustomizableObjectInstance(
 		USkeletalMesh* SkeletalMesh = Cast<USkeletalMesh>(DupObject);
 		if (SkeletalMesh)
 		{
+			SkeletalMesh->ResetLODInfo();
 			for (int32 LODIndex = 0; LODIndex < Mesh->GetLODNum(); ++LODIndex)
-			{
-				SkeletalMesh->AddLODInfo(*Mesh->GetLODInfo(LODIndex));
-			}
-
+        	{
+        		SkeletalMesh->AddLODInfo(*Mesh->GetLODInfo(LODIndex));
+        	}
+		
 			SkeletalMesh->GetImportedModel()->SkeletalMeshModelGUID = FGuid::NewGuid();
 
 			// Duplicate AssetUserData
