[Mutable] Add tangent frame skinning to OpMeshApplyPose and modify resulting ref pose to match the passed pose.

- Now vertices are processd in parallel batches using ParallelFor.

#rnx

diff --git a/Engine/Plugins/Experimental/Mutable/Source/MutableRuntime/Internal/MuR/MeshPrivate.h b/Engine/Plugins/Experimental/Mutable/Source/MutableRuntime/Internal/MuR/MeshPrivate.h
--- a/Engine/Plugins/Experimental/Mutable/Source/MutableRuntime/Internal/MuR/MeshPrivate.h
+++ b/Engine/Plugins/Experimental/Mutable/Source/MutableRuntime/Internal/MuR/MeshPrivate.h
@@ -353,7 +353,7 @@ namespace mu
 			return reinterpret_cast<CTYPE*>(m_pBuf);
 		}
 
-		inline MeshBufferIterator<FORMAT,CTYPE,COMPONENTS> operator+(int c)
+		inline MeshBufferIterator<FORMAT,CTYPE,COMPONENTS> operator+(int c) const
 		{
 			MeshBufferIterator<FORMAT,CTYPE,COMPONENTS> res = *this;
 			res += c;
diff --git a/Engine/Plugins/Experimental/Mutable/Source/MutableRuntime/Private/MuR/OpMeshApplyPose.h b/Engine/Plugins/Experimental/Mutable/Source/MutableRuntime/Private/MuR/OpMeshApplyPose.h
--- a/Engine/Plugins/Experimental/Mutable/Source/MutableRuntime/Private/MuR/OpMeshApplyPose.h
+++ b/Engine/Plugins/Experimental/Mutable/Source/MutableRuntime/Private/MuR/OpMeshApplyPose.h
@@ -7,38 +7,36 @@
 #include "MuR/MutableTrace.h"
 #include "MuR/Platform.h"
 
+#include "Async/ParallelFor.h"
 
 namespace mu
 {
-
-    //---------------------------------------------------------------------------------------------
-    //! Reference version
-    //---------------------------------------------------------------------------------------------
-    inline void MeshApplyPose(Mesh* Result, const Mesh* pBase, const Mesh* pPose, bool& bOutSuccess)
+	/**
+    * Reference version
+    */
+    inline void MeshApplyPose(Mesh* Result, const Mesh* BaseMesh, const Mesh* PoseMesh, bool& bOutSuccess)
     {
         MUTABLE_CPUPROFILER_SCOPE(MeshApplyPose);
 
 		bOutSuccess = true;
 
-		mu::SkeletonPtrConst pSkeleton = pBase->GetSkeleton();
-		if (!pSkeleton)
+		Ptr<const Skeleton> Skeleton = BaseMesh->GetSkeleton();
+		if (!Skeleton)
 		{
 			bOutSuccess = false;
 			return;
 		}
 
         // We assume the matrices are transforms from the binding pose bone to the new pose
-        // For now we only convert the vertex positions.
-        // \TODO: normals and tangents
 
 		// Find closest bone affected by the pose for each bone in the skeleton.
-		const int32 NumBones = pSkeleton->GetBoneCount();
+		const int32 NumBones = Skeleton->GetBoneCount();
 		TArray<int32> BoneToPoseIndex;
 		BoneToPoseIndex.Init(INDEX_NONE, NumBones);
 
 		for (int32 BoneIndex = 0; BoneIndex < NumBones; ++BoneIndex)
 		{
-			const int32 BonePoseIndex = pPose->FindBonePose(pSkeleton->GetBoneName(BoneIndex));
+			const int32 BonePoseIndex = PoseMesh->FindBonePose(Skeleton->GetBoneName(BoneIndex));
 			if (BonePoseIndex != INDEX_NONE)
 			{
 				BoneToPoseIndex[BoneIndex] = BonePoseIndex;
@@ -46,15 +44,14 @@ namespace mu
 			}
 
 			// Parent bones are in a strictly increassing order. Set the pose index from the parent bone.
-			const int32 ParentIndex = pSkeleton->GetBoneParent(BoneIndex);
+			const int32 ParentIndex = Skeleton->GetBoneParent(BoneIndex);
 			if (ParentIndex != INDEX_NONE)
 			{
 				BoneToPoseIndex[BoneIndex] = BoneToPoseIndex[ParentIndex];
 			}
 		}
 
-
-		const TArray<FBoneName>& BoneMap = pBase->GetBoneMap();
+		const TArray<FBoneName>& BoneMap = BaseMesh->GetBoneMap();
 		const int32 NumBonesBoneMap = BoneMap.Num();
 
         // Prepare the skin matrices. They may be in different order, and we only need the ones
@@ -65,10 +62,10 @@ namespace mu
 		bool bBonesAffected = false;
 		for (int32 Index = 0; Index < NumBonesBoneMap; ++Index)
 		{
-			const int32 BoneIndex = pSkeleton->FindBone(BoneMap[Index]);
+			const int32 BoneIndex = Skeleton->FindBone(BoneMap[Index]);
 			if (BoneToPoseIndex[BoneIndex] != INDEX_NONE)
 			{
-				SkinTransforms.Add(pPose->BonePoses[BoneToPoseIndex[BoneIndex]].BoneTransform);
+				SkinTransforms.Add(PoseMesh->BonePoses[BoneToPoseIndex[BoneIndex]].BoneTransform);
 				bBonesAffected = true;
 			}
 			else
@@ -86,19 +83,19 @@ namespace mu
 		}
 
         // Get pointers to vertex position data
-		MeshBufferIteratorConst<MBF_FLOAT32, float, 3> itSource(pBase->VertexBuffers, MBS_POSITION, 0);
-        if (!itSource.ptr())
+		MeshBufferIteratorConst<MBF_FLOAT32, float, 3> SourcePositionIterBegin(BaseMesh->VertexBuffers, MBS_POSITION, 0);
+        if (!SourcePositionIterBegin.ptr())
         {
             // Formats not implemented
             check(false);
 			bOutSuccess = false;
             return;
         }
-
-        // Get pointers to skinning data
-		UntypedMeshBufferIteratorConst itBoneIndices(pBase->VertexBuffers, MBS_BONEINDICES, 0);
-		UntypedMeshBufferIteratorConst itBoneWeights(pBase->VertexBuffers, MBS_BONEWEIGHTS, 0);
-        if (!itBoneIndices.ptr() || !itBoneWeights.ptr())
+ 
+		// Get pointers to skinning data
+		UntypedMeshBufferIteratorConst BoneIndicesIterBegin(BaseMesh->VertexBuffers, MBS_BONEINDICES, 0);
+		UntypedMeshBufferIteratorConst BoneWeightsIterBegin(BaseMesh->VertexBuffers, MBS_BONEWEIGHTS, 0);
+        if (!BoneIndicesIterBegin.ptr() || !BoneWeightsIterBegin.ptr())
         {
             // No skinning data
             check(false);
@@ -106,61 +103,173 @@ namespace mu
             return;
         }
 
-		Result->CopyFrom(*pBase);
+		Result->CopyFrom(*BaseMesh);
 
-		MeshBufferIterator<MBF_FLOAT32, float, 3> itTarget(Result->VertexBuffers, MBS_POSITION, 0);
-		check(itTarget.ptr());
+		// Set the pose as the Result reference pose.
+		// PoseMesh PoseXform can be decomposed to PoseXform = ModelPoseXform * ModelRefXform^-1, 
+		// BaseMesh Poses have ModelRefXform, to get the new ModelRefXform to be ModelPoseXform we need 
+		// to multiply PoseXform * ModelRefXform. 
+		// (Note, transform multiplication application is from right to left for this comment, but TTransform::operator* is reversed).
+		const int32 ResultNumBones = Result->GetBonePoseCount();
+		for (int32 BoneIndex = 0; BoneIndex < ResultNumBones; ++BoneIndex)
+		{
+			const FBoneName BoneName = Result->GetBonePoseId(BoneIndex);
+			const int32 FoundPoseIndex = PoseMesh->FindBonePose(BoneName);
+			if (FoundPoseIndex == INDEX_NONE)
+			{
+				continue;
+			}
 
-        // Proceed
-        int vertexCount = pBase->GetVertexCount();
-        int weightCount = itBoneIndices.GetComponents();
-        check( weightCount == itBoneWeights.GetComponents() );
+			FTransform3f ModelRefTransform;
+			Result->GetBonePoseTransform(BoneIndex, ModelRefTransform);
 
-        constexpr int MAX_BONES_PER_VERTEX = 16;
-        check( weightCount < MAX_BONES_PER_VERTEX );
+			FTransform3f PoseTransform;
+			PoseMesh->GetBonePoseTransform(FoundPoseIndex, PoseTransform); 
 
-        for ( int v=0; v<vertexCount; ++v )
-        {
-            FVector3f sourcePosition = itSource.GetAsVec3f();
-			FVector3f position = FVector3f(0,0,0);
-
-            float totalWeight = 0.0f;
-
-            for ( int w=0; w<weightCount; ++w )
-            {
-                float weight[MAX_BONES_PER_VERTEX] = {};
-                ConvertData( w, &weight, MBF_FLOAT32, itBoneWeights.ptr(), itBoneWeights.GetFormat() );
-
-                uint32_t boneIndex[MAX_BONES_PER_VERTEX] = {};
-                ConvertData( w, &boneIndex, MBF_UINT32, itBoneIndices.ptr(), itBoneIndices.GetFormat() );
-
-                totalWeight += weight[w];
-
-                //vec4f p = skinMatrices[ boneIndex[w] ]
-                //        * vec4f( sourcePosition, 1.0f )
-                //        * weight[w];
-				FVector3f p = SkinTransforms[boneIndex[w]].TransformPosition(sourcePosition);
-				position += p * weight[w];
-            }
-
-            if (totalWeight<1e-5f)
-            {
-                position = sourcePosition;
-            }
-            else
-            {
-                position /= totalWeight;
-            }
-
-			(*itTarget)[0] = position[0];
-			(*itTarget)[1] = position[1];
-			(*itTarget)[2] = position[2];
-
-            ++itTarget;
-            ++itSource;
-            ++itBoneIndices;
-            ++itBoneWeights;
-        }
-    }
+			// Reshape flag is added so the bone is prioritized in case of merge conflict.
+			// TODO: Add another flag to indicate this case or generalize the Reshaped flag. 
+			EBoneUsageFlags UsageFlags = BaseMesh->GetBoneUsageFlags(BoneIndex);
+			EnumAddFlags(UsageFlags, EBoneUsageFlags::Reshaped);
+			
+			Result->SetBonePose(FoundPoseIndex, BoneName, ModelRefTransform * PoseTransform, UsageFlags);
+		}
 
+		MeshBufferIterator<MBF_FLOAT32, float, 3> TargetPositionIterBegin(Result->VertexBuffers, MBS_POSITION, 0);
+		
+		// Tangent frame buffers are optional.
+		UntypedMeshBufferIteratorConst SourceNormalIterBegin(BaseMesh->VertexBuffers, MBS_NORMAL, 0);
+		UntypedMeshBufferIteratorConst SourceTangentIterBegin(BaseMesh->VertexBuffers, MBS_TANGENT, 0);
+		UntypedMeshBufferIteratorConst SourceBiNormalIterBegin(BaseMesh->VertexBuffers, MBS_BINORMAL, 0);
+
+		UntypedMeshBufferIterator TargetNormalIterBegin(Result->VertexBuffers, MBS_NORMAL, 0);
+		UntypedMeshBufferIterator TargetTangentIterBegin(Result->VertexBuffers, MBS_TANGENT, 0);
+		UntypedMeshBufferIterator TargetBiNormalIterBegin(Result->VertexBuffers, MBS_BINORMAL, 0);
+		
+		check(TargetPositionIterBegin.ptr());
+
+        const int32 VertexCount = BaseMesh->GetVertexCount();
+        const int32 WeightCount = BoneIndicesIterBegin.GetComponents();
+        check(WeightCount == BoneWeightsIterBegin.GetComponents());
+
+        constexpr int32 MAX_BONES_PER_VERTEX = 16;
+        check(WeightCount <= MAX_BONES_PER_VERTEX);
+
+		constexpr int32 NumVertsPerBatch = 1 << 11;
+		auto ProcessVertexBatch = 
+			[
+				SkinTransforms = MakeArrayView<const FTransform3f>(SkinTransforms.GetData(), SkinTransforms.Num()),
+				SourcePositionIterBegin,
+				SourceNormalIterBegin,
+				SourceTangentIterBegin,
+				SourceBiNormalIterBegin,
+				BoneIndicesIterBegin,
+				BoneWeightsIterBegin,
+				TargetPositionIterBegin,
+				TargetNormalIterBegin,
+				TargetTangentIterBegin,
+				TargetBiNormalIterBegin,
+				VertexCount,
+				WeightCount,
+				NumVertsPerBatch
+			](int32 BatchId)
+		{
+			const int32 BatchBeginVertIndex = BatchId * NumVertsPerBatch;
+			const int32 BatchEndVertIndex = FMath::Min(BatchBeginVertIndex + NumVertsPerBatch, VertexCount);
+			
+			const EMeshBufferFormat WeightsFormat = BoneWeightsIterBegin.GetFormat(); 
+			const EMeshBufferFormat BoneIndexFormat = BoneIndicesIterBegin.GetFormat(); 
+			const EMeshBufferFormat NormalFormat = TargetNormalIterBegin.ptr() 
+					? TargetNormalIterBegin.GetFormat() : EMeshBufferFormat::MBF_NONE;	
+
+			check(TargetNormalIterBegin.GetComponents() >= 3)
+
+			for (int32 VertexIndex = BatchBeginVertIndex; VertexIndex < BatchEndVertIndex; ++VertexIndex)
+			{
+				float TotalWeight = 0.0f;
+
+				float Weights[MAX_BONES_PER_VERTEX];
+				const uint8* VertexBoneWeightData = (BoneWeightsIterBegin + VertexIndex).ptr();
+				for (int32 WeightIndex = 0; WeightIndex < WeightCount; ++WeightIndex)
+				{
+					ConvertData(WeightIndex, Weights, MBF_FLOAT32, VertexBoneWeightData, WeightsFormat);
+					TotalWeight += Weights[WeightIndex];
+				}
+
+				uint32 BoneIndices[MAX_BONES_PER_VERTEX];
+				const uint8* VertexBoneIndexData = (BoneIndicesIterBegin + VertexIndex).ptr();
+				for (int32 WeightIndex = 0; WeightIndex < WeightCount; ++WeightIndex)
+				{
+					ConvertData(WeightIndex, BoneIndices, MBF_UINT32, VertexBoneIndexData, BoneIndexFormat);
+				}
+		
+				FVector3f SourcePosition = (SourcePositionIterBegin + VertexIndex).GetAsVec3f();
+				FVector3f SourceNormal = SourceNormalIterBegin.ptr() 
+						? (SourceNormalIterBegin + VertexIndex).GetAsVec3f() : FVector3f::ZeroVector;
+				FVector3f SourceTangent  = SourceTangentIterBegin.ptr() 
+						? (SourceTangentIterBegin + VertexIndex).GetAsVec3f() : FVector3f::ZeroVector;
+				FVector3f SourceBiNormal = SourceBiNormalIterBegin.ptr() 
+						? (SourceBiNormalIterBegin + VertexIndex).GetAsVec3f() : FVector3f::ZeroVector;
+
+				FVector3f Position = FVector3f::ZeroVector;
+				FVector3f Normal   = FVector3f::ZeroVector;
+				FVector3f Tangent  = FVector3f::ZeroVector;
+				FVector3f BiNormal = FVector3f::ZeroVector;
+				
+				for (int32 WeightIndex = 0; WeightIndex < WeightCount; ++WeightIndex)
+				{
+					Position += SkinTransforms[BoneIndices[WeightIndex]].TransformPosition(SourcePosition) * Weights[WeightIndex];
+				
+					if (SourceNormalIterBegin.ptr())
+					{
+						Normal += SkinTransforms[BoneIndices[WeightIndex]].TransformVector(SourceNormal) * Weights[WeightIndex];
+					}
+
+					if (SourceTangentIterBegin.ptr())
+					{
+						Tangent += SkinTransforms[BoneIndices[WeightIndex]].TransformVector(SourceTangent) * Weights[WeightIndex];
+					}
+
+					if (SourceBiNormalIterBegin.ptr())
+					{
+						BiNormal += SkinTransforms[BoneIndices[WeightIndex]].TransformVector(SourceBiNormal) * Weights[WeightIndex];
+					}
+				}
+
+				const float TotalWeightRcp = TotalWeight > UE_SMALL_NUMBER ? 1.0f/TotalWeight : 1.0f; 
+				Position *= TotalWeightRcp;
+
+				float* TargetPositionData = *(TargetPositionIterBegin + VertexIndex); 
+				TargetPositionData[0] = Position[0];
+				TargetPositionData[1] = Position[1];
+				TargetPositionData[2] = Position[2];
+
+				if (TargetNormalIterBegin.ptr())
+				{
+					// This will maintain any packed format sign component if present.
+					(TargetNormalIterBegin + VertexIndex).SetFromVec3f(Normal.GetSafeNormal());
+				}
+
+				if (TargetTangentIterBegin.ptr())
+				{
+					(TargetTangentIterBegin + VertexIndex).SetFromVec3f(Tangent.GetSafeNormal());
+				}
+
+				if (TargetBiNormalIterBegin.ptr())
+				{
+					(TargetBiNormalIterBegin + VertexIndex).SetFromVec3f(BiNormal.GetSafeNormal());
+				}
+			}
+		};
+
+		const int32 NumBatches = FMath::DivideAndRoundUp<int32>(VertexCount, NumVertsPerBatch);
+
+		if (NumBatches == 1)
+		{
+			ProcessVertexBatch(0);
+		}
+		else
+		{
+			ParallelFor(NumBatches, ProcessVertexBatch);
+		}
+    }
 }
