[Mutable] Remove uneeded check in OpMeshApplyPose that could be triggered if no normal buffers are present.

#rnx

diff --git a/Engine/Plugins/Experimental/Mutable/Source/MutableRuntime/Private/MuR/OpMeshApplyPose.h b/Engine/Plugins/Experimental/Mutable/Source/MutableRuntime/Private/MuR/OpMeshApplyPose.h
--- a/Engine/Plugins/Experimental/Mutable/Source/MutableRuntime/Private/MuR/OpMeshApplyPose.h
+++ b/Engine/Plugins/Experimental/Mutable/Source/MutableRuntime/Private/MuR/OpMeshApplyPose.h
@@ -181,8 +181,6 @@ namespace mu
 			const EMeshBufferFormat NormalFormat = TargetNormalIterBegin.ptr() 
 					? TargetNormalIterBegin.GetFormat() : EMeshBufferFormat::MBF_NONE;	
 
-			check(TargetNormalIterBegin.GetComponents() >= 3)
-
 			for (int32 VertexIndex = BatchBeginVertIndex; VertexIndex < BatchEndVertIndex; ++VertexIndex)
 			{
 				float TotalWeight = 0.0f;
