[Mutable] Sink Mesh format operation against ASTOpApplyPose.

#rnx

diff --git a/Engine/Plugins/Experimental/Mutable/Source/MutableTools/Private/MuT/ASTOpMeshFormat.cpp b/Engine/Plugins/Experimental/Mutable/Source/MutableTools/Private/MuT/ASTOpMeshFormat.cpp
--- a/Engine/Plugins/Experimental/Mutable/Source/MutableTools/Private/MuT/ASTOpMeshFormat.cpp
+++ b/Engine/Plugins/Experimental/Mutable/Source/MutableTools/Private/MuT/ASTOpMeshFormat.cpp
@@ -17,6 +17,7 @@
 #include "MuT/ASTOpMeshRemoveMask.h"
 #include "MuT/ASTOpMeshMorph.h"
 #include "MuT/ASTOpMeshAddTags.h"
+#include "MuT/ASTOpMeshApplyPose.h"
 #include "MuT/ASTOpSwitch.h"
 
 
@@ -297,6 +298,14 @@ mu::Ptr<ASTOp> Sink_MeshFormatAST::Visit(const mu::Ptr<ASTOp>& at, const ASTOpMe
 		break;
 	}
 
+	case OP_TYPE::ME_APPLYPOSE:
+	{
+		Ptr<ASTOpMeshApplyPose> NewOp = mu::Clone<ASTOpMeshApplyPose>(at);
+		NewOp->base = Visit(NewOp->base.child(), currentFormatOp);
+		newAt = NewOp;
+		break;
+	}
+
 	case OP_TYPE::ME_INTERPOLATE:
 	{
 		// Move the format down the base of the morph
