[mutable] Fixed the baking of skeletal meshes so that they don't lose their PhysicsAsset when restarting the editor.
#tests Editor
#rnx

diff --git a/Engine/Plugins/Experimental/Mutable/Source/CustomizableObjectEditor/Private/MuCOE/CustomizableObjectInstanceBakingUtils.cpp b/Engine/Plugins/Experimental/Mutable/Source/CustomizableObjectEditor/Private/MuCOE/CustomizableObjectInstanceBakingUtils.cpp
--- a/Engine/Plugins/Experimental/Mutable/Source/CustomizableObjectEditor/Private/MuCOE/CustomizableObjectInstanceBakingUtils.cpp
+++ b/Engine/Plugins/Experimental/Mutable/Source/CustomizableObjectEditor/Private/MuCOE/CustomizableObjectInstanceBakingUtils.cpp
@@ -15,6 +15,7 @@
 #include "Materials/MaterialInstanceConstant.h"
 #include "Materials/MaterialInstanceDynamic.h"
 #include "Materials/Material.h"
+#include "PhysicsEngine/PhysicsAsset.h"
 
 #include "MuCO/CustomizableObject.h"
 #include "MuCO/CustomizableObjectInstanceAssetUserData.h"
@@ -857,6 +858,18 @@ bool BakeCustomizableObjectInstance(
 				}
 			}
 
+			// Duplicate PhysicsAsset
+			{
+				const UPhysicsAsset* PhysicsAsset = Mesh->GetPhysicsAsset();
+
+				if (PhysicsAsset)
+				{
+					// Duplicate to change from the Transient Package to the baked mesh one
+					UPhysicsAsset* NewPhysicsAsset = Cast<UPhysicsAsset>(StaticDuplicateObject(PhysicsAsset, SkeletalMesh));
+					SkeletalMesh->SetPhysicsAsset(NewPhysicsAsset);
+				}
+			}
+
 			// Generate render data
 			SkeletalMesh->Build();
 		}
