[Mutable] Force loading all node referencing objects before reconstructing them.

Only added to nodes whose referenced objects are used in PostLoad (for example ReconstructNode, AllocateDefaultPins...).

#rnx

diff --git a/Engine/Plugins/Experimental/Mutable/Source/CustomizableObjectEditor/Private/MuCOE/CustomizableObjectEditorUtilities.cpp b/Engine/Plugins/Experimental/Mutable/Source/CustomizableObjectEditor/Private/MuCOE/CustomizableObjectEditorUtilities.cpp
new file mode 100644
--- /dev/null
+++ b/Engine/Plugins/Experimental/Mutable/Source/CustomizableObjectEditor/Private/MuCOE/CustomizableObjectEditorUtilities.cpp
@@ -0,0 +1,39 @@
+﻿// Copyright Epic Games, Inc. All Rights Reserved.
+
+#pragma once
+
+#include "MuCOE/CustomizableObjectEditorUtilities.h"
+
+#include "UObject/Object.h"
+#include "UObject/LinkerLoad.h"
+
+
+bool CompareNames(const TSharedPtr<FString>& sp1, const TSharedPtr<FString>& sp2)
+{
+	if (const FString* s1 = sp1.Get())
+	{
+		if (const FString* s2 = sp2.Get())
+		{
+			return (s1->Compare(*s2, ESearchCase::IgnoreCase) < 0);
+		}
+		else
+		{
+			return false;
+		}
+	}
+	else
+	{
+		return true;
+	}
+}
+
+
+void ConditionalPostLoadReference(UObject& Object)
+{
+	if (Object.HasAnyFlags(RF_NeedLoad))
+	{
+		Object.GetLinker()->Preload(&Object);
+	}
+
+	Object.ConditionalPostLoad();
+}
diff --git a/Engine/Plugins/Experimental/Mutable/Source/CustomizableObjectEditor/Private/MuCOE/CustomizableObjectEditorUtilities.h b/Engine/Plugins/Experimental/Mutable/Source/CustomizableObjectEditor/Private/MuCOE/CustomizableObjectEditorUtilities.h
--- a/Engine/Plugins/Experimental/Mutable/Source/CustomizableObjectEditor/Private/MuCOE/CustomizableObjectEditorUtilities.h
+++ b/Engine/Plugins/Experimental/Mutable/Source/CustomizableObjectEditor/Private/MuCOE/CustomizableObjectEditorUtilities.h
@@ -4,21 +4,9 @@
 
 #include "Templates/SharedPointer.h"
 
-static bool CompareNames(const TSharedPtr<FString>& sp1, const TSharedPtr<FString>& sp2)
-{
-    if (const FString* s1 = sp1.Get())
-    {
-        if (const FString* s2 = sp2.Get())
-        {
-            return (s1->Compare(*s2, ESearchCase::IgnoreCase) < 0);
-        }
-        else
-        {
-            return false;
-        }
-    }
-    else
-    {
-        return true;
-    }
-}
+class UObject;
+
+
+bool CompareNames(const TSharedPtr<FString>& sp1, const TSharedPtr<FString>& sp2);
+
+void ConditionalPostLoadReference(UObject& Object);
diff --git a/Engine/Plugins/Experimental/Mutable/Source/CustomizableObjectEditor/Private/MuCOE/Nodes/CustomizableObjectNodeMaterial.cpp b/Engine/Plugins/Experimental/Mutable/Source/CustomizableObjectEditor/Private/MuCOE/Nodes/CustomizableObjectNodeMaterial.cpp
--- a/Engine/Plugins/Experimental/Mutable/Source/CustomizableObjectEditor/Private/MuCOE/Nodes/CustomizableObjectNodeMaterial.cpp
+++ b/Engine/Plugins/Experimental/Mutable/Source/CustomizableObjectEditor/Private/MuCOE/Nodes/CustomizableObjectNodeMaterial.cpp
@@ -21,6 +21,7 @@
 #include "Modules/ModuleManager.h"
 #include "MuCO/CustomizableObjectCustomVersion.h"
 #include "MuCOE/CustomizableObjectEditorLogger.h"
+#include "MuCOE/CustomizableObjectEditorUtilities.h"
 #include "MuCOE/Nodes/CustomizableObjectNodeExternalPin.h"
 
 class SGraphNode;
@@ -662,7 +663,7 @@ void UCustomizableObjectNodeMaterial::PostLoad()
 	
 	if (Material)
 	{
-		Material->ConditionalPostLoad(); // Make sure the Material has been fully loaded.
+		ConditionalPostLoadReference(*Material); // Make sure the Material has been fully loaded.
 	}
 }
 
diff --git a/Engine/Plugins/Experimental/Mutable/Source/CustomizableObjectEditor/Private/MuCOE/Nodes/CustomizableObjectNodeSkeletalMesh.cpp b/Engine/Plugins/Experimental/Mutable/Source/CustomizableObjectEditor/Private/MuCOE/Nodes/CustomizableObjectNodeSkeletalMesh.cpp
--- a/Engine/Plugins/Experimental/Mutable/Source/CustomizableObjectEditor/Private/MuCOE/Nodes/CustomizableObjectNodeSkeletalMesh.cpp
+++ b/Engine/Plugins/Experimental/Mutable/Source/CustomizableObjectEditor/Private/MuCOE/Nodes/CustomizableObjectNodeSkeletalMesh.cpp
@@ -9,6 +9,7 @@
 #include "MuCO/CustomizableObjectCustomVersion.h"
 #include "MuCOE/CustomizableObjectEditor.h"
 #include "MuCOE/CustomizableObjectEditorStyle.h"
+#include "MuCOE/CustomizableObjectEditorUtilities.h"
 #include "MuCOE/CustomizableObjectLayout.h"
 #include "MuCOE/EdGraphSchema_CustomizableObject.h"
 #include "MuCOE/GraphTraversal.h"
@@ -30,6 +31,25 @@ struct FSlateBrush;
 static const TCHAR* SKELETAL_MESH_PIN_NAME = TEXT("Skeletal Mesh");
 
 
+void UCustomizableObjectNodeSkeletalMesh::PostLoad()
+{
+	Super::PostLoad();
+
+	if (SkeletalMesh)
+	{
+		ConditionalPostLoadReference(*SkeletalMesh);
+
+		for (const FSkeletalMaterial& SkeletalMaterial : SkeletalMesh->GetMaterials())
+		{
+			if (UMaterialInterface* Material = SkeletalMaterial.MaterialInterface)
+			{
+				ConditionalPostLoadReference(*Material);
+			}
+		}
+	}
+}
+
+
 void UCustomizableObjectNodeSkeletalMesh::PostEditChangeProperty(FPropertyChangedEvent& PropertyChangedEvent)
 {
 	Super::PostEditChangeProperty(PropertyChangedEvent);
diff --git a/Engine/Plugins/Experimental/Mutable/Source/CustomizableObjectEditor/Private/MuCOE/Nodes/CustomizableObjectNodeStaticMesh.cpp b/Engine/Plugins/Experimental/Mutable/Source/CustomizableObjectEditor/Private/MuCOE/Nodes/CustomizableObjectNodeStaticMesh.cpp
--- a/Engine/Plugins/Experimental/Mutable/Source/CustomizableObjectEditor/Private/MuCOE/Nodes/CustomizableObjectNodeStaticMesh.cpp
+++ b/Engine/Plugins/Experimental/Mutable/Source/CustomizableObjectEditor/Private/MuCOE/Nodes/CustomizableObjectNodeStaticMesh.cpp
@@ -8,6 +8,7 @@
 #include "MuCO/CustomizableObjectCustomVersion.h"
 #include "MuCOE/CustomizableObjectEditor.h"
 #include "MuCOE/CustomizableObjectEditorStyle.h"
+#include "MuCOE/CustomizableObjectEditorUtilities.h"
 #include "MuCOE/CustomizableObjectLayout.h"
 #include "MuCOE/EdGraphSchema_CustomizableObject.h"
 #include "MuCOE/GraphTraversal.h"
@@ -28,6 +29,25 @@ struct FSlateBrush;
 static const TCHAR* STATIC_MESH_PIN_NAME = TEXT("Static Mesh");
 
 
+void UCustomizableObjectNodeStaticMesh::PostLoad()
+{
+	Super::PostLoad();
+
+	if (StaticMesh)
+	{
+		ConditionalPostLoadReference(*StaticMesh);
+
+		for (const FStaticMaterial& StaticMaterial : StaticMesh->GetStaticMaterials())
+		{
+			if (UMaterialInterface* Material = StaticMaterial.MaterialInterface)
+			{
+				ConditionalPostLoadReference(*Material);
+			}
+		}
+	}
+}
+
+
 void UCustomizableObjectNodeStaticMesh::PostEditChangeProperty(FPropertyChangedEvent& PropertyChangedEvent)
 {
 	FProperty* PropertyThatChanged = PropertyChangedEvent.Property;
diff --git a/Engine/Plugins/Experimental/Mutable/Source/CustomizableObjectEditor/Private/MuCOE/Nodes/CustomizableObjectNodeTable.cpp b/Engine/Plugins/Experimental/Mutable/Source/CustomizableObjectEditor/Private/MuCOE/Nodes/CustomizableObjectNodeTable.cpp
--- a/Engine/Plugins/Experimental/Mutable/Source/CustomizableObjectEditor/Private/MuCOE/Nodes/CustomizableObjectNodeTable.cpp
+++ b/Engine/Plugins/Experimental/Mutable/Source/CustomizableObjectEditor/Private/MuCOE/Nodes/CustomizableObjectNodeTable.cpp
@@ -25,6 +25,7 @@
 #include "AssetRegistry/AssetRegistryModule.h"
 #include "AssetRegistry/IAssetRegistry.h"
 #include "MuCO/CustomizableObjectCustomVersion.h"
+#include "MuCOE/CustomizableObjectEditorUtilities.h"
 #include "UObject/LinkerLoad.h"
 
 class ICustomizableObjectEditor;
@@ -197,6 +198,58 @@ void UCustomizableObjectNodeTable::PostBackwardsCompatibleFixup()
 }
 
 
+void UCustomizableObjectNodeTable::PostLoad()
+{
+	Super::PostLoad();
+	
+	if (Table)
+	{
+		ConditionalPostLoadReference(*Table);
+	}
+
+	if (Structure)
+	{
+		ConditionalPostLoadReference(*Structure);
+	}
+	
+	if (const UScriptStruct* TableStruct = GetTableNodeStruct())
+	{
+		TArray<int8> DefaultDataArray;
+		DefaultDataArray.SetNumZeroed(TableStruct->GetStructureSize());
+		TableStruct->InitializeStruct(DefaultDataArray.GetData());
+
+		for (TFieldIterator<FProperty> It(TableStruct); It; ++It)
+		{
+			FProperty* ColumnProperty = *It;
+			if (!ColumnProperty)
+			{
+				continue;
+			}
+		
+			if (const FSoftObjectProperty* SoftObjectProperty = CastField<FSoftObjectProperty>(ColumnProperty))
+			{
+				if (SoftObjectProperty->PropertyClass)
+				{
+					uint8* CellData = SoftObjectProperty->ContainerPtrToValuePtr<uint8>(DefaultDataArray.GetData());
+					if (!CellData)
+					{
+						continue;
+					}
+
+					UObject* Object = SoftObjectProperty->GetPropertyValue(CellData).LoadSynchronous();
+					if (!Object)
+					{
+						continue;
+					}
+
+					ConditionalPostLoadReference(*Object);
+				}
+			}
+		}
+	}
+}
+
+
 void UCustomizableObjectNodeTable::PreEditChange(FProperty* PropertyAboutToChange)
 {
 	Super::PreEditChange(PropertyAboutToChange);
@@ -549,15 +602,6 @@ void UCustomizableObjectNodeTable::GenerateMeshPins(UObject* Mesh, const FString
 {
 	if (USkeletalMesh* SkeletalMesh = Cast<USkeletalMesh>(Mesh))
 	{
-		// This can be called from a codepath that starts in PostLoad. The SkeletalMesh may not be serialized at this time so here we will preload it to make sure we can read the LOD data
-		if (SkeletalMesh->HasAnyFlags(RF_NeedLoad))
-		{
-			if (FLinkerLoad* Linker = GetLinker())
-			{
-				Linker->Preload(SkeletalMesh);
-			}
-		}
-
 		const int NumLODs = SkeletalMesh->GetLODNum();
 		const UEdGraphSchema_CustomizableObject* Schema = GetDefault<UEdGraphSchema_CustomizableObject>();
 
diff --git a/Engine/Plugins/Experimental/Mutable/Source/CustomizableObjectEditor/Public/MuCOE/Nodes/CustomizableObjectNode.h b/Engine/Plugins/Experimental/Mutable/Source/CustomizableObjectEditor/Public/MuCOE/Nodes/CustomizableObjectNode.h
--- a/Engine/Plugins/Experimental/Mutable/Source/CustomizableObjectEditor/Public/MuCOE/Nodes/CustomizableObjectNode.h
+++ b/Engine/Plugins/Experimental/Mutable/Source/CustomizableObjectEditor/Public/MuCOE/Nodes/CustomizableObjectNode.h
@@ -109,7 +109,7 @@ public:
 
 	/** Add post load work here.
 	 * When called, it is guaranteed that all nodes in this graph will have executed the BackwardsCompatibleFixup function. */
-	virtual void PostBackwardsCompatibleFixup() {};
+	virtual void PostBackwardsCompatibleFixup() {}
 	
 	/** Virtual implementation of RemovePin. Allows to do work before removing a pin.
 	 * Use this function instead of RemovePin. RemovePin does not removes possible attached PinData. */
diff --git a/Engine/Plugins/Experimental/Mutable/Source/CustomizableObjectEditor/Public/MuCOE/Nodes/CustomizableObjectNodeSkeletalMesh.h b/Engine/Plugins/Experimental/Mutable/Source/CustomizableObjectEditor/Public/MuCOE/Nodes/CustomizableObjectNodeSkeletalMesh.h
--- a/Engine/Plugins/Experimental/Mutable/Source/CustomizableObjectEditor/Public/MuCOE/Nodes/CustomizableObjectNodeSkeletalMesh.h
+++ b/Engine/Plugins/Experimental/Mutable/Source/CustomizableObjectEditor/Public/MuCOE/Nodes/CustomizableObjectNodeSkeletalMesh.h
@@ -201,6 +201,7 @@ public:
 	FGameplayTagContainer AnimationGameplayTags;
 
 	// UObject interface.
+	virtual void PostLoad() override;
 	virtual void PostEditChangeProperty(FPropertyChangedEvent& PropertyChangedEvent) override;
 	
 	// UEdGraphNode interface
diff --git a/Engine/Plugins/Experimental/Mutable/Source/CustomizableObjectEditor/Public/MuCOE/Nodes/CustomizableObjectNodeStaticMesh.h b/Engine/Plugins/Experimental/Mutable/Source/CustomizableObjectEditor/Public/MuCOE/Nodes/CustomizableObjectNodeStaticMesh.h
--- a/Engine/Plugins/Experimental/Mutable/Source/CustomizableObjectEditor/Public/MuCOE/Nodes/CustomizableObjectNodeStaticMesh.h
+++ b/Engine/Plugins/Experimental/Mutable/Source/CustomizableObjectEditor/Public/MuCOE/Nodes/CustomizableObjectNodeStaticMesh.h
@@ -131,13 +131,14 @@ public:
 	FEdGraphPinReference DefaultPin;
 	
 	// UObject interface.
-	void PostEditChangeProperty(FPropertyChangedEvent& PropertyChangedEvent) override;
+	virtual void PostLoad() override;
+	virtual void PostEditChangeProperty(FPropertyChangedEvent& PropertyChangedEvent) override;
 	virtual void Serialize(FArchive& Ar) override;
 
 	// UEdGraphNode interface
-	FText GetNodeTitle(ENodeTitleType::Type TitleType) const override;
-	FLinearColor GetNodeTitleColor() const override;
-	FText GetTooltipText() const override;
+	virtual FText GetNodeTitle(ENodeTitleType::Type TitleType) const override;
+	virtual FLinearColor GetNodeTitleColor() const override;
+	virtual FText GetTooltipText() const override;
 
 	// UCustomizableObjectNode interface
 	virtual void AllocateDefaultPins(UCustomizableObjectNodeRemapPins* RemapPins) override;
@@ -145,9 +146,8 @@ public:
 	virtual bool IsPinRelevant(const UEdGraphPin* Pin) const override;
 	virtual UCustomizableObjectNodeRemapPinsByName* CreateRemapPinsByName() const override;
 	virtual bool HasPinViewer() const override;
-
-	bool IsNodeOutDatedAndNeedsRefresh() override;
-	FString GetRefreshMessage() const override;
+	virtual bool IsNodeOutDatedAndNeedsRefresh() override;
+	virtual FString GetRefreshMessage() const override;
 
 	// UCustomizableObjectNodeMesh interface
 	virtual UTexture2D* FindTextureForPin(const UEdGraphPin* Pin) const override;
@@ -158,11 +158,11 @@ public:
 	virtual UEdGraphPin* GetLayoutPin(int32 LODIndex, int32 SectionIndex, int32 LayoutIndex) const override;
 	virtual void GetPinSection(const UEdGraphPin& Pin, int32& OutLODIndex, int32& OutSectionIndex, int32& OutLayoutIndex) const override;
 
-	/** Returns the material assossiated to the given output pin. */
+	/** Returns the material associated to the given output pin. */
 	UMaterialInterface* GetMaterialFor(const UEdGraphPin* Pin) const;
 
 	// Creates the SGraph Node widget for the thumbnail
-	TSharedPtr<SGraphNode> CreateVisualWidget() override;
+	virtual TSharedPtr<SGraphNode> CreateVisualWidget() override;
 
 	// Determines if the Node is collapsed or not
 	bool bCollapsed = true;
diff --git a/Engine/Plugins/Experimental/Mutable/Source/CustomizableObjectEditor/Public/MuCOE/Nodes/CustomizableObjectNodeTable.h b/Engine/Plugins/Experimental/Mutable/Source/CustomizableObjectEditor/Public/MuCOE/Nodes/CustomizableObjectNodeTable.h
--- a/Engine/Plugins/Experimental/Mutable/Source/CustomizableObjectEditor/Public/MuCOE/Nodes/CustomizableObjectNodeTable.h
+++ b/Engine/Plugins/Experimental/Mutable/Source/CustomizableObjectEditor/Public/MuCOE/Nodes/CustomizableObjectNodeTable.h
@@ -239,6 +239,7 @@ public:
 public:
 
 	// UObject interface
+	virtual void PostLoad() override;
 	virtual void PreEditChange(FProperty* PropertyAboutToChange) override;
 	virtual void PostEditChangeChainProperty(FPropertyChangedChainEvent& PropertyChangedEvent) override;
 	virtual void PostEditChangeProperty(FPropertyChangedEvent& PropertyChangedEvent) override;
