[mutable] Fixed duplicated LOCTEXT key

#rnx

diff --git a/Engine/Plugins/Experimental/Mutable/Source/CustomizableObjectEditor/Private/MuCOE/EdGraphSchema_CustomizableObject.cpp b/Engine/Plugins/Experimental/Mutable/Source/CustomizableObjectEditor/Private/MuCOE/EdGraphSchema_CustomizableObject.cpp
--- a/Engine/Plugins/Experimental/Mutable/Source/CustomizableObjectEditor/Private/MuCOE/EdGraphSchema_CustomizableObject.cpp
+++ b/Engine/Plugins/Experimental/Mutable/Source/CustomizableObjectEditor/Private/MuCOE/EdGraphSchema_CustomizableObject.cpp
@@ -880,8 +880,8 @@ void UEdGraphSchema_CustomizableObject::GetContextMenuActions(UToolMenu* Menu, U
 				{
 					FToolMenuSection& Section = Menu->FindOrAddSection("EdGraphSchemaPinActions");
 					Section.AddMenuEntry("HidePin",
-						LOCTEXT("HidePin", "Hide Pin"),
-						LOCTEXT("SetTexturePinModeDefault_Tooltip", "Hides the selected pin."),
+						LOCTEXT("HidePin_Label", "Hide Pin"),
+						LOCTEXT("HidePin_Tooltip", "Hides the selected pin."),
 						FSlateIcon(),
 						FUIAction(FExecuteAction::CreateLambda([Node = const_cast<UCustomizableObjectNode*>(Node), Pin]()
 						{
