[Mutable] Removed pragma once from source file.

#rnx

diff --git a/Engine/Plugins/Experimental/Mutable/Source/CustomizableObjectEditor/Private/MuCOE/CustomizableObjectEditorUtilities.cpp b/Engine/Plugins/Experimental/Mutable/Source/CustomizableObjectEditor/Private/MuCOE/CustomizableObjectEditorUtilities.cpp
--- a/Engine/Plugins/Experimental/Mutable/Source/CustomizableObjectEditor/Private/MuCOE/CustomizableObjectEditorUtilities.cpp
+++ b/Engine/Plugins/Experimental/Mutable/Source/CustomizableObjectEditor/Private/MuCOE/CustomizableObjectEditorUtilities.cpp
@@ -1,7 +1,5 @@
 ﻿// Copyright Epic Games, Inc. All Rights Reserved.
 
-#pragma once
-
 #include "MuCOE/CustomizableObjectEditorUtilities.h"
 
 #include "UObject/Object.h"
