[Mutable] Added Pin Viewer to the NodeTable.

#rnx

diff --git a/Engine/Plugins/Experimental/Mutable/Source/CustomizableObject/Internal/MuCO/CustomizableObjectCustomVersion.h b/Engine/Plugins/Experimental/Mutable/Source/CustomizableObject/Internal/MuCO/CustomizableObjectCustomVersion.h
--- a/Engine/Plugins/Experimental/Mutable/Source/CustomizableObject/Internal/MuCO/CustomizableObjectCustomVersion.h
+++ b/Engine/Plugins/Experimental/Mutable/Source/CustomizableObject/Internal/MuCO/CustomizableObjectCustomVersion.h
@@ -120,6 +120,8 @@ struct CUSTOMIZABLEOBJECT_API FCustomizableObjectCustomVersion
 
 		CompilationOptions,
 
+		NodeTablePinViewer,
+
 		// -----<new versions can be added above this line>-------------------------------------------------
 		VersionPlusOne,
 		LatestVersion = VersionPlusOne - 1
diff --git a/Engine/Plugins/Experimental/Mutable/Source/CustomizableObjectEditor/Private/MuCOE/EdGraphSchema_CustomizableObject.cpp b/Engine/Plugins/Experimental/Mutable/Source/CustomizableObjectEditor/Private/MuCOE/EdGraphSchema_CustomizableObject.cpp
--- a/Engine/Plugins/Experimental/Mutable/Source/CustomizableObjectEditor/Private/MuCOE/EdGraphSchema_CustomizableObject.cpp
+++ b/Engine/Plugins/Experimental/Mutable/Source/CustomizableObjectEditor/Private/MuCOE/EdGraphSchema_CustomizableObject.cpp
@@ -5,6 +5,7 @@
 #include "EdGraphNode_Comment.h"
 #include "Engine/SkeletalMesh.h"
 #include "Engine/StaticMesh.h"
+#include "Framework/Application/SlateApplication.h"
 #include "Framework/Commands/GenericCommands.h"
 #include "GraphEditorActions.h"
 #include "Materials/MaterialInterface.h"
@@ -874,8 +875,6 @@ void UEdGraphSchema_CustomizableObject::GetContextMenuActions(UToolMenu* Menu, U
 		{
 			if (const UEdGraphPin* Pin = Context->Pin)
 			{
-				FPropertyEditorModule& PropPlugin = FModuleManager::LoadModuleChecked<FPropertyEditorModule>("PropertyEditor");
-				
 				if (Node->CanPinBeHidden(*Pin))
 				{
 					FToolMenuSection& Section = Menu->FindOrAddSection("EdGraphSchemaPinActions");
@@ -890,53 +889,21 @@ void UEdGraphSchema_CustomizableObject::GetContextMenuActions(UToolMenu* Menu, U
 				}
 				
 				
-				if (TSharedPtr<SWidget> Widget = Node->CustomizePinDetails(*Pin))
+				if (TSharedPtr<IDetailsView> Widget = Node->CustomizePinDetails(*Pin))
 				{
+					Widget->OnFinishedChangingProperties().AddLambda([WeakMenu = Widget->AsWeak()](const FPropertyChangedEvent& Event){
+						if (TSharedPtr<SWidget> Menu = WeakMenu.Pin())
+						{
+							FSlateApplication::Get().DismissMenuByWidget(Menu.ToSharedRef());
+						}
+					});
+					
 					FToolMenuSection& Section = Menu->FindOrAddSection("EdGraphSchemaPinActions");
 					Section.AddSeparator("Pin Viewer");
 					Section.AddEntry(FToolMenuEntry::InitWidget("Pin Viewer", Widget.ToSharedRef(), {}));
 				}
 			}
 		}
-		
-		const UCustomizableObjectNodeTable* TableNode = Cast<UCustomizableObjectNodeTable>(Context->Node);
-		UEdGraphPin* TexturePin = (UEdGraphPin*)Context->Pin;
-
-		if (TableNode && TexturePin && !TableNode->IsNotTexture2DPin(TexturePin) && !TexturePin->LinkedTo.Num() && (TexturePin->PinType.PinCategory == PC_Image || TexturePin->PinType.PinCategory == PC_PassThroughImage))
-		{
-			FText ActionText = FText::Format(LOCTEXT("ChangeTexturePinMode_Label", 
-				"Set as {0} texture"), Context->Pin->PinType.PinCategory == PC_PassThroughImage ? FText::FromString("Mutable") : FText::FromString("Pass Through"));
-			
-			FText ToolTipText = FText::Format(LOCTEXT("ChangeTexturePinMode_Tooltip",
-				"Set the texture pin as {0} texture."), Context->Pin->PinType.PinCategory == PC_PassThroughImage ? FText::FromString("Mutable") : FText::FromString("Pass Through"));
-
-			FToolMenuSection& Section = Menu->FindOrAddSection("EdGraphSchemaPinActions");
-			Section.InitSection("EdGraphSchemaPinActions", LOCTEXT("PinActionsMenuHeader", "Pin Actions"), FToolMenuInsert());
-
-			Section.AddMenuEntry
-			(
-				"ChangeTexturePinMode", ActionText, ToolTipText, FSlateIcon(),
-				FUIAction(FExecuteAction::CreateLambda([TableNode = const_cast<UCustomizableObjectNodeTable*>(TableNode), TexturePin]()
-					{
-						TableNode->ChangeImagePinMode(TexturePin);
-					}))
-			);
-
-			if (!TableNode->IsImagePinDefault(TexturePin))
-			{
-				Section.AddMenuEntry
-				(
-					"SetTexturePinModeDefault",
-					LOCTEXT("SetTexturePinModeDefault_Label","Set pin as default."),
-					LOCTEXT("SetTexturePinModeDefault_Tooltip","Set the selected texture pin to use the default node mode."),
-					FSlateIcon(),
-					FUIAction(FExecuteAction::CreateLambda([TableNode = const_cast<UCustomizableObjectNodeTable*>(TableNode), TexturePin]()
-						{
-							TableNode->ChangeImagePinMode(TexturePin, true);
-						}))
-				);
-			}
-		}
 	}
 }
 
diff --git a/Engine/Plugins/Experimental/Mutable/Source/CustomizableObjectEditor/Private/MuCOE/Nodes/CustomizableObjectNode.cpp b/Engine/Plugins/Experimental/Mutable/Source/CustomizableObjectEditor/Private/MuCOE/Nodes/CustomizableObjectNode.cpp
--- a/Engine/Plugins/Experimental/Mutable/Source/CustomizableObjectEditor/Private/MuCOE/Nodes/CustomizableObjectNode.cpp
+++ b/Engine/Plugins/Experimental/Mutable/Source/CustomizableObjectEditor/Private/MuCOE/Nodes/CustomizableObjectNode.cpp
@@ -658,12 +658,26 @@ bool UCustomizableObjectNode::HasPinViewer() const
 }
 
 
-TSharedPtr<SWidget> UCustomizableObjectNode::CustomizePinDetails(const UEdGraphPin& Pin) const
+TSharedPtr<IDetailsView> UCustomizableObjectNode::CustomizePinDetails(const UEdGraphPin& Pin) const
 {
 	return nullptr;
 }
 
 
+UEdGraphPin* UCustomizableObjectNode::GetPin(const UCustomizableObjectNodePinData& PinData)
+{
+	for (const TTuple<FGuid, TObjectPtr<UCustomizableObjectNodePinData>>& Pair : PinsDataId)
+	{
+		if (Pair.Value == &PinData)
+		{
+			return FindPinById(Pair.Key);
+		}
+	}
+	
+	return nullptr;
+}
+
+
 UCustomizableObjectNodePinData* UCustomizableObjectNode::GetPinData(const UEdGraphPin& Pin) const
 {
 	if (TObjectPtr<UCustomizableObjectNodePinData> const* Result = PinsDataId.Find(Pin.PinId))
diff --git a/Engine/Plugins/Experimental/Mutable/Source/CustomizableObjectEditor/Private/MuCOE/Nodes/CustomizableObjectNodeMaterial.cpp b/Engine/Plugins/Experimental/Mutable/Source/CustomizableObjectEditor/Private/MuCOE/Nodes/CustomizableObjectNodeMaterial.cpp
--- a/Engine/Plugins/Experimental/Mutable/Source/CustomizableObjectEditor/Private/MuCOE/Nodes/CustomizableObjectNodeMaterial.cpp
+++ b/Engine/Plugins/Experimental/Mutable/Source/CustomizableObjectEditor/Private/MuCOE/Nodes/CustomizableObjectNodeMaterial.cpp
@@ -1074,7 +1074,7 @@ FString UCustomizableObjectNodeMaterial::GetRefreshMessage() const
 }
 
 
-TSharedPtr<SWidget> UCustomizableObjectNodeMaterial::CustomizePinDetails(const UEdGraphPin& Pin) const
+TSharedPtr<IDetailsView> UCustomizableObjectNodeMaterial::CustomizePinDetails(const UEdGraphPin& Pin) const
 {
 	if (UCustomizableObjectNodeMaterialPinDataImage* PinData = Cast<UCustomizableObjectNodeMaterialPinDataImage>(GetPinData(Pin)))
 	{
diff --git a/Engine/Plugins/Experimental/Mutable/Source/CustomizableObjectEditor/Private/MuCOE/Nodes/CustomizableObjectNodeTable.cpp b/Engine/Plugins/Experimental/Mutable/Source/CustomizableObjectEditor/Private/MuCOE/Nodes/CustomizableObjectNodeTable.cpp
--- a/Engine/Plugins/Experimental/Mutable/Source/CustomizableObjectEditor/Private/MuCOE/Nodes/CustomizableObjectNodeTable.cpp
+++ b/Engine/Plugins/Experimental/Mutable/Source/CustomizableObjectEditor/Private/MuCOE/Nodes/CustomizableObjectNodeTable.cpp
@@ -27,6 +27,7 @@
 #include "MuCO/CustomizableObjectCustomVersion.h"
 #include "MuCOE/CustomizableObjectEditorUtilities.h"
 #include "UObject/LinkerLoad.h"
+#include "PropertyEditorModule.h"
 
 class ICustomizableObjectEditor;
 class UCustomizableObjectNodeRemapPins;
@@ -171,16 +172,14 @@ void UCustomizableObjectNodeTable::BackwardsCompatibleFixup()
 			}
 		}
 	}
-
-	if (CustomizableObjectCustomVersion < FCustomizableObjectCustomVersion::AddedAnyTextureTypeToPassThroughTextures)
+	
+	if (CustomizableObjectCustomVersion < FCustomizableObjectCustomVersion::NodeTablePinViewer)
 	{
-		for (UEdGraphPin* Pin : Pins)
+		for (UEdGraphPin* Pin : GetAllPins())
 		{
-			UCustomizableObjectNodePinData* PinData = GetPinData(*Pin);
-
-			if (UCustomizableObjectNodeTableImagePinData* ImagePinData = Cast<UCustomizableObjectNodeTableImagePinData>(PinData))
+			if (UCustomizableObjectNodeTableImagePinData* ImagePinData = Cast<UCustomizableObjectNodeTableImagePinData>(GetPinData(*Pin)))
 			{
-				ImagePinData->ConvertArrayTextureToAnyTextureFixup();
+				ImagePinData->NodeTable = this;
 			}
 		}
 	}
@@ -334,23 +333,6 @@ void UCustomizableObjectNodeTable::OnRenameNode(const FString& NewName)
 }
 
 
-void UCustomizableObjectNodeTable::PinConnectionListChanged(UEdGraphPin* Pin)
-{
-	if (Pin && !Pin->LinkedTo.Num() && 
-		(Pin->PinType.PinCategory == UEdGraphSchema_CustomizableObject::PC_Image || Pin->PinType.PinCategory == UEdGraphSchema_CustomizableObject::PC_PassThroughImage))
-	{
-		if (UCustomizableObjectNodeTableImagePinData* ImagePinData = Cast<UCustomizableObjectNodeTableImagePinData>(GetPinData(*(Pin))))
-		{
-			if (!ImagePinData->IsNotTexture2D() && ImagePinData->IsDefaultImageMode() && ImagePinData->ImageMode != DefaultImageMode)
-			{
-				ImagePinData->ImageMode = DefaultImageMode;
-				ReconstructNode();
-			}
-		}
-	}
-}
-
-
 void UCustomizableObjectNodeTable::PostEditChangeProperty(FPropertyChangedEvent& PropertyChangedEvent)
 {
 	Super::PostEditChangeProperty(PropertyChangedEvent);
@@ -359,36 +341,13 @@ void UCustomizableObjectNodeTable::PostEditChangeProperty(FPropertyChangedEvent&
 
 	if (PropertyThatChanged)
 	{
-		if (PropertyThatChanged->GetName() == TEXT("Table") || PropertyThatChanged->GetName() == TEXT("Structure"))
+		if (PropertyThatChanged->GetName() == GET_MEMBER_NAME_CHECKED(UCustomizableObjectNodeTable, Table) || PropertyThatChanged->GetName() == GET_MEMBER_NAME_CHECKED(UCustomizableObjectNodeTable, Structure))
 		{
 			ParamUIMetadataColumn = NAME_None;
 			VersionColumn = NAME_None;
 			ReconstructNode();
 		}
-		else if (PropertyThatChanged->GetName() == TEXT("DefaultImageMode"))
-		{
-			bool bReconstruct = false;
-
-			for (UEdGraphPin* Pin : Pins)
-			{
-				if (UCustomizableObjectNodeTableImagePinData* TexturePinData = Cast<UCustomizableObjectNodeTableImagePinData>(GetPinData(*(Pin))))
-				{
-					if (!Pin->LinkedTo.Num() && !TexturePinData->IsNotTexture2D() && TexturePinData->IsDefaultImageMode())
-					{
-						ETableTextureType Mode = TexturePinData->ImageMode == ETableTextureType::PASSTHROUGH_TEXTURE ? ETableTextureType::MUTABLE_TEXTURE : ETableTextureType::PASSTHROUGH_TEXTURE;
-						TexturePinData->ImageMode = Mode;
-
-						bReconstruct = true;
-					}
-				}
-			}
-
-			if (bReconstruct)
-			{
-				ReconstructNode();
-			}
-		}
-		else if (PropertyThatChanged->GetName() == TEXT("TableDataGatheringMode"))
+		else if (PropertyThatChanged->GetName() == GET_MEMBER_NAME_CHECKED(UCustomizableObjectNodeTable, TableDataGatheringMode))
 		{
 			if (Table)
 			{
@@ -492,9 +451,10 @@ void UCustomizableObjectNodeTable::AllocateDefaultPins(UCustomizableObjectNodeRe
 						UCustomizableObjectNodeTableImagePinData* PinData = NewObject<UCustomizableObjectNodeTableImagePinData>(this);
 						PinData->ColumnName = ColumnName;
 						PinData->StructColumnId = ColumnPropertyId;
-						PinData->SetIsNotTexture2D(false);
+						PinData->bIsNotTexture2D = false;
+						PinData->NodeTable = this;
 
-						FName PinCategory = DefaultImageMode == ETableTextureType::PASSTHROUGH_TEXTURE ? Schema->PC_PassThroughImage : Schema->PC_Image;
+						FName PinCategory = PinData->ImageMode == ETableTextureType::PASSTHROUGH_TEXTURE ? Schema->PC_PassThroughImage : Schema->PC_Image;
 
 						for (UEdGraphPin* Pin : OldPins)
 						{
@@ -518,7 +478,8 @@ void UCustomizableObjectNodeTable::AllocateDefaultPins(UCustomizableObjectNodeRe
 						PinData->ColumnName = ColumnName;
 						PinData->StructColumnId = ColumnPropertyId;
 						PinData->ImageMode = ETableTextureType::PASSTHROUGH_TEXTURE;
-						PinData->SetIsNotTexture2D(true);
+						PinData->bIsNotTexture2D = true;
+						PinData->NodeTable = this;
 
 						OutPin = CustomCreatePin(EGPD_Output, Schema->PC_PassThroughImage, FName(*PinName), PinData);
 					}
@@ -874,6 +835,8 @@ FString UCustomizableObjectNodeTable::GetRefreshMessage() const
 // TODO(MTBL-1652): Move this to the Copy method of the NodePinData class which is more PinType specific (and the right plece to do this)
 void UCustomizableObjectNodeTable::RemapPinsData(const TMap<UEdGraphPin*, UEdGraphPin*>& PinsToRemap)
 {
+	Super::RemapPinsData(PinsToRemap);
+	
 	const UEdGraphSchema_CustomizableObject* Schema = GetDefault<UEdGraphSchema_CustomizableObject>();
 
 	for (const TTuple<UEdGraphPin*, UEdGraphPin*>& Pair : PinsToRemap)
@@ -914,8 +877,7 @@ void UCustomizableObjectNodeTable::RemapPinsData(const TMap<UEdGraphPin*, UEdGra
 			if (PinDataOldPin && PinDataNewPin)
 			{
 				PinDataNewPin->ImageMode = PinDataOldPin->ImageMode;
-				PinDataNewPin->SetDefaultImageMode(PinDataOldPin->IsDefaultImageMode());
-				PinDataNewPin->SetIsNotTexture2D(PinDataOldPin->IsNotTexture2D());
+				PinDataNewPin->bIsNotTexture2D = PinDataOldPin->bIsNotTexture2D;
 			}
 		}
 	}
@@ -1232,49 +1194,45 @@ TSoftClassPtr<UAnimInstance> UCustomizableObjectNodeTable::GetAnimInstanceAt(con
 }
 
 
-void UCustomizableObjectNodeTableImagePinData::ConvertArrayTextureToAnyTextureFixup()
+void UCustomizableObjectNodeTableImagePinData::PostLoad()
 {
-	if (bIsArrayTexture_DEPRECATED)
+	Super::PostLoad();
+
+	const int32 CustomizableObjectCustomVersion = GetLinkerCustomVersion(FCustomizableObjectCustomVersion::GUID);
+
+	if (CustomizableObjectCustomVersion < FCustomizableObjectCustomVersion::AddedAnyTextureTypeToPassThroughTextures)
 	{
-		bIsNotTexture2D = true;
-		bIsArrayTexture_DEPRECATED = false;
+		if (bIsArrayTexture_DEPRECATED)
+		{
+			bIsNotTexture2D = true;
+			bIsArrayTexture_DEPRECATED = false;
+		}
 	}
 }
 
 
-void UCustomizableObjectNodeTable::ChangeImagePinMode(UEdGraphPin* Pin, bool bSetDefault)
+void UCustomizableObjectNodeTableImagePinData::PostEditChangeProperty(FPropertyChangedEvent& PropertyChangedEvent)
 {
-	if (UCustomizableObjectNodeTableImagePinData* PinData = Cast<UCustomizableObjectNodeTableImagePinData>(GetPinData(*(Pin))))
+	Super::PostEditChangeProperty(PropertyChangedEvent);
+
+	if (const FProperty* PropertyThatChanged = PropertyChangedEvent.Property)
 	{
-		if (PinData->IsNotTexture2D())
+		if (PropertyThatChanged->GetFName() == GET_MEMBER_NAME_CHECKED(UCustomizableObjectNodeTableImagePinData, ImageMode))
 		{
-			return;
+			NodeTable->ReconstructNode();
 		}
+	}
+}
 
-		if (PinData->IsDefaultImageMode() && !bSetDefault)
-		{
-			ETableTextureType Mode = DefaultImageMode == ETableTextureType::PASSTHROUGH_TEXTURE ? ETableTextureType::MUTABLE_TEXTURE : ETableTextureType::PASSTHROUGH_TEXTURE;
-			PinData->ImageMode = Mode;
-			PinData->SetDefaultImageMode(false);
 
-			ReconstructNode();
-		}
-		else
-		{
-			if (bSetDefault)
-			{
-				PinData->ImageMode = DefaultImageMode;
-				PinData->SetDefaultImageMode(true);
-			}
-			else
-			{
-				ETableTextureType Mode = PinData->ImageMode == ETableTextureType::PASSTHROUGH_TEXTURE ? ETableTextureType::MUTABLE_TEXTURE : ETableTextureType::PASSTHROUGH_TEXTURE;
-				PinData->ImageMode = Mode;
-			}
-
-			ReconstructNode();
-		}
+bool UCustomizableObjectNodeTableImagePinData::CanEditChange(const FProperty* InProperty) const
+{
+	if (InProperty->GetFName() == GET_MEMBER_NAME_CHECKED(UCustomizableObjectNodeTableImagePinData, ImageMode))
+	{
+		return !NodeTable->GetPin(*this)->LinkedTo.Num() && !bIsNotTexture2D;
 	}
+	
+	return Super::CanEditChange(InProperty);
 }
 
 
@@ -1290,14 +1248,25 @@ bool UCustomizableObjectNodeTable::HasPinViewer() const
 }
 
 
-bool UCustomizableObjectNodeTable::IsImagePinDefault(const UEdGraphPin* Pin) const
+TSharedPtr<IDetailsView> UCustomizableObjectNodeTable::CustomizePinDetails(const UEdGraphPin& Pin) const
 {
-	if (UCustomizableObjectNodeTableImagePinData* PinData = Cast<UCustomizableObjectNodeTableImagePinData>(GetPinData(*(Pin))))
+	if (UCustomizableObjectNodeTableImagePinData* PinData = Cast<UCustomizableObjectNodeTableImagePinData>(GetPinData(Pin)))
 	{
-		return PinData->IsDefaultImageMode();
-	}
+		FPropertyEditorModule& EditModule = FModuleManager::Get().GetModuleChecked<FPropertyEditorModule>("PropertyEditor");
 
-	return true;
+		FDetailsViewArgs DetailsViewArgs;
+		DetailsViewArgs.bAllowSearch = false;
+		DetailsViewArgs.bHideSelectionTip = true;
+		
+		const TSharedRef<IDetailsView> SettingsView = EditModule.CreateDetailView(DetailsViewArgs);
+		SettingsView->SetObject(PinData);
+		
+		return SettingsView;
+	}
+	else
+	{
+		return nullptr;
+	}
 }
 
 
@@ -1314,18 +1283,8 @@ ETableTextureType UCustomizableObjectNodeTable::GetColumnImageMode(const FString
 		}
 	}
 
-	return DefaultImageMode;
-}
-
-
-bool UCustomizableObjectNodeTable::IsNotTexture2DPin(const UEdGraphPin* Pin) const
-{
-	if (UCustomizableObjectNodeTableImagePinData* PinData = Cast<UCustomizableObjectNodeTableImagePinData>(GetPinData(*(Pin))))
-	{
-		return PinData->IsNotTexture2D();
-	}
-
-	return false;
+	unimplemented();
+	return {};
 }
 
 
diff --git a/Engine/Plugins/Experimental/Mutable/Source/CustomizableObjectEditor/Private/MuCOE/PinViewer/SPinViewerListRow.cpp b/Engine/Plugins/Experimental/Mutable/Source/CustomizableObjectEditor/Private/MuCOE/PinViewer/SPinViewerListRow.cpp
--- a/Engine/Plugins/Experimental/Mutable/Source/CustomizableObjectEditor/Private/MuCOE/PinViewer/SPinViewerListRow.cpp
+++ b/Engine/Plugins/Experimental/Mutable/Source/CustomizableObjectEditor/Private/MuCOE/PinViewer/SPinViewerListRow.cpp
@@ -8,6 +8,7 @@
 #include "ScopedTransaction.h"
 #include "Widgets/Input/SCheckBox.h"
 #include "Widgets/Views/SListView.h"
+#include "IDetailsView.h"
 
 class STableViewBase;
 class SWidget;
diff --git a/Engine/Plugins/Experimental/Mutable/Source/CustomizableObjectEditor/Public/MuCOE/Nodes/CustomizableObjectNode.h b/Engine/Plugins/Experimental/Mutable/Source/CustomizableObjectEditor/Public/MuCOE/Nodes/CustomizableObjectNode.h
--- a/Engine/Plugins/Experimental/Mutable/Source/CustomizableObjectEditor/Public/MuCOE/Nodes/CustomizableObjectNode.h
+++ b/Engine/Plugins/Experimental/Mutable/Source/CustomizableObjectEditor/Public/MuCOE/Nodes/CustomizableObjectNode.h
@@ -15,6 +15,7 @@ class UCustomizableObjectNodeObject;
 class UCustomizableObjectNodeRemapPins;
 class UCustomizableObjectNodeRemapPinsByName;
 class UCustomizableObjectNodeRemapPinsByPosition;
+class IDetailsView;
 struct FPropertyChangedEvent;
 
 /** Abstract base class for all pin data. */
@@ -234,7 +235,10 @@ public:
 	virtual bool HasPinViewer() const;
 
 	/** Returns pin custom details. Override if required. */
-	virtual TSharedPtr<SWidget> CustomizePinDetails(const UEdGraphPin& Pin) const;
+	virtual TSharedPtr<IDetailsView> CustomizePinDetails(const UEdGraphPin& Pin) const;
+
+	/** Given the pin data, returns its associated pin. */
+	UEdGraphPin* GetPin(const UCustomizableObjectNodePinData& PinData);
 	
 	/** See GetPinData(const FEdGraphPinReference&). */
 	UCustomizableObjectNodePinData* GetPinData(const UEdGraphPin &Pin) const;
diff --git a/Engine/Plugins/Experimental/Mutable/Source/CustomizableObjectEditor/Public/MuCOE/Nodes/CustomizableObjectNodeMaterial.h b/Engine/Plugins/Experimental/Mutable/Source/CustomizableObjectEditor/Public/MuCOE/Nodes/CustomizableObjectNodeMaterial.h
--- a/Engine/Plugins/Experimental/Mutable/Source/CustomizableObjectEditor/Public/MuCOE/Nodes/CustomizableObjectNodeMaterial.h
+++ b/Engine/Plugins/Experimental/Mutable/Source/CustomizableObjectEditor/Public/MuCOE/Nodes/CustomizableObjectNodeMaterial.h
@@ -131,7 +131,7 @@ public:
 	virtual void BreakExistingConnectionsPostConnection(UEdGraphPin* InputPin, UEdGraphPin* OutputPin) override;
 	virtual bool IsNodeOutDatedAndNeedsRefresh() override;
 	virtual FString GetRefreshMessage() const override;
-	virtual TSharedPtr<SWidget> CustomizePinDetails(const UEdGraphPin& Pin) const override;
+	virtual TSharedPtr<IDetailsView> CustomizePinDetails(const UEdGraphPin& Pin) const override;
 
 	
 	// UCustomizableObjectNodeMaterialBase interface
diff --git a/Engine/Plugins/Experimental/Mutable/Source/CustomizableObjectEditor/Public/MuCOE/Nodes/CustomizableObjectNodeTable.h b/Engine/Plugins/Experimental/Mutable/Source/CustomizableObjectEditor/Public/MuCOE/Nodes/CustomizableObjectNodeTable.h
--- a/Engine/Plugins/Experimental/Mutable/Source/CustomizableObjectEditor/Public/MuCOE/Nodes/CustomizableObjectNodeTable.h
+++ b/Engine/Plugins/Experimental/Mutable/Source/CustomizableObjectEditor/Public/MuCOE/Nodes/CustomizableObjectNodeTable.h
@@ -14,6 +14,7 @@ namespace ENodeTitleType { enum Type : int; }
 class UCustomizableObjectLayout;
 class UCustomizableObjectNodeRemapPins;
 class UCustomizableObjectNodeRemapPinsByName;
+class UCustomizableObjectNodeTable;
 class UEdGraphPin;
 class UObject;
 class USkeletalMesh;
@@ -99,28 +100,23 @@ class CUSTOMIZABLEOBJECTEDITOR_API UCustomizableObjectNodeTableImagePinData : pu
 	GENERATED_BODY()
 
 public:
+	// UObject interface
+	virtual void PostLoad() override;
+	virtual void PostEditChangeProperty(FPropertyChangedEvent& PropertyChangedEvent) override;
+	virtual bool CanEditChange(const FProperty* InProperty) const override;
 
-	bool IsDefaultImageMode() { return bIsDefault; }
-	void SetDefaultImageMode(bool bValue) { bIsDefault = bValue; }
-
-	bool IsNotTexture2D() { return bIsNotTexture2D; }
-	void SetIsNotTexture2D(bool bValue) { bIsNotTexture2D = bValue; }
-
-	void ConvertArrayTextureToAnyTextureFixup();
-
-	// Pin Type
-	UPROPERTY()
+	UPROPERTY(EditAnywhere, Category = NoCategory,  DisplayName = "Texture Parameter Mode")
 	ETableTextureType ImageMode = ETableTextureType::MUTABLE_TEXTURE;
 
-private:
-
 	UPROPERTY()
-	bool bIsDefault = true;
+	TObjectPtr<UCustomizableObjectNodeTable> NodeTable = nullptr;
 
+private:
 	// Replaced by the more general bIsNot2DTexture
 	UPROPERTY()
 	bool bIsArrayTexture_DEPRECATED = false;
 
+public:
 	UPROPERTY()
 	bool bIsNotTexture2D = false;
 };
@@ -210,13 +206,7 @@ public:
 
 	UPROPERTY(EditAnywhere, Category = TableProperties, meta = (EditCondition = "TableDataGatheringMode == ETableDataGatheringSource::ETDGM_AssetRegistry", EditConditionHides))
 	TArray<FName> FilterPaths;
-
-	/** Decides the default type of the texture pins (passtrhough or mutable)
-	*   Right click on a non-linked image pin to customize its image mode
-	*/
-	UPROPERTY(EditAnywhere, Category = TableProperties)
-	ETableTextureType DefaultImageMode = ETableTextureType::MUTABLE_TEXTURE;
-
+	
 	/** Name of the column that contains the Version options. */
 	UPROPERTY(EditAnywhere, Category = TableProperties)
 	FName VersionColumn;
@@ -248,7 +238,6 @@ public:
 	FText GetNodeTitle(ENodeTitleType::Type TitleType) const override;
 	FLinearColor GetNodeTitleColor() const override;
 	FText GetTooltipText() const override;
-	virtual void PinConnectionListChanged(UEdGraphPin* Pin) override;
 	virtual void OnRenameNode(const FString& NewName) override;
 	virtual bool GetCanRenameNode() const override { return true; }
 	
@@ -262,6 +251,7 @@ public:
 	virtual bool IsPinRelevant(const UEdGraphPin* Pin) const override;
 	UCustomizableObjectNodeTableRemapPins* CreateRemapPinsDefault() const;
 	virtual bool HasPinViewer() const override;
+	virtual TSharedPtr<IDetailsView> CustomizePinDetails(const UEdGraphPin& Pin) const override;
 	
 	/*** Allows to perform work when remapping the pin data. */
 	virtual void RemapPinsData(const TMap<UEdGraphPin*, UEdGraphPin*>& PinsToRemap) override;
@@ -369,17 +359,7 @@ public:
 	// We should do this in a template!
 	USkeletalMesh* GetSkeletalMeshAt(const UEdGraphPin* Pin, const UDataTable* DataTable, const FName& RowName) const;
 	TSoftClassPtr<UAnimInstance> GetAnimInstanceAt(const UEdGraphPin* Pin, const UDataTable* DataTable, const FName& RowName) const;
-
-	// Changes the image mode of a pin
-	// bSetDefault param: if true sets the pin to be equal to the default mode (same as node)
-	void ChangeImagePinMode(UEdGraphPin* Pin, bool bSetDefault = false);
-
-	// Returns true if the pin is in the default mode (same as node)
-	bool IsImagePinDefault(const UEdGraphPin* Pin) const;
-
-	// Returns true if the pin is not a UTexture2D but a different kind of UTexture
-	bool IsNotTexture2DPin(const UEdGraphPin* Pin) const;
-
+	
 	// Returns the image mode of the column
 	ETableTextureType GetColumnImageMode(const FString& ColumnName) const;
 
