Remove the now deleted option to EnableAutoRTFM, this is default and now youve to explicitly disable it for a module

diff --git a/Engine/Plugins/Experimental/LocalizableMessage/Source/LocalizableMessage/LocalizableMessage.Build.cs b/Engine/Plugins/Experimental/LocalizableMessage/Source/LocalizableMessage/LocalizableMessage.Build.cs
--- a/Engine/Plugins/Experimental/LocalizableMessage/Source/LocalizableMessage/LocalizableMessage.Build.cs
+++ b/Engine/Plugins/Experimental/LocalizableMessage/Source/LocalizableMessage/LocalizableMessage.Build.cs
@@ -16,8 +16,6 @@ namespace UnrealBuildTool.Rules
 				new string[] {
 					"StructUtils" // used in header
 				});
-
-			bAllowAutoRTFMInstrumentation = true;
 		}
 	}
 }
diff --git a/Engine/Plugins/Experimental/Mutable/Source/CustomizableObject/CustomizableObject.Build.cs b/Engine/Plugins/Experimental/Mutable/Source/CustomizableObject/CustomizableObject.Build.cs
--- a/Engine/Plugins/Experimental/Mutable/Source/CustomizableObject/CustomizableObject.Build.cs
+++ b/Engine/Plugins/Experimental/Mutable/Source/CustomizableObject/CustomizableObject.Build.cs
@@ -41,8 +41,6 @@ public class CustomizableObject : ModuleRules
 			"StructUtils"
 		});
 
-		bAllowAutoRTFMInstrumentation = true;
-
 		if (TargetRules.bBuildEditor == true)
         {
             PublicDependencyModuleNames.Add("UnrealEd");	// @todo api: Only public because of WITH_EDITOR
diff --git a/Engine/Plugins/Experimental/StructUtils/Source/StructUtils/StructUtils.Build.cs b/Engine/Plugins/Experimental/StructUtils/Source/StructUtils/StructUtils.Build.cs
--- a/Engine/Plugins/Experimental/StructUtils/Source/StructUtils/StructUtils.Build.cs
+++ b/Engine/Plugins/Experimental/StructUtils/Source/StructUtils/StructUtils.Build.cs
@@ -21,8 +21,6 @@ namespace UnrealBuildTool.Rules
 			{
 				PrivateDependencyModuleNames.Add("Engine");
 			}
-
-			bAllowAutoRTFMInstrumentation = true;
 		}
 	}
 }
diff --git a/Engine/Plugins/Experimental/UIFramework/Source/UIFramework.Build.cs b/Engine/Plugins/Experimental/UIFramework/Source/UIFramework.Build.cs
--- a/Engine/Plugins/Experimental/UIFramework/Source/UIFramework.Build.cs
+++ b/Engine/Plugins/Experimental/UIFramework/Source/UIFramework.Build.cs
@@ -35,8 +35,6 @@ public class UIFramework : ModuleRules
 			}
 		);
 
-		bAllowAutoRTFMInstrumentation = true;
-
 		SetupIrisSupport(Target);
 	}
 }
diff --git a/Engine/Plugins/Experimental/Water/Source/Runtime/Water.Build.cs b/Engine/Plugins/Experimental/Water/Source/Runtime/Water.Build.cs
--- a/Engine/Plugins/Experimental/Water/Source/Runtime/Water.Build.cs
+++ b/Engine/Plugins/Experimental/Water/Source/Runtime/Water.Build.cs
@@ -68,7 +68,5 @@ public class Water : ModuleRules
 		}
 		// Add a feature define instead of relying on the generic WITH_EDITOR define
 		PublicDefinitions.Add("WITH_WATER_SELECTION_SUPPORT=" + (bWithWaterSelectionSupport ? 1 : 0));
-
-		bAllowAutoRTFMInstrumentation = true;
 	}
 }
diff --git a/Engine/Plugins/Online/OnlineServices/Source/OnlineServicesInterface/OnlineServicesInterface.Build.cs b/Engine/Plugins/Online/OnlineServices/Source/OnlineServicesInterface/OnlineServicesInterface.Build.cs
--- a/Engine/Plugins/Online/OnlineServices/Source/OnlineServicesInterface/OnlineServicesInterface.Build.cs
+++ b/Engine/Plugins/Online/OnlineServices/Source/OnlineServicesInterface/OnlineServicesInterface.Build.cs
@@ -21,7 +21,5 @@ public class OnlineServicesInterface : ModuleRules
 				"Core",
 			}
 		);
-
-		bAllowAutoRTFMInstrumentation = true;
 	}
 }
diff --git a/Engine/Plugins/Online/OnlineSubsystemUtils/Source/OnlineSubsystemUtils/OnlineSubsystemUtils.Build.cs b/Engine/Plugins/Online/OnlineSubsystemUtils/Source/OnlineSubsystemUtils/OnlineSubsystemUtils.Build.cs
--- a/Engine/Plugins/Online/OnlineSubsystemUtils/Source/OnlineSubsystemUtils/OnlineSubsystemUtils.Build.cs
+++ b/Engine/Plugins/Online/OnlineSubsystemUtils/Source/OnlineSubsystemUtils/OnlineSubsystemUtils.Build.cs
@@ -66,7 +66,5 @@ public class OnlineSubsystemUtils : ModuleRules
 		}
 
 		PublicDependencyModuleNames.Add("OnlineSubsystem");
-
-		bAllowAutoRTFMInstrumentation = true;
 	}
 }
diff --git a/Engine/Plugins/Runtime/GameplayAbilities/Source/GameplayAbilities/GameplayAbilities.Build.cs b/Engine/Plugins/Runtime/GameplayAbilities/Source/GameplayAbilities/GameplayAbilities.Build.cs
--- a/Engine/Plugins/Runtime/GameplayAbilities/Source/GameplayAbilities/GameplayAbilities.Build.cs
+++ b/Engine/Plugins/Runtime/GameplayAbilities/Source/GameplayAbilities/GameplayAbilities.Build.cs
@@ -39,8 +39,6 @@ namespace UnrealBuildTool.Rules
 			SetupGameplayDebuggerSupport(Target);
 
 			SetupIrisSupport(Target);
-
-			bAllowAutoRTFMInstrumentation = true;
 		}
 	}
 }
diff --git a/Engine/Plugins/Runtime/ReplicationGraph/Source/ReplicationGraph.Build.cs b/Engine/Plugins/Runtime/ReplicationGraph/Source/ReplicationGraph.Build.cs
--- a/Engine/Plugins/Runtime/ReplicationGraph/Source/ReplicationGraph.Build.cs
+++ b/Engine/Plugins/Runtime/ReplicationGraph/Source/ReplicationGraph.Build.cs
@@ -19,8 +19,6 @@ public class ReplicationGraph : ModuleRules
 
 		UnsafeTypeCastWarningLevel = WarningLevel.Error;
 
-		bAllowAutoRTFMInstrumentation = true;
-
 		SetupGameplayDebuggerSupport(Target);
 	}
 }
diff --git a/Engine/Source/Programs/AutoRTFMTests/AutoRTFMTests.Build.cs b/Engine/Source/Programs/AutoRTFMTests/AutoRTFMTests.Build.cs
--- a/Engine/Source/Programs/AutoRTFMTests/AutoRTFMTests.Build.cs
+++ b/Engine/Source/Programs/AutoRTFMTests/AutoRTFMTests.Build.cs
@@ -25,8 +25,6 @@ public class AutoRTFMTests : ModuleRules
 			);
 		}
 
-		bAllowAutoRTFMInstrumentation = true;
-
 		PrivateIncludePaths.AddRange(
 			new string[]
 			{
diff --git a/Engine/Source/Runtime/AIModule/AIModule.Build.cs b/Engine/Source/Runtime/AIModule/AIModule.Build.cs
--- a/Engine/Source/Runtime/AIModule/AIModule.Build.cs
+++ b/Engine/Source/Runtime/AIModule/AIModule.Build.cs
@@ -50,7 +50,6 @@ namespace UnrealBuildTool.Rules
             }
 
             SetupGameplayDebuggerSupport(Target);
-            bAllowAutoRTFMInstrumentation = true;
         }
     }
 }
diff --git a/Engine/Source/Runtime/Core/Core.Build.cs b/Engine/Source/Runtime/Core/Core.Build.cs
--- a/Engine/Source/Runtime/Core/Core.Build.cs
+++ b/Engine/Source/Runtime/Core/Core.Build.cs
@@ -354,8 +354,6 @@ public class Core : ModuleRules
 
 		IWYUSupport = IWYUSupport.KeepAsIs;
 
-		bAllowAutoRTFMInstrumentation = true;
-
 		PrivateDefinitions.Add("UE_PERSISTENT_ALLOCATOR_RESERVE_SIZE=" + GetPersistentAllocatorReserveSize().ToString() + "ULL");
 
 		PrivateDefinitions.Add("UE_MERGED_MODULES=" + (Target.bMergeModules ? "1" : "0"));
diff --git a/Engine/Source/Runtime/CoreOnline/CoreOnline.Build.cs b/Engine/Source/Runtime/CoreOnline/CoreOnline.Build.cs
--- a/Engine/Source/Runtime/CoreOnline/CoreOnline.Build.cs
+++ b/Engine/Source/Runtime/CoreOnline/CoreOnline.Build.cs
@@ -22,8 +22,6 @@ public class CoreOnline : ModuleRules
 
 		PrivateDefinitions.Add("COREONLINE_PACKAGE=1");
 		PublicDefinitions.Add("PLATFORM_MAX_LOCAL_PLAYERS=" + GetPlatformMaxLocalPlayers(Target));
-
-		bAllowAutoRTFMInstrumentation = true;
 	}
 
 	protected virtual int GetPlatformMaxLocalPlayers(ReadOnlyTargetRules Target)
diff --git a/Engine/Source/Runtime/CorePreciseFP/CorePreciseFP.build.cs b/Engine/Source/Runtime/CorePreciseFP/CorePreciseFP.build.cs
--- a/Engine/Source/Runtime/CorePreciseFP/CorePreciseFP.build.cs
+++ b/Engine/Source/Runtime/CorePreciseFP/CorePreciseFP.build.cs
@@ -15,7 +15,5 @@ public class CorePreciseFP : ModuleRules
 				"Core",
 			}
 		);
-
-		bAllowAutoRTFMInstrumentation = true;
 	}
 }
diff --git a/Engine/Source/Runtime/CoreUObject/CoreUObject.Build.cs b/Engine/Source/Runtime/CoreUObject/CoreUObject.Build.cs
--- a/Engine/Source/Runtime/CoreUObject/CoreUObject.Build.cs
+++ b/Engine/Source/Runtime/CoreUObject/CoreUObject.Build.cs
@@ -63,7 +63,5 @@ public class CoreUObject : ModuleRules
 		}
 
 		PrivateDefinitions.Add("UE_DEFINE_LEGACY_MATH_CONSTANT_MACRO_NAMES=0");
-
-		bAllowAutoRTFMInstrumentation = true;
 	}
 }
diff --git a/Engine/Source/Runtime/Engine/Engine.Build.cs b/Engine/Source/Runtime/Engine/Engine.Build.cs
--- a/Engine/Source/Runtime/Engine/Engine.Build.cs
+++ b/Engine/Source/Runtime/Engine/Engine.Build.cs
@@ -480,7 +480,5 @@ public class Engine : ModuleRules
 		SetupIrisSupport(Target, bIrisAddAsPublicDepedency);
 
 		PrivateDefinitions.Add("UE_DEPRECATE_LEGACY_MATH_CONSTANT_MACRO_NAMES=1");
-
-		bAllowAutoRTFMInstrumentation = true;
 	}
 }
diff --git a/Engine/Source/Runtime/Experimental/Chaos/Chaos.Build.cs b/Engine/Source/Runtime/Experimental/Chaos/Chaos.Build.cs
--- a/Engine/Source/Runtime/Experimental/Chaos/Chaos.Build.cs
+++ b/Engine/Source/Runtime/Experimental/Chaos/Chaos.Build.cs
@@ -44,8 +44,6 @@ namespace UnrealBuildTool.Rules
 			
 			UnsafeTypeCastWarningLevel = WarningLevel.Error;
 
-			bAllowAutoRTFMInstrumentation = true;
-
 			StaticAnalyzerDisabledCheckers.Add("cplusplus.NewDeleteLeaks"); // To be reevalulated, believed to be invalid warnings.
 			StaticAnalyzerDisabledCheckers.Add("core.UndefinedBinaryOperatorResult"); // Invalid warning in mass property calculation.
 		}
diff --git a/Engine/Source/Runtime/Experimental/Iris/Core/IrisCore.Build.cs b/Engine/Source/Runtime/Experimental/Iris/Core/IrisCore.Build.cs
--- a/Engine/Source/Runtime/Experimental/Iris/Core/IrisCore.Build.cs
+++ b/Engine/Source/Runtime/Experimental/Iris/Core/IrisCore.Build.cs
@@ -30,8 +30,6 @@ namespace UnrealBuildTool.Rules
 				}
 			}
 
-			bAllowAutoRTFMInstrumentation = true;
-
 			UnsafeTypeCastWarningLevel = WarningLevel.Error;
 		}
 	}
diff --git a/Engine/Source/Runtime/GameplayTags/GameplayTags.Build.cs b/Engine/Source/Runtime/GameplayTags/GameplayTags.Build.cs
--- a/Engine/Source/Runtime/GameplayTags/GameplayTags.Build.cs
+++ b/Engine/Source/Runtime/GameplayTags/GameplayTags.Build.cs
@@ -36,8 +36,6 @@ namespace UnrealBuildTool.Rules
                 );
             }
 
-			bAllowAutoRTFMInstrumentation = true;
-
 			SetupIrisSupport(Target);
 		}
 	}
diff --git a/Engine/Source/Runtime/Json/Json.Build.cs b/Engine/Source/Runtime/Json/Json.Build.cs
--- a/Engine/Source/Runtime/Json/Json.Build.cs
+++ b/Engine/Source/Runtime/Json/Json.Build.cs
@@ -15,8 +15,6 @@ namespace UnrealBuildTool.Rules
 			); 
 
 			UnsafeTypeCastWarningLevel = WarningLevel.Error;
-
-			bAllowAutoRTFMInstrumentation = true;
 		}
 	}
 }
diff --git a/Engine/Source/Runtime/JsonUtilities/JsonUtilities.Build.cs b/Engine/Source/Runtime/JsonUtilities/JsonUtilities.Build.cs
--- a/Engine/Source/Runtime/JsonUtilities/JsonUtilities.Build.cs
+++ b/Engine/Source/Runtime/JsonUtilities/JsonUtilities.Build.cs
@@ -14,7 +14,5 @@ public class JsonUtilities : ModuleRules
 				"Json",
 			}
 		);
-
-		bAllowAutoRTFMInstrumentation = true;
 	}
 }
diff --git a/Engine/Source/Runtime/LevelSequence/LevelSequence.Build.cs b/Engine/Source/Runtime/LevelSequence/LevelSequence.Build.cs
--- a/Engine/Source/Runtime/LevelSequence/LevelSequence.Build.cs
+++ b/Engine/Source/Runtime/LevelSequence/LevelSequence.Build.cs
@@ -43,7 +43,5 @@ public class LevelSequence : ModuleRules
 				}
 			);
 		}
-
-		bAllowAutoRTFMInstrumentation = true;
 	}
 }
diff --git a/Engine/Source/Runtime/MovieScene/MovieScene.Build.cs b/Engine/Source/Runtime/MovieScene/MovieScene.Build.cs
--- a/Engine/Source/Runtime/MovieScene/MovieScene.Build.cs
+++ b/Engine/Source/Runtime/MovieScene/MovieScene.Build.cs
@@ -26,7 +26,5 @@ public class MovieScene : ModuleRules
 			}
 		);
 		SetupIrisSupport(Target);
-
-		bAllowAutoRTFMInstrumentation = true;
 	}
 }
diff --git a/Engine/Source/Runtime/MovieSceneTracks/MovieSceneTracks.Build.cs b/Engine/Source/Runtime/MovieSceneTracks/MovieSceneTracks.Build.cs
--- a/Engine/Source/Runtime/MovieSceneTracks/MovieSceneTracks.Build.cs
+++ b/Engine/Source/Runtime/MovieSceneTracks/MovieSceneTracks.Build.cs
@@ -40,7 +40,5 @@ public class MovieSceneTracks : ModuleRules
 					"UnrealEd"
 				});
 		}
-
-		bAllowAutoRTFMInstrumentation = true;
 	}
 }
diff --git a/Engine/Source/Runtime/NavigationSystem/NavigationSystem.Build.cs b/Engine/Source/Runtime/NavigationSystem/NavigationSystem.Build.cs
--- a/Engine/Source/Runtime/NavigationSystem/NavigationSystem.Build.cs
+++ b/Engine/Source/Runtime/NavigationSystem/NavigationSystem.Build.cs
@@ -76,8 +76,6 @@ namespace UnrealBuildTool.Rules
                 PublicDependencyModuleNames.Add("UnrealEd");
                 CircularlyReferencedDependentModules.Add("UnrealEd");
             }
-
-			bAllowAutoRTFMInstrumentation = true;
         }
     }
 }
diff --git a/Engine/Source/Runtime/Net/Core/NetCore.Build.cs b/Engine/Source/Runtime/Net/Core/NetCore.Build.cs
--- a/Engine/Source/Runtime/Net/Core/NetCore.Build.cs
+++ b/Engine/Source/Runtime/Net/Core/NetCore.Build.cs
@@ -19,7 +19,5 @@ public class NetCore : ModuleRules
 		);
 
 		UnsafeTypeCastWarningLevel = WarningLevel.Error;
-
-		bAllowAutoRTFMInstrumentation = true;
 	}
 }
diff --git a/Engine/Source/Runtime/PakFile/PakFile.Build.cs b/Engine/Source/Runtime/PakFile/PakFile.Build.cs
--- a/Engine/Source/Runtime/PakFile/PakFile.Build.cs
+++ b/Engine/Source/Runtime/PakFile/PakFile.Build.cs
@@ -11,7 +11,5 @@ public class PakFile : ModuleRules
 		PrivateDependencyModuleNames.Add("CoreUObject");
 		PrivateDependencyModuleNames.Add("TraceLog");
 		PublicDependencyModuleNames.Add("RSA");
-
-		bAllowAutoRTFMInstrumentation = true;
 	}
 }
diff --git a/Engine/Source/Runtime/PhysicsCore/PhysicsCore.Build.cs b/Engine/Source/Runtime/PhysicsCore/PhysicsCore.Build.cs
--- a/Engine/Source/Runtime/PhysicsCore/PhysicsCore.Build.cs
+++ b/Engine/Source/Runtime/PhysicsCore/PhysicsCore.Build.cs
@@ -27,7 +27,5 @@ public class PhysicsCore: ModuleRules
 		PublicDependencyModuleNames.Remove("PhysicsCore");
 
 		PrivateDefinitions.Add("CHAOS_INCLUDE_LEVEL_1=1");
-
-		bAllowAutoRTFMInstrumentation = true;
 	}
 }
diff --git a/Engine/Source/Runtime/Renderer/Renderer.Build.cs b/Engine/Source/Runtime/Renderer/Renderer.Build.cs
--- a/Engine/Source/Runtime/Renderer/Renderer.Build.cs
+++ b/Engine/Source/Runtime/Renderer/Renderer.Build.cs
@@ -57,7 +57,5 @@ public class Renderer : ModuleRules
 		DynamicallyLoadedModuleNames.AddRange(new string[] { "EyeTracker" });
 
 		RuntimeDependencies.Add("$(EngineDir)/Content/Renderer/TessellationTable.bin");
-
-		bAllowAutoRTFMInstrumentation = true;
 	}
 }
