[mutable] Get LOD Settings from Reference Skeletal Mesh Data Asset and apply them to generated and baked meshes.
#rnx
#tests Editor

diff --git a/Engine/Plugins/Experimental/Mutable/Source/CustomizableObject/Internal/MuCO/CustomizableObjectPrivate.h b/Engine/Plugins/Experimental/Mutable/Source/CustomizableObject/Internal/MuCO/CustomizableObjectPrivate.h
--- a/Engine/Plugins/Experimental/Mutable/Source/CustomizableObject/Internal/MuCO/CustomizableObjectPrivate.h
+++ b/Engine/Plugins/Experimental/Mutable/Source/CustomizableObject/Internal/MuCO/CustomizableObjectPrivate.h
@@ -32,6 +32,7 @@ class UTexture;
 class UAnimInstance;
 class UAssetUserData;
 class UCustomizableObject;
+class USkeletalMeshLODSettings;
 
 
 FGuid CUSTOMIZABLEOBJECT_API GenerateIdentifier(const UCustomizableObject& CustomizableObject);
@@ -385,6 +386,10 @@ struct CUSTOMIZABLEOBJECT_API FMutableRefSkeletalMeshData
 	UPROPERTY()
 	TSoftObjectPtr<USkeletalMesh> SoftSkeletalMesh;
 
+	// Optional USkeletalMeshLODSettings
+	UPROPERTY()
+	TObjectPtr<USkeletalMeshLODSettings> SkeletalMeshLODSettings;
+
 	// LOD info
 	UPROPERTY()
 	TArray<FMutableRefLODData> LODData;
@@ -1062,6 +1067,8 @@ public:
 
 		FixClipMeshWithMeshCrash,
 
+		SkeletalMeshLODSettingsSupport,
+
 		LastCustomizableObjectVersion
 	};
 	static constexpr int32 CurrentSupportedVersion = ECustomizableObjectVersions::LastCustomizableObjectVersion;
diff --git a/Engine/Plugins/Experimental/Mutable/Source/CustomizableObject/Private/MuCO/CustomizableObject.cpp b/Engine/Plugins/Experimental/Mutable/Source/CustomizableObject/Private/MuCO/CustomizableObject.cpp
--- a/Engine/Plugins/Experimental/Mutable/Source/CustomizableObject/Private/MuCO/CustomizableObject.cpp
+++ b/Engine/Plugins/Experimental/Mutable/Source/CustomizableObject/Private/MuCO/CustomizableObject.cpp
@@ -8,6 +8,7 @@
 #include "EdGraph/EdGraph.h"
 #include "Engine/Engine.h"
 #include "Engine/SkeletalMesh.h"
+#include "Engine/SkeletalMeshLODSettings.h"
 #include "Animation/AnimInstance.h"
 #include "Animation/Skeleton.h"
 #include "Engine/AssetUserData.h"
@@ -2698,6 +2699,9 @@ FArchive& operator<<(FArchive& Ar, FMutableRefSkeletalMeshData& Data)
 		FString AssetPath = Data.SoftSkeletalMesh.ToString();
 		Ar << AssetPath;
 
+		AssetPath = TSoftObjectPtr<USkeletalMeshLODSettings>(Data.SkeletalMeshLODSettings).ToString();
+		Ar << AssetPath;
+
 		AssetPath = TSoftObjectPtr<USkeleton>(Data.Skeleton).ToString();
 		Ar << AssetPath;
 
@@ -2717,6 +2721,10 @@ FArchive& operator<<(FArchive& Ar, FMutableRefSkeletalMeshData& Data)
 		Ar << SkeletalMeshAssetPath;
 		Data.SoftSkeletalMesh = SkeletalMeshAssetPath;
 
+		FString SkeletalMeshLODSettingsAssetPath;
+		Ar << SkeletalMeshLODSettingsAssetPath;
+		Data.SkeletalMeshLODSettings = TSoftObjectPtr<USkeletalMeshLODSettings>(SkeletalMeshLODSettingsAssetPath).LoadSynchronous();
+
 		FString SkeletonAssetPath;
 		Ar << SkeletonAssetPath;
 		Data.Skeleton = TSoftObjectPtr<USkeleton>(SkeletonAssetPath).LoadSynchronous();
diff --git a/Engine/Plugins/Experimental/Mutable/Source/CustomizableObject/Private/MuCO/CustomizableObjectInstance.cpp b/Engine/Plugins/Experimental/Mutable/Source/CustomizableObject/Private/MuCO/CustomizableObjectInstance.cpp
--- a/Engine/Plugins/Experimental/Mutable/Source/CustomizableObject/Private/MuCO/CustomizableObjectInstance.cpp
+++ b/Engine/Plugins/Experimental/Mutable/Source/CustomizableObject/Private/MuCO/CustomizableObjectInstance.cpp
@@ -3590,6 +3590,16 @@ void UCustomizableInstancePrivate::InitSkeletalMeshData(const TSharedRef<FUpdate
 		}
 	}
 
+	if (RefSkeletalMeshData.SkeletalMeshLODSettings)
+	{
+#if WITH_EDITORONLY_DATA
+		SkeletalMesh->SetLODSettings(RefSkeletalMeshData.SkeletalMeshLODSettings);
+#else
+		// This is the part from the above SkeletalMesh->SetLODSettings that's available in-game
+		RefSkeletalMeshData.SkeletalMeshLODSettings->SetLODSettingsToMesh(SkeletalMesh);
+#endif
+	}
+
 	// Set up unreal's default material, will be replaced when building materials
 	{
 		MUTABLE_CPUPROFILER_SCOPE(InitSkeletalMesh_AddDefaultMaterial);
diff --git a/Engine/Plugins/Experimental/Mutable/Source/CustomizableObjectEditor/Private/MuCOE/CustomizableObjectInstanceBakingUtils.cpp b/Engine/Plugins/Experimental/Mutable/Source/CustomizableObjectEditor/Private/MuCOE/CustomizableObjectInstanceBakingUtils.cpp
--- a/Engine/Plugins/Experimental/Mutable/Source/CustomizableObjectEditor/Private/MuCOE/CustomizableObjectInstanceBakingUtils.cpp
+++ b/Engine/Plugins/Experimental/Mutable/Source/CustomizableObjectEditor/Private/MuCOE/CustomizableObjectInstanceBakingUtils.cpp
@@ -870,6 +870,15 @@ bool BakeCustomizableObjectInstance(
 				}
 			}
 
+			// Copy LODSettings from the Reference Skeletal Mesh
+			{
+				if (InstanceCO && InstanceCO->GetPrivate()->GetModelResources().ReferenceSkeletalMeshesData.IsValidIndex(ComponentIndex))
+				{
+					USkeletalMeshLODSettings* LODSettings = InstanceCO->GetPrivate()->GetModelResources().ReferenceSkeletalMeshesData[ComponentIndex].SkeletalMeshLODSettings;
+					SkeletalMesh->SetLODSettings(LODSettings);
+				}
+			}
+
 			// Generate render data
 			SkeletalMesh->Build();
 		}
diff --git a/Engine/Plugins/Experimental/Mutable/Source/CustomizableObjectEditor/Private/MuCOE/GenerateMutableSource/GenerateMutableSource.cpp b/Engine/Plugins/Experimental/Mutable/Source/CustomizableObjectEditor/Private/MuCOE/GenerateMutableSource/GenerateMutableSource.cpp
--- a/Engine/Plugins/Experimental/Mutable/Source/CustomizableObjectEditor/Private/MuCOE/GenerateMutableSource/GenerateMutableSource.cpp
+++ b/Engine/Plugins/Experimental/Mutable/Source/CustomizableObjectEditor/Private/MuCOE/GenerateMutableSource/GenerateMutableSource.cpp
@@ -1634,6 +1634,9 @@ void PopulateReferenceSkeletalMeshesData(FMutableGraphGenerationContext& Generat
 		Data.SkeletalMesh = TObjectPtr<USkeletalMesh>(RefSkeletalMesh);
 		Data.SoftSkeletalMesh = RefSkeletalMesh;
 
+		// Set the optional USkeletalMeshLODSettings that will be applied to the generated transient meshes or the baked meshes
+		Data.SkeletalMeshLODSettings = RefSkeletalMesh->GetLODSettings();
+
 		// Gather LODData, this may include per LOD settings such as render data config or LODDataInfoArray
 		Data.LODData.AddDefaulted(GenerationContext.NumLODsInRoot);
 		
