[Mutable] Add sink optimization to ASTOpMeshApplyPose operations against ASTOpMeshAddTags.

- Allow ASTOpMeshTags operations to float past ApplyPose.

#rnx

diff --git a/Engine/Plugins/Experimental/Mutable/Source/MutableTools/Private/MuT/ASTOpMeshApplyPose.cpp b/Engine/Plugins/Experimental/Mutable/Source/MutableTools/Private/MuT/ASTOpMeshApplyPose.cpp
--- a/Engine/Plugins/Experimental/Mutable/Source/MutableTools/Private/MuT/ASTOpMeshApplyPose.cpp
+++ b/Engine/Plugins/Experimental/Mutable/Source/MutableTools/Private/MuT/ASTOpMeshApplyPose.cpp
@@ -1,13 +1,13 @@
 // Copyright Epic Games, Inc. All Rights Reserved.
 
 #include "MuT/ASTOpMeshApplyPose.h"
+#include "MuT/ASTOpMeshAddTags.h"
 
 #include "HAL/PlatformMath.h"
 #include "MuR/ModelPrivate.h"
 #include "MuR/RefCounted.h"
 #include "MuR/Types.h"
 
-
 namespace mu
 {
 
@@ -83,4 +83,39 @@ namespace mu
 
 	}
 
-}
\ No newline at end of file
+	Ptr<ASTOp> ASTOpMeshApplyPose::OptimiseSink(const FModelOptimizationOptions&, FOptimizeSinkContext&) const
+	{
+		Ptr<ASTOp> NewOp;
+
+		Ptr<ASTOp> MeshAt = base.child();
+
+		if (!MeshAt)
+		{
+			return nullptr;
+		}
+
+		OP_TYPE MeshType = MeshAt->GetOpType();
+
+		switch (MeshType)
+		{
+			case OP_TYPE::ME_ADDTAGS:
+			{
+				Ptr<ASTOpMeshAddTags> New = mu::Clone<ASTOpMeshAddTags>(MeshAt);
+				if (New->Source)
+				{
+					Ptr<ASTOpMeshApplyPose> NewApplyPose = mu::Clone<ASTOpMeshApplyPose>(this);
+					NewApplyPose->base = New->Source.child();
+					New->Source = NewApplyPose;
+				}
+
+				NewOp = New;
+				break;
+			}
+
+			default:
+				break;
+		}
+
+		return NewOp;
+	}
+}
diff --git a/Engine/Plugins/Experimental/Mutable/Source/MutableTools/Private/MuT/ASTOpMeshApplyPose.h b/Engine/Plugins/Experimental/Mutable/Source/MutableTools/Private/MuT/ASTOpMeshApplyPose.h
--- a/Engine/Plugins/Experimental/Mutable/Source/MutableTools/Private/MuT/ASTOpMeshApplyPose.h
+++ b/Engine/Plugins/Experimental/Mutable/Source/MutableTools/Private/MuT/ASTOpMeshApplyPose.h
@@ -35,6 +35,7 @@ struct FProgram;
 		Ptr<ASTOp> Clone(MapChildFuncRef mapChild) const override;
 		void ForEachChild(const TFunctionRef<void(ASTChild&)>) override;
 		void Link(FProgram& program, FLinkerOptions* Options) override;
+		Ptr<ASTOp> OptimiseSink(const FModelOptimizationOptions&, FOptimizeSinkContext&) const override;
 	};
 
 
