[Mutable] Changed editor show name to normal combine layer operation to better describe what it does.

- Renamed form COMBINE to BLEND NORMALS.

#rnx

diff --git a/Engine/Plugins/Experimental/Mutable/Source/CustomizableObjectEditor/Public/MuCOE/Nodes/CustomizableObjectNodeTextureLayer.h b/Engine/Plugins/Experimental/Mutable/Source/CustomizableObjectEditor/Public/MuCOE/Nodes/CustomizableObjectNodeTextureLayer.h
--- a/Engine/Plugins/Experimental/Mutable/Source/CustomizableObjectEditor/Public/MuCOE/Nodes/CustomizableObjectNodeTextureLayer.h
+++ b/Engine/Plugins/Experimental/Mutable/Source/CustomizableObjectEditor/Public/MuCOE/Nodes/CustomizableObjectNodeTextureLayer.h
@@ -27,7 +27,7 @@ enum ECustomizableObjectTextureLayerEffect : int
 	COTLE_SCREEN				UMETA(DisplayName = "SCREEN"),
 	COTLE_OVERLAY				UMETA(DisplayName = "OVERLAY"),
 	COTLE_ALPHA_OVERLAY			UMETA(DisplayName = "LIGHTEN"),
-	COTLE_NORMAL_COMBINE		UMETA(DisplayName = "COMBINE")
+	COTLE_NORMAL_COMBINE		UMETA(DisplayName = "BLEND NORMALS")
 };
 
 
