[mutable] Replace custom curves with Unreal's.
#rnx

diff --git a/Engine/Plugins/Experimental/Mutable/Source/CustomizableObject/Internal/MuCO/CustomizableObjectPrivate.h b/Engine/Plugins/Experimental/Mutable/Source/CustomizableObject/Internal/MuCO/CustomizableObjectPrivate.h
--- a/Engine/Plugins/Experimental/Mutable/Source/CustomizableObject/Internal/MuCO/CustomizableObjectPrivate.h
+++ b/Engine/Plugins/Experimental/Mutable/Source/CustomizableObject/Internal/MuCO/CustomizableObjectPrivate.h
@@ -1069,6 +1069,8 @@ public:
 
 		SkeletalMeshLODSettingsSupport,
 
+		RemoveCustomCurve,
+
 		LastCustomizableObjectVersion
 	};
 	static constexpr int32 CurrentSupportedVersion = ECustomizableObjectVersions::LastCustomizableObjectVersion;
diff --git a/Engine/Plugins/Experimental/Mutable/Source/CustomizableObjectEditor/Private/MuCOE/GenerateMutableSource/GenerateMutableSourceFloat.cpp b/Engine/Plugins/Experimental/Mutable/Source/CustomizableObjectEditor/Private/MuCOE/GenerateMutableSource/GenerateMutableSourceFloat.cpp
--- a/Engine/Plugins/Experimental/Mutable/Source/CustomizableObjectEditor/Private/MuCOE/GenerateMutableSource/GenerateMutableSourceFloat.cpp
+++ b/Engine/Plugins/Experimental/Mutable/Source/CustomizableObjectEditor/Private/MuCOE/GenerateMutableSource/GenerateMutableSourceFloat.cpp
@@ -26,36 +26,6 @@
 #define LOCTEXT_NAMESPACE "CustomizableObjectEditor"
 
 
-void SetCurveNodeValues(mu::NodeScalarCurvePtr& CurveNode, const FRichCurve& Curve)
-{
-	if (Curve.HasAnyData())
-	{
-		CurveNode->SetDefaultValue(Curve.GetDefaultValue());
-		CurveNode->SetKeyFrameCount(Curve.GetNumKeys());
-		int32 KeyNum = 0;
-
-		for (auto Itr = Curve.GetKeyIterator(); Itr; ++Itr)
-		{
-			CurveNode->SetKeyFrame(KeyNum,
-				Itr->Time,
-				Itr->Value,
-				Itr->ArriveTangent,
-				Itr->ArriveTangentWeight,
-				Itr->LeaveTangent,
-				Itr->LeaveTangentWeight,
-				Itr->InterpMode,
-				Itr->TangentMode,
-				Itr->TangentWeightMode);
-
-			KeyNum++;
-		}
-	}
-	else
-	{
-		CurveNode->SetDefaultValue(0.0f); // FRichCurve default DefaultValue is MAX_flt
-	}
-}
-
 
 mu::NodeScalarPtr GenerateMutableSourceFloat(const UEdGraphPin* Pin, FMutableGraphGenerationContext& GenerationContext)
 {
@@ -213,13 +183,12 @@ mu::NodeScalarPtr GenerateMutableSourceFloat(const UEdGraphPin* Pin, FMutableGra
 
 	else if (const UCustomizableObjectNodeCurve* TypedNodeCurve = Cast<UCustomizableObjectNodeCurve>(Node))
 	{
-		mu::NodeScalarCurvePtr CurveNode = new mu::NodeScalarCurve();
+		mu::Ptr<mu::NodeScalarCurve> CurveNode = new mu::NodeScalarCurve();
 		Result = CurveNode;
 
 		if (const UEdGraphPin* ConnectedPin = FollowInputPin(*TypedNodeCurve->InputPin()))
 		{
-			mu::NodeScalarPtr ScalarNode = GenerateMutableSourceFloat(ConnectedPin, GenerationContext);
-			CurveNode->SetT(ScalarNode);
+			CurveNode->CurveSampleValue = GenerateMutableSourceFloat(ConnectedPin, GenerationContext);
 		}
 
 		if (UCurveBase* CurveAsset = TypedNodeCurve->CurveAsset)
@@ -239,24 +208,21 @@ mu::NodeScalarPtr GenerateMutableSourceFloat(const UEdGraphPin* Pin, FMutableGra
 			{
 				if (PinIndex >= 0 && PinIndex <= 3)
 				{
-					const FRichCurve& Curve = CurveColor->FloatCurves[PinIndex];
-					SetCurveNodeValues(CurveNode, Curve);
+					CurveNode->Curve = CurveColor->FloatCurves[PinIndex];
 				}
 			}
 			else if (const UCurveVector* const CurveVector = Cast<UCurveVector>(CurveAsset))
 			{
 				if (PinIndex >= 0 && PinIndex <= 2)
 				{
-					const FRichCurve& Curve = CurveVector->FloatCurves[PinIndex];
-					SetCurveNodeValues(CurveNode, Curve);
+					CurveNode->Curve = CurveVector->FloatCurves[PinIndex];
 				}
 			}
 			else if (const UCurveFloat* const CurveFloat = Cast<UCurveFloat>(CurveAsset))
 			{
 				if (PinIndex == 0)
 				{
-					const FRichCurve& Curve = CurveFloat->FloatCurve;
-					SetCurveNodeValues(CurveNode, Curve);
+					CurveNode->Curve = CurveFloat->FloatCurve;
 				}
 			}
 		}
diff --git a/Engine/Plugins/Experimental/Mutable/Source/CustomizableObjectEditor/Private/MuCOE/SMutableCodeViewer.cpp b/Engine/Plugins/Experimental/Mutable/Source/CustomizableObjectEditor/Private/MuCOE/SMutableCodeViewer.cpp
--- a/Engine/Plugins/Experimental/Mutable/Source/CustomizableObjectEditor/Private/MuCOE/SMutableCodeViewer.cpp
+++ b/Engine/Plugins/Experimental/Mutable/Source/CustomizableObjectEditor/Private/MuCOE/SMutableCodeViewer.cpp
@@ -43,7 +43,6 @@ class FExtender;
 class FReferenceCollector;
 class FUICommandList;
 class SWidget;
-namespace mu { struct Curve; }
 namespace mu { struct FProjector; }
 namespace mu { struct FShape; }
 struct FGeometry;
@@ -2815,14 +2814,8 @@ void SMutableCodeViewer::PreviewMutableSkeleton(mu::SkeletonPtrConst Skeleton)
 }
 
 
-void SMutableCodeViewer::PreviewMutableCurve(const mu::Curve* Curve)
+void SMutableCodeViewer::PreviewMutableCurve(const FRichCurve& Curve)
 {
-	if (!Curve)
-	{
-		UE_LOG(LogTemp,Error,TEXT("Unable to preview data on null Curve pointer."))
-		return;
-	}
-	
 	if (!PreviewCurveViewer)
 	{
 		PreviewCurveViewer = SNew(SMutableCurveViewer);
@@ -2830,7 +2823,7 @@ void SMutableCodeViewer::PreviewMutableCurve(const mu::Curve* Curve)
 
 	PreviewBorder->SetContent(PreviewCurveViewer.ToSharedRef());
 	
-	PreviewCurveViewer->SetCurve(*Curve);
+	PreviewCurveViewer->SetCurve(Curve);
 }
 
 // TODO: Implement matrix viewer
diff --git a/Engine/Plugins/Experimental/Mutable/Source/CustomizableObjectEditor/Private/MuCOE/SMutableCodeViewer.h b/Engine/Plugins/Experimental/Mutable/Source/CustomizableObjectEditor/Private/MuCOE/SMutableCodeViewer.h
--- a/Engine/Plugins/Experimental/Mutable/Source/CustomizableObjectEditor/Private/MuCOE/SMutableCodeViewer.h
+++ b/Engine/Plugins/Experimental/Mutable/Source/CustomizableObjectEditor/Private/MuCOE/SMutableCodeViewer.h
@@ -38,7 +38,6 @@ class SMutableSkeletonViewer;
 class SMutableStringViewer;
 class STextComboBox;
 class SWidget;
-namespace mu { struct Curve; }
 namespace mu { struct FProjector; }
 namespace mu { struct FShape; }
 struct FGeometry;
@@ -796,7 +795,7 @@ public:
 	void PreviewMutableProjector(const mu::FProjector* Projector);
 	void PreviewMutableMatrix(const FMatrix44f& Mat);
 	void PreviewMutableShape(const mu::FShape* Shape);
-	void PreviewMutableCurve(const mu::Curve* Curve);
+	void PreviewMutableCurve(const FRichCurve& Curve);
 	
 private:
 	void PrepareStringViewer();
diff --git a/Engine/Plugins/Experimental/Mutable/Source/CustomizableObjectEditor/Private/MuCOE/SMutableConstantsWidget.cpp b/Engine/Plugins/Experimental/Mutable/Source/CustomizableObjectEditor/Private/MuCOE/SMutableConstantsWidget.cpp
--- a/Engine/Plugins/Experimental/Mutable/Source/CustomizableObjectEditor/Private/MuCOE/SMutableConstantsWidget.cpp
+++ b/Engine/Plugins/Experimental/Mutable/Source/CustomizableObjectEditor/Private/MuCOE/SMutableConstantsWidget.cpp
@@ -1312,7 +1312,7 @@ void SMutableConstantsWidget::LoadConstantShapes()
 void SMutableConstantsWidget::LoadConstantCurves()
 {
 	check (MutableProgramPtr);
-	const int32 ConstantsCount = MutableProgramPtr->m_constantCurves.Num();
+	const int32 ConstantsCount = MutableProgramPtr->ConstantCurves.Num();
 	ConstantCurveElements.Empty(ConstantsCount);
 	
 	mu::OutputMemoryStream Stream;
@@ -1321,10 +1321,10 @@ void SMutableConstantsWidget::LoadConstantCurves()
 	for (int32 CurveIndex = 0; CurveIndex < ConstantsCount; CurveIndex++)
 	{
 		TSharedPtr<FMutableConstantCurveElement> ConstantCurveElement = MakeShared<FMutableConstantCurveElement>();
-		ConstantCurveElement->Curve = &(MutableProgramPtr->m_constantCurves[CurveIndex]);
+		ConstantCurveElement->Curve = MutableProgramPtr->ConstantCurves[CurveIndex];
 		ConstantCurveElement->IndexOnSourceVector = CurveIndex;
 
-		ConstantCurveElement->Curve->Serialise(Archive);
+		Archive << ConstantCurveElement->Curve;
 		
 		ConstantCurveElements.Add(ConstantCurveElement);
 	}
diff --git a/Engine/Plugins/Experimental/Mutable/Source/CustomizableObjectEditor/Private/MuCOE/SMutableConstantsWidget.h b/Engine/Plugins/Experimental/Mutable/Source/CustomizableObjectEditor/Private/MuCOE/SMutableConstantsWidget.h
--- a/Engine/Plugins/Experimental/Mutable/Source/CustomizableObjectEditor/Private/MuCOE/SMutableConstantsWidget.h
+++ b/Engine/Plugins/Experimental/Mutable/Source/CustomizableObjectEditor/Private/MuCOE/SMutableConstantsWidget.h
@@ -17,7 +17,6 @@ class SExpandableArea;
 class SMutableCodeViewer;
 class STableViewBase;
 class SWidget;
-namespace mu { struct Curve; }
 namespace mu { struct FProgram; }
 namespace mu { struct FProjector; }
 namespace mu { struct FShape; }
@@ -100,7 +99,7 @@ struct FMutableConstantShapeElement : public FMutableConstantElement
 */ 
 struct FMutableConstantCurveElement : public FMutableConstantElement
 {
-	const mu::Curve* Curve = nullptr;
+	FRichCurve Curve;
 };
 
 
diff --git a/Engine/Plugins/Experimental/Mutable/Source/CustomizableObjectEditor/Private/MuCOE/SMutableCurveViewer.cpp b/Engine/Plugins/Experimental/Mutable/Source/CustomizableObjectEditor/Private/MuCOE/SMutableCurveViewer.cpp
--- a/Engine/Plugins/Experimental/Mutable/Source/CustomizableObjectEditor/Private/MuCOE/SMutableCurveViewer.cpp
+++ b/Engine/Plugins/Experimental/Mutable/Source/CustomizableObjectEditor/Private/MuCOE/SMutableCurveViewer.cpp
@@ -263,7 +263,7 @@ public:
 			return SNew(SHorizontalBox)+SHorizontalBox::Slot()
 			[
 				SNew(STextBlock).
-				Text(FText::FromString( FString::SanitizeFloat(RowItem->CurveKeyFrame.time) ))
+				Text(FText::FromString( FString::SanitizeFloat(RowItem->CurveKeyFrame.Time) ))
 			];
 		}
 		
@@ -273,7 +273,7 @@ public:
 			return SNew(SHorizontalBox)+SHorizontalBox::Slot()
 			[
 				SNew(STextBlock).
-				Text(FText::FromString( FString::SanitizeFloat(RowItem->CurveKeyFrame.value) ))
+				Text(FText::FromString( FString::SanitizeFloat(RowItem->CurveKeyFrame.Value) ))
 			];
 		}
 		
@@ -283,7 +283,7 @@ public:
 			return SNew(SHorizontalBox)+SHorizontalBox::Slot()
 			[
 				SNew(STextBlock).
-				Text(FText::FromString( FString::SanitizeFloat(RowItem->CurveKeyFrame.in_tangent) ))
+				Text(FText::FromString( FString::SanitizeFloat(RowItem->CurveKeyFrame.ArriveTangent) ))
 			];
 		}
 
@@ -293,7 +293,7 @@ public:
 			return SNew(SHorizontalBox)+SHorizontalBox::Slot()
 			[
 				SNew(STextBlock).
-				Text(FText::FromString( FString::SanitizeFloat(RowItem->CurveKeyFrame.in_tangent_weight) ))
+				Text(FText::FromString( FString::SanitizeFloat(RowItem->CurveKeyFrame.ArriveTangentWeight) ))
 			];
 		}
 
@@ -303,7 +303,7 @@ public:
 			return SNew(SHorizontalBox)+SHorizontalBox::Slot()
 			[
 				SNew(STextBlock).
-				Text(FText::FromString( FString::SanitizeFloat(RowItem->CurveKeyFrame.out_tangent) ))
+				Text(FText::FromString( FString::SanitizeFloat(RowItem->CurveKeyFrame.LeaveTangent) ))
 			];
 		}
 
@@ -313,43 +313,43 @@ public:
 			return SNew(SHorizontalBox)+SHorizontalBox::Slot()
 			[
 				SNew(STextBlock).
-				Text(FText::FromString( FString::SanitizeFloat(RowItem->CurveKeyFrame.out_tangent_weight) ))
+				Text(FText::FromString( FString::SanitizeFloat(RowItem->CurveKeyFrame.LeaveTangentWeight) ))
 			];
 		}
 
 		// interp_mode 
 		if (InColumnName == MutableCurveKeyFramesListColumns::KeyFrameInterpolationModeColumnID)
 		{
-			const uint8_t InterpolationMode = RowItem->CurveKeyFrame.interp_mode;
+			const uint8 InterpolationMode = RowItem->CurveKeyFrame.InterpMode;
 			
 			return SNew(SHorizontalBox)+SHorizontalBox::Slot()
 			[
 				SNew(STextBlock).
-				Text(FText::FromString(* FString(mu::TypeInfo::s_curveInterpolationModeName[InterpolationMode])))
+				Text(FText::FromString(* FString::Printf(TEXT("%d"),InterpolationMode)))
 			];
 		}
 
 		// tangent mode 
 		if (InColumnName == MutableCurveKeyFramesListColumns::KeyFrameTangentModeColumnID)
 		{
-			const uint8_t TangentMode = RowItem->CurveKeyFrame.tangent_mode;
+			const uint8 TangentMode = RowItem->CurveKeyFrame.TangentMode;
 			
 			return SNew(SHorizontalBox)+SHorizontalBox::Slot()
 			[
 				SNew(STextBlock).
-				Text(FText::FromString(* FString(mu::TypeInfo::s_curveTangentModeName[TangentMode])))
+					Text(FText::FromString(*FString::Printf(TEXT("%d"), TangentMode)))
 			];
 		}
 
 		// tangent weight mode 
 		if (InColumnName == MutableCurveKeyFramesListColumns::KeyFrameTangentWeightModeColumnID)
 		{
-			const uint8_t TangentWeightMode = RowItem->CurveKeyFrame.tangent_weight_mode;
+			const uint8 TangentWeightMode = RowItem->CurveKeyFrame.TangentWeightMode;
 			
 			return SNew(SHorizontalBox)+SHorizontalBox::Slot()
 			[
 				SNew(STextBlock).
-				Text(FText::FromString(* FString(mu::TypeInfo::s_curveTangentWeightModeName[TangentWeightMode])))
+					Text(FText::FromString(*FString::Printf(TEXT("%d"), TangentWeightMode)))
 			];
 		}
 		
@@ -443,7 +443,7 @@ void SMutableCurveViewer::Construct(const FArguments& InArgs)
 }
 
 
-void SMutableCurveViewer::SetCurve(const mu::Curve& InMutableCurve)
+void SMutableCurveViewer::SetCurve(const FRichCurve& InMutableCurve)
 {
 	this->MutableCurve = InMutableCurve;
 
@@ -457,11 +457,11 @@ void SMutableCurveViewer::SetCurve(const mu::Curve& InMutableCurve)
 
 void SMutableCurveViewer::SetupMutableCurveListView()
 {
-	const int32 KeyFrameCount = MutableCurve.keyFrames.Num() ;
+	const int32 KeyFrameCount = MutableCurve.Keys.Num() ;
 	CurveElements.SetNum(KeyFrameCount);
 	for	(int32 KeyFrameIndex = 0; KeyFrameIndex < KeyFrameCount; KeyFrameIndex++)
 	{
-		const mu::CurveKeyFrame& CurrentKeyFrame = MutableCurve.keyFrames[KeyFrameIndex];
+		const FRichCurveKey& CurrentKeyFrame = MutableCurve.Keys[KeyFrameIndex];
 
 		const TSharedPtr<FMutableCurveElement> NewCurveElement =
 			MakeShareable(new FMutableCurveElement(KeyFrameIndex,CurrentKeyFrame));
@@ -497,18 +497,18 @@ void SMutableCurveViewer::SetupMutableCurveGraph() const
 	float PreviousFrameTime =  TNumericLimits<float>::Lowest();
 	
 	// Fill the curve with data
-	const int32 KeyFrameCount = MutableCurve.keyFrames.Num() ;
+	const int32 KeyFrameCount = MutableCurve.Keys.Num() ;
 	for (int32 KeyframeIndex = 0; KeyframeIndex < KeyFrameCount; KeyframeIndex++)
 	{
 		// Load the mutable data
-		const mu::CurveKeyFrame& CurrentKeyframe = MutableCurve.keyFrames[KeyframeIndex];
+		const FRichCurveKey& CurrentKeyframe = MutableCurve.Keys[KeyframeIndex];
 		
 		// Setup basic data (x and y axis)
 		FKeyPosition KeyPosition;
 		{
 			// time and value
-			KeyPosition.InputValue = CurrentKeyframe.time;
-			KeyPosition.OutputValue = CurrentKeyframe.value;
+			KeyPosition.InputValue = CurrentKeyframe.Time;
+			KeyPosition.OutputValue = CurrentKeyframe.Value;
 		}
 		
 		/*
@@ -524,8 +524,8 @@ void SMutableCurveViewer::SetupMutableCurveGraph() const
 
 			// Shift the position of the current element if we are too close to the last processed element or if
 			// the previous element for whatever reason is now in front of us (previousTime > currentTime>.
-			if (FMath::IsNearlyEqual(PreviousFrameTime,CurrentKeyframe.time,ApplicableTimeDelta)
-				|| PreviousFrameTime > CurrentKeyframe.time)
+			if (FMath::IsNearlyEqual(PreviousFrameTime,CurrentKeyframe.Time,ApplicableTimeDelta)
+				|| PreviousFrameTime > CurrentKeyframe.Time)
 			{
 				KeyPosition.InputValue += ApplicableTimeDelta;
 			}
@@ -538,26 +538,25 @@ void SMutableCurveViewer::SetupMutableCurveGraph() const
 		FKeyAttributes KeyAttributes;
 		{
 			// tangents
-			KeyAttributes.SetArriveTangent(CurrentKeyframe.in_tangent);
-			KeyAttributes.SetLeaveTangent(CurrentKeyframe.out_tangent);
+			KeyAttributes.SetArriveTangent(CurrentKeyframe.ArriveTangent);
+			KeyAttributes.SetLeaveTangent(CurrentKeyframe.LeaveTangent);
 
 			// Weights
-			KeyAttributes.SetArriveTangentWeight(CurrentKeyframe.in_tangent_weight);
-			KeyAttributes.SetLeaveTangentWeight(CurrentKeyframe.out_tangent_weight);
+			KeyAttributes.SetArriveTangentWeight(CurrentKeyframe.ArriveTangentWeight);
+			KeyAttributes.SetLeaveTangentWeight(CurrentKeyframe.LeaveTangentWeight);
 			
 			// Interp mode
-			const TEnumAsByte UnrealInterpolationMode =
-				static_cast<ERichCurveInterpMode>(CurrentKeyframe.interp_mode);
+			const TEnumAsByte UnrealInterpolationMode = static_cast<ERichCurveInterpMode>(CurrentKeyframe.InterpMode);
 			KeyAttributes.SetInterpMode(UnrealInterpolationMode);
 			
 			// Tangent mode
 			const TEnumAsByte UnrealTangentMode =
-				static_cast<ERichCurveTangentMode>(CurrentKeyframe.tangent_mode);
+				static_cast<ERichCurveTangentMode>(CurrentKeyframe.TangentMode);
 			KeyAttributes.SetTangentMode(UnrealTangentMode);
 
 			// Tangent weight
 			const TEnumAsByte UnrealTangentWeightMode =
-				static_cast<ERichCurveTangentWeightMode>(CurrentKeyframe.tangent_weight_mode);
+				static_cast<ERichCurveTangentWeightMode>(CurrentKeyframe.TangentWeightMode);
 			KeyAttributes.SetTangentWeightMode(UnrealTangentWeightMode);
 		}
 
diff --git a/Engine/Plugins/Experimental/Mutable/Source/CustomizableObjectEditor/Private/MuCOE/SMutableCurveViewer.h b/Engine/Plugins/Experimental/Mutable/Source/CustomizableObjectEditor/Private/MuCOE/SMutableCurveViewer.h
--- a/Engine/Plugins/Experimental/Mutable/Source/CustomizableObjectEditor/Private/MuCOE/SMutableCurveViewer.h
+++ b/Engine/Plugins/Experimental/Mutable/Source/CustomizableObjectEditor/Private/MuCOE/SMutableCurveViewer.h
@@ -96,7 +96,7 @@ public:
 /** Defines an element representing a keyframe of the mutable curve. Used by a ListView */
 struct FMutableCurveElement
 {
-	FMutableCurveElement(const int32 InIndex,const  mu::CurveKeyFrame& InCurveKeyFrame)
+	FMutableCurveElement(const int32 InIndex,const  FRichCurveKey& InCurveKeyFrame)
 		: KeyFrameIndex(InIndex),
 		CurveKeyFrame(InCurveKeyFrame)
 	{
@@ -104,7 +104,7 @@ struct FMutableCurveElement
 	}
 	
 	const int32 KeyFrameIndex;
-	const mu::CurveKeyFrame& CurveKeyFrame;
+	const FRichCurveKey& CurveKeyFrame;
 };
 
 
@@ -114,7 +114,7 @@ struct FMutableCurveElement
 class SMutableCurveViewer final : public SCompoundWidget
 {
 	SLATE_BEGIN_ARGS(SMutableCurveViewer) {}
-		SLATE_ARGUMENT(mu::Curve, MutableCurve) ;
+		SLATE_ARGUMENT(FRichCurve, MutableCurve) ;
 	SLATE_END_ARGS()
 	
 public:
@@ -129,12 +129,12 @@ public:
 	 * @param InMutableCurve - Reference to a mutable curve object to be set as origin for the data displayed
 	 * on this slate object
 	 */
-	void SetCurve(const mu::Curve& InMutableCurve);
+	void SetCurve(const FRichCurve& InMutableCurve);
 	
 private:
 	
 	/** Mutable curve to read the data from. It is the origin for all the data displayed on this slate object */
-	mu::Curve MutableCurve;
+	FRichCurve MutableCurve;
 
 	/*
 	 * Curve graph view
diff --git a/Engine/Plugins/Experimental/Mutable/Source/CustomizableObjectEditor/Private/MuCOE/SMutableGraphViewer.cpp b/Engine/Plugins/Experimental/Mutable/Source/CustomizableObjectEditor/Private/MuCOE/SMutableGraphViewer.cpp
--- a/Engine/Plugins/Experimental/Mutable/Source/CustomizableObjectEditor/Private/MuCOE/SMutableGraphViewer.cpp
+++ b/Engine/Plugins/Experimental/Mutable/Source/CustomizableObjectEditor/Private/MuCOE/SMutableGraphViewer.cpp
@@ -694,8 +694,7 @@ void SMutableGraphViewer::GetChildrenForInfo(TSharedPtr<FMutableGraphTreeElement
 	else if (ParentNode->GetType() == mu::NodeScalarCurve::GetStaticType())
 	{
 		mu::NodeScalarCurve* ScalarCurveVar = StaticCast<mu::NodeScalarCurve*>(ParentNode);
-		mu::NodeScalarCurve::Private* Private = ScalarCurveVar->GetPrivate();
-		AddChildFunc(Private->m_input_scalar.get(), FString::Printf(TEXT("INPUT")));
+		AddChildFunc(ScalarCurveVar->CurveSampleValue.get(), FString::Printf(TEXT("INPUT")));
 	}
 
 	else if (ParentNode->GetType() == mu::NodeMeshMakeMorph::GetStaticType())
diff --git a/Engine/Plugins/Experimental/Mutable/Source/MutableRuntime/Internal/MuR/ModelPrivate.h b/Engine/Plugins/Experimental/Mutable/Source/MutableRuntime/Internal/MuR/ModelPrivate.h
--- a/Engine/Plugins/Experimental/Mutable/Source/MutableRuntime/Internal/MuR/ModelPrivate.h
+++ b/Engine/Plugins/Experimental/Mutable/Source/MutableRuntime/Internal/MuR/ModelPrivate.h
@@ -253,7 +253,7 @@ namespace mu
 		TArray<FShape> m_constantShapes;
 
         //! Constant curves
-		TArray<Curve> m_constantCurves;
+		TArray<FRichCurve> ConstantCurves;
 
         //! Constant skeletons
 		TArray<Ptr<const Skeleton>> m_constantSkeletons;
@@ -295,7 +295,7 @@ namespace mu
             arch << m_constantProjectors;
 			arch << m_constantMatrices;
 			arch << m_constantShapes;
-            arch << m_constantCurves;
+            arch << ConstantCurves;
             arch << m_constantSkeletons;
 			arch << m_constantPhysicsBodies;
             arch << m_parameters;
@@ -320,7 +320,7 @@ namespace mu
             arch >> m_constantProjectors;
 			arch >> m_constantMatrices;
 			arch >> m_constantShapes;
-            arch >> m_constantCurves;
+            arch >> ConstantCurves;
             arch >> m_constantSkeletons;
 			arch >> m_constantPhysicsBodies;
             arch >> m_parameters;
@@ -553,19 +553,9 @@ namespace mu
             return index;
         }
 
-        OP::ADDRESS AddConstant( const Curve& m )
+        OP::ADDRESS AddConstant( const FRichCurve& m )
         {
-            // Ensure unique
-//            for ( SIZE_T i=0; i<m_constantCurves.Num(); ++i)
-//            {
-//                if (m_constantCurves[i]==m)
-//                {
-//                    return (OP::ADDRESS)i;
-//                }
-//            }
-
-            OP::ADDRESS index = OP::ADDRESS( m_constantCurves.Num() );
-            m_constantCurves.Add( m );
+            OP::ADDRESS index = OP::ADDRESS( ConstantCurves.AddUnique( m ) );
             return index;
         }
 
diff --git a/Engine/Plugins/Experimental/Mutable/Source/MutableRuntime/Internal/MuR/ParametersPrivate.h b/Engine/Plugins/Experimental/Mutable/Source/MutableRuntime/Internal/MuR/ParametersPrivate.h
--- a/Engine/Plugins/Experimental/Mutable/Source/MutableRuntime/Internal/MuR/ParametersPrivate.h
+++ b/Engine/Plugins/Experimental/Mutable/Source/MutableRuntime/Internal/MuR/ParametersPrivate.h
@@ -248,128 +248,6 @@ namespace Private
 	};
 
 
-	//---------------------------------------------------------------------------------------------
-	//! Information about a curve interpolated from keyframes.
-	//---------------------------------------------------------------------------------------------
-	struct CurveKeyFrame
-	{
-		float time;
-		float value;
-		float in_tangent;
-		float in_tangent_weight;
-		float out_tangent;
-		float out_tangent_weight;
-
-		enum class InterpMode : uint8
-		{
-			Linear = 0,
-			Constant,
-			Cubic,
-			None,
-
-			//! Utility enumeration value, not really a Interp mode.
-			Count
-		};
-        uint8 interp_mode;
-
-
-		enum class TangentMode : uint8
-		{
-			Auto = 0,
-			Manual,
-			Broken,
-			None,
-
-			//! Utility enumeration value, not really a Tangent mode.
-			Count
-		};
-        uint8 tangent_mode;
-
-
-		enum class TangentWeightMode : uint8
-		{
-			Auto = 0,
-			Manual,
-			Broken,
-			None,
-
-			//! Utility enumeration value, not really a Tangent weight mode.
-			Count
-		};
-        uint8 tangent_weight_mode;
-
-        //!
-        bool operator==( const CurveKeyFrame& other ) const
-        {
-            return time==other.time &&
-                    value==other.value&&
-                    in_tangent==other.in_tangent&&
-                    in_tangent_weight==other.in_tangent_weight&&
-                    out_tangent==other.out_tangent&&
-                    out_tangent_weight==other.out_tangent_weight&&
-                    interp_mode==other.interp_mode&&
-                    tangent_mode==other.tangent_mode&&
-                    tangent_weight_mode==other.tangent_weight_mode;
-        }
-
-		//!
-		void Serialise(OutputArchive& arch) const
-		{
-			arch << time;
-			arch << value;
-			arch << in_tangent;
-			arch << in_tangent_weight;
-			arch << out_tangent;
-			arch << out_tangent_weight;
-
-			arch << interp_mode;
-			arch << tangent_mode;
-			arch << tangent_weight_mode;
-		}
-
-		//!
-		void Unserialise(InputArchive& arch)
-		{
-			arch >> time;
-			arch >> value;
-			arch >> in_tangent;
-			arch >> in_tangent_weight;
-			arch >> out_tangent;
-			arch >> out_tangent_weight;
-
-			arch >> interp_mode;
-			arch >> tangent_mode;
-			arch >> tangent_weight_mode;
-		}
-	};
-
-
-	struct Curve
-	{
-		TArray<CurveKeyFrame> keyFrames;
-		float defaultValue = 0.f;
-
-        //!
-        bool operator==( const Curve& other ) const
-        {
-            return keyFrames == other.keyFrames && defaultValue == other.defaultValue;
-        }
-
-		//!
-		void Serialise(OutputArchive& arch) const
-		{
-            arch << keyFrames;
-            arch << defaultValue;
-		}
-
-		//!
-		void Unserialise(InputArchive& arch)
-		{
-			arch >> keyFrames;
-            arch >> defaultValue;
-        }
-	};
-
 	using ParamBoolType = bool;
 	using ParamIntType = int32;
 	using ParamFloatType = float;
diff --git a/Engine/Plugins/Experimental/Mutable/Source/MutableRuntime/Internal/MuR/SerialisationPrivate.h b/Engine/Plugins/Experimental/Mutable/Source/MutableRuntime/Internal/MuR/SerialisationPrivate.h
--- a/Engine/Plugins/Experimental/Mutable/Source/MutableRuntime/Internal/MuR/SerialisationPrivate.h
+++ b/Engine/Plugins/Experimental/Mutable/Source/MutableRuntime/Internal/MuR/SerialisationPrivate.h
@@ -105,7 +105,7 @@ namespace mu
 		Arch << Num;                                                           \
 		if (Num)                                                               \
 		{                                                                      \
-			Arch.Stream->Write(&V[0], Num * sizeof(Type));					   \
+			Arch.Stream->Write(V.GetData(), Num * sizeof(Type));					   \
 		}                                                                      \
 	}                                                                          \
                                                                                \
@@ -117,7 +117,7 @@ namespace mu
 		V.SetNum(Num);                                                         \
 		if (Num)                                                               \
 		{                                                                      \
-			Arch.Stream->Read(&V[0], Num * sizeof(Type));					   \
+			Arch.Stream->Read(V.GetData(), Num * sizeof(Type));					   \
 		}                                                                      \
 	}                                                                          \
 
@@ -269,6 +269,7 @@ namespace mu
 	MUTABLE_IMPLEMENT_POD_VECTOR_SERIALISABLE(FVector2f);
 	MUTABLE_IMPLEMENT_POD_VECTOR_SERIALISABLE(FVector4f);
 	MUTABLE_IMPLEMENT_POD_VECTOR_SERIALISABLE(FMatrix44f);
+	MUTABLE_IMPLEMENT_POD_VECTOR_SERIALISABLE(FRichCurveKey);
 
 	
 	//---------------------------------------------------------------------------------------------
diff --git a/Engine/Plugins/Experimental/Mutable/Source/MutableRuntime/MutableRuntime.Build.cs b/Engine/Plugins/Experimental/Mutable/Source/MutableRuntime/MutableRuntime.Build.cs
--- a/Engine/Plugins/Experimental/Mutable/Source/MutableRuntime/MutableRuntime.Build.cs
+++ b/Engine/Plugins/Experimental/Mutable/Source/MutableRuntime/MutableRuntime.Build.cs
@@ -18,6 +18,7 @@ namespace UnrealBuildTool.Rules
 				new string[] {
 				"Core",
 				"CoreUObject",
+				"Engine",
 				}
 			);
 
diff --git a/Engine/Plugins/Experimental/Mutable/Source/MutableRuntime/Private/MuR/CodeRunner.cpp b/Engine/Plugins/Experimental/Mutable/Source/MutableRuntime/Private/MuR/CodeRunner.cpp
--- a/Engine/Plugins/Experimental/Mutable/Source/MutableRuntime/Private/MuR/CodeRunner.cpp
+++ b/Engine/Plugins/Experimental/Mutable/Source/MutableRuntime/Private/MuR/CodeRunner.cpp
@@ -21,7 +21,6 @@
 #include "MuR/MutableMath.h"
 #include "MuR/MutableString.h"
 #include "MuR/MutableTrace.h"
-#include "MuR/OpEvalCurve.h"
 #include "MuR/OpImageApplyComposite.h"
 #include "MuR/OpImageBinarise.h"
 #include "MuR/OpImageBlend.h"
@@ -5476,18 +5475,18 @@ namespace mu
             switch (item.Stage)
             {
             case 0:
-                    AddOp( FScheduledOp( item.At, item, 1),
-                           FScheduledOp( args.time, item) );
+                AddOp( FScheduledOp( item.At, item, 1),
+                        FScheduledOp( args.time, item) );
                 break;
 
             case 1:
             {
                 float time = LoadScalar( FCacheAddress(args.time,item) );
 
-                const Curve& curve = Program.m_constantCurves[args.curve];
-                float result = EvalCurve(curve, time);
+                const FRichCurve& Curve = Program.ConstantCurves[args.curve];
+                float Result = Curve.Eval(time);
 
-                StoreScalar( item, result );
+                StoreScalar( item, Result );
                 break;
             }
 
diff --git a/Engine/Plugins/Experimental/Mutable/Source/MutableRuntime/Private/MuR/OpEvalCurve.h b/Engine/Plugins/Experimental/Mutable/Source/MutableRuntime/Private/MuR/OpEvalCurve.h
deleted file mode 100644
diff --git a/Engine/Plugins/Experimental/Mutable/Source/MutableRuntime/Private/MuR/Serialisation.cpp b/Engine/Plugins/Experimental/Mutable/Source/MutableRuntime/Private/MuR/Serialisation.cpp
--- a/Engine/Plugins/Experimental/Mutable/Source/MutableRuntime/Private/MuR/Serialisation.cpp
+++ b/Engine/Plugins/Experimental/Mutable/Source/MutableRuntime/Private/MuR/Serialisation.cpp
@@ -35,6 +35,7 @@ namespace mu
 	MUTABLE_IMPLEMENT_POD_SERIALISABLE(FVector2f);
 	MUTABLE_IMPLEMENT_POD_SERIALISABLE(FVector4f);
 	MUTABLE_IMPLEMENT_POD_SERIALISABLE(FMatrix44f);
+	MUTABLE_IMPLEMENT_POD_SERIALISABLE(FRichCurveKey);
 
 	
     void operator<<(OutputArchive& arch, const FString& t)
@@ -53,6 +54,18 @@ namespace mu
     }
 
 
+	void operator<<(OutputArchive& arch, const FRichCurve& t)
+	{
+		arch << t.Keys;
+	}
+
+
+	void operator>>(InputArchive& arch, FRichCurve& t)
+	{
+		arch >> t.Keys;
+	}
+
+
     void operator<<(OutputArchive& arch, const FName& v)
     {
 	    arch << v.ToString();
diff --git a/Engine/Plugins/Experimental/Mutable/Source/MutableRuntime/Public/MuR/Serialisation.h b/Engine/Plugins/Experimental/Mutable/Source/MutableRuntime/Public/MuR/Serialisation.h
--- a/Engine/Plugins/Experimental/Mutable/Source/MutableRuntime/Public/MuR/Serialisation.h
+++ b/Engine/Plugins/Experimental/Mutable/Source/MutableRuntime/Public/MuR/Serialisation.h
@@ -10,6 +10,7 @@
 #include "MuR/RefCounted.h"
 #include "MuR/MutableMath.h"
 
+#include "Curves/RichCurve.h"
 #include "Math/Vector.h"
 #include "Math/IntVector.h"
 #include "Math/Vector4.h"
@@ -334,6 +335,7 @@ namespace mu
 	MUTABLE_DEFINE_POD_SERIALISABLE(FVector2f);
 	MUTABLE_DEFINE_POD_SERIALISABLE(FVector4f);
 	MUTABLE_DEFINE_POD_SERIALISABLE(FMatrix44f);
+	MUTABLE_DEFINE_POD_SERIALISABLE(FRichCurveKey);
 
 	MUTABLE_DEFINE_POD_SERIALISABLE(FGuid);
 
@@ -352,6 +354,7 @@ namespace mu
 	MUTABLE_DEFINE_POD_VECTOR_SERIALISABLE(FMatrix44f);
 	MUTABLE_DEFINE_POD_VECTOR_SERIALISABLE(FIntVector2);
 	MUTABLE_DEFINE_POD_VECTOR_SERIALISABLE(TCHAR);
+	MUTABLE_DEFINE_POD_VECTOR_SERIALISABLE(FRichCurveKey);
 
 	MUTABLE_DEFINE_POD_VECTOR_SERIALISABLE(FUintVector2);
 	MUTABLE_DEFINE_POD_VECTOR_SERIALISABLE(UE::Math::TIntVector2<uint16>);
@@ -363,6 +366,11 @@ namespace mu
 	
 	MUTABLERUNTIME_API void operator>>(InputArchive& arch, FString& t);
 
+	//---------------------------------------------------------------------------------------------
+	MUTABLERUNTIME_API void operator<<(OutputArchive& arch, const FRichCurve& t);
+
+	MUTABLERUNTIME_API void operator>>(InputArchive& arch, FRichCurve& t);
+
 
 	//---------------------------------------------------------------------------------------------
 	MUTABLERUNTIME_API void operator<<(OutputArchive& arch, const FName& v);
diff --git a/Engine/Plugins/Experimental/Mutable/Source/MutableTools/Internal/MuT/NodeScalarCurvePrivate.h b/Engine/Plugins/Experimental/Mutable/Source/MutableTools/Internal/MuT/NodeScalarCurvePrivate.h
--- a/Engine/Plugins/Experimental/Mutable/Source/MutableTools/Internal/MuT/NodeScalarCurvePrivate.h
+++ b/Engine/Plugins/Experimental/Mutable/Source/MutableTools/Internal/MuT/NodeScalarCurvePrivate.h
@@ -19,8 +19,6 @@ namespace mu
 
 		static FNodeType s_type;
 
-		NodeScalarPtr m_input_scalar;
-		Curve m_curve;
 	};
 
 
diff --git a/Engine/Plugins/Experimental/Mutable/Source/MutableTools/MutableTools.Build.cs b/Engine/Plugins/Experimental/Mutable/Source/MutableTools/MutableTools.Build.cs
--- a/Engine/Plugins/Experimental/Mutable/Source/MutableTools/MutableTools.Build.cs
+++ b/Engine/Plugins/Experimental/Mutable/Source/MutableTools/MutableTools.Build.cs
@@ -22,6 +22,7 @@ namespace UnrealBuildTool.Rules
 					"ImageCore",
 					"TextureCompressor",
 					"TextureBuildUtilities",
+					"Engine",
 				}
 			);
 		}
diff --git a/Engine/Plugins/Experimental/Mutable/Source/MutableTools/Private/MuT/ASTOpScalarCurve.cpp b/Engine/Plugins/Experimental/Mutable/Source/MutableTools/Private/MuT/ASTOpScalarCurve.cpp
--- a/Engine/Plugins/Experimental/Mutable/Source/MutableTools/Private/MuT/ASTOpScalarCurve.cpp
+++ b/Engine/Plugins/Experimental/Mutable/Source/MutableTools/Private/MuT/ASTOpScalarCurve.cpp
@@ -34,7 +34,7 @@ namespace mu
 	uint64 ASTOpScalarCurve::Hash() const
 	{
 		uint64 res = std::hash<uint64>()(size_t(OP_TYPE::SC_CURVE));
-		hash_combine(res, curve.keyFrames.Num());
+		hash_combine(res, Curve.Keys.Num());
 		return res;
 	}
 
@@ -44,7 +44,7 @@ namespace mu
 		if (otherUntyped.GetOpType() == GetOpType())
 		{
 			const ASTOpScalarCurve* other = static_cast<const ASTOpScalarCurve*>(&otherUntyped);
-			return time == other->time && curve == other->curve;
+			return time == other->time && Curve == other->Curve;
 		}
 		return false;
 	}
@@ -53,7 +53,7 @@ namespace mu
 	mu::Ptr<ASTOp> ASTOpScalarCurve::Clone(MapChildFuncRef mapChild) const
 	{
 		Ptr<ASTOpScalarCurve> n = new ASTOpScalarCurve();
-		n->curve = curve;
+		n->Curve = Curve;
 		n->time = mapChild(time.child());
 		return n;
 	}
@@ -66,7 +66,7 @@ namespace mu
 			OP::ScalarCurveArgs args;
 			memset(&args, 0, sizeof(args));
 			args.time = time ? time->linkedAddress : 0;
-			args.curve = program.AddConstant(curve);
+			args.curve = program.AddConstant(Curve);
 
 			linkedAddress = (OP::ADDRESS)program.m_opAddress.Num();
 			//program.m_code.push_back(op);
diff --git a/Engine/Plugins/Experimental/Mutable/Source/MutableTools/Private/MuT/ASTOpScalarCurve.h b/Engine/Plugins/Experimental/Mutable/Source/MutableTools/Private/MuT/ASTOpScalarCurve.h
--- a/Engine/Plugins/Experimental/Mutable/Source/MutableTools/Private/MuT/ASTOpScalarCurve.h
+++ b/Engine/Plugins/Experimental/Mutable/Source/MutableTools/Private/MuT/ASTOpScalarCurve.h
@@ -2,11 +2,12 @@
 
 #pragma once
 
-#include "HAL/Platform.h"
 #include "MuR/Operations.h"
 #include "MuR/ParametersPrivate.h"
 #include "MuR/Ptr.h"
 #include "MuT/AST.h"
+#include "Curves/RichCurve.h"
+#include "HAL/Platform.h"
 
 
 namespace mu
@@ -24,7 +25,7 @@ struct FProgram;
 		ASTChild time;
 
 		//!
-		Curve curve;
+		FRichCurve Curve;
 
 	public:
 
@@ -32,7 +33,7 @@ struct FProgram;
 		ASTOpScalarCurve(const ASTOpScalarCurve&) = delete;
 		~ASTOpScalarCurve() override;
 
-		OP_TYPE GetOpType() const override { return OP_TYPE::PR_CONSTANT; }
+		OP_TYPE GetOpType() const override { return OP_TYPE::SC_CURVE; }
 		uint64 Hash() const override;
 		void ForEachChild(const TFunctionRef<void(ASTChild&)>) override;
 		bool IsEqual(const ASTOp& otherUntyped) const override;
diff --git a/Engine/Plugins/Experimental/Mutable/Source/MutableTools/Private/MuT/CodeGenerator_Scalar.cpp b/Engine/Plugins/Experimental/Mutable/Source/MutableTools/Private/MuT/CodeGenerator_Scalar.cpp
--- a/Engine/Plugins/Experimental/Mutable/Source/MutableTools/Private/MuT/CodeGenerator_Scalar.cpp
+++ b/Engine/Plugins/Experimental/Mutable/Source/MutableTools/Private/MuT/CodeGenerator_Scalar.cpp
@@ -340,21 +340,19 @@ namespace mu
 	//-------------------------------------------------------------------------------------------------
 	void CodeGenerator::GenerateScalar_Curve(FScalarGenerationResult& result, const FGenericGenerationOptions& Options, const Ptr<const NodeScalarCurve>& Typed)
 	{
-		const NodeScalarCurve::Private& node = *Typed->GetPrivate();
-
 		Ptr<ASTOpScalarCurve> op = new ASTOpScalarCurve();
 
 		// T
-		if (Node* pA = node.m_input_scalar.get())
+		if (Node* pA = Typed->CurveSampleValue.get())
 		{
 			op->time = Generate(pA, Options);
 		}
 		else
 		{
-			op->time = CodeGenerator::GenerateMissingScalarCode(TEXT("Curve T"), 0.5f, node.m_errorContext);
+			op->time = CodeGenerator::GenerateMissingScalarCode(TEXT("Curve T"), 0.5f, Typed->GetPrivate()->m_errorContext);
 		}
 
-		op->curve = node.m_curve;
+		op->Curve = Typed->Curve;
 
 		result.op = op;
 	}
diff --git a/Engine/Plugins/Experimental/Mutable/Source/MutableTools/Private/MuT/NodeComponentEdit.cpp b/Engine/Plugins/Experimental/Mutable/Source/MutableTools/Private/MuT/NodeComponentEdit.cpp
--- a/Engine/Plugins/Experimental/Mutable/Source/MutableTools/Private/MuT/NodeComponentEdit.cpp
+++ b/Engine/Plugins/Experimental/Mutable/Source/MutableTools/Private/MuT/NodeComponentEdit.cpp
@@ -52,7 +52,7 @@ namespace mu
 
 
     //---------------------------------------------------------------------------------------------
-    void NodeComponentEdit::SetSurfaceCount( int num )
+    void NodeComponentEdit::SetSurfaceCount( int32 num )
     {
         check( num >=0 );
         m_pD->m_surfaces.SetNum( num );
@@ -60,7 +60,7 @@ namespace mu
 
 
     //---------------------------------------------------------------------------------------------
-    NodeSurface* NodeComponentEdit::GetSurface( int index ) const
+    NodeSurface* NodeComponentEdit::GetSurface( int32 index ) const
     {
         check( index >=0 && index < m_pD->m_surfaces.Num() );
 
@@ -69,7 +69,7 @@ namespace mu
 
 
     //---------------------------------------------------------------------------------------------
-    void NodeComponentEdit::SetSurface( int index, NodeSurface* pNode )
+    void NodeComponentEdit::SetSurface( int32 index, NodeSurface* pNode )
     {
         check( index >=0 && index < m_pD->m_surfaces.Num() );
 
diff --git a/Engine/Plugins/Experimental/Mutable/Source/MutableTools/Private/MuT/NodeScalarCurve.cpp b/Engine/Plugins/Experimental/Mutable/Source/MutableTools/Private/MuT/NodeScalarCurve.cpp
--- a/Engine/Plugins/Experimental/Mutable/Source/MutableTools/Private/MuT/NodeScalarCurve.cpp
+++ b/Engine/Plugins/Experimental/Mutable/Source/MutableTools/Private/MuT/NodeScalarCurve.cpp
@@ -13,79 +13,9 @@
 namespace mu
 {
 
-
-	//---------------------------------------------------------------------------------------------
-	// Static initialisation
-	//---------------------------------------------------------------------------------------------
     FNodeType NodeScalarCurve::Private::s_type = FNodeType( "Curve", NodeScalar::GetStaticType() );
 
-
-	//---------------------------------------------------------------------------------------------
-	//---------------------------------------------------------------------------------------------
-	//---------------------------------------------------------------------------------------------
-
     MUTABLE_IMPLEMENT_NODE( NodeScalarCurve, EType::Curve, Node, Node::EType::Scalar)
 
-
-	//---------------------------------------------------------------------------------------------
-	// Own Interface
-	//---------------------------------------------------------------------------------------------
-    int NodeScalarCurve::GetKeyFrameCount() const
-	{
-		return (int)m_pD->m_curve.keyFrames.Num();
-	}
-
-
-	//---------------------------------------------------------------------------------------------
-    void NodeScalarCurve::SetKeyFrameCount( int num )
-	{
-		check( num >=0 );
-		m_pD->m_curve.keyFrames.SetNum( num );
-	}
-
-	//---------------------------------------------------------------------------------------------
-    void NodeScalarCurve::SetDefaultValue(float defaultValue)
-	{
-		m_pD->m_curve.defaultValue = defaultValue;
-	}
-
-
-	//---------------------------------------------------------------------------------------------
-    void NodeScalarCurve::SetKeyFrame(
-            int index, float time, float value,
-            float in_tangent, float in_tangent_weight,
-            float out_tangent, float out_tangent_weight,
-            uint8_t interp_mode, uint8_t tangent_mode, uint8_t tangent_weight_mode )
-	{
-		check( index >=0 && index < m_pD->m_curve.keyFrames.Num() );
-
-		CurveKeyFrame& key_frame = m_pD->m_curve.keyFrames[index];
-
-		key_frame.time = time;
-		key_frame.value = value;
-		key_frame.in_tangent = in_tangent;
-		key_frame.in_tangent_weight = in_tangent_weight;
-		key_frame.out_tangent = out_tangent;
-		key_frame.out_tangent_weight = out_tangent_weight;
-		
-		key_frame.interp_mode = interp_mode;
-		key_frame.tangent_mode = tangent_mode;
-		key_frame.tangent_weight_mode = tangent_weight_mode;
-	}
-
-
-	//---------------------------------------------------------------------------------------------
-    NodeScalarPtr NodeScalarCurve::GetT() const
-	{
-		return m_pD->m_input_scalar.get();
-	}
-
-
-	//---------------------------------------------------------------------------------------------
-    void NodeScalarCurve::SetT(NodeScalarPtr pNode)
-	{
-		m_pD->m_input_scalar = pNode;
-	}
-
 }
 
diff --git a/Engine/Plugins/Experimental/Mutable/Source/MutableTools/Private/MuT/TypeInfo.cpp b/Engine/Plugins/Experimental/Mutable/Source/MutableTools/Private/MuT/TypeInfo.cpp
--- a/Engine/Plugins/Experimental/Mutable/Source/MutableTools/Private/MuT/TypeInfo.cpp
+++ b/Engine/Plugins/Experimental/Mutable/Source/MutableTools/Private/MuT/TypeInfo.cpp
@@ -133,30 +133,6 @@ namespace mu
 	};
 	static_assert(sizeof(TypeInfo::s_projectorTypeName) / sizeof(void*) == int32(PROJECTOR_TYPE::COUNT));
 
-	const char* TypeInfo::s_curveInterpolationModeName[static_cast<uint8_t>(mu::CurveKeyFrame::InterpMode::Count)] =
-	{
-		"Linear",
-		"Constant",
-		"Cubic",
-		"None",
-	};
-
-	const char* TypeInfo::s_curveTangentModeName[static_cast<uint8_t> (mu::CurveKeyFrame::TangentMode::Count)] =
-	{
-		"Auto",
-		"Manual",
-		"Broken",
-		"None",
-	};
-
-	const char* TypeInfo::s_curveTangentWeightModeName[static_cast<uint8_t> (mu::CurveKeyFrame::TangentWeightMode::Count)] =
-	{
-		"Auto",
-		"Manual",
-		"Broken",
-		"None",
-	};
-
 	const char* TypeInfo::s_blendModeName[uint32(EBlendType::_BT_COUNT)] =
 	{
 		"None",
@@ -168,7 +144,7 @@ namespace mu
 		"Overlay",
 		"Blend",
 		"Multiply",
-		"Lighten", // TODO: This name is not descriptive.
+		"Lighten",
 		"NormalCombine",
 	};
 	static_assert(sizeof(TypeInfo::s_blendModeName) / sizeof(void*) == int32(EBlendType::_BT_COUNT));
diff --git a/Engine/Plugins/Experimental/Mutable/Source/MutableTools/Public/MuT/NodeScalarCurve.h b/Engine/Plugins/Experimental/Mutable/Source/MutableTools/Public/MuT/NodeScalarCurve.h
--- a/Engine/Plugins/Experimental/Mutable/Source/MutableTools/Public/MuT/NodeScalarCurve.h
+++ b/Engine/Plugins/Experimental/Mutable/Source/MutableTools/Public/MuT/NodeScalarCurve.h
@@ -7,19 +7,19 @@
 #include "MuT/Node.h"
 #include "MuT/NodeScalar.h"
 
+#include "Curves/RichCurve.h"
+
 namespace mu
 {
 
-	// Forward definitions
-    class NodeScalarCurve;
-    typedef Ptr<NodeScalarCurve> NodeScalarCurvePtr;
-    typedef Ptr<const NodeScalarCurve> NodeScalarCurvePtrConst;
-
-
-    //! This node makes a new scalar value transforming another scalar value with a curve.
-	//! \ingroup model
+    /** This node makes a new scalar value transforming another scalar value with a curve. */
     class MUTABLETOOLS_API NodeScalarCurve : public NodeScalar
 	{
+	public:
+
+		FRichCurve Curve;
+		Ptr<NodeScalar> CurveSampleValue;
+
 	public:
 
         NodeScalarCurve();
@@ -35,26 +35,6 @@ namespace mu
 		// Own Interface
 		//-----------------------------------------------------------------------------------------
 
-        //! Get the number of key frames in the Curve.
-		int GetKeyFrameCount() const;
-
-        //! Set the number of key frames in the Curve.
-		void SetKeyFrameCount( int );
-
-		//! Set the default value of the curve
-		void SetDefaultValue(float);
-
-		//! \param index index of the curve, from 0 to GetKeyFrameCount()-1
-        void SetKeyFrame(int index, float time, float value,
-                         float in_tangent, float in_tangent_weight,
-                         float out_tangent, float out_tangent_weight,
-                         uint8_t interp_mode,
-                         uint8_t tangent_mode, uint8_t tangent_weight_mode);
-
-        //! Get the curve time parameter
-		NodeScalarPtr GetT() const;
-		void SetT(NodeScalarPtr);
-
 		//-----------------------------------------------------------------------------------------
 		// Interface pattern
 		//-----------------------------------------------------------------------------------------
diff --git a/Engine/Plugins/Experimental/Mutable/Source/MutableTools/Public/MuT/TypeInfo.h b/Engine/Plugins/Experimental/Mutable/Source/MutableTools/Public/MuT/TypeInfo.h
--- a/Engine/Plugins/Experimental/Mutable/Source/MutableTools/Public/MuT/TypeInfo.h
+++ b/Engine/Plugins/Experimental/Mutable/Source/MutableTools/Public/MuT/TypeInfo.h
@@ -28,11 +28,6 @@ namespace mu
 
 		static const char* s_projectorTypeName [ static_cast<uint32>(mu::PROJECTOR_TYPE::COUNT) ];
 
-		static const char* s_curveInterpolationModeName [ static_cast<uint32>(mu::CurveKeyFrame::InterpMode::Count) ];
-
-		static const char* s_curveTangentModeName [ static_cast<uint32> (mu::CurveKeyFrame::TangentMode::Count) ];
-
-		static const char* s_curveTangentWeightModeName [ static_cast<uint32> (mu::CurveKeyFrame::TangentWeightMode::Count) ];
 	};
 
 }
