[Mutable] Add unimplemented 16-bit skin weight format in MeshOptimizeBuffers. Enable 16-bit skin weights.

#rnx

diff --git a/Engine/Plugins/Experimental/Mutable/Source/CustomizableObject/Public/MuCO/CustomizableObject.h b/Engine/Plugins/Experimental/Mutable/Source/CustomizableObject/Public/MuCO/CustomizableObject.h
--- a/Engine/Plugins/Experimental/Mutable/Source/CustomizableObject/Public/MuCO/CustomizableObject.h
+++ b/Engine/Plugins/Experimental/Mutable/Source/CustomizableObject/Public/MuCO/CustomizableObject.h
@@ -431,8 +431,7 @@ public:
 	UPROPERTY(EditAnywhere, Category = CompileOptions)
 	bool bEnableClothing = false;
 
-	// TODO: Enable 16 bit weights 
-	UPROPERTY(VisibleAnywhere, Category = CompileOptions)
+	UPROPERTY(EditAnywhere, Category = CompileOptions)
 	bool bEnable16BitBoneWeights = false;
 
 	UPROPERTY(EditAnywhere, Category = CompileOptions)
diff --git a/Engine/Plugins/Experimental/Mutable/Source/MutableRuntime/Private/MuR/OpMeshFormat.cpp b/Engine/Plugins/Experimental/Mutable/Source/MutableRuntime/Private/MuR/OpMeshFormat.cpp
--- a/Engine/Plugins/Experimental/Mutable/Source/MutableRuntime/Private/MuR/OpMeshFormat.cpp
+++ b/Engine/Plugins/Experimental/Mutable/Source/MutableRuntime/Private/MuR/OpMeshFormat.cpp
@@ -602,6 +602,19 @@ namespace mu
 					}
 					break;
 				}
+				case MBF_NUINT16:
+				{
+					const uint16* Data = reinterpret_cast<const uint16*>(WeightIt.ptr());
+					for (int32 InfluenceIndex = 0; InfluenceIndex < BufferInfluences; ++InfluenceIndex)
+					{
+						if (*Data > 0)
+						{
+							++ThisVertexInfluences;
+						}
+						++Data;
+					}
+					break;
+				}
 
 				default:
 					// Unsupported
