[mutable] Fixed the baking of mutable instances when asking for the baking of all resources.
- Non mutable textures were being ignored when checking for textures in the UMaterialInstances.
- Removed now redundant searching for non-mutable textures in the base materials.

diff --git a/Engine/Plugins/Experimental/Mutable/Source/CustomizableObjectEditor/Private/MuCOE/CustomizableObjectInstanceBakingUtils.cpp b/Engine/Plugins/Experimental/Mutable/Source/CustomizableObjectEditor/Private/MuCOE/CustomizableObjectInstanceBakingUtils.cpp
--- a/Engine/Plugins/Experimental/Mutable/Source/CustomizableObjectEditor/Private/MuCOE/CustomizableObjectInstanceBakingUtils.cpp
+++ b/Engine/Plugins/Experimental/Mutable/Source/CustomizableObjectEditor/Private/MuCOE/CustomizableObjectInstanceBakingUtils.cpp
@@ -465,126 +465,88 @@ bool BakeCustomizableObjectInstance(
 			UObject* DuplicatedObject;
 			TArray<TMap<int, UTexture*>> TextureReplacementMaps;
 
-			// Duplicate Mutable generated textures
+			// Duplicate Textures found in the Material Instances of the SkeletalMesh so we can later assign them to the
+			// duplicates of those material instances. At the end of the baking we will have a series of materials with the 
+			// parameters set as the material instances they are based of.
 			for (int32 m = 0; m < Mesh->GetMaterials().Num(); ++m)
 			{
 				UMaterialInterface* Interface = Mesh->GetMaterials()[m].MaterialInterface;
 				Material = Interface->GetMaterial();
 				MaterialName = Material ? Material->GetName() : "Material";
-				Inst = Cast<UMaterialInstance>(Mesh->GetMaterials()[m].MaterialInterface);
+				Inst = Cast<UMaterialInstance>(Interface);
 
-				TMap<int, UTexture*> ReplacementTextures;
-				TextureReplacementMaps.Add(ReplacementTextures);
-
-				// The material will only have Mutable generated textures if it's actually a UMaterialInstance
+				TextureReplacementMaps.AddDefaulted();
+				
 				if (Material != nullptr && Inst != nullptr)
 				{
 					TArray<FName> ParameterNames = FUnrealBakeHelpers::GetTextureParameterNames(Material);
-
 					for (int32 i = 0; i < ParameterNames.Num(); i++)
 					{
 						if (Inst->GetTextureParameterValue(ParameterNames[i], Texture))
 						{
 							UTexture2D* SrcTex = Cast<UTexture2D>(Texture);
-							if (!SrcTex) continue;
+							if (!SrcTex)
+							{
+								continue;
+							}
 
-							bool bIsMutableTexture = false;
+							FString ParameterSanitized = ParameterNames[i].GetPlainNameString();
+							RemoveRestrictedChars(ParameterSanitized);
+							ResourceName = ObjectName + "_" + MaterialName + "_" + ParameterSanitized;
+							if (!GetUniqueResourceName(SrcTex, ResourceName, ArrayCachedObject, ArrayCachedElement))
+							{
+								continue;
+							}
+							
+							EPackageSaveResolutionType SaveType = EPackageSaveResolutionType::None;
+							if (!ManageBakingAction(AssetPath, ResourceName, bUsedGrantedOverridingRights, bIsUnattendedExecution, SaveType))
+							{
+								return false;
+							}
+
+							// Skip already processed resource
+							if (ArrayCachedElement.Find(ResourceName) != INDEX_NONE)
+							{
+								continue;
+							}
 
+							bool bIsMutableTexture = false;
 							for (UAssetUserData* UserData : *SrcTex->GetAssetUserDataArray())
 							{
-								UTextureMipDataProviderFactory* CustomMipDataProviderFactory = Cast<UMutableTextureMipDataProviderFactory>(UserData);
-								if (CustomMipDataProviderFactory)
+								if (Cast<UMutableTextureMipDataProviderFactory>(UserData))
 								{
 									bIsMutableTexture = true;
 								}
 							}
-
-							if ((SrcTex->GetPlatformData() != nullptr) &&
-								(SrcTex->GetPlatformData()->Mips.Num() > 0) &&
-								bIsMutableTexture)
+							
+							// Duplicating mutable generated textures
+							if (bIsMutableTexture)
 							{
-								FString ParameterSanitized = ParameterNames[i].GetPlainNameString();
-								RemoveRestrictedChars(ParameterSanitized);
-								ResourceName = ObjectName + "_" + MaterialName + "_" + ParameterSanitized;
-
-								if (!GetUniqueResourceName(SrcTex, ResourceName, ArrayCachedObject, ArrayCachedElement))
-								{
-									continue;
-								}
-
-								EPackageSaveResolutionType SaveType = EPackageSaveResolutionType::None;
-								if (!ManageBakingAction(AssetPath, ResourceName, bUsedGrantedOverridingRights, bIsUnattendedExecution, SaveType))
-								{
-									return false;
-								}
-
-								// Recover original name of the texture parameter value, now substituted by the generated Mutable texture
-								UTexture* OriginalTexture = nullptr;
-								UMaterialInstanceDynamic* InstDynamic = Cast<UMaterialInstanceDynamic>(Mesh->GetMaterials()[m].MaterialInterface);
-								if (InstDynamic != nullptr)
-								{
-									InstDynamic->Parent->GetTextureParameterValue(FName(*ParameterNames[i].GetPlainNameString()), OriginalTexture);
-								}
-								else
+								if (SrcTex->GetPlatformData() && SrcTex->GetPlatformData()->Mips.Num() > 0)
 								{
-									UMaterialInstanceConstant* InstConstant = Cast<UMaterialInstanceConstant>(Mesh->GetMaterials()[m].MaterialInterface);
+									// Recover original name of the texture parameter value, now substituted by the generated Mutable texture
+									UTexture* OriginalTexture = nullptr;
+									Inst->Parent->GetTextureParameterValue(FName(*ParameterNames[i].GetPlainNameString()), OriginalTexture);
+
+									PackageName = AssetPath + FString("/") + ResourceName;
+									TMap<UObject*, UObject*> FakeReplacementMap;
+									UTexture2D* DupTex = FUnrealBakeHelpers::BakeHelper_CreateAssetTexture(SrcTex, ResourceName, PackageName, OriginalTexture, true, FakeReplacementMap, bUsedGrantedOverridingRights);
+									ArrayCachedElement.Add(ResourceName);
+									ArrayCachedObject.Add(DupTex);
+								
+									TPair<EPackageSaveResolutionType, UPackage*> PackageToSave {SaveType, DupTex->GetPackage()};
+									OutSavedPackages.Add(PackageToSave);
 
-									if (InstConstant != nullptr)
+									if (OriginalTexture != nullptr)
 									{
-										InstConstant->Parent->GetTextureParameterValue(FName(*ParameterNames[i].GetPlainNameString()), OriginalTexture);
+										TextureReplacementMaps[m].Add(i, DupTex);
 									}
 								}
-
-								PackageName = AssetPath + FString("/") + ResourceName;
-								TMap<UObject*, UObject*> FakeReplacementMap;
-								UTexture2D* DupTex = FUnrealBakeHelpers::BakeHelper_CreateAssetTexture(SrcTex, ResourceName, PackageName, OriginalTexture, true, FakeReplacementMap, bUsedGrantedOverridingRights);
-								ArrayCachedElement.Add(ResourceName);
-								ArrayCachedObject.Add(DupTex);
-								
-								TPair<EPackageSaveResolutionType, UPackage*> PackageToSave {SaveType, DupTex->GetPackage()};
-								OutSavedPackages.Add(PackageToSave);
-
-								if (OriginalTexture != nullptr)
-								{
-									TextureReplacementMaps[m].Add(i, DupTex);
-								}
 							}
-						}
-					}
-				}
-			}
-
-			// Duplicate non-Mutable material textures
-			for (int32 m = 0; m < Mesh->GetMaterials().Num(); ++m)
-			{
-				UMaterialInterface* Interface = Mesh->GetMaterials()[m].MaterialInterface;
-				Material = Interface->GetMaterial();
-				MaterialName = Material ? Material->GetName() : "Material";
-
-				if (Material != nullptr)
-				{
-					TArray<FName> ParameterNames = FUnrealBakeHelpers::GetTextureParameterNames(Material);
-
-					for (int32 i = 0; i < ParameterNames.Num(); i++)
-					{
-						TArray<FMaterialParameterInfo> InfoArray;
-						TArray<FGuid> GuidArray;
-						Material->GetAllTextureParameterInfo(InfoArray, GuidArray);
-						
-						if (Material->GetTextureParameterValue(InfoArray[i], Texture))
-						{
-							FString ParameterSanitized = ParameterNames[i].GetPlainNameString();
-							RemoveRestrictedChars(ParameterSanitized);
-							ResourceName = ObjectName + "_" + MaterialName + "_" + ParameterSanitized;
-
-							if (ArrayCachedElement.Find(ResourceName) == INDEX_NONE)
+							else
 							{
-								EPackageSaveResolutionType SaveType = EPackageSaveResolutionType::None;
-								if (!ManageBakingAction(AssetPath, ResourceName, bUsedGrantedOverridingRights, bIsUnattendedExecution, SaveType ))
-								{
-									return false;
-								}
-
+								// Duplicate the non-mutable textures of the Material instance (pass-through textures)
+								
 								PackageName = AssetPath + FString("/") + ResourceName;
 								TMap<UObject*, UObject*> FakeReplacementMap;
 								DuplicatedObject = FUnrealBakeHelpers::BakeHelper_DuplicateAsset(Texture, ResourceName, PackageName, true, FakeReplacementMap, bUsedGrantedOverridingRights, false);
@@ -731,7 +693,6 @@ bool BakeCustomizableObjectInstance(
 										TPair<EPackageSaveResolutionType, UPackage*> TexturePackageToSave {TextureSaveType, DupTex->GetPackage()};
 										OutSavedPackages.Add(TexturePackageToSave);
 										
-
 										if (InstDynamic)
 										{
 											InstDynamic->SetTextureParameterValue(Inst->TextureParameterValues[TextureIndex].ParameterInfo.Name, DupTex);
diff --git a/Engine/Plugins/Experimental/Mutable/Source/CustomizableObjectEditor/Private/MuCOE/UnrealBakeHelpers.cpp b/Engine/Plugins/Experimental/Mutable/Source/CustomizableObjectEditor/Private/MuCOE/UnrealBakeHelpers.cpp
--- a/Engine/Plugins/Experimental/Mutable/Source/CustomizableObjectEditor/Private/MuCOE/UnrealBakeHelpers.cpp
+++ b/Engine/Plugins/Experimental/Mutable/Source/CustomizableObjectEditor/Private/MuCOE/UnrealBakeHelpers.cpp
@@ -47,33 +47,15 @@ UObject* FUnrealBakeHelpers::BakeHelper_DuplicateAsset(UObject* Object, const FS
 	
 	if (bGenerateConstantMaterialInstances && MatInterface && MatInstance)
 	{
-		UMaterialInterface* ParentInterface = MatInterface;
-
-		UMaterialInstanceDynamic* InstanceDynamic = Cast<UMaterialInstanceDynamic>(Object);
-
-		if (InstanceDynamic)
-		{
-			ParentInterface = InstanceDynamic->Parent;
-		}
-		else
-		{
-			UMaterialInstanceConstant* InstanceConstant = Cast<UMaterialInstanceConstant>(Object);
-
-			if (InstanceConstant)
-			{
-				ParentInterface = InstanceConstant->Parent;
-			}
-		}
-
 		UMaterialInstanceConstantFactoryNew* MaterialFactory = NewObject<UMaterialInstanceConstantFactoryNew>();
-		MaterialFactory->InitialParent = ParentInterface;
+		MaterialFactory->InitialParent = MatInstance->Parent;
 		FString MaterialInstanceName = FinalObjName;
 		DupObject = (UMaterialInstanceConstant*)MaterialFactory->FactoryCreateNew(UMaterialInstanceConstant::StaticClass(),
 			Package, FName(MaterialInstanceName), RF_NoFlags, NULL, GWarn);
 		ensure(DupObject);
 
 		TMap<int, UTexture*> EmptyTextureReplacementMap;
-		FUnrealBakeHelpers::CopyAllMaterialParameters(DupObject, MatInterface, EmptyTextureReplacementMap);
+		CopyAllMaterialParameters(DupObject, MatInterface, EmptyTextureReplacementMap);
 	}
 	else
 	{
@@ -300,9 +282,6 @@ void FUnrealBakeHelpers::CopyAllMaterialParameters(UObject* DestMaterial, UMater
 
 	if (Material && DupMaterial)
 	{
-		TArray<FMaterialParameterInfo> parametersInfo;
-		TArray<FGuid> parametersGuids;
-
 		// copy scalar parameters
 		TArray<FMaterialParameterInfo> ScalarParameterInfoArray;
 		TArray<FGuid> GuidArray;
