[Mutable] Fix NodeEditMaterial having two Pin Viewers.

#rnx

diff --git a/Engine/Plugins/Experimental/Mutable/Source/CustomizableObjectEditor/Private/MuCOE/Nodes/CustomizableObjectNodeEditMaterialBaseDetails.cpp b/Engine/Plugins/Experimental/Mutable/Source/CustomizableObjectEditor/Private/MuCOE/Nodes/CustomizableObjectNodeEditMaterialBaseDetails.cpp
--- a/Engine/Plugins/Experimental/Mutable/Source/CustomizableObjectEditor/Private/MuCOE/Nodes/CustomizableObjectNodeEditMaterialBaseDetails.cpp
+++ b/Engine/Plugins/Experimental/Mutable/Source/CustomizableObjectEditor/Private/MuCOE/Nodes/CustomizableObjectNodeEditMaterialBaseDetails.cpp
@@ -25,8 +25,6 @@ TSharedRef<IDetailCustomization> FCustomizableObjectNodeEditMaterialBaseDetails:
 
 void FCustomizableObjectNodeEditMaterialBaseDetails::CustomizeDetails(IDetailLayoutBuilder& DetailBuilder)
 {
-	FCustomizableObjectNodeParentedMaterialDetails::CustomizeDetails(DetailBuilder);
-
 	DetailBuilderPtr = &DetailBuilder;
 	NodeEditMaterialBase = nullptr;
 
