[mutable] Added CVar to use for the CO compilation instead of the value set in the COs 'PackagedDataBytesLimit' var.
- The value in the UProperty is left untouched.

diff --git a/Engine/Plugins/Experimental/Mutable/Source/CustomizableObject/Private/MuCO/CustomizableObject.cpp b/Engine/Plugins/Experimental/Mutable/Source/CustomizableObject/Private/MuCO/CustomizableObject.cpp
--- a/Engine/Plugins/Experimental/Mutable/Source/CustomizableObject/Private/MuCO/CustomizableObject.cpp
+++ b/Engine/Plugins/Experimental/Mutable/Source/CustomizableObject/Private/MuCO/CustomizableObject.cpp
@@ -52,6 +52,17 @@
 
 DEFINE_LOG_CATEGORY(LogMutable);
 
+#if WITH_EDITOR
+
+TAutoConsoleVariable<int32> CVarPackagedDataBytesLimitOverride(
+	TEXT("mutable.PackagedDataBytesLimitOverride"),
+	-1,
+	TEXT("Defines the value to be used as 'PackagedDataBytesLimitOverride' for the compilation of all COs.\n")
+	TEXT(" <0 : Use value defined in the CO\n")
+	TEXT(" >=0  : Use this value instead\n"));
+
+#endif
+
 #if WITH_EDITORONLY_DATA
 
 namespace UE::Mutable::Private
@@ -2100,7 +2111,18 @@ FCompilationOptions UCustomizableObjectPrivate::GetCompileOptions() const
 	Options.TextureCompression = TextureCompression;
 	Options.OptimizationLevel = OptimizationLevel;
 	Options.bUseDiskCompilation = bUseDiskCompilation;
-	Options.PackagedDataBytesLimit = PackagedDataBytesLimit;
+	
+	const int32 TargetBulkDataFileBytesOverride = CVarPackagedDataBytesLimitOverride.GetValueOnAnyThread();
+	if ( TargetBulkDataFileBytesOverride >= 0)
+	{
+		Options.PackagedDataBytesLimit = TargetBulkDataFileBytesOverride;
+		UE_LOG(LogMutable,Display, TEXT("Ignoring CO PackagedDataBytesLimit value in favour of overriding CVar value : mutable.PackagedDataBytesLimitOverride %llu"), Options.PackagedDataBytesLimit);
+	}
+	else
+	{
+		Options.PackagedDataBytesLimit =  PackagedDataBytesLimit;
+	}
+	
 	Options.EmbeddedDataBytesLimit = EmbeddedDataBytesLimit;
 	Options.CustomizableObjectNumBoneInfluences = ICustomizableObjectModule::Get().GetNumBoneInfluences();
 	Options.bRealTimeMorphTargetsEnabled = GetPublic()->bEnableRealTimeMorphTargets;
diff --git a/Engine/Plugins/Experimental/Mutable/Source/MutableValidation/Private/MuV/CustomizableObjectCompilationUtility.cpp b/Engine/Plugins/Experimental/Mutable/Source/MutableValidation/Private/MuV/CustomizableObjectCompilationUtility.cpp
--- a/Engine/Plugins/Experimental/Mutable/Source/MutableValidation/Private/MuV/CustomizableObjectCompilationUtility.cpp
+++ b/Engine/Plugins/Experimental/Mutable/Source/MutableValidation/Private/MuV/CustomizableObjectCompilationUtility.cpp
@@ -55,6 +55,7 @@ bool FCustomizableObjectCompilationUtility::CompileCustomizableObject(UCustomiza
 		UE_LOG(LogMutable, Log, TEXT("(string) model_texture_compression : %s "), *UEnum::GetValueAsString(CompilationOptions.TextureCompression));
 		UE_LOG(LogMutable, Log, TEXT("(string) model_disk_compilation : %s "), CompilationOptions.bUseDiskCompilation ? TEXT("true") : TEXT("false"));
 		UE_LOG(LogMutable, Log, TEXT("(string) model_compile_platform_name : %s "), *CompilationOptions.TargetPlatform->PlatformName());
+		UE_LOG(LogMutable, Log, TEXT("(int) model_package_data_bytes_limit_bytes : %llu "), CompilationOptions.PackagedDataBytesLimit);
 	}
 
 	// Compilation completed successfully
@@ -175,4 +176,4 @@ bool FCustomizableObjectCompilationUtility::CompileCustomizableObject(UCustomiza
 	
 	// Return the success state of the operation
 	return bCompilationSuccess;
-}
\ No newline at end of file
+}
