[Mutable] Add gameplay tags filter for int parameter available options.

#rnx

diff --git a/Engine/Plugins/Experimental/Mutable/Source/CustomizableObject/Private/MuCO/CustomizableObjectUIData.cpp b/Engine/Plugins/Experimental/Mutable/Source/CustomizableObject/Private/MuCO/CustomizableObjectUIData.cpp
--- a/Engine/Plugins/Experimental/Mutable/Source/CustomizableObject/Private/MuCO/CustomizableObjectUIData.cpp
+++ b/Engine/Plugins/Experimental/Mutable/Source/CustomizableObject/Private/MuCO/CustomizableObjectUIData.cpp
@@ -1,4 +1,4 @@
-﻿// Copyright Epic Games, Inc. All Rights Reserved.
+// Copyright Epic Games, Inc. All Rights Reserved.
 
 #include "MuCO/CustomizableObjectUIData.h"
 
@@ -67,6 +67,24 @@ FArchive& operator<<(FArchive& Ar, FMutableParamUIMetadata& Struct)
 	Ar << Struct.MinimumValue;
 	Ar << Struct.MaximumValue;
 
+#if WITH_EDITORONLY_DATA
+	FString ExportString;
+	if (Ar.IsSaving())
+	{
+		ExportString = Struct.EditorGameplayTags.ToString();
+	}
+
+	Ar << ExportString;
+
+	if (Ar.IsLoading())
+	{
+		if (!ExportString.IsEmpty())
+		{
+			Struct.EditorGameplayTags.FromExportString(ExportString);
+		}
+	}
+#endif
+
 	return Ar;
 }
 
diff --git a/Engine/Plugins/Experimental/Mutable/Source/CustomizableObject/Public/MuCO/CustomizableObjectUIData.h b/Engine/Plugins/Experimental/Mutable/Source/CustomizableObject/Public/MuCO/CustomizableObjectUIData.h
--- a/Engine/Plugins/Experimental/Mutable/Source/CustomizableObject/Public/MuCO/CustomizableObjectUIData.h
+++ b/Engine/Plugins/Experimental/Mutable/Source/CustomizableObject/Public/MuCO/CustomizableObjectUIData.h
@@ -6,6 +6,10 @@
 #include "Engine/Texture2D.h"
 #include "MuCO/CustomizableObjectParameterTypeDefinitions.h"
 
+#if WITH_EDITORONLY_DATA
+#include "GameplayTagContainer.h"
+#endif
+
 #include "CustomizableObjectUIData.generated.h"
 
 class UTexture2D;
@@ -55,6 +59,13 @@ struct FMutableParamUIMetadata : public FMutableUIMetadata
 	UPROPERTY(BlueprintReadWrite, EditAnywhere, Category = UI)
 	float MaximumValue = 1.0f;
 
+#if WITH_EDITORONLY_DATA
+	/** Gameplay tags to take into consideration when filtering parameters.
+	  * Only applies to the options shown in drop downs. Child object parameters and data table rows. */
+	UPROPERTY(EditAnywhere, Category = UI)
+	FGameplayTagContainer EditorGameplayTags;
+#endif
+
 	friend FArchive& operator<<(FArchive& Ar, FMutableParamUIMetadata& Struct);
 };
 
diff --git a/Engine/Plugins/Experimental/Mutable/Source/CustomizableObjectEditor/Private/MuCOE/CustomizableInstanceDetails.cpp b/Engine/Plugins/Experimental/Mutable/Source/CustomizableObjectEditor/Private/MuCOE/CustomizableInstanceDetails.cpp
--- a/Engine/Plugins/Experimental/Mutable/Source/CustomizableObjectEditor/Private/MuCOE/CustomizableInstanceDetails.cpp
+++ b/Engine/Plugins/Experimental/Mutable/Source/CustomizableObjectEditor/Private/MuCOE/CustomizableInstanceDetails.cpp
@@ -186,6 +186,16 @@ void FCustomizableInstanceDetails::CustomizeDetails(const TSharedPtr<IDetailLayo
 		.OnCheckStateChanged(this, &FCustomizableInstanceDetails::OnUseUISectionsSelectionChanged)
 	];
 
+	TSharedPtr<ICustomizableObjectInstanceEditor> Editor = GetEditorChecked();
+	if (UCustomizableObjectEditorProperties* EditorProperties = Editor->GetEditorProperties())
+	{
+		IDetailPropertyRow* FilterPropertyRow = VisibilitySettingsCategory.AddExternalObjectProperty({ EditorProperties }, FName("Filter"));
+		FilterPropertyRow->GetPropertyHandle()->SetOnChildPropertyValueChanged(FSimpleDelegate::CreateSP(this, &FCustomizableInstanceDetails::Refresh));
+
+		Filter = EditorProperties->Filter.GameplayTagsFilter;
+		FilterType = EditorProperties->Filter.GameplayTagsFilterType;
+	}
+
 	// Copy, Paste and Reset Parameters
 	ParametersCategory.AddCustomRow(LOCTEXT("CustomizableInstanceDetails_CopyPasteResetButtons", "Copy Paste Reset"))
 	[
@@ -622,12 +632,12 @@ bool FCustomizableInstanceDetails::GenerateParametersView(IDetailCategoryBuilder
 
 			if (CustomInstance->GetPrivate()->bShowUISections)
 			{
-				IDetailGroup* CurrentSection = GenerateParameterSection(ParamInfo.ParamIndexInObject, DetailsCategory);
+				IDetailGroup* CurrentSection = GenerateParameterSection(DetailsCategory, *CustomizableObject, ParamInfo.ParamName);
 				check(CurrentSection);
 
 				if (!IsMultidimensionalProjector(ParamInfo.ParamIndexInObject))
 				{
-					GenerateWidgetRow(CurrentSection->AddWidgetRow(), ParamInfo.ParamName, ParamInfo.ParamIndexInObject);
+					GenerateWidgetRow(CurrentSection->AddWidgetRow(), *CustomizableObject, ParamInfo.ParamName, ParamInfo.ParamIndexInObject);
 				}
 				else
 				{
@@ -635,14 +645,14 @@ bool FCustomizableInstanceDetails::GenerateParametersView(IDetailCategoryBuilder
 
 					// Call Order between the following lines maters.
 					ParentsGroups.Add(ParamInfo.ParamName, ProjectorGroup);
-					GenerateWidgetRow(ProjectorGroup->HeaderRow(), ParamInfo.ParamName, ParamInfo.ParamIndexInObject);
+					GenerateWidgetRow(ProjectorGroup->HeaderRow(), *CustomizableObject, ParamInfo.ParamName, ParamInfo.ParamIndexInObject);
 				}
 			}
 			else
 			{
 				if (!IsMultidimensionalProjector(ParamInfo.ParamIndexInObject))
 				{
-					GenerateWidgetRow(DetailsCategory.AddCustomRow(FText::FromString(ParamInfo.ParamName)), ParamInfo.ParamName, ParamInfo.ParamIndexInObject);
+					GenerateWidgetRow(DetailsCategory.AddCustomRow(FText::FromString(ParamInfo.ParamName)), *CustomizableObject, ParamInfo.ParamName, ParamInfo.ParamIndexInObject);
 				}
 				else
 				{
@@ -650,7 +660,7 @@ bool FCustomizableInstanceDetails::GenerateParametersView(IDetailCategoryBuilder
 
 					// Call Order between the following lines maters.
 					ParentsGroups.Add(ParamInfo.ParamName, ProjectorGroup);
-					GenerateWidgetRow(ProjectorGroup->HeaderRow(), ParamInfo.ParamName, ParamInfo.ParamIndexInObject);
+					GenerateWidgetRow(ProjectorGroup->HeaderRow(), *CustomizableObject, ParamInfo.ParamName, ParamInfo.ParamIndexInObject);
 				}
 			}
 		}
@@ -683,7 +693,7 @@ bool FCustomizableInstanceDetails::GenerateParametersView(IDetailCategoryBuilder
 		{
 			if (!ParamHasParent.Find(ParametersTree[ParamIndexInObject].ParamIndexInObject))
 			{
-				RecursivelyAddParamAndChildren(ParametersTree[ParamIndexInObject].ParamIndexInObject, "", DetailsCategory);
+				RecursivelyAddParamAndChildren(*CustomizableObject, ParametersTree[ParamIndexInObject].ParamIndexInObject, "", DetailsCategory);
 			}
 		}
 	}
@@ -692,10 +702,9 @@ bool FCustomizableInstanceDetails::GenerateParametersView(IDetailCategoryBuilder
 }
 
 
-void FCustomizableInstanceDetails::RecursivelyAddParamAndChildren(const int32 ParamIndexInObject, const FString ParentName, IDetailCategoryBuilder& DetailsCategory)
+void FCustomizableInstanceDetails::RecursivelyAddParamAndChildren(const UCustomizableObject& CustomizableObject, const int32 ParamIndexInObject, const FString ParentName, IDetailCategoryBuilder& DetailsCategory)
 {
-	const UCustomizableObject* CustomizableObject = CustomInstance->GetCustomizableObject();
-	const FString ParamName = CustomizableObject->GetParameterName(ParamIndexInObject);	
+	const FString ParamName = CustomizableObject.GetParameterName(ParamIndexInObject);	
 	TArray<int32> Children;
 	
 	ParamChildren.MultiFind(ParamName, Children, true);
@@ -704,12 +713,12 @@ void FCustomizableInstanceDetails::RecursivelyAddParamAndChildren(const int32 Pa
 	{
 		if (CustomInstance->GetPrivate()->bShowUISections)
 		{
-			IDetailGroup* CurrentSection = GenerateParameterSection(ParamIndexInObject, DetailsCategory);
+			IDetailGroup* CurrentSection = GenerateParameterSection(DetailsCategory, CustomizableObject, ParamName);
 			check(CurrentSection);
 
 			if (Children.Num() == 0 && !IsMultidimensionalProjector(ParamIndexInObject))
 			{
-				GenerateWidgetRow(CurrentSection->AddWidgetRow(), ParamName, ParamIndexInObject);
+				GenerateWidgetRow(CurrentSection->AddWidgetRow(), CustomizableObject, ParamName, ParamIndexInObject);
 			}
 			else
 			{
@@ -717,14 +726,14 @@ void FCustomizableInstanceDetails::RecursivelyAddParamAndChildren(const int32 Pa
 				
 				// Call Order between the following lines maters.
 				ParentsGroups.Add(ParamName, ParentGroup);
-				GenerateWidgetRow(ParentGroup->HeaderRow(), ParamName, ParamIndexInObject);
+				GenerateWidgetRow(ParentGroup->HeaderRow(), CustomizableObject, ParamName, ParamIndexInObject);
 			}
 		}
 		else
 		{
 			if (Children.Num() == 0 && !IsMultidimensionalProjector(ParamIndexInObject))
 			{
-				GenerateWidgetRow(DetailsCategory.AddCustomRow(FText::FromString(ParamName)), ParamName, ParamIndexInObject);
+				GenerateWidgetRow(DetailsCategory.AddCustomRow(FText::FromString(ParamName)), CustomizableObject, ParamName, ParamIndexInObject);
 			}
 			else
 			{
@@ -732,7 +741,7 @@ void FCustomizableInstanceDetails::RecursivelyAddParamAndChildren(const int32 Pa
 				
 				// Call Order between the following lines maters.
 				ParentsGroups.Add(ParamName, ParentGroup);
-				GenerateWidgetRow(ParentGroup->HeaderRow(), ParamName, ParamIndexInObject);
+				GenerateWidgetRow(ParentGroup->HeaderRow(), CustomizableObject, ParamName, ParamIndexInObject);
 			}
 		}
 	}
@@ -742,7 +751,7 @@ void FCustomizableInstanceDetails::RecursivelyAddParamAndChildren(const int32 Pa
 
 		if (Children.Num() == 0 && !IsMultidimensionalProjector(ParamIndexInObject))
 		{
-			GenerateWidgetRow(ParentGroup->AddWidgetRow(), ParamName, ParamIndexInObject);
+			GenerateWidgetRow(ParentGroup->AddWidgetRow(), CustomizableObject, ParamName, ParamIndexInObject);
 		}
 		else
 		{
@@ -750,13 +759,13 @@ void FCustomizableInstanceDetails::RecursivelyAddParamAndChildren(const int32 Pa
 
 			// Call Order between the following lines maters
 			ParentsGroups.Add(ParamName, ChildGroup);
-			GenerateWidgetRow(ParentGroup->HeaderRow(), ParamName, ParamIndexInObject);
+			GenerateWidgetRow(ParentGroup->HeaderRow(), CustomizableObject, ParamName, ParamIndexInObject);
 		}
 	}
 
 	for (const int32 ChildIndexInObject : Children)
 	{
-		RecursivelyAddParamAndChildren(ChildIndexInObject, ParamName, DetailsCategory);
+		RecursivelyAddParamAndChildren(CustomizableObject, ChildIndexInObject, ParamName, DetailsCategory);
 	}
 }
 
@@ -801,6 +810,33 @@ bool FCustomizableInstanceDetails::IsVisible(int32 ParamIndexInObject)
 }
 
 
+bool FCustomizableInstanceDetails::IsIntParameterFilteredOut(const UCustomizableObject& CustomizableObject, const FString& ParamName, const FString& ParamOption) const
+{
+	if (Filter.IsEmpty())
+	{
+		return false;
+	}
+
+	FMutableParamUIMetadata Metadata = CustomizableObject.GetIntParameterOptionUIMetadata(ParamName, ParamOption);
+
+	switch (FilterType)
+	{
+	case EGameplayContainerMatchType::Any:
+	{
+		return !Metadata.EditorGameplayTags.HasAny(Filter);
+	}
+	case EGameplayContainerMatchType::All:
+	{
+		return !Metadata.EditorGameplayTags.HasAll(Filter);
+	}
+	default:
+		unimplemented();
+		return false;
+	}
+}
+
+
+
 bool FCustomizableInstanceDetails::IsMultidimensionalProjector(int32 ParamIndexInObject)
 {
 	const UCustomizableObject* CustomizableObject = CustomInstance->GetCustomizableObject();
@@ -809,12 +845,9 @@ bool FCustomizableInstanceDetails::IsMultidimensionalProjector(int32 ParamIndexI
 }
 
 
-IDetailGroup* FCustomizableInstanceDetails::GenerateParameterSection(const int32 ParamIndexInObject, IDetailCategoryBuilder& DetailsCategory)
+IDetailGroup* FCustomizableInstanceDetails::GenerateParameterSection(IDetailCategoryBuilder& DetailsCategory, const UCustomizableObject& CustomizableObject, const FString& ParamName)
 {
-	const UCustomizableObject* CustomizableObject = CustomInstance->GetCustomizableObject();
-
-	const FString& ParamName = CustomizableObject->GetParameterName(ParamIndexInObject);
-	const FMutableParamUIMetadata UIMetadata = CustomizableObject->GetParameterUIMetadata(ParamName);
+	const FMutableParamUIMetadata UIMetadata = CustomizableObject.GetParameterUIMetadata(ParamName);
 	const FString SectionName = UIMetadata.UISectionName.IsEmpty() ? "Miscellaneous" : UIMetadata.UISectionName;
 	IDetailGroup* CurrentSection = nullptr;
 
@@ -834,7 +867,7 @@ IDetailGroup* FCustomizableInstanceDetails::GenerateParameterSection(const int32
 }
 
 
-void FCustomizableInstanceDetails::GenerateWidgetRow(FDetailWidgetRow& WidgetRow, const FString& ParamName, const int32 ParamIndexInObject)
+void FCustomizableInstanceDetails::GenerateWidgetRow(FDetailWidgetRow& WidgetRow, const UCustomizableObject& CustomizableObject, const FString& ParamName, const int32 ParamIndexInObject)
 {
 	WidgetRow.NameContent()
 	[
@@ -850,7 +883,7 @@ void FCustomizableInstanceDetails::GenerateWidgetRow(FDetailWidgetRow& WidgetRow
 		.HAlign(HAlign_Fill)
 		.Padding(0.0f, 5.0f, 0.0f, 5.0f)
 		[
-			GenerateParameterWidget(ParamIndexInObject)
+			GenerateParameterWidget(CustomizableObject, ParamName, ParamIndexInObject)
 		]
 	]
 	.OverrideResetToDefault(FResetToDefaultOverride::Create(FSimpleDelegate::CreateSP(this, &FCustomizableInstanceDetails::OnResetParameterButtonClicked, ParamIndexInObject)))
@@ -858,27 +891,25 @@ void FCustomizableInstanceDetails::GenerateWidgetRow(FDetailWidgetRow& WidgetRow
 }
 
 
-TSharedRef<SWidget> FCustomizableInstanceDetails::GenerateParameterWidget(const int32 ParamIndexInObject)
+TSharedRef<SWidget> FCustomizableInstanceDetails::GenerateParameterWidget(const UCustomizableObject& CustomizableObject, const FString& ParamName, const int32 ParamIndexInObject)
 {
-	UCustomizableObject* CustomizableObject = CustomInstance->GetCustomizableObject();
-
-	switch (CustomizableObject->GetParameterType(ParamIndexInObject))
+	switch (CustomizableObject.GetParameterType(ParamIndexInObject))
 	{
 		case EMutableParameterType::Bool:
 		{
-			return GenerateBoolWidget(ParamIndexInObject);
+			return GenerateBoolWidget(ParamName);
 		}
 		case EMutableParameterType::Float:
 		{
-			return GenerateFloatWidget(ParamIndexInObject);
+			return GenerateFloatWidget(CustomizableObject, ParamName);
 		}
 		case EMutableParameterType::Color:
 		{
-			return GenerateColorWidget(ParamIndexInObject);
+			return GenerateColorWidget(ParamName);
 		}
 		case EMutableParameterType::Texture:
 		{
-			return GenerateTextureWidget(ParamIndexInObject);
+			return GenerateTextureWidget(CustomizableObject, ParamName);
 		}
 		case EMutableParameterType::Projector:
 		{
@@ -886,16 +917,16 @@ TSharedRef<SWidget> FCustomizableInstanceDetails::GenerateParameterWidget(const
 
 			if (!bIsParamMultidimensional)
 			{
-				return GenerateSimpleProjector(ParamIndexInObject);
+				return GenerateSimpleProjector(ParamName);
 			}
 			else 
 			{
-				return GenerateMultidimensionalProjector(ParamIndexInObject);
+				return GenerateMultidimensionalProjector(CustomizableObject, ParamName, ParamIndexInObject);
 			}
 		}
 		case EMutableParameterType::Int:
 		{
-			return GenerateIntWidget(ParamIndexInObject);
+			return GenerateIntWidget(CustomizableObject, ParamName, ParamIndexInObject);
 		}
 		case EMutableParameterType::None:
 		{
@@ -909,26 +940,20 @@ TSharedRef<SWidget> FCustomizableInstanceDetails::GenerateParameterWidget(const
 
 // INT PARAMETERS -----------------------------------------------------------------------------------------------------------------
 
-TSharedRef<SWidget> FCustomizableInstanceDetails::GenerateIntWidget(const int32 ParamIndexInObject)
+TSharedRef<SWidget> FCustomizableInstanceDetails::GenerateIntWidget(const UCustomizableObject& CustomizableObject, const FString& ParamName, const int32 ParamIndexInObject)
 {
-	UCustomizableObject* CustomizableObject = CustomInstance->GetCustomizableObject();
-	
 	const bool bMultidimensional = CustomInstance->GetCustomizableObject()->IsParameterMultidimensional(ParamIndexInObject);
 	if (bMultidimensional)
 	{
 		return SNew(STextBlock).Text(LOCTEXT("MultidimensionalINTParameter_Text", "Multidimensional INT Parameter not supported"));
 	}
 
-	const int32 NumValues = CustomizableObject->GetIntParameterNumOptions(ParamIndexInObject);
+	const int32 NumValues = CustomizableObject.GetIntParameterNumOptions(ParamIndexInObject);
 	if (!NumValues)
 	{
 		return SNew(STextBlock).Text(LOCTEXT("NoAvailableOptions", "No Available Options"));
 	}
 	
-	FString ToolTipText = FString("None");
-	FString ParamName = CustomizableObject->GetParameterName(ParamIndexInObject);
-	FString Value = CustomInstance->GetIntParameterSelectedOption(ParamName, -1);
-
 	TSharedPtr<TArray<TSharedPtr<FString>>>* FoundOptions = IntParameterOptions.Find(ParamIndexInObject);
 	TArray<TSharedPtr<FString>>& OptionNamesAttribute = FoundOptions && FoundOptions->IsValid() ? 
 														*FoundOptions->Get() :
@@ -936,40 +961,46 @@ TSharedRef<SWidget> FCustomizableInstanceDetails::GenerateIntWidget(const int32
 
 	OptionNamesAttribute.Empty();
 
-	int32 ValueIndex = 0;
+	const FString SelectedOption = CustomInstance->GetIntParameterSelectedOption(ParamName, INDEX_NONE);
+
+	// Tooltip for the selected option
+	FString ToolTipText = FString("None");
+	if (const FString* Identifier = CustomizableObject.GetPrivate()->GroupNodeMap.FindKey(FCustomizableObjectIdPair(ParamName, SelectedOption)))
+	{
+		if (FString* CustomizableObjectPath = CustomizableObject.GetPrivate()->CustomizableObjectPathMap.Find(*Identifier))
+		{
+			ToolTipText = *CustomizableObjectPath;
+		}
+	}
 
+	TSharedPtr<FString> SelectedOptionString;
 	for (int32 i = 0; i < NumValues; ++i)
 	{
-		FString PossibleValue = CustomizableObject->GetIntParameterAvailableOption(ParamIndexInObject, i);
+		const FString PossibleValue = CustomizableObject.GetIntParameterAvailableOption(ParamIndexInObject, i);
 
-		if (PossibleValue == Value)
+		if (PossibleValue == SelectedOption) // Always add the selected option, even if it should be filtered.
 		{
-			ValueIndex = i;
-
-			const FString* Identifier = CustomizableObject->GetPrivate()->GroupNodeMap.FindKey(FCustomizableObjectIdPair(ParamName, PossibleValue));
-			if (Identifier)
-			{
-				if (FString* CustomizableObjectPath = CustomizableObject->GetPrivate()->CustomizableObjectPathMap.Find(*Identifier))
-				{
-					ToolTipText = *CustomizableObjectPath;
-				}
-			}
+			SelectedOptionString = TSharedPtr<FString>(new FString(PossibleValue));
+			OptionNamesAttribute.Add(SelectedOptionString);
 		}
 
-		OptionNamesAttribute.Add(TSharedPtr<FString>(new FString(CustomizableObject->GetIntParameterAvailableOption(ParamIndexInObject, i))));
+		else if (!IsIntParameterFilteredOut(CustomizableObject, ParamName, PossibleValue))
+		{
+			OptionNamesAttribute.Add(TSharedPtr<FString>(new FString(PossibleValue)));
+		}
 	}
 
 	return SNew(SSearchableComboBox)
 		.ToolTipText(FText::FromString(ToolTipText))
 		.OptionsSource(&OptionNamesAttribute)
-		.InitiallySelectedItem(OptionNamesAttribute[ValueIndex])
+		.InitiallySelectedItem(SelectedOptionString)
 		.Method(EPopupMethod::UseCurrentWindow)
 		.OnSelectionChanged(this, &FCustomizableInstanceDetails::OnIntParameterComboBoxChanged, ParamName)
 		.OnGenerateWidget(this, &FCustomizableInstanceDetails::OnGenerateWidgetIntParameter)
 		.Content()
 		[
 			SNew(STextBlock)
-				.Text(FText::FromString(*OptionNamesAttribute[ValueIndex]))
+				.Text(FText::FromString(*SelectedOption))
 		];
 }
 
@@ -997,11 +1028,9 @@ void FCustomizableInstanceDetails::OnIntParameterComboBoxChanged(TSharedPtr<FStr
 
 // FLOAT PARAMETERS -----------------------------------------------------------------------------------------------------------------
 
-TSharedRef<SWidget> FCustomizableInstanceDetails::GenerateFloatWidget(const int32 ParamIndexInObject)
+TSharedRef<SWidget> FCustomizableInstanceDetails::GenerateFloatWidget(const UCustomizableObject& CustomizableObject, const FString& ParamName)
 {
-	const UCustomizableObject* CustomizableObject = CustomInstance->GetCustomizableObject();
-	const FString& ParamName = CustomizableObject->GetParameterName(ParamIndexInObject);
-	const FMutableParamUIMetadata& UIMetadata = CustomizableObject->GetParameterUIMetadata(ParamName);
+	const FMutableParamUIMetadata& UIMetadata = CustomizableObject.GetParameterUIMetadata(ParamName);
 
 	if (const TSoftObjectPtr<UObject>* FloatDecoratorAsset = UIMetadata.ExtraAssets.Find(UIMetadataKeyWords::FloatDecoratorName))
 	{
@@ -1132,10 +1161,8 @@ void FCustomizableInstanceDetails::OnFloatParameterCommited(float Value, ETextCo
 
 // TEXTURE PARAMETERS -----------------------------------------------------------------------------------------------------------------
 
-TSharedRef<SWidget> FCustomizableInstanceDetails::GenerateTextureWidget(const int32 ParamIndexInObject)
+TSharedRef<SWidget> FCustomizableInstanceDetails::GenerateTextureWidget(const UCustomizableObject& CustomizableObject, const FString& ParamName)
 {
-	UCustomizableObject* CustomizableObject = CustomInstance->GetCustomizableObject();
-	FString ParamName = CustomizableObject->GetParameterName(ParamIndexInObject);
 	const FName ParameterValue = CustomInstance->GetTextureParameterSelectedOption(ParamName);
 	TSharedPtr<FString> InitiallySelected;
 
@@ -1235,11 +1262,8 @@ void FCustomizableInstanceDetails::OnTextureParameterComboBoxSelectionChanged(TS
 
 // COLOR PARAMETERS -----------------------------------------------------------------------------------------------------------------
 
-TSharedRef<SWidget> FCustomizableInstanceDetails::GenerateColorWidget(const int32 ParamIndexInObject)
+TSharedRef<SWidget> FCustomizableInstanceDetails::GenerateColorWidget(const FString& ParamName)
 {
-	UCustomizableObject* CustomizableObject = CustomInstance->GetCustomizableObject();
-	FString ParamName = CustomizableObject->GetParameterName(ParamIndexInObject);
-
 	return SNew(SColorBlock)
 		.Color(this, &FCustomizableInstanceDetails::GetColorParameterValue, ParamName)
 		.ShowBackgroundForAlpha(false)
@@ -1288,11 +1312,8 @@ void FCustomizableInstanceDetails::OnSetColorFromColorPicker(FLinearColor NewCol
 
 // PROJECTOR PARAMETERS -----------------------------------------------------------------------------------------------------------------
 
-TSharedRef<SWidget> FCustomizableInstanceDetails::GenerateSimpleProjector(const int32 ParamIndexInObject)
+TSharedRef<SWidget> FCustomizableInstanceDetails::GenerateSimpleProjector(const FString& ParamName)
 {
-	UCustomizableObject* CustomizableObject = CustomInstance->GetCustomizableObject();
-	FString ParamName = CustomizableObject->GetParameterName(ParamIndexInObject);
-
 	const TSharedPtr<ICustomizableObjectInstanceEditor> Editor = GetEditorChecked();
 	const UProjectorParameter* ProjectorParameter = Editor->GetProjectorParameter();
 	const bool bSelectedProjector = ProjectorParameter->IsProjectorSelected(ParamName);
@@ -1469,11 +1490,8 @@ FReply FCustomizableInstanceDetails::OnProjectorResetTransform(const FString Par
 }
 
 
-TSharedRef<SWidget> FCustomizableInstanceDetails::GenerateMultidimensionalProjector(const int32 ParamIndexInObject)
+TSharedRef<SWidget> FCustomizableInstanceDetails::GenerateMultidimensionalProjector(const UCustomizableObject& CustomizableObject, const FString& ParamName, const int32 ParamIndexInObject)
 {
-	UCustomizableObject* CustomizableObject = CustomInstance->GetCustomizableObject();
-	FString ParamName = CustomizableObject->GetParameterName(ParamIndexInObject);
-
 	const TSharedPtr<ICustomizableObjectInstanceEditor> Editor = GetEditorChecked();
 	const TArray<FCustomizableObjectProjectorParameterValue>& ProjectorParameters = CustomInstance->GetPrivate()->GetDescriptor().GetProjectorParameters();
 	const int32 ProjectorParamIndex = CustomInstance->FindProjectorParameterNameIndex(ParamName);
@@ -1482,12 +1500,12 @@ TSharedRef<SWidget> FCustomizableInstanceDetails::GenerateMultidimensionalProjec
 
 	// Selected Pose UI
 	const FString PoseSwitchEnumParamName = ParamName + FMultilayerProjector::POSE_PARAMETER_POSTFIX;
-	const int32 PoseSwitchEnumParamIndexInObject = CustomizableObject->FindParameter(PoseSwitchEnumParamName);
+	const int32 PoseSwitchEnumParamIndexInObject = CustomizableObject.FindParameter(PoseSwitchEnumParamName);
 	TSharedPtr<SVerticalBox> ProjectorBox = SNew(SVerticalBox);
 
 	if (PoseSwitchEnumParamIndexInObject != INDEX_NONE)
 	{
-		const int32 NumPoseValues = CustomizableObject->GetIntParameterNumOptions(PoseSwitchEnumParamIndexInObject);
+		const int32 NumPoseValues = CustomizableObject.GetIntParameterNumOptions(PoseSwitchEnumParamIndexInObject);
 
 		TSharedPtr<TArray<TSharedPtr<FString>>>* FoundOptions = ProjectorParameterPoseOptions.Find(PoseSwitchEnumParamIndexInObject);
 		TArray<TSharedPtr<FString>>& PoseOptionNamesAttribute = FoundOptions && FoundOptions->IsValid() ?
@@ -1501,7 +1519,7 @@ TSharedRef<SWidget> FCustomizableInstanceDetails::GenerateMultidimensionalProjec
 
 		for (int32 j = 0; j < NumPoseValues; ++j)
 		{
-			const FString PossibleValue = CustomizableObject->GetIntParameterAvailableOption(PoseSwitchEnumParamIndexInObject, j);
+			const FString PossibleValue = CustomizableObject.GetIntParameterAvailableOption(PoseSwitchEnumParamIndexInObject, j);
 			if (PossibleValue == PoseValue)
 			{
 				PoseValueIndex = j;
@@ -1571,7 +1589,7 @@ TSharedRef<SWidget> FCustomizableInstanceDetails::GenerateMultidimensionalProjec
 
 	for (int32 RangeIndex = 0; RangeIndex < ProjectorParameters[ProjectorParamIndex].RangeValues.Num(); ++RangeIndex)
 	{
-		const int32 TextureSwitchEnumParamIndexInObject = CustomizableObject->FindParameter(TextureSwitchEnumParamName);
+		const int32 TextureSwitchEnumParamIndexInObject = CustomizableObject.FindParameter(TextureSwitchEnumParamName);
 		check(TextureSwitchEnumParamIndexInObject >= 0); TSharedPtr<FString> CurrentStateName = nullptr;
 
 		const UProjectorParameter* ProjectorParameter = Editor->GetProjectorParameter();
@@ -1666,8 +1684,8 @@ TSharedRef<SWidget> FCustomizableInstanceDetails::GenerateMultidimensionalProjec
 		Button->SetBorderBackgroundColor(bSelectedProjector ? FLinearColor::Green : FLinearColor::White);
 
 		// If number of options is equal to 1, Mutable does not consider it multidimensional parameters
-		int32 NumValues = CustomizableObject->GetIntParameterNumOptions(TextureSwitchEnumParamIndexInObject);
-		FString Value = CustomizableObject->IsParameterMultidimensional(TextureSwitchEnumParamName) ? 
+		int32 NumValues = CustomizableObject.GetIntParameterNumOptions(TextureSwitchEnumParamIndexInObject);
+		FString Value = CustomizableObject.IsParameterMultidimensional(TextureSwitchEnumParamName) ? 
 			CustomInstance->GetIntParameterSelectedOption(TextureSwitchEnumParamName, RangeIndex) : CustomInstance->GetIntParameterSelectedOption(TextureSwitchEnumParamName);
 
 		TArray<TSharedPtr<FString>> OptionNamesAttribute;
@@ -1675,13 +1693,13 @@ TSharedRef<SWidget> FCustomizableInstanceDetails::GenerateMultidimensionalProjec
 
 		for (int32 CandidateIndex = 0; CandidateIndex < NumValues; ++CandidateIndex)
 		{
-			FString PossibleValue = CustomizableObject->GetIntParameterAvailableOption(TextureSwitchEnumParamIndexInObject, CandidateIndex);
+			FString PossibleValue = CustomizableObject.GetIntParameterAvailableOption(TextureSwitchEnumParamIndexInObject, CandidateIndex);
 			if (PossibleValue == Value)
 			{
 				ValueIndex = CandidateIndex;
 			}
 
-			OptionNamesAttribute.Add(MakeShared<FString>(CustomizableObject->GetIntParameterAvailableOption(TextureSwitchEnumParamIndexInObject, CandidateIndex)));
+			OptionNamesAttribute.Add(MakeShared<FString>(CustomizableObject.GetIntParameterAvailableOption(TextureSwitchEnumParamIndexInObject, CandidateIndex)));
 		}
 		
 		// Avoid filling this arraw with repeated array  options
@@ -1867,11 +1885,8 @@ TSharedRef<SWidget> FCustomizableInstanceDetails::MakeTextureComboEntryWidget(TS
 
 // BOOL PARAMETERS -----------------------------------------------------------------------------------------------------------------
 
-TSharedRef<SWidget> FCustomizableInstanceDetails::GenerateBoolWidget(const int32 ParamIndexInObject)
+TSharedRef<SWidget> FCustomizableInstanceDetails::GenerateBoolWidget(const FString& ParamName)
 {
-	UCustomizableObject* CustomizableObject = CustomInstance->GetCustomizableObject();
-	FString ParamName = CustomizableObject->GetParameterName(ParamIndexInObject);
-
 	return SNew(SHorizontalBox)
 		+ SHorizontalBox::Slot()
 		.HAlign(HAlign_Left)
diff --git a/Engine/Plugins/Experimental/Mutable/Source/CustomizableObjectEditor/Private/MuCOE/CustomizableInstanceDetails.h b/Engine/Plugins/Experimental/Mutable/Source/CustomizableObjectEditor/Private/MuCOE/CustomizableInstanceDetails.h
--- a/Engine/Plugins/Experimental/Mutable/Source/CustomizableObjectEditor/Private/MuCOE/CustomizableInstanceDetails.h
+++ b/Engine/Plugins/Experimental/Mutable/Source/CustomizableObjectEditor/Private/MuCOE/CustomizableInstanceDetails.h
@@ -5,6 +5,7 @@
 #include "IDetailCustomization.h"
 #include "UObject/ObjectPtr.h"
 #include "UObject/WeakObjectPtr.h"
+#include "GameplayTagContainer.h"
 
 class FDetailWidgetRow;
 class FReply;
@@ -15,6 +16,7 @@ class IDetailCategoryBuilder;
 class IDetailGroup;
 class IDetailLayoutBuilder; 
 class SWidget;
+class UCustomizableObject;
 class UCustomizableObjectInstance;
 
 enum class ECheckBoxState : uint8;
@@ -70,25 +72,28 @@ private:
 	// Main parameter generation functions
 	// Returns true if parameters have been hidden due to runtime type
 	bool GenerateParametersView(IDetailCategoryBuilder& MainCategory);
-	void RecursivelyAddParamAndChildren(const int32 ParamIndexInObject, const FString ParentName, IDetailCategoryBuilder& DetailsCategory);
+	void RecursivelyAddParamAndChildren(const UCustomizableObject& CustomizableObject, const int32 ParamIndexInObject, const FString ParentName, IDetailCategoryBuilder& DetailsCategory);
 	void FillChildrenMap(int32 ParamIndexInObject);
 
 	// Function to determine if a parameter widget should be generated
 	bool IsVisible(int32 ParamIndexInObject);
 	bool IsMultidimensionalProjector(int32 ParamIndexInObject);
 
+	// @return true if the parameter doesn't match the current filter.
+	bool IsIntParameterFilteredOut(const UCustomizableObject& CustomizableObject, const FString& ParamName, const FString& ParamOption) const;
+
 	// Main widget generation functions
-	IDetailGroup* GenerateParameterSection(const int32 ParamIndexInObject, IDetailCategoryBuilder& DetailsCategory);
-	void GenerateWidgetRow(FDetailWidgetRow& WidgetRow, const FString& ParamName, const int32 ParamIndexInObject);
-	TSharedRef<SWidget> GenerateParameterWidget(const int32 ParamIndexInObject);
+	IDetailGroup* GenerateParameterSection(IDetailCategoryBuilder& DetailsCategory, const UCustomizableObject& CustomizableObject, const FString& ParamName);
+	void GenerateWidgetRow(FDetailWidgetRow& WidgetRow, const UCustomizableObject& CustomizableObject, const FString& ParamName, const int32 ParamIndexInObject);
+	TSharedRef<SWidget> GenerateParameterWidget(const UCustomizableObject& CustomizableObject, const FString& ParamName, const int32 ParamIndexInObject);
 
 	// Int Parameters Functions
-	TSharedRef<SWidget> GenerateIntWidget(const int32 ParamIndexInObject);
+	TSharedRef<SWidget> GenerateIntWidget(const UCustomizableObject& CustomizableObject, const FString& ParamName, const int32 ParamIndexInObject);
 	void OnIntParameterComboBoxChanged(TSharedPtr<FString> Selection, ESelectInfo::Type SelectInfo, const FString ParamName);
 	TSharedRef<SWidget> OnGenerateWidgetIntParameter(TSharedPtr<FString> InItem) const;
 	
 	// Float Parameters Functions
-	TSharedRef<SWidget> GenerateFloatWidget(const int32 ParamIndexInObject);
+	TSharedRef<SWidget> GenerateFloatWidget(const UCustomizableObject& CustomizableObject, const FString& ParamName);
 	float GetFloatParameterValue(const FString ParamName, int32 RangeIndex) const;
 	void OnFloatParameterChanged(float Value, const FString ParamName, int32 RangeIndex);
 	void OnFloatParameterSliderBegin();
@@ -97,24 +102,24 @@ private:
 	void OnFloatParameterCommited(float Value, ETextCommit::Type Type, const FString ParamName, int32 RangeIndex);	// Needed to have undo/redo
  
 	// Texture Parameters Functions
-	TSharedRef<SWidget> GenerateTextureWidget(const int32 ParamIndexInObject);
+	TSharedRef<SWidget> GenerateTextureWidget(const UCustomizableObject& CustomizableObject, const FString& ParamName);
 	void GenerateTextureParameterOptions();
 	void OnTextureParameterComboBoxSelectionChanged(TSharedPtr<FString> Selection, ESelectInfo::Type SelectInfo, const FString ParamName);
 	
 	// Color Parameters Functions
-	TSharedRef<SWidget> GenerateColorWidget(const int32 ParamIndexInObject);
+	TSharedRef<SWidget> GenerateColorWidget(const FString& ParamName);
 	FLinearColor GetColorParameterValue(const FString ParamName) const;
 	FReply OnColorBlockMouseButtonDown(const FGeometry& MyGeometry, const FPointerEvent& MouseEvent, const FString ParamName);
 	void OnSetColorFromColorPicker(FLinearColor NewColor, const FString PickerParamName);
 
 	// Bool Parameters Functions
-	TSharedRef<SWidget> GenerateBoolWidget(const int32 ParamIndexInObject);
+	TSharedRef<SWidget> GenerateBoolWidget(const FString& ParamName);
 	ECheckBoxState GetBoolParameterValue(const FString ParamName) const;
 	void OnBoolParameterChanged(ECheckBoxState InCheckboxState, const FString ParamName);
 
 	// Projector Parameters Functions
-	TSharedRef<SWidget> GenerateSimpleProjector(const int32 ParamIndexInObject);
-	TSharedRef<SWidget> GenerateMultidimensionalProjector(const int32 ParamIndexInObject);
+	TSharedRef<SWidget> GenerateSimpleProjector(const FString& ParamName);
+	TSharedRef<SWidget> GenerateMultidimensionalProjector(const UCustomizableObject& CustomizableObject, const FString& ParamName, const int32 ParamIndexInObject);
 	TSharedPtr<ICustomizableObjectInstanceEditor> GetEditorChecked() const;
 	FReply OnProjectorSelectChanged(const FString ParamName, const int32 RangeIndex) const;
 	FReply OnProjectorCopyTransform(const FString ParamName, const int32 RangeIndex) const;
@@ -188,5 +193,9 @@ private:
 
 	// Unique transaction pointer to allow transactions that start and finish in different funtion scopes.
 	TUniquePtr<FScopedTransaction> Transaction;
+
+	//** Editor gameplay tags filter and filter type. Used to filter int parameter options. */
+	FGameplayTagContainer Filter;
+	EGameplayContainerMatchType FilterType;
 };
 
diff --git a/Engine/Plugins/Experimental/Mutable/Source/CustomizableObjectEditor/Private/MuCOE/CustomizableObjectEditor.cpp b/Engine/Plugins/Experimental/Mutable/Source/CustomizableObjectEditor/Private/MuCOE/CustomizableObjectEditor.cpp
--- a/Engine/Plugins/Experimental/Mutable/Source/CustomizableObjectEditor/Private/MuCOE/CustomizableObjectEditor.cpp
+++ b/Engine/Plugins/Experimental/Mutable/Source/CustomizableObjectEditor/Private/MuCOE/CustomizableObjectEditor.cpp
@@ -186,6 +186,8 @@ void FCustomizableObjectEditor::InitCustomizableObjectEditor(const EToolkitMode:
 	CustomSettings = NewObject<UCustomSettings>();
 	CustomSettings->SetEditor(SharedThis(this));
 	
+	EditorProperties = NewObject<UCustomizableObjectEditorProperties>();
+
 	// Support undo/redo
 	CustomizableObject->SetFlags(RF_Transactional);
 
@@ -348,6 +350,7 @@ void FCustomizableObjectEditor::AddReferencedObjects( FReferenceCollector& Colle
 	Collector.AddReferencedObject( PreviewInstance );
 	Collector.AddReferencedObject( ProjectorParameter );
 	Collector.AddReferencedObject( CustomSettings );
+	Collector.AddReferencedObject(EditorProperties);
 }
 
 
@@ -1123,6 +1126,12 @@ void FCustomizableObjectEditor::HideGizmoLight()
 }
 
 
+UCustomizableObjectEditorProperties* FCustomizableObjectEditor::GetEditorProperties()
+{
+	return EditorProperties;
+}
+
+
 void FCustomizableObjectEditor::PostUndo(bool bSuccess)
 {
 	if (bSuccess)
diff --git a/Engine/Plugins/Experimental/Mutable/Source/CustomizableObjectEditor/Private/MuCOE/CustomizableObjectEditor.h b/Engine/Plugins/Experimental/Mutable/Source/CustomizableObjectEditor/Private/MuCOE/CustomizableObjectEditor.h
--- a/Engine/Plugins/Experimental/Mutable/Source/CustomizableObjectEditor/Private/MuCOE/CustomizableObjectEditor.h
+++ b/Engine/Plugins/Experimental/Mutable/Source/CustomizableObjectEditor/Private/MuCOE/CustomizableObjectEditor.h
@@ -211,6 +211,7 @@ public:
 	virtual void HideGizmoClipMesh() override;
 	virtual void ShowGizmoLight(ULightComponent& SelectedLight) override;
 	virtual void HideGizmoLight() override;
+	virtual UCustomizableObjectEditorProperties* GetEditorProperties() override;
 
 	/** Select only this node only. Do nothing if already was only selected. */
 	void SelectSingleNode(UCustomizableObjectNode& Node);
@@ -447,6 +448,8 @@ private:
 
 	TObjectPtr<UCustomSettings> CustomSettings = nullptr;
 
+	TObjectPtr<UCustomizableObjectEditorProperties> EditorProperties = nullptr;
+	
 	bool bRecursionGuard = false;
 	
 	EGizmoType GizmoType = EGizmoType::Hidden;
diff --git a/Engine/Plugins/Experimental/Mutable/Source/CustomizableObjectEditor/Private/MuCOE/CustomizableObjectInstanceEditor.cpp b/Engine/Plugins/Experimental/Mutable/Source/CustomizableObjectEditor/Private/MuCOE/CustomizableObjectInstanceEditor.cpp
--- a/Engine/Plugins/Experimental/Mutable/Source/CustomizableObjectEditor/Private/MuCOE/CustomizableObjectInstanceEditor.cpp
+++ b/Engine/Plugins/Experimental/Mutable/Source/CustomizableObjectEditor/Private/MuCOE/CustomizableObjectInstanceEditor.cpp
@@ -229,6 +229,8 @@ void FCustomizableObjectInstanceEditor::InitCustomizableObjectInstanceEditor( co
 	CustomSettings = NewObject<UCustomSettings>();
 	CustomSettings->SetEditor(SharedThis(this));
 
+	EditorProperties = NewObject<UCustomizableObjectEditorProperties>();
+
 	// Register our commands. This will only register them if not previously registered
 	FCustomizableObjectInstanceEditorCommands::Register();
 	FCustomizableObjectEditorViewportCommands::Register();
@@ -380,6 +382,7 @@ void FCustomizableObjectInstanceEditor::AddReferencedObjects( FReferenceCollecto
 	Collector.AddReferencedObject( HelperCallback );
 	Collector.AddReferencedObject( ProjectorParameter );
 	Collector.AddReferencedObject( CustomSettings );
+	Collector.AddReferencedObject(EditorProperties);
 }
 
 
@@ -580,6 +583,12 @@ void FCustomizableObjectInstanceEditor::HideGizmoProjectorParameter()
 }
 
 
+UCustomizableObjectEditorProperties* FCustomizableObjectInstanceEditor::GetEditorProperties()
+{
+	return EditorProperties;
+}
+
+
 void FCustomizableObjectInstanceEditor::OnInstancePropertySelectionChanged(FProperty* InProperty)
 {
 	Viewport->GetViewportClient()->Invalidate();
diff --git a/Engine/Plugins/Experimental/Mutable/Source/CustomizableObjectEditor/Private/MuCOE/CustomizableObjectInstanceEditor.h b/Engine/Plugins/Experimental/Mutable/Source/CustomizableObjectEditor/Private/MuCOE/CustomizableObjectInstanceEditor.h
--- a/Engine/Plugins/Experimental/Mutable/Source/CustomizableObjectEditor/Private/MuCOE/CustomizableObjectInstanceEditor.h
+++ b/Engine/Plugins/Experimental/Mutable/Source/CustomizableObjectEditor/Private/MuCOE/CustomizableObjectInstanceEditor.h
@@ -2,10 +2,12 @@
 
 #pragma once
 
+#include "MuCO/CustomizableObjectPrivate.h"
 #include "MuCOE/ICustomizableObjectInstanceEditor.h"
 #include "MuCOE/CustomizableObjectEditorViewportLights.h"
+
 #include "TickableEditorObject.h"
-#include "MuCO/CustomizableObjectPrivate.h"
+#include "GameplayTagContainer.h"
 
 #include "CustomizableObjectInstanceEditor.generated.h"
 
@@ -127,6 +129,35 @@ private:
 	TWeakPtr<ICustomizableObjectInstanceEditor> Editor;
 };
 
+USTRUCT()
+struct FCustomizableObjectGameplayTagsFilter
+{
+	GENERATED_BODY();
+
+public:
+	/** Filter Gameplay tags to match with. */
+	UPROPERTY(EditAnywhere, meta = (DisplayName = "Tags"), Category = NoCategory)
+	FGameplayTagContainer GameplayTagsFilter;
+
+	/** Filter match type. */
+	UPROPERTY(EditAnywhere, meta = (DisplayName = "Match type"), Category = NoCategory)
+	EGameplayContainerMatchType GameplayTagsFilterType = EGameplayContainerMatchType::Any;
+};
+
+UCLASS(Transient)
+class UCustomizableObjectEditorProperties : public UObject
+{
+	GENERATED_BODY()
+
+public:
+
+	/** Gameplay tags based parameter options filter.
+	  * Tags may be added to the Editor Gameplay Tags container in the parameters metadata.
+	  * Only applies to the options shown in drop downs. Child object parameters and data table rows. */
+	UPROPERTY(meta = (DisplayName = "Gameplay Tags Filter"))
+	FCustomizableObjectGameplayTagsFilter Filter;
+};
+
 
 /**
  * CustomizableObject Editor class
@@ -177,6 +208,7 @@ public:
 	virtual void HideGizmo() override;
 	virtual void ShowGizmoProjectorParameter(const FString& ParamName, int32 RangeIndex) override;
 	virtual void HideGizmoProjectorParameter() override;
+	virtual UCustomizableObjectEditorProperties* GetEditorProperties() override;
 	
 	/** Callback to notify the editor when the PreviewInstance has been updated */
 	void OnUpdatePreviewInstance(UCustomizableObjectInstance* Instance);
@@ -298,6 +330,8 @@ private:
 
 	TObjectPtr<UCustomSettings> CustomSettings = nullptr;
 
+	TObjectPtr<UCustomizableObjectEditorProperties> EditorProperties = nullptr;
+
 	/** Adds the customizable Object Instance Editor commands to the default toolbar */
 	void ExtendToolbar();
 };
diff --git a/Engine/Plugins/Experimental/Mutable/Source/CustomizableObjectEditor/Public/MuCOE/ICustomizableObjectInstanceEditor.h b/Engine/Plugins/Experimental/Mutable/Source/CustomizableObjectEditor/Public/MuCOE/ICustomizableObjectInstanceEditor.h
--- a/Engine/Plugins/Experimental/Mutable/Source/CustomizableObjectEditor/Public/MuCOE/ICustomizableObjectInstanceEditor.h
+++ b/Engine/Plugins/Experimental/Mutable/Source/CustomizableObjectEditor/Public/MuCOE/ICustomizableObjectInstanceEditor.h
@@ -17,6 +17,7 @@ class UCustomizableObjectNodeMeshClipMorph;
 class UCustomizableObjectNodeMeshClipWithMesh;
 class ULightComponent;
 class UPoseAsset;
+class UCustomizableObjectEditorProperties;
 
 /**
  * Public interface to Customizable Object Instance Editor
@@ -90,6 +91,9 @@ public:
 	/** Hide the light gizmo of the light.
 	 * Synchronizes all editor widgets so that the selection and the widgets is always consistent. */
 	virtual void HideGizmoLight() {}
+
+	/** @return Editor properties. */
+	virtual UCustomizableObjectEditorProperties* GetEditorProperties() = 0;
 };
 
 
