[Mutable] Added connections color and emphasis highligh.

#rnx

diff --git a/Engine/Plugins/Experimental/Mutable/Source/CustomizableObjectEditor/Private/MuCOE/CustomizableObjectConnectionDrawingPolicy.cpp b/Engine/Plugins/Experimental/Mutable/Source/CustomizableObjectEditor/Private/MuCOE/CustomizableObjectConnectionDrawingPolicy.cpp
new file mode 100644
--- /dev/null
+++ b/Engine/Plugins/Experimental/Mutable/Source/CustomizableObjectEditor/Private/MuCOE/CustomizableObjectConnectionDrawingPolicy.cpp
@@ -0,0 +1,18 @@
+﻿// Copyright Epic Games, Inc. All Rights Reserved.
+
+#include "CustomizableObjectConnectionDrawingPolicy.h"
+
+#include "MuCOE/EdGraphSchema_CustomizableObject.h"
+
+
+void FCustomizableObjectConnectionDrawingPolicy::DetermineWiringStyle(UEdGraphPin* OutputPin, UEdGraphPin* InputPin, FConnectionParams& Params)
+{
+	FConnectionDrawingPolicy::DetermineWiringStyle(OutputPin, InputPin, Params);
+
+	Params.WireColor = UEdGraphSchema_CustomizableObject::GetPinTypeColor(OutputPin->PinType.PinCategory); // Color.
+
+	if (HoveredPins.Num() > 0)
+	{
+		ApplyHoverDeemphasis(OutputPin, InputPin,Params.WireThickness, Params.WireColor); // Emphasis highlight.
+	}
+}
diff --git a/Engine/Plugins/Experimental/Mutable/Source/CustomizableObjectEditor/Private/MuCOE/CustomizableObjectConnectionDrawingPolicy.h b/Engine/Plugins/Experimental/Mutable/Source/CustomizableObjectEditor/Private/MuCOE/CustomizableObjectConnectionDrawingPolicy.h
new file mode 100644
--- /dev/null
+++ b/Engine/Plugins/Experimental/Mutable/Source/CustomizableObjectEditor/Private/MuCOE/CustomizableObjectConnectionDrawingPolicy.h
@@ -0,0 +1,16 @@
+﻿// Copyright Epic Games, Inc. All Rights Reserved.
+
+#pragma once
+
+#include "ConnectionDrawingPolicy.h"
+
+
+class FCustomizableObjectConnectionDrawingPolicy : public FConnectionDrawingPolicy
+{
+public:
+	FCustomizableObjectConnectionDrawingPolicy(int32 InBackLayerID, int32 InFrontLayerID, float ZoomFactor, const FSlateRect& InClippingRect, FSlateWindowElementList& InDrawElements)
+		: FConnectionDrawingPolicy(InBackLayerID, InFrontLayerID, ZoomFactor, InClippingRect, InDrawElements) {}
+	
+	// FConnectionDrawingPolicy interface
+	virtual void DetermineWiringStyle(UEdGraphPin* OutputPin, UEdGraphPin* InputPin, FConnectionParams& Params) override;
+};
diff --git a/Engine/Plugins/Experimental/Mutable/Source/CustomizableObjectEditor/Private/MuCOE/EdGraphSchema_CustomizableObject.cpp b/Engine/Plugins/Experimental/Mutable/Source/CustomizableObjectEditor/Private/MuCOE/EdGraphSchema_CustomizableObject.cpp
--- a/Engine/Plugins/Experimental/Mutable/Source/CustomizableObjectEditor/Private/MuCOE/EdGraphSchema_CustomizableObject.cpp
+++ b/Engine/Plugins/Experimental/Mutable/Source/CustomizableObjectEditor/Private/MuCOE/EdGraphSchema_CustomizableObject.cpp
@@ -2,6 +2,7 @@
 
 #include "MuCOE/EdGraphSchema_CustomizableObject.h"
 
+#include "CustomizableObjectConnectionDrawingPolicy.h"
 #include "EdGraphNode_Comment.h"
 #include "Engine/SkeletalMesh.h"
 #include "Engine/StaticMesh.h"
@@ -640,7 +641,7 @@ const FPinConnectionResponse UEdGraphSchema_CustomizableObject::CanCreateConnect
 }
 
 
-FLinearColor UEdGraphSchema_CustomizableObject::GetPinTypeColor(const FName& TypeString) const
+FLinearColor UEdGraphSchema_CustomizableObject::GetPinTypeColor(const FName& TypeString)
 {
 	if (TypeString == PC_Enum)
 	{
@@ -1092,6 +1093,12 @@ void UEdGraphSchema_CustomizableObject::OnPinConnectionDoubleCicked(UEdGraphPin*
 }
 
 
+FConnectionDrawingPolicy* UEdGraphSchema_CustomizableObject::CreateConnectionDrawingPolicy(int32 InBackLayerID, int32 InFrontLayerID, float InZoomFactor, const FSlateRect& InClippingRect, FSlateWindowElementList& InDrawElements, UEdGraph* InGraphObj) const
+{
+	return new FCustomizableObjectConnectionDrawingPolicy(InBackLayerID, InFrontLayerID, InZoomFactor, InClippingRect, InDrawElements);
+}
+
+
 bool UEdGraphSchema_CustomizableObject::IsSpawnableAsset(const FAssetData& InAsset, ESpawnableObjectType& OutObjectType) const
 {
 	UObject* Object = InAsset.GetAsset();
diff --git a/Engine/Plugins/Experimental/Mutable/Source/CustomizableObjectEditor/Public/MuCOE/EdGraphSchema_CustomizableObject.h b/Engine/Plugins/Experimental/Mutable/Source/CustomizableObjectEditor/Public/MuCOE/EdGraphSchema_CustomizableObject.h
--- a/Engine/Plugins/Experimental/Mutable/Source/CustomizableObjectEditor/Public/MuCOE/EdGraphSchema_CustomizableObject.h
+++ b/Engine/Plugins/Experimental/Mutable/Source/CustomizableObjectEditor/Public/MuCOE/EdGraphSchema_CustomizableObject.h
@@ -113,8 +113,9 @@ public:
 	virtual bool ShouldHidePinDefaultValue(UEdGraphPin* Pin) const override;
 	virtual void DroppedAssetsOnGraph(const TArray<FAssetData>& Assets, const FVector2D& GraphPosition, UEdGraph* Graph) const override;
 	virtual void OnPinConnectionDoubleCicked(UEdGraphPin* PinA, UEdGraphPin* PinB, const FVector2D& GraphPosition) const override;
+	virtual FConnectionDrawingPolicy* CreateConnectionDrawingPolicy(int32 InBackLayerID, int32 InFrontLayerID, float InZoomFactor, const FSlateRect& InClippingRect, FSlateWindowElementList& InDrawElements, UEdGraph* InGraphObj) const override;
 
-	FLinearColor GetPinTypeColor(const FName& PinType) const;
+	static FLinearColor GetPinTypeColor(const FName& PinType);
 
 	void GetContextMenuActionsReconstructAllChildNodes(UToolMenu* Menu, TWeakObjectPtr<UGraphNodeContextMenuContext> Context) const;
 	virtual void GetContextMenuActions(UToolMenu* Menu, UGraphNodeContextMenuContext* Context) const override;
