[mutable] Added thumbnail preview for int parameter options.

#rnx

diff --git a/Engine/Plugins/Experimental/Mutable/Source/CustomizableObject/Internal/MuCO/CustomizableObjectInstancePrivate.h b/Engine/Plugins/Experimental/Mutable/Source/CustomizableObject/Internal/MuCO/CustomizableObjectInstancePrivate.h
--- a/Engine/Plugins/Experimental/Mutable/Source/CustomizableObject/Internal/MuCO/CustomizableObjectInstancePrivate.h
+++ b/Engine/Plugins/Experimental/Mutable/Source/CustomizableObject/Internal/MuCO/CustomizableObjectInstancePrivate.h
@@ -388,6 +388,7 @@ public:
 	bool bShowOnlyRuntimeParameters = true;
 	bool bShowOnlyRelevantParameters = true;
 	bool bShowUISections = false;
+	bool bShowUIThumbnails = false;
 	
 	// TEMP VARIABLE to check the Min desired LODs for this instance
 	TWeakObjectPtr<UCustomizableObjectInstanceUsage> NearestToActor;
diff --git a/Engine/Plugins/Experimental/Mutable/Source/CustomizableObject/Internal/MuCO/CustomizableObjectPrivate.h b/Engine/Plugins/Experimental/Mutable/Source/CustomizableObject/Internal/MuCO/CustomizableObjectPrivate.h
--- a/Engine/Plugins/Experimental/Mutable/Source/CustomizableObject/Internal/MuCO/CustomizableObjectPrivate.h
+++ b/Engine/Plugins/Experimental/Mutable/Source/CustomizableObject/Internal/MuCO/CustomizableObjectPrivate.h
@@ -1073,6 +1073,9 @@ public:
 
 		AddEditorGamePlayTags,
 
+		AddedParameterThumbnailsToEditor,
+
+		// -----<new versions can be added above this line>--------
 		LastCustomizableObjectVersion
 	};
 	static constexpr int32 CurrentSupportedVersion = ECustomizableObjectVersions::LastCustomizableObjectVersion;
diff --git a/Engine/Plugins/Experimental/Mutable/Source/CustomizableObject/Private/MuCO/CustomizableObjectUIData.cpp b/Engine/Plugins/Experimental/Mutable/Source/CustomizableObject/Private/MuCO/CustomizableObjectUIData.cpp
--- a/Engine/Plugins/Experimental/Mutable/Source/CustomizableObject/Private/MuCO/CustomizableObjectUIData.cpp
+++ b/Engine/Plugins/Experimental/Mutable/Source/CustomizableObject/Private/MuCO/CustomizableObjectUIData.cpp
@@ -25,6 +25,20 @@ FArchive& operator<<(FArchive& Ar, FMutableUIMetadata& Metadata)
 		Ar << StringRef;
 	}
 
+#if WITH_EDITORONLY_DATA
+	if (Ar.IsLoading())
+	{
+		FString StringRef;
+		Ar << StringRef;
+		Metadata.EditorUIThumbnailObject = TSoftObjectPtr<UObject>(FSoftObjectPath(StringRef));
+	}
+	else
+	{
+		FString StringRef = Metadata.EditorUIThumbnailObject.ToSoftObjectPath().ToString();
+		Ar << StringRef;
+	}
+#endif // WITH_EDITORONLY_DATA
+
 	Ar << Metadata.ExtraInformation;
 
 	if (Ar.IsLoading())
@@ -55,7 +69,6 @@ FArchive& operator<<(FArchive& Ar, FMutableUIMetadata& Metadata)
 		}
 	}
 
-
 	return Ar;
 }
 
diff --git a/Engine/Plugins/Experimental/Mutable/Source/CustomizableObject/Public/MuCO/CustomizableObjectUIData.h b/Engine/Plugins/Experimental/Mutable/Source/CustomizableObject/Public/MuCO/CustomizableObjectUIData.h
--- a/Engine/Plugins/Experimental/Mutable/Source/CustomizableObject/Public/MuCO/CustomizableObjectUIData.h
+++ b/Engine/Plugins/Experimental/Mutable/Source/CustomizableObject/Public/MuCO/CustomizableObjectUIData.h
@@ -36,6 +36,12 @@ struct FMutableUIMetadata
 	UPROPERTY(BlueprintReadWrite, EditAnywhere, Category = UI)
 	TSoftObjectPtr<UTexture2D> UIThumbnail;
 
+#if WITH_EDITORONLY_DATA
+	/** Editor Only Parameter Thumbnail. In an editor parameter combobox, this option will be represented with the selected asset's thumbnail. */
+	UPROPERTY(EditAnywhere, Category = UI)
+	TSoftObjectPtr<UObject> EditorUIThumbnailObject;
+#endif
+
 	/** Extra information to be used in UI building, with semantics completely defined by the game/UI programmer, with a key to identify the semantic of its related value */
 	UPROPERTY(BlueprintReadWrite, EditAnywhere, Category = UI)
 	TMap<FString, FString> ExtraInformation;
diff --git a/Engine/Plugins/Experimental/Mutable/Source/CustomizableObjectEditor/Private/MuCOE/CustomizableInstanceDetails.cpp b/Engine/Plugins/Experimental/Mutable/Source/CustomizableObjectEditor/Private/MuCOE/CustomizableInstanceDetails.cpp
--- a/Engine/Plugins/Experimental/Mutable/Source/CustomizableObjectEditor/Private/MuCOE/CustomizableInstanceDetails.cpp
+++ b/Engine/Plugins/Experimental/Mutable/Source/CustomizableObjectEditor/Private/MuCOE/CustomizableInstanceDetails.cpp
@@ -2,8 +2,8 @@
 
 #include "MuCOE/CustomizableInstanceDetails.h"
 
+#include "AssetRegistry/AssetRegistryModule.h"
 #include "CustomizableObjectInstanceEditor.h"
-
 #include "ContentBrowserModule.h"
 #include "DetailLayoutBuilder.h"
 #include "DetailWidgetRow.h"
@@ -21,6 +21,7 @@
 #include "MuCO/CustomizableObjectSystem.h"
 #include "MuCOE/UnrealEditorPortabilityHelpers.h"
 
+#include "AssetThumbnail.h"
 #include "Slate/DeferredCleanupSlateBrush.h"
 #include "SSearchableComboBox.h"
 #include "Types/SlateEnums.h"
@@ -186,6 +187,23 @@ void FCustomizableInstanceDetails::CustomizeDetails(const TSharedPtr<IDetailLayo
 		.OnCheckStateChanged(this, &FCustomizableInstanceDetails::OnUseUISectionsSelectionChanged)
 	];
 
+	
+	// Show UI thumbnails Option
+	VisibilitySettingsCategory.AddCustomRow( LOCTEXT("CustomizableInstanceDetails_UIThumbnails", "UI Thumbnails") )
+	.NameContent()
+	[
+		SNew(STextBlock)
+		.Text(FText::FromString("UI Thumbnails"))
+	]
+	.ValueContent()
+	.HAlign(EHorizontalAlignment::HAlign_Fill)
+	[
+		SNew(SCheckBox)
+		.IsChecked(CustomInstance->GetPrivate()->bShowUIThumbnails ? ECheckBoxState::Checked : ECheckBoxState::Unchecked)
+		.OnCheckStateChanged(this, &FCustomizableInstanceDetails::OnUseUIThumbnailsSelectionChanged)
+	];
+
+
 	TSharedPtr<ICustomizableObjectInstanceEditor> Editor = GetEditorChecked();
 	if (UCustomizableObjectEditorProperties* EditorProperties = Editor->GetEditorProperties())
 	{
@@ -555,6 +573,12 @@ void FCustomizableInstanceDetails::OnUseUISectionsSelectionChanged(ECheckBoxStat
 	Refresh();
 }
 
+void FCustomizableInstanceDetails::OnUseUIThumbnailsSelectionChanged(ECheckBoxState InCheckboxState)
+{
+	CustomInstance->GetPrivate()->bShowUIThumbnails = InCheckboxState == ECheckBoxState::Checked;
+	Refresh();
+}
+
 
 // PARAMETER INFO STRUCT -----------------------------------------------------------------------------------------------------------------
 
@@ -996,7 +1020,7 @@ TSharedRef<SWidget> FCustomizableInstanceDetails::GenerateIntWidget(const UCusto
 		.InitiallySelectedItem(SelectedOptionString)
 		.Method(EPopupMethod::UseCurrentWindow)
 		.OnSelectionChanged(this, &FCustomizableInstanceDetails::OnIntParameterComboBoxChanged, ParamName)
-		.OnGenerateWidget(this, &FCustomizableInstanceDetails::OnGenerateWidgetIntParameter)
+		.OnGenerateWidget(this, &FCustomizableInstanceDetails::OnGenerateWidgetIntParameter, ParamName)
 		.Content()
 		[
 			SNew(STextBlock)
@@ -1005,9 +1029,78 @@ TSharedRef<SWidget> FCustomizableInstanceDetails::GenerateIntWidget(const UCusto
 }
 
 
-TSharedRef<SWidget> FCustomizableInstanceDetails::OnGenerateWidgetIntParameter(TSharedPtr<FString> InItem) const
+TSharedRef<SWidget> FCustomizableInstanceDetails::OnGenerateWidgetIntParameter(TSharedPtr<FString> OptionName, const FString ParameterName)
 {
-	return SNew(STextBlock).Text(FText::FromString(*InItem.Get()));
+	// Final widget
+	TSharedPtr<SHorizontalBox> IntWidgetBox = SNew(SHorizontalBox);
+
+	if (CustomInstance->GetPrivate()->bShowUIThumbnails)
+	{
+		bool bUsesCustomThumbnail = false;
+
+		// Asset with the thumbnail info
+		FAssetData AssetData;
+		FAssetRegistryModule& AssetRegistryModule = FModuleManager::LoadModuleChecked<FAssetRegistryModule>("AssetRegistry");
+
+		// Metadata of the represented int option
+		const FMutableParamUIMetadata ParameterMetadata = CustomInstance->GetCustomizableObject()->GetIntParameterOptionUIMetadata(ParameterName, *OptionName);
+
+		// Custom thumbnail has preference
+		if (!ParameterMetadata.UIThumbnail.IsNull())
+		{
+			// Custom Thumbnail
+			AssetData = AssetRegistryModule.Get().GetAssetByObjectPath(ParameterMetadata.UIThumbnail.ToSoftObjectPath());
+
+			if (AssetData.IsValid())
+			{
+				if (UTexture2D* UIThumbnail = Cast<UTexture2D>(AssetData.GetAsset()))
+				{
+					// We need to store the generated texture.
+					TSharedRef<FDeferredCleanupSlateBrush> Brush = FDeferredCleanupSlateBrush::CreateBrush(UIThumbnail, FVector2D(68.f, 68.0f), /* Texture size (64) + padding (4) */
+						FLinearColor(1.0f, 1.0f, 1.0f, 1.0f), ESlateBrushTileType::NoTile, ESlateBrushImageType::Linear);
+					DynamicBrushes.Add(Brush);
+
+					IntWidgetBox->AddSlot().AutoWidth()
+					[
+						SNew(SImage).Image(Brush->GetSlateBrush())
+					];
+
+					bUsesCustomThumbnail = true;
+				}
+			}
+		}
+
+		if (!bUsesCustomThumbnail)
+		{
+			// Asset thumbnail
+			AssetData = AssetRegistryModule.Get().GetAssetByObjectPath(ParameterMetadata.EditorUIThumbnailObject.ToSoftObjectPath());
+
+			// We don't need to check if the asset data is valid here. We want to use the default thumbnail if there is no asset data.
+			TSharedPtr<FAssetThumbnail> Thumbnail = MakeShareable(new FAssetThumbnail(AssetData, 64, 64, UThumbnailManager::Get().GetSharedThumbnailPool()));
+			FAssetThumbnailConfig ThumbnailConfig;
+
+			ThumbnailConfig.ColorStripOrientation = EThumbnailColorStripOrientation::VerticalRightEdge;
+			ThumbnailConfig.Padding = FMargin(2.0f); // Prevents overlap with rounded corners; this matches what the Content Browser tiles do
+
+			IntWidgetBox->AddSlot().AutoWidth()
+			[
+				Thumbnail->MakeThumbnailWidget(ThumbnailConfig)
+			];
+
+			if (!AssetData.IsValid())
+			{
+				IntWidgetBox->GetSlot(0).SetFillWidth(2.0f);
+				IntWidgetBox->GetSlot(0).SetMaxSize(68.0f);
+			}
+		}
+	}
+
+	IntWidgetBox->AddSlot().VAlign(EVerticalAlignment::VAlign_Center).Padding(10.0f, 0.0f, 0.0f, 0.0f)
+	[
+		SNew(STextBlock).Text(FText::FromString(*OptionName.Get()))
+	];
+
+	return IntWidgetBox.ToSharedRef();
 }
 
 
diff --git a/Engine/Plugins/Experimental/Mutable/Source/CustomizableObjectEditor/Private/MuCOE/CustomizableInstanceDetails.h b/Engine/Plugins/Experimental/Mutable/Source/CustomizableObjectEditor/Private/MuCOE/CustomizableInstanceDetails.h
--- a/Engine/Plugins/Experimental/Mutable/Source/CustomizableObjectEditor/Private/MuCOE/CustomizableInstanceDetails.h
+++ b/Engine/Plugins/Experimental/Mutable/Source/CustomizableObjectEditor/Private/MuCOE/CustomizableInstanceDetails.h
@@ -68,6 +68,7 @@ private:
 	void OnShowOnlyRuntimeSelectionChanged(ECheckBoxState InCheckboxState);
 	void OnShowOnlyRelevantSelectionChanged(ECheckBoxState InCheckboxState);
 	void OnUseUISectionsSelectionChanged(ECheckBoxState InCheckboxState);
+	void OnUseUIThumbnailsSelectionChanged(ECheckBoxState InCheckboxState);
 
 	// Main parameter generation functions
 	// Returns true if parameters have been hidden due to runtime type
@@ -90,7 +91,7 @@ private:
 	// Int Parameters Functions
 	TSharedRef<SWidget> GenerateIntWidget(const UCustomizableObject& CustomizableObject, const FString& ParamName, const int32 ParamIndexInObject);
 	void OnIntParameterComboBoxChanged(TSharedPtr<FString> Selection, ESelectInfo::Type SelectInfo, const FString ParamName);
-	TSharedRef<SWidget> OnGenerateWidgetIntParameter(TSharedPtr<FString> InItem) const;
+	TSharedRef<SWidget> OnGenerateWidgetIntParameter(TSharedPtr<FString> OptionName, const FString ParameterName);
 	
 	// Float Parameters Functions
 	TSharedRef<SWidget> GenerateFloatWidget(const UCustomizableObject& CustomizableObject, const FString& ParamName);
diff --git a/Engine/Plugins/Experimental/Mutable/Source/CustomizableObjectEditor/Private/MuCOE/GenerateMutableSource/GenerateMutableSourceTable.cpp b/Engine/Plugins/Experimental/Mutable/Source/CustomizableObjectEditor/Private/MuCOE/GenerateMutableSource/GenerateMutableSourceTable.cpp
--- a/Engine/Plugins/Experimental/Mutable/Source/CustomizableObjectEditor/Private/MuCOE/GenerateMutableSource/GenerateMutableSourceTable.cpp
+++ b/Engine/Plugins/Experimental/Mutable/Source/CustomizableObjectEditor/Private/MuCOE/GenerateMutableSource/GenerateMutableSourceTable.cpp
@@ -862,47 +862,63 @@ void GenerateTableParameterUIData(const UDataTable* DataTable, const UCustomizab
 
 		FMutableParameterData ParameterUIData(TableNode->ParamUIMetadata, EMutableParameterType::Int);
 		ParameterUIData.IntegerParameterGroupType = TableNode->bAddNoneOption ? ECustomizableObjectGroupType::COGT_ONE_OR_NONE : ECustomizableObjectGroupType::COGT_ONE;
+		
 		FMutableParameterData& ParameterUIDataRef = GenerationContext.ParameterUIDataMap.Add(TableNode->ParameterName, ParameterUIData);
+		FProperty* MetadataColumnProperty = DataTable->FindTableProperty(TableNode->ParamUIMetadataColumn);
+		bool bIsValidMetadataColumn = MetadataColumnProperty &&
+			CastField<FStructProperty>(MetadataColumnProperty) &&
+			CastField<FStructProperty>(MetadataColumnProperty)->Struct == FMutableParamUIMetadata::StaticStruct();
 
-		if (TableNode->ParamUIMetadataColumn.IsNone())
+		// Trigger warning only if the name is different than "None"
+		if (!TableNode->ParamUIMetadataColumn.IsNone() && !bIsValidMetadataColumn)
 		{
-			return;
+			FText LogMessage = FText::Format(LOCTEXT("InvalidParamUIMetadataColumn_Warning",
+				"UI Metadata Column [{0}] is not a valid type or does not exist in the Structure of the Node."), FText::FromName(TableNode->ParamUIMetadataColumn));
+			GenerationContext.Compiler->CompilerLog(LogMessage, TableNode);
 		}
 
-		FProperty* ColumnProperty = DataTable->FindTableProperty(TableNode->ParamUIMetadataColumn);
+		FProperty* ThumbnailColumnProperty = DataTable->FindTableProperty(TableNode->ThumbnailColumn);
+		bool bIsValidThumbnailColumn = ThumbnailColumnProperty && CastField<FSoftObjectProperty>(ThumbnailColumnProperty);
 
-		if (!ColumnProperty)
+		// Trigger warning only if the name is different than "None"
+		if (!TableNode->ThumbnailColumn.IsNone() && !bIsValidThumbnailColumn)
 		{
-			FString msg = "Couldn't find Options UI Metadata Column [" + TableNode->ParamUIMetadataColumn.ToString() + "] in the Structure of the Node.";
-			GenerationContext.Compiler->CompilerLog(FText::FromString(msg), TableNode);
+			FText LogMessage = FText::Format(LOCTEXT("InvalidThumbnailColumn_Warning",
+				"Thumbnail Column [{0}] is not an objet type or does not exist in the Structure of the Node."), FText::FromName(TableNode->ThumbnailColumn));
+			GenerationContext.Compiler->CompilerLog(LogMessage, TableNode);
+		}
 
+		if (!bIsValidMetadataColumn && !bIsValidThumbnailColumn)
+		{
+			// Do no add a meta data if it is not needed.
 			return;
 		}
 
-		FString WrongTypeMessage = "Column with name [" + TableNode->ParamUIMetadataColumn.ToString() + "] is not a Mutable Param UI Metadata type.";
-
-		if (const FStructProperty* StructProperty = CastField<FStructProperty>(ColumnProperty))
+		for (int32 NameIndex = 0; NameIndex < RowNames.Num(); ++NameIndex)
 		{
-			if (StructProperty->Struct != FMutableParamUIMetadata::StaticStruct())
-			{
-				GenerationContext.Compiler->CompilerLog(FText::FromString(WrongTypeMessage), TableNode);
+			FName RowName = RowNames[NameIndex];
+			FMutableParamUIMetadata MetadataValue;
 
-				return;
+			if (bIsValidMetadataColumn)
+			{
+				if (uint8* CellData = GetCellData(RowName, *DataTable, *MetadataColumnProperty))
+				{
+					MetadataValue = *(FMutableParamUIMetadata*)CellData;
+				}
 			}
-			
-			for (int32 NameIndex = 0; NameIndex < RowNames.Num(); ++NameIndex)
+
+			FIntegerParameterUIData IntegerMetadata = FIntegerParameterUIData(MetadataValue);
+
+			if (bIsValidThumbnailColumn && MetadataValue.EditorUIThumbnailObject.IsNull())
 			{
-				if (uint8* CellData = GetCellData(RowNames[NameIndex], *DataTable, *ColumnProperty))
+				if (uint8* CellData = GetCellData(RowName, *DataTable, *ThumbnailColumnProperty))
 				{
-					FMutableParamUIMetadata Value = *(FMutableParamUIMetadata*)CellData;
-					ParameterUIDataRef.ArrayIntegerParameterOption.Add(RowNames[NameIndex].ToString(), FIntegerParameterUIData(Value));
+					FSoftObjectPtr ObjectPtr = *(FSoftObjectPtr*)CellData;
+					IntegerMetadata.ParamUIMetadata.EditorUIThumbnailObject = ObjectPtr.ToSoftObjectPath();
 				}
 			}
-		}
-		else
-		{
-			GenerationContext.Compiler->CompilerLog(FText::FromString(WrongTypeMessage), TableNode);
-			return;
+
+			ParameterUIDataRef.ArrayIntegerParameterOption.Add(RowName.ToString(), IntegerMetadata);
 		}
 	}
 }
diff --git a/Engine/Plugins/Experimental/Mutable/Source/CustomizableObjectEditor/Private/MuCOE/Nodes/CustomizableObjectNodeTableDetails.cpp b/Engine/Plugins/Experimental/Mutable/Source/CustomizableObjectEditor/Private/MuCOE/Nodes/CustomizableObjectNodeTableDetails.cpp
--- a/Engine/Plugins/Experimental/Mutable/Source/CustomizableObjectEditor/Private/MuCOE/Nodes/CustomizableObjectNodeTableDetails.cpp
+++ b/Engine/Plugins/Experimental/Mutable/Source/CustomizableObjectEditor/Private/MuCOE/Nodes/CustomizableObjectNodeTableDetails.cpp
@@ -54,6 +54,7 @@ void FCustomizableObjectNodeTableDetails::CustomizeDetails(const TSharedPtr<IDet
 		DetailBuilder->HideProperty("VersionColumn");
 		IDetailCategoryBuilder& UICategory = DetailBuilder->EditCategory("UI");
 		DetailBuilder->HideProperty("ParamUIMetadataColumn");
+		DetailBuilder->HideProperty("ThumbnailColumn");
 		IDetailCategoryBuilder& AnimationCategory = DetailBuilder->EditCategory("AnimationProperties");
 		IDetailCategoryBuilder& LayoutCategory = DetailBuilder->EditCategory("DefaultMeshLayoutEditor");
 
@@ -63,6 +64,7 @@ void FCustomizableObjectNodeTableDetails::CustomizeDetails(const TSharedPtr<IDet
 		GenerateMeshColumnComboBoxOptions();
 		TSharedPtr<FString> CurrentMutableMetadataColumn = GenerateMutableMetaDataColumnComboBoxOptions();
 		TSharedPtr<FString> CurrentVersionColumn = GenerateVersionColumnComboBoxOptions();
+		TSharedPtr<FString> CurrentThumbnailColumn = GenerateThumbnailColumnComboBoxOptions();
 
 		CustomizableObjectCategory.AddProperty("ParameterName");
 		CustomizableObjectCategory.AddCustomRow(LOCTEXT("VersionColumn_Selector","VersionColumn"))
@@ -101,9 +103,29 @@ void FCustomizableObjectNodeTableDetails::CustomizeDetails(const TSharedPtr<IDet
 			.OnComboBoxOpening(this, &FCustomizableObjectNodeTableDetails::OnOpenMutableMetadataComboBox)
 			.OnSelectionChanged(this, &FCustomizableObjectNodeTableDetails::OnMutableMetaDataColumnComboBoxSelectionChanged)
 			.Font(IDetailLayoutBuilder::GetDetailFont())
-			.ColorAndOpacity(this, &FCustomizableObjectNodeTableDetails::GetMetadataUIComboBoxTextColor, &MutableMetaDataColumnsOptionNames)
+			.ColorAndOpacity(this, &FCustomizableObjectNodeTableDetails::GetComboBoxTextColor, &MutableMetaDataColumnsOptionNames, Node->ParamUIMetadataColumn)
 		]
 		.OverrideResetToDefault(FResetToDefaultOverride::Create(FSimpleDelegate::CreateSP(this, &FCustomizableObjectNodeTableDetails::OnMutableMetaDataColumnComboBoxSelectionReset)));
+		
+		UICategory.AddCustomRow(LOCTEXT("ThumbnailColumn_Selector", "ThumbnailColumn"))
+		.NameContent()
+		[
+			SNew(STextBlock)
+			.Text(LOCTEXT("ThumbnailColumn_SelectorText", "Options Thumbnail Column"))
+			.ToolTipText(LOCTEXT("ThumbnailColumn_SelectorTooltip", "Select a column that contains the assets to use its thumbnails as Option thumbnails."))
+			.Font(IDetailLayoutBuilder::GetDetailFont())
+		]
+		.ValueContent()
+		[
+			SAssignNew(ThumbnailComboBox,STextComboBox)
+			.InitiallySelectedItem(CurrentThumbnailColumn)
+			.OptionsSource(&ThumbnailColumnOptionNames)
+			.OnComboBoxOpening(this, &FCustomizableObjectNodeTableDetails::OnOpenThumbnailComboBox)
+			.OnSelectionChanged(this, &FCustomizableObjectNodeTableDetails::OnThumbnailColumnComboBoxSelectionChanged)
+			.Font(IDetailLayoutBuilder::GetDetailFont())
+			.ColorAndOpacity(this, &FCustomizableObjectNodeTableDetails::GetComboBoxTextColor, &MutableMetaDataColumnsOptionNames, Node->ThumbnailColumn)
+		]
+		.OverrideResetToDefault(FResetToDefaultOverride::Create(FSimpleDelegate::CreateSP(this, &FCustomizableObjectNodeTableDetails::OnThumbnailColumnComboBoxSelectionReset)));
 
 
 
@@ -1017,9 +1039,9 @@ void FCustomizableObjectNodeTableDetails::OnMutableMetaDataColumnComboBoxSelecti
 }
 
 
-FSlateColor FCustomizableObjectNodeTableDetails::GetMetadataUIComboBoxTextColor(TArray<TSharedPtr<FString>>* CurrentOptions) const
+FSlateColor FCustomizableObjectNodeTableDetails::GetComboBoxTextColor(TArray<TSharedPtr<FString>>* CurrentOptions, const FName ColumnName) const
 {	
-	if (Node->FindTableProperty(Node->GetTableNodeStruct(), Node->ParamUIMetadataColumn) || Node->ParamUIMetadataColumn.IsNone())
+	if (Node->FindTableProperty(Node->GetTableNodeStruct(), ColumnName) || ColumnName.IsNone())
 	{
 		return FSlateColor::UseForeground();
 	}
@@ -1042,6 +1064,86 @@ void FCustomizableObjectNodeTableDetails::OnMutableMetaDataColumnComboBoxSelecti
 }
 
 
+TSharedPtr<FString> FCustomizableObjectNodeTableDetails::GenerateThumbnailColumnComboBoxOptions()
+{
+	const UScriptStruct* TableStruct = Node->GetTableNodeStruct();
+	TSharedPtr<FString> CurrentSelection;
+	ThumbnailColumnOptionNames.Reset();
+
+	if (!TableStruct)
+	{
+		return CurrentSelection;
+	}
+
+	// Iterating struct Options
+	for (TFieldIterator<FProperty> It(TableStruct); It; ++It)
+	{
+		FProperty* ColumnProperty = *It;
+
+		if (!ColumnProperty)
+		{
+			continue;
+		}
+
+		if (const FSoftObjectProperty* ObjectProperty = CastField<FSoftObjectProperty>(ColumnProperty))
+		{
+			TSharedPtr<FString> Option = MakeShareable(new FString(DataTableUtils::GetPropertyExportName(ColumnProperty)));
+			ThumbnailColumnOptionNames.Add(Option);
+
+			if (*Option == Node->ThumbnailColumn)
+			{
+				CurrentSelection = ThumbnailColumnOptionNames.Last();
+			}
+		}
+	}
+
+	if (!Node->ThumbnailColumn.IsNone() && !CurrentSelection)
+	{
+		ThumbnailColumnOptionNames.Add(MakeShareable(new FString(Node->ThumbnailColumn.ToString())));
+		CurrentSelection = ThumbnailColumnOptionNames.Last();
+	}
+
+	return CurrentSelection;
+}
+
+
+void FCustomizableObjectNodeTableDetails::OnOpenThumbnailComboBox()
+{
+	TSharedPtr<FString> CurrentSelection = GenerateThumbnailColumnComboBoxOptions();
+
+	if (ThumbnailComboBox.IsValid())
+	{
+		ThumbnailComboBox->ClearSelection();
+		ThumbnailComboBox->RefreshOptions();
+		ThumbnailComboBox->SetSelectedItem(CurrentSelection);
+	}
+}
+
+
+void FCustomizableObjectNodeTableDetails::OnThumbnailColumnComboBoxSelectionChanged(TSharedPtr<FString> Selection, ESelectInfo::Type SelectInfo)
+{
+	if (Selection && Node->ThumbnailColumn != FName(*Selection)
+		&& (SelectInfo == ESelectInfo::OnKeyPress || SelectInfo == ESelectInfo::OnMouseClick))
+	{
+		Node->ThumbnailColumn = FName(*Selection);
+		Node->MarkPackageDirty();
+	}
+}
+
+
+void FCustomizableObjectNodeTableDetails::OnThumbnailColumnComboBoxSelectionReset()
+{
+	Node->ThumbnailColumn = NAME_None;
+
+	if (ThumbnailComboBox.IsValid())
+	{
+		GenerateThumbnailColumnComboBoxOptions();
+		ThumbnailComboBox->ClearSelection();
+		ThumbnailComboBox->RefreshOptions();
+	}
+}
+
+
 // Version Bridge Category --------------------------------------------------------------------------------
 
 TSharedPtr<FString> FCustomizableObjectNodeTableDetails::GenerateVersionColumnComboBoxOptions()
diff --git a/Engine/Plugins/Experimental/Mutable/Source/CustomizableObjectEditor/Private/MuCOE/Nodes/CustomizableObjectNodeTableDetails.h b/Engine/Plugins/Experimental/Mutable/Source/CustomizableObjectEditor/Private/MuCOE/Nodes/CustomizableObjectNodeTableDetails.h
--- a/Engine/Plugins/Experimental/Mutable/Source/CustomizableObjectEditor/Private/MuCOE/Nodes/CustomizableObjectNodeTableDetails.h
+++ b/Engine/Plugins/Experimental/Mutable/Source/CustomizableObjectEditor/Private/MuCOE/Nodes/CustomizableObjectNodeTableDetails.h
@@ -126,17 +126,24 @@ private:
 	// Returns the current selected option or a null pointer
 	TSharedPtr<FString> GenerateMutableMetaDataColumnComboBoxOptions();
 
-	// Callback to regenerate the combobox options
+	// Generates Thumbnail columns combobox options
+	// Returns the current selected option or a null pointer
+	TSharedPtr<FString> GenerateThumbnailColumnComboBoxOptions();
+
+	// Callbacks to regenerate the combobox options
 	void OnOpenMutableMetadataComboBox();
+	void OnOpenThumbnailComboBox();
 
-	// OnComboBoxSelectionChanged Callback for Layout ComboBox
+	// OnComboBoxSelectionChanged Callbacks
 	void OnMutableMetaDataColumnComboBoxSelectionChanged(TSharedPtr<FString> Selection, ESelectInfo::Type SelectInfo);
+	void OnThumbnailColumnComboBoxSelectionChanged(TSharedPtr<FString> Selection, ESelectInfo::Type SelectInfo);
 
-	// Sets the combo box selection color
-	FSlateColor GetMetadataUIComboBoxTextColor(TArray<TSharedPtr<FString>>* CurrentOptions) const;
+	// If the property selected in the combobox does not exist anymore, returns a red color.
+	FSlateColor GetComboBoxTextColor(TArray<TSharedPtr<FString>>* CurrentOptions, const FName ColumnName) const;
 
-	// OnComboBoxSelectionChanged Callback for Layout ComboBox
+	// OnComboBoxSelectionChanged Callbacks for ComboBox
 	void OnMutableMetaDataColumnComboBoxSelectionReset();
+	void OnThumbnailColumnComboBoxSelectionReset();
 
 
 	// Version Control Category -------------
@@ -229,9 +236,14 @@ private:
 	// Mutable UI Metadata -------------
 	// Array with the name of the MutableMetaData columns
 	TArray<TSharedPtr<FString>> MutableMetaDataColumnsOptionNames;
+	
+	TArray<TSharedPtr<FString>> ThumbnailColumnOptionNames;
 
 	// ComboBox widget to select a MutableMetaDatacolumn from the NodeTable
 	TSharedPtr<STextComboBox> MutableMetaDataComboBox;
+	
+	// ComboBox widget to select a Thumbnail Column from the NodeTable
+	TSharedPtr<STextComboBox> ThumbnailComboBox;
 
 
 	// Version Bridge -------------
diff --git a/Engine/Plugins/Experimental/Mutable/Source/CustomizableObjectEditor/Public/MuCOE/Nodes/CustomizableObjectNodeTable.h b/Engine/Plugins/Experimental/Mutable/Source/CustomizableObjectEditor/Public/MuCOE/Nodes/CustomizableObjectNodeTable.h
--- a/Engine/Plugins/Experimental/Mutable/Source/CustomizableObjectEditor/Public/MuCOE/Nodes/CustomizableObjectNodeTable.h
+++ b/Engine/Plugins/Experimental/Mutable/Source/CustomizableObjectEditor/Public/MuCOE/Nodes/CustomizableObjectNodeTable.h
@@ -221,6 +221,10 @@ public:
 	/** Name of the column that contains the MutableUIMetadata of the row options. */
 	UPROPERTY(EditAnywhere, Category = UI)
 	FName ParamUIMetadataColumn;
+	
+	/** Name of the column that contains the asset to use its thumbnails as option thumbnails. */
+	UPROPERTY(EditAnywhere, Category = UI)
+	FName ThumbnailColumn;
 
 	/** Map to relate a Structure Column with its Data */
 	UPROPERTY()
