[mutable] Improved tag explorer style.
+Now it only has 2 rows, one for the Customizable Object and one for the node type, instead of a row for each node type.
+Removed duplicated function to search all the COs in a graph.

#rnx

diff --git a/Engine/Plugins/Experimental/Mutable/Source/CustomizableObjectEditor/Private/MuCOE/CustomizableObjectEditor.cpp b/Engine/Plugins/Experimental/Mutable/Source/CustomizableObjectEditor/Private/MuCOE/CustomizableObjectEditor.cpp
--- a/Engine/Plugins/Experimental/Mutable/Source/CustomizableObjectEditor/Private/MuCOE/CustomizableObjectEditor.cpp
+++ b/Engine/Plugins/Experimental/Mutable/Source/CustomizableObjectEditor/Private/MuCOE/CustomizableObjectEditor.cpp
@@ -2309,60 +2309,6 @@ void FCustomizableObjectEditor::AddCachedReferencers(const FName& PathName, TArr
 }
 
 
-void FCustomizableObjectEditor::GetExternalChildObjects(const UCustomizableObject* const Object, TArray<UCustomizableObject*>& ExternalChildren, const bool bRecursively, const EObjectFlags ExcludeFlags)
-{
-	TArray<FAssetData> ArrayAssetData;
-	TArray<FName> ArrayAlreadyProcessedChild;
-	TArray<FName> ArrayReferenceNames;
-
-	AddCachedReferencers(*Object->GetOuter()->GetPathName(), ArrayReferenceNames, ArrayAssetData);
-
-	bool bMultipleBaseObjectsFound = false;
-
-	FAssetData* AssetData = nullptr;
-	UCustomizableObject* ChildObject;
-
-	for (const FName& ReferenceName : ArrayReferenceNames)
-	{
-		const int32 MaxIndex = ArrayAssetData.Num();
-
-		for (int32 i = 0; i < MaxIndex; ++i)
-		{
-			if (ArrayAssetData[i].PackageName.ToString() == ReferenceName.ToString())
-			{
-				AssetData = &ArrayAssetData[i];
-			}
-		}
-
-		if (AssetData != nullptr) // Elements in ArrayAssetData are already of static class UCustomizableObject
-		{
-			ChildObject = Cast<UCustomizableObject>(AssetData->GetAsset());
-			if (!ChildObject) continue;
-
-			if (ChildObject != Object && !ChildObject->HasAnyFlags(ExcludeFlags))
-			{
-				ExternalChildren.AddUnique(ChildObject);
-			}
-
-			TArray<UCustomizableObjectNodeObjectGroup*> GroupNodes;
-			ChildObject->GetPrivate()->GetSource()->GetNodesOfClass<UCustomizableObjectNodeObjectGroup>(GroupNodes);
-
-			if (GroupNodes.Num() > 0) // Only grafs with group nodes should have child grafs
-			{
-				if (ArrayAlreadyProcessedChild.Find(ReferenceName) == INDEX_NONE)
-				{
-					ArrayAlreadyProcessedChild.Add(ReferenceName);
-					if (bRecursively)
-					{
-						GetExternalChildObjects(ChildObject, ExternalChildren, bRecursively);
-					}
-				}
-			}
-		}
-	}
-}
-
-
 void FCustomizableObjectEditor::CreatePreviewActor()
 {
 	if (Actor)
diff --git a/Engine/Plugins/Experimental/Mutable/Source/CustomizableObjectEditor/Private/MuCOE/CustomizableObjectEditor.h b/Engine/Plugins/Experimental/Mutable/Source/CustomizableObjectEditor/Private/MuCOE/CustomizableObjectEditor.h
--- a/Engine/Plugins/Experimental/Mutable/Source/CustomizableObjectEditor/Private/MuCOE/CustomizableObjectEditor.h
+++ b/Engine/Plugins/Experimental/Mutable/Source/CustomizableObjectEditor/Private/MuCOE/CustomizableObjectEditor.h
@@ -372,7 +372,6 @@ public:
 	// Helpers to get the absolute parent of a Customizable Object
 	static UCustomizableObject* GetAbsoluteCOParent(const UCustomizableObjectNodeObject* const Root);
 	static void AddCachedReferencers(const FName& PathName, TArray<FName>& ArrayReferenceNames, TArray<FAssetData>& ArrayAssetData);
-	static void GetExternalChildObjects(const UCustomizableObject* const Object, TArray<UCustomizableObject*>& ExternalChildren, const bool bRecursively = true, const EObjectFlags ExcludeFlags = EObjectFlags::RF_Transient);
 
 	/**	The tab ids for all the tabs used */
 	static const FName ViewportTabId;
diff --git a/Engine/Plugins/Experimental/Mutable/Source/CustomizableObjectEditor/Private/MuCOE/Nodes/CustomizableObjectNodeMaterial.cpp b/Engine/Plugins/Experimental/Mutable/Source/CustomizableObjectEditor/Private/MuCOE/Nodes/CustomizableObjectNodeMaterial.cpp
--- a/Engine/Plugins/Experimental/Mutable/Source/CustomizableObjectEditor/Private/MuCOE/Nodes/CustomizableObjectNodeMaterial.cpp
+++ b/Engine/Plugins/Experimental/Mutable/Source/CustomizableObjectEditor/Private/MuCOE/Nodes/CustomizableObjectNodeMaterial.cpp
@@ -688,17 +688,17 @@ void UCustomizableObjectNodeMaterial::PostEditChangeProperty(FPropertyChangedEve
 
 FText UCustomizableObjectNodeMaterial::GetNodeTitle(ENodeTitleType::Type TitleType) const
 {
-	if (Material)
+	if (TitleType == ENodeTitleType::ListView || !Material)
+	{
+		return LOCTEXT("Material", "Material");
+	}
+	else
 	{
 		FFormatNamedArguments Args;
 		Args.Add(TEXT("MaterialName"), FText::FromString(Material->GetName()));
 
 		return FText::Format(LOCTEXT("Material_Title", "{MaterialName}\nMaterial"), Args);
 	}
-	else
-	{
-		return LOCTEXT("Material", "Material");
-	}
 }
 
 
diff --git a/Engine/Plugins/Experimental/Mutable/Source/CustomizableObjectEditor/Private/MuCOE/SCustomizableObjectEditorTagExplorer.cpp b/Engine/Plugins/Experimental/Mutable/Source/CustomizableObjectEditor/Private/MuCOE/SCustomizableObjectEditorTagExplorer.cpp
--- a/Engine/Plugins/Experimental/Mutable/Source/CustomizableObjectEditor/Private/MuCOE/SCustomizableObjectEditorTagExplorer.cpp
+++ b/Engine/Plugins/Experimental/Mutable/Source/CustomizableObjectEditor/Private/MuCOE/SCustomizableObjectEditorTagExplorer.cpp
@@ -8,7 +8,6 @@
 #include "MuCOE/CustomizableObjectEditor.h"
 #include "MuCOE/CustomizableObjectGraph.h"
 #include "MuCOE/GraphTraversal.h"
-#include "MuCOE/Nodes/CustomizableObjectNodeObject.h"
 #include "MuCOE/Nodes/CustomizableObjectNodeExtendMaterial.h"
 #include "MuCOE/Nodes/CustomizableObjectNodeMaterial.h"
 #include "MuCOE/Nodes/CustomizableObjectNodeMaterialVariation.h"
@@ -18,15 +17,19 @@
 #include "MuCOE/Nodes/CustomizableObjectNodeModifierClipWithUVMask.h"
 #include "Widgets/Input/SButton.h"
 #include "Widgets/Input/SComboButton.h"
-#include "Widgets/Views/SListView.h"
+
 
 class ITableRow;
 class STableViewBase;
 class SWidget;
 class UObject;
 
+const FName SCustomizableObjectEditorTagExplorer::COLUMN_OBJECT(TEXT("Customizable Object"));
+const FName SCustomizableObjectEditorTagExplorer::COLUMN_TYPE(TEXT("Node Type"));
+
 #define LOCTEXT_NAMESPACE "SCustomizableObjectEditorTagExplorer"
 
+
 void SCustomizableObjectEditorTagExplorer::Construct(const FArguments & InArgs)
 {
 	CustomizableObjectEditorPtr = InArgs._CustomizableObjectEditor;
@@ -34,7 +37,6 @@ void SCustomizableObjectEditorTagExplorer::Construct(const FArguments & InArgs)
 	ChildSlot
 	[
 		SNew(SVerticalBox)
-
 		+ SVerticalBox::Slot()
 		.AutoHeight()
 		.Padding(5.0f)
@@ -80,181 +82,78 @@ void SCustomizableObjectEditorTagExplorer::Construct(const FArguments & InArgs)
 		.AutoHeight()
 		.Padding(2.0f)
 		[
-			SNew(SHorizontalBox)
-			
-			+ SHorizontalBox::Slot()
-			[
-				SAssignNew(ColumnMat, SListView<UCustomizableObjectNode*>)
-				.ListItemsSource(&MaterialNodes)
-				.OnGenerateRow(this, &SCustomizableObjectEditorTagExplorer::OnGenerateTableRow)
-				.OnSelectionChanged(this, &SCustomizableObjectEditorTagExplorer::OnTagTableSelectionChanged)
-				.SelectionMode(ESelectionMode::Single)
-				.HeaderRow
-				(
-					SNew(SHeaderRow)
-					+ SHeaderRow::Column("MaterialNode")
-					.DefaultLabel(LOCTEXT("MaterialNode_ColumnName", "Material Node"))
-					.FillWidth(0.3f)
-					.HAlignHeader(EHorizontalAlignment::HAlign_Center)
-					.HAlignCell(EHorizontalAlignment::HAlign_Center)
-				)
-			]
-			
-			+ SHorizontalBox::Slot()
-			[
-				SAssignNew(ColumnClipMesh, SListView<UCustomizableObjectNode*>)
-				.ListItemsSource(&ClipMeshNodes)
-				.OnGenerateRow(this, &SCustomizableObjectEditorTagExplorer::OnGenerateTableRow)
-				.OnSelectionChanged(this, &SCustomizableObjectEditorTagExplorer::OnTagTableSelectionChanged)
-				.HeaderRow
-				(
-					SNew(SHeaderRow)
-					+ SHeaderRow::Column("ClipMeshNode")
-					.DefaultLabel(LOCTEXT("ClipMeshNode_ColumnName", "Clip Mesh Node"))
-					.FillWidth(0.3f)
-					.HAlignHeader(EHorizontalAlignment::HAlign_Center)
-					.HAlignCell(EHorizontalAlignment::HAlign_Center)
-				)
-			]
-
-			+ SHorizontalBox::Slot()
-			[
-				SAssignNew(ColumnClipMorph, SListView<UCustomizableObjectNode*>)
-				.ListItemsSource(&ClipMorphNodes)
-				.OnGenerateRow(this, &SCustomizableObjectEditorTagExplorer::OnGenerateTableRow)
-				.OnSelectionChanged(this, &SCustomizableObjectEditorTagExplorer::OnTagTableSelectionChanged)
-				.HeaderRow
-				(
-					SNew(SHeaderRow)
-					+ SHeaderRow::Column("ClipMorphNode")
-					.DefaultLabel(LOCTEXT("ClipMorphNode_ColumnName", "Clip Morph Node"))
-					.FillWidth(0.3f)
-					.HAlignHeader(EHorizontalAlignment::HAlign_Center)
-					.HAlignCell(EHorizontalAlignment::HAlign_Center)
-				)
-			]
-	
-			+ SHorizontalBox::Slot()
-			[
-				SAssignNew(ColumnVar, SListView<UCustomizableObjectNode*>)
-				.ListItemsSource(&VariationNodes)
-				.OnGenerateRow(this, &SCustomizableObjectEditorTagExplorer::OnGenerateTableRow)
-				.OnSelectionChanged(this, &SCustomizableObjectEditorTagExplorer::OnTagTableSelectionChanged)
-				.HeaderRow
-				(
-					SNew(SHeaderRow)
-					+ SHeaderRow::Column("VariationNode")
-					.DefaultLabel(LOCTEXT("VariationNode_ColumnName", "Variation Node"))
-					.FillWidth(0.3f)
-					.HAlignHeader(EHorizontalAlignment::HAlign_Center)
-					.HAlignCell(EHorizontalAlignment::HAlign_Center)
-				)
-			]
-
-			+ SHorizontalBox::Slot()
-			[
-				SAssignNew(ColumnExtend, SListView<UCustomizableObjectNode*>)
-				.ListItemsSource(&ExtendNodes)
-				.OnGenerateRow(this, &SCustomizableObjectEditorTagExplorer::OnGenerateTableRow)
-				.OnSelectionChanged(this, &SCustomizableObjectEditorTagExplorer::OnTagTableSelectionChanged)
-				.HeaderRow
-				(
-					SNew(SHeaderRow)
-					+ SHeaderRow::Column("ExtendNode")
-					.DefaultLabel(LOCTEXT("ExtendNode_ColumnName", "Extend Node"))
-					.FillWidth(0.3f)
-					.HAlignHeader(EHorizontalAlignment::HAlign_Center)
-					.HAlignCell(EHorizontalAlignment::HAlign_Center)
-				)
-			]
-
-			+ SHorizontalBox::Slot()
-			[
-				SAssignNew(ColumnClipDeform, SListView<UCustomizableObjectNode*>)
-				.ListItemsSource(&ClipDeformNodes)
-				.OnGenerateRow(this, &SCustomizableObjectEditorTagExplorer::OnGenerateTableRow)
-				.OnSelectionChanged(this, &SCustomizableObjectEditorTagExplorer::OnTagTableSelectionChanged)
-				.HeaderRow
-				(
-					SNew(SHeaderRow)
-					+ SHeaderRow::Column("ClipDeform")
-					.DefaultLabel(LOCTEXT("ClipDeform_ColumnName", "Clip Deform Node"))
-					.FillWidth(0.3f)
-					.HAlignHeader(EHorizontalAlignment::HAlign_Center)
-					.HAlignCell(EHorizontalAlignment::HAlign_Center)
-				)
-			]
+			SAssignNew(ListViewWidget, SListView<TWeakObjectPtr<UCustomizableObjectNode>>)
+			.ListItemsSource(&Nodes)
+			.OnGenerateRow(this, &SCustomizableObjectEditorTagExplorer::OnGenerateTableRow)
+			.OnSelectionChanged(this, &SCustomizableObjectEditorTagExplorer::OnTagTableSelectionChanged)
+			.SelectionMode(ESelectionMode::Single)
+			.HeaderRow
+			(
+				SNew(SHeaderRow)
+				+ SHeaderRow::Column(COLUMN_OBJECT)
+				.DefaultLabel(LOCTEXT("CustomizableObject_ColumnName", "Customizable Object"))
+				.FillWidth(0.5f)
+				.HAlignHeader(EHorizontalAlignment::HAlign_Center)
+				.HAlignCell(EHorizontalAlignment::HAlign_Left)
+				.SortMode(this, &SCustomizableObjectEditorTagExplorer::GetColumnSortMode, COLUMN_OBJECT)
+				.OnSort(this, &SCustomizableObjectEditorTagExplorer::SortListView)
+
+				+ SHeaderRow::Column(COLUMN_TYPE)
+				.DefaultLabel(LOCTEXT("NodeType_ColumnName", "Node Type"))
+				.FillWidth(0.5f)
+				.HAlignHeader(EHorizontalAlignment::HAlign_Center)
+				.HAlignCell(EHorizontalAlignment::HAlign_Left)
+				.SortMode(this, &SCustomizableObjectEditorTagExplorer::GetColumnSortMode, COLUMN_TYPE)
+				.OnSort(this, &SCustomizableObjectEditorTagExplorer::SortListView)
+			)
 		]
 	];
 }
 
+
 TSharedRef<SWidget> SCustomizableObjectEditorTagExplorer::OnGetTagsMenuContent()
 {
-	bool bMultipleBaseObjectsFound = false;
 	NodeTags.Empty();
 
-	UCustomizableObject* CustomizableObject = CustomizableObjectEditorPtr->GetCustomizableObject();
-
-	if (!CustomizableObject)
+	if (UCustomizableObject* CustomizableObject = CustomizableObjectEditorPtr->GetCustomizableObject())
 	{
-		return SNullWidget::NullWidget;
-	}
-
-	UCustomizableObjectNodeObject* Root = GetRootNode(CustomizableObject, bMultipleBaseObjectsFound);
-	TArray<UCustomizableObject*> CustomizableObjectFamily;
-	TArray<FString> Tags;
+		TArray<FString> Tags;
+		TSet<UCustomizableObject*> CustomizableObjectTree;
 
-	if (Root && !bMultipleBaseObjectsFound)
-	{
-		UCustomizableObject* AbsoluteCO = Root->ParentObject ? CustomizableObjectEditorPtr->GetAbsoluteCOParent(Root) : CustomizableObject;
+		GetAllObjectsInGraph(CustomizableObject, CustomizableObjectTree);
 
-		if (AbsoluteCO)
+		for (const UCustomizableObject* CustObject : CustomizableObjectTree)
 		{
-			CustomizableObjectFamily.Add(AbsoluteCO);
-			// Stores external customizable objects
-			CustomizableObjectEditorPtr->GetExternalChildObjects(AbsoluteCO, CustomizableObjectFamily, true);
-
-			for (UCustomizableObject* CustObject : CustomizableObjectFamily)
+			if (CustObject)
 			{
-				FillTagInformation(CustObject, Tags);
+				FillTagInformation(*CustObject, Tags);
 			}
 		}
-	}
-	else
-	{
-		if (bMultipleBaseObjectsFound)
-		{
-			UE_LOG(LogTemp, Error, TEXT("Multiple base object nodes found. Can't check tags."));
-		}
 
-		if (!Root)
+		if (Tags.Num())
 		{
-			UE_LOG(LogTemp, Error, TEXT("No base object node found. Can't check tags."));
-		}
-	}
+			FMenuBuilder MenuBuilder(true, NULL);
 
-	if (Tags.Num())
-	{
-		FMenuBuilder MenuBuilder(true, NULL);
+			for (int32 TagIndex = 0; TagIndex < Tags.Num(); ++TagIndex)
+			{
+				FText TagText = FText::FromString(Tags[TagIndex]);
+				FUIAction Action(FExecuteAction::CreateSP(this, &SCustomizableObjectEditorTagExplorer::OnComboBoxSelectionChanged, Tags[TagIndex]));
+				MenuBuilder.AddMenuEntry(TagText, FText::GetEmpty(), FSlateIcon(), Action);
+			}
 
-		for (int32 TagIndex = 0; TagIndex < Tags.Num(); ++TagIndex)
-		{
-			FText TagText = FText::FromString(Tags[TagIndex]);
-			FUIAction Action(FExecuteAction::CreateSP(this, &SCustomizableObjectEditorTagExplorer::OnComboBoxSelectionChanged, Tags[TagIndex]));
-			MenuBuilder.AddMenuEntry(TagText, FText::GetEmpty(), FSlateIcon(), Action);
+			return MenuBuilder.MakeWidget();
 		}
-
-		return MenuBuilder.MakeWidget();
 	}
 
 	return SNullWidget::NullWidget;
 }
 
-void SCustomizableObjectEditorTagExplorer::FillTagInformation(UCustomizableObject * Object, TArray<FString>& Tags)
+
+void SCustomizableObjectEditorTagExplorer::FillTagInformation(const UCustomizableObject& Object, TArray<FString>& Tags)
 {
-	if (Object->GetPrivate()->GetSource())
+	if (Object.GetPrivate()->GetSource())
 	{
-		for (const TObjectPtr<UEdGraphNode>& Node : Object->GetPrivate()->GetSource()->Nodes)
+		for (const TObjectPtr<UEdGraphNode>& Node : Object.GetPrivate()->GetSource()->Nodes)
 		{
 			if (UCustomizableObjectNodeMaterialBase* TypedNodeMat = Cast<UCustomizableObjectNodeMaterialBase>(Node))
 			{
@@ -276,6 +175,7 @@ void SCustomizableObjectEditorTagExplorer::FillTagInformation(UCustomizableObjec
 				{
 					const FString& VairationTag = TypedNodeVariations->GetVariation(i).Tag; 
 					NodeTags.Add(VairationTag, TypedNodeVariations);
+					
 					if (Tags.Find(VairationTag) == INDEX_NONE)
 					{
 						Tags.Add(VairationTag);
@@ -288,6 +188,7 @@ void SCustomizableObjectEditorTagExplorer::FillTagInformation(UCustomizableObjec
 				for (int32 i = 0; i < TypedNodeClipMorph->Tags.Num(); ++i)
 				{
 					NodeTags.Add(TypedNodeClipMorph->Tags[i], TypedNodeClipMorph);
+					
 					if (Tags.Find(TypedNodeClipMorph->Tags[i]) == INDEX_NONE)
 					{
 						Tags.Add(TypedNodeClipMorph->Tags[i]);
@@ -300,6 +201,7 @@ void SCustomizableObjectEditorTagExplorer::FillTagInformation(UCustomizableObjec
 				for (int32 i = 0; i < TypedNodeClipMesh->Tags.Num(); ++i)
 				{
 					NodeTags.Add(TypedNodeClipMesh->Tags[i], TypedNodeClipMesh);
+					
 					if (Tags.Find(TypedNodeClipMesh->Tags[i]) == INDEX_NONE)
 					{
 						Tags.Add(TypedNodeClipMesh->Tags[i]);
@@ -312,6 +214,7 @@ void SCustomizableObjectEditorTagExplorer::FillTagInformation(UCustomizableObjec
 				for (int32 i = 0; i < TypedNodeClipMask->Tags.Num(); ++i)
 				{
 					NodeTags.Add(TypedNodeClipMask->Tags[i], TypedNodeClipMask);
+					
 					if (Tags.Find(TypedNodeClipMask->Tags[i]) == INDEX_NONE)
 					{
 						Tags.Add(TypedNodeClipMask->Tags[i]);
@@ -324,6 +227,7 @@ void SCustomizableObjectEditorTagExplorer::FillTagInformation(UCustomizableObjec
 				for (int32 i = 0; i < TypedNodeExtend->Tags.Num(); ++i)
 				{
 					NodeTags.Add(TypedNodeExtend->Tags[i], TypedNodeExtend);
+					
 					if (Tags.Find(TypedNodeExtend->Tags[i]) == INDEX_NONE)
 					{
 						Tags.Add(TypedNodeExtend->Tags[i]);
@@ -336,6 +240,7 @@ void SCustomizableObjectEditorTagExplorer::FillTagInformation(UCustomizableObjec
 				for (int32 i = 0; i < TypedNodeClipDeform->Tags.Num(); ++i)
 				{
 					NodeTags.Add(TypedNodeClipDeform->Tags[i], TypedNodeClipDeform);
+					
 					if (Tags.Find(TypedNodeClipDeform->Tags[i]) == INDEX_NONE)
 					{
 						Tags.Add(TypedNodeClipDeform->Tags[i]);
@@ -346,11 +251,13 @@ void SCustomizableObjectEditorTagExplorer::FillTagInformation(UCustomizableObjec
 	}
 }
 
+
 TSharedRef<SWidget> SCustomizableObjectEditorTagExplorer::MakeComboButtonItemWidget(TSharedPtr<FString> StringItem)
 {
 	return SNew(STextBlock).Text(FText::FromString(*StringItem));
 }
 
+
 FText SCustomizableObjectEditorTagExplorer::GetCurrentItemLabel() const
 {
 	if (!SelectedTag.IsEmpty())
@@ -361,6 +268,7 @@ FText SCustomizableObjectEditorTagExplorer::GetCurrentItemLabel() const
 	return LOCTEXT("InvalidComboEntryText", "None");
 }
 
+
 FReply SCustomizableObjectEditorTagExplorer::CopyTagToClipboard()
 {
 	if (!SelectedTag.IsEmpty())
@@ -371,115 +279,153 @@ FReply SCustomizableObjectEditorTagExplorer::CopyTagToClipboard()
 	return FReply::Handled();
 }
 
-void SCustomizableObjectEditorTagExplorer::OnComboBoxSelectionChanged(FString NewValue)
+
+void SCustomizableObjectEditorTagExplorer::OnComboBoxSelectionChanged(const FString NewValue)
 {
 	SelectedTag = NewValue;
 
 	if (!NewValue.IsEmpty())
 	{
-		TArray<UCustomizableObjectNode*> auxNodes;
-		NodeTags.MultiFind(NewValue, auxNodes, false);
+		TArray<UCustomizableObjectNode*> AllNodes;
+		NodeTags.MultiFind(NewValue, AllNodes, false);
 
-		MaterialNodes.Empty();
-		VariationNodes.Empty();
-		ClipMeshNodes.Empty();
-		ClipMorphNodes.Empty();
-		ExtendNodes.Empty();
-		ClipDeformNodes.Empty();
+		Nodes.Empty();
 
-		for (UCustomizableObjectNode* node : auxNodes)
+		for (UCustomizableObjectNode* Node : AllNodes)
 		{
-			if (Cast<UCustomizableObjectNodeMaterialBase>(node) && MaterialNodes.Find(node) == INDEX_NONE)
+			if (Node)
 			{
-				MaterialNodes.Add(node);
+				Nodes.Add(MakeWeakObjectPtr(Node));
 			}
+		}
+	}
 
-			if (Cast<UCustomizableObjectNodeMaterialVariation>(node) && VariationNodes.Find(node) == INDEX_NONE)
-			{
-				VariationNodes.Add(node);
-			}
+	if (ListViewWidget.IsValid())
+	{
+		ListViewWidget->RequestListRefresh();
+	}
+}
 
-			if (Cast<UCustomizableObjectNodeMeshClipMorph>(node) && ClipMorphNodes.Find(node) == INDEX_NONE)
-			{
-				ClipMorphNodes.Add(node);
-			}
 
-			if (Cast<UCustomizableObjectNodeMeshClipWithMesh>(node) && ClipMeshNodes.Find(node) == INDEX_NONE)
-			{
-				ClipMeshNodes.Add(node);
-			}
-			
-			if (Cast<UCustomizableObjectNodeExtendMaterial>(node) && ExtendNodes.Find(node) == INDEX_NONE)
+TSharedRef<ITableRow> SCustomizableObjectEditorTagExplorer::OnGenerateTableRow(TWeakObjectPtr<UCustomizableObjectNode> Node, const TSharedRef<STableViewBase>& OwnerTable)
+{
+	return SNew(STagExplorerTableRow, OwnerTable).CustomizableObjectNode(Node);
+}
+
+
+void SCustomizableObjectEditorTagExplorer::OnTagTableSelectionChanged(TWeakObjectPtr<UCustomizableObjectNode> Entry, ESelectInfo::Type SelectInfo) const
+{
+	if (Entry.IsValid())
+	{
+		UObject* Object = Entry->GetCustomizableObjectGraph()->GetOuter();
+
+		// Make sure the editor exists for this asset
+		GEditor->GetEditorSubsystem<UAssetEditorSubsystem>()->OpenEditorForAsset(Object);
+
+		// Find it
+		TSharedPtr<ICustomizableObjectEditor> Editor = Entry->GetGraphEditor();
+		Editor->SelectNode(Entry.Get());
+
+		if (ListViewWidget.IsValid())
+		{
+			ListViewWidget->ClearSelection();
+			ListViewWidget->RequestListRefresh();
+		}
+	}
+}
+
+
+void SCustomizableObjectEditorTagExplorer::SortListView(const EColumnSortPriority::Type SortPriority, const FName& ColumnId, const EColumnSortMode::Type NewSortMode)
+{
+	CurrentSortColumn = ColumnId;
+	SortMode = NewSortMode;
+
+	if (ColumnId == SCustomizableObjectEditorTagExplorer::COLUMN_OBJECT)
+	{
+		Nodes.Sort([&](const TWeakObjectPtr<UCustomizableObjectNode>& NodeA, const TWeakObjectPtr<UCustomizableObjectNode>& NodeB)
+		{
+			if (NodeA.IsValid() && NodeB.IsValid())
 			{
-				ExtendNodes.Add(node);
+				FString NameA = GetNameSafe(NodeA->GetOutermostObject());
+				FString NameB = GetNameSafe(NodeB->GetOutermostObject());
+
+				return NewSortMode == EColumnSortMode::Ascending ? NameA < NameB : NameA > NameB;
 			}
 
-			if (Cast<UCustomizableObjectNodeMeshClipDeform>(node) && ClipDeformNodes.Find(node) == INDEX_NONE)
+			return NodeA.IsValid();
+		});
+	}
+	else if (ColumnId == SCustomizableObjectEditorTagExplorer::COLUMN_TYPE)
+	{
+		Nodes.Sort([&](const TWeakObjectPtr<UCustomizableObjectNode>& NodeA, const TWeakObjectPtr<UCustomizableObjectNode>& NodeB)
+		{
+			if (NodeA.IsValid() && NodeB.IsValid())
 			{
-				ClipDeformNodes.Add(node);
+				const FString NodeTypeA = NodeA->GetNodeTitle(ENodeTitleType::ListView).ToString();
+				const FString NodeTypeB = NodeB->GetNodeTitle(ENodeTitleType::ListView).ToString();
+
+				return NewSortMode == EColumnSortMode::Ascending ? NodeTypeA < NodeTypeB : NodeTypeA > NodeTypeB;
 			}
-		}
-	}
 
-	if (ColumnMat.IsValid() && ColumnClipMesh.IsValid() && ColumnClipMorph.IsValid() && ColumnVar.IsValid() && ColumnExtend.IsValid() && ColumnClipDeform.IsValid())
+			return NodeA.IsValid();
+		});
+	}
+	else
 	{
-		ColumnMat.ToSharedRef()->RequestListRefresh();
-		ColumnClipMesh.ToSharedRef()->RequestListRefresh();
-		ColumnClipMorph.ToSharedRef()->RequestListRefresh();
-		ColumnVar.ToSharedRef()->RequestListRefresh();
-		ColumnExtend.ToSharedRef()->RequestListRefresh();
-		ColumnClipDeform.ToSharedRef()->RequestListRefresh();
+		check(false); // Unknown method.
 	}
+
+	ListViewWidget->RequestListRefresh();
 }
 
-TSharedRef<ITableRow> SCustomizableObjectEditorTagExplorer::OnGenerateTableRow(UCustomizableObjectNode * Node, const TSharedRef<STableViewBase>& OwnerTable)
-{
-	bool multiRoot = false;
-	UCustomizableObject* NodeCO = Cast<UCustomizableObject>(Node->GetGraph()->GetOuter());
-	FString ObjectNode = Node->GetNodeTitle(ENodeTitleType::ListView).ToString();
 
-	if (NodeCO)
+EColumnSortMode::Type SCustomizableObjectEditorTagExplorer::GetColumnSortMode(const FName ColumnName) const
+{
+	if (CurrentSortColumn != ColumnName)
 	{
-		ObjectNode = GetRootNode(NodeCO, multiRoot)->GetNodeTitle(ENodeTitleType::MenuTitle).ToString();
-		int Indx = ObjectNode.Find("\n");
-		ObjectNode.RemoveAt(Indx, ObjectNode.Len() - Indx);
+		return EColumnSortMode::None;
 	}
 
-	return SNew(STableRow<UCustomizableObjectNode*>, OwnerTable)
-		[
-			SNew(SHorizontalBox)
-			+ SHorizontalBox::Slot()
-		.HAlign(EHorizontalAlignment::HAlign_Center)
-		[
-			SNew(STextBlock)
-			.Text(FText::FromString(ObjectNode))
-		]
-		];
+	return SortMode;
 }
 
-void SCustomizableObjectEditorTagExplorer::OnTagTableSelectionChanged(UCustomizableObjectNode* Entry, ESelectInfo::Type SelectInfo) const
+
+void STagExplorerTableRow::Construct(const FArguments& InArgs, const TSharedRef<STableViewBase>& OwnerTableView)
 {
-	if (Entry)
-	{
-		UObject* Object = Entry->GetCustomizableObjectGraph()->GetOuter();
+	Node = InArgs._CustomizableObjectNode;
 
-		// Make sure the editor exists for this asset
-		GEditor->GetEditorSubsystem<UAssetEditorSubsystem>()->OpenEditorForAsset(Object);
+	SMultiColumnTableRow<TWeakObjectPtr<UCustomizableObjectNode> >::Construct(FSuperRowType::FArguments(), OwnerTableView);
+}
 
-		// Find it
-		TSharedPtr<ICustomizableObjectEditor> Editor = Entry->GetGraphEditor();
-		Editor->SelectNode(Entry);
 
-		if (ColumnMat.IsValid() && ColumnClipMesh.IsValid() && ColumnClipMorph.IsValid() && ColumnVar.IsValid() && ColumnExtend.IsValid() && ColumnClipDeform.IsValid())
+TSharedRef<SWidget> STagExplorerTableRow::GenerateWidgetForColumn(const FName& InColumnName)
+{
+	if (InColumnName == SCustomizableObjectEditorTagExplorer::COLUMN_OBJECT)
+	{
+		if (Node.IsValid())
 		{
-			ColumnMat.ToSharedRef()->ClearSelection();
-			ColumnClipMesh.ToSharedRef()->ClearSelection();
-			ColumnClipMorph.ToSharedRef()->ClearSelection();
-			ColumnVar.ToSharedRef()->ClearSelection();
-			ColumnExtend.ToSharedRef()->ClearSelection();
-			ColumnClipDeform.ToSharedRef()->RequestListRefresh();
+			if (UCustomizableObject* CustomizableObject = Cast<UCustomizableObject>(Node->GetOutermostObject()))
+			{
+				return SNew(SBox).Padding(5.0f,0.0f,0.0f,0.0f)
+				[
+					SNew(STextBlock).Text(FText::FromString(CustomizableObject->GetName()))
+				];
+			}
 		}
 	}
+	
+	else if (InColumnName == SCustomizableObjectEditorTagExplorer::COLUMN_TYPE)
+	{
+		if (Node.IsValid())
+		{
+			return SNew(SBox).Padding(5.0f, 0.0f, 0.0f, 0.0f)
+			[
+				SNew(STextBlock).Text(Node->GetNodeTitle(ENodeTitleType::ListView))
+			];
+		}
+	}
+	
+	return SNullWidget::NullWidget;
 }
 
 #undef LOCTEXT_NAMESPACE 
diff --git a/Engine/Plugins/Experimental/Mutable/Source/CustomizableObjectEditor/Private/MuCOE/SCustomizableObjectEditorTagExplorer.h b/Engine/Plugins/Experimental/Mutable/Source/CustomizableObjectEditor/Private/MuCOE/SCustomizableObjectEditorTagExplorer.h
--- a/Engine/Plugins/Experimental/Mutable/Source/CustomizableObjectEditor/Private/MuCOE/SCustomizableObjectEditorTagExplorer.h
+++ b/Engine/Plugins/Experimental/Mutable/Source/CustomizableObjectEditor/Private/MuCOE/SCustomizableObjectEditorTagExplorer.h
@@ -2,7 +2,11 @@
 
 #pragma once
 
+#include "MuCOE/Nodes/CustomizableObjectNode.h"
+#include "UObject/WeakObjectPtr.h"
 #include "Widgets/SCompoundWidget.h"
+#include "Widgets/Views/SListView.h"
+#include "Widgets/Views/STableRow.h"
 
 namespace ESelectInfo { enum Type : int; }
 template <typename ItemType> class SListView;
@@ -13,7 +17,27 @@ class STableViewBase;
 class SComboButton;
 class SWidget;
 class UCustomizableObject;
-class UCustomizableObjectNode;
+
+
+class STagExplorerTableRow : public SMultiColumnTableRow<TWeakObjectPtr<UCustomizableObjectNode>>
+{
+public:
+
+	SLATE_BEGIN_ARGS(STagExplorerTableRow) {}
+		SLATE_ARGUMENT(TWeakObjectPtr<UCustomizableObjectNode>, CustomizableObjectNode)
+	SLATE_END_ARGS()
+
+	void Construct(const FArguments& InArgs, const TSharedRef<STableViewBase>& OwnerTableView);
+
+	virtual TSharedRef<SWidget> GenerateWidgetForColumn(const FName& InColumnName) override;
+
+private:
+
+	/** The info about the widget that we are visualizing */
+	TWeakObjectPtr<UCustomizableObjectNode> Node;
+
+};
+
 
 class SCustomizableObjectEditorTagExplorer : public SCompoundWidget
 {
@@ -25,13 +49,16 @@ public:
 
 	void Construct(const FArguments& InArgs);
 
+	static const FName COLUMN_OBJECT;
+	static const FName COLUMN_TYPE;
+
 private:
 
 	/** Callback to fill the combobox options */
 	TSharedRef<SWidget> OnGetTagsMenuContent();
 
 	/** Fills a list with all the tags found in the nodes of a graph */
-	void FillTagInformation(UCustomizableObject* Object, TArray<FString>& Tags);
+	void FillTagInformation(const UCustomizableObject& Object, TArray<FString>& Tags);
 
 	/** Generates the combobox with all the tags  */
 	TSharedRef<SWidget> MakeComboButtonItemWidget(TSharedPtr<FString> StringItem);
@@ -43,11 +70,17 @@ private:
 	FReply CopyTagToClipboard();
 
 	/** OnSelectionChanged callback of the tags combobox */
-	void OnComboBoxSelectionChanged(FString NewValue);
+	void OnComboBoxSelectionChanged(const FString NewValue);
 
 	/** Tags table callbacks*/
-	TSharedRef<ITableRow> OnGenerateTableRow(UCustomizableObjectNode* Node, const TSharedRef<STableViewBase>& OwnerTable);
-	void OnTagTableSelectionChanged(UCustomizableObjectNode* Entry, ESelectInfo::Type SelectInfo) const;
+	TSharedRef<ITableRow> OnGenerateTableRow(TWeakObjectPtr<UCustomizableObjectNode> Node, const TSharedRef<STableViewBase>& OwnerTable);
+	void OnTagTableSelectionChanged(TWeakObjectPtr<UCustomizableObjectNode> Entry, ESelectInfo::Type SelectInfo) const;
+
+	/** Sorts the content of the list view alphabetically */
+	void SortListView(const EColumnSortPriority::Type SortPriority, const FName& ColumnName, const EColumnSortMode::Type NewSortMode);
+
+	/** Returns the sorting mode of the specified column */
+	EColumnSortMode::Type GetColumnSortMode(const FName ColumnName) const;
 
 private:
 
@@ -61,22 +94,17 @@ private:
 	FString SelectedTag;
 
 	/** List views for the nodes that contain the same tag */
-	TSharedPtr<SListView<UCustomizableObjectNode*>> ColumnMat;
-	TSharedPtr<SListView<UCustomizableObjectNode*>> ColumnVar;
-	TSharedPtr<SListView<UCustomizableObjectNode*>> ColumnClipMesh;
-	TSharedPtr<SListView<UCustomizableObjectNode*>> ColumnClipMorph;
-	TSharedPtr<SListView<UCustomizableObjectNode*>> ColumnExtend;
-	TSharedPtr<SListView<UCustomizableObjectNode*>> ColumnClipDeform;
+	TSharedPtr<SListView<TWeakObjectPtr<UCustomizableObjectNode>>> ListViewWidget;
 
 	/** Arrays for each type of node that contains a tag */
-	TArray<UCustomizableObjectNode*> MaterialNodes;
-	TArray<UCustomizableObjectNode*> VariationNodes;
-	TArray<UCustomizableObjectNode*> ClipMeshNodes;
-	TArray<UCustomizableObjectNode*> ClipMorphNodes;
-	TArray<UCustomizableObjectNode*> ExtendNodes;
-	TArray<UCustomizableObjectNode*> ClipDeformNodes;
-
-	/** Multimap to relatre the nodes with the tags */
+	TArray<TWeakObjectPtr<UCustomizableObjectNode>> Nodes;
+
+	/** Multimap to relate nodes with tags */
 	TMultiMap<FString, UCustomizableObjectNode*> NodeTags;
 
+	/** Stores name of the column to sort */
+	FName CurrentSortColumn;
+
+	/** Stores the sorting mode of the selected column */
+	EColumnSortMode::Type SortMode = EColumnSortMode::Type::None;
 };
