[Mutable] Fix Wrap and Cylindrical projectors to not be afected by FadeStart and FadeEnd parameters

- Fix Cylindrical projection lod selection when minfilter method is set to Total Area Heuristic.

#rnx

diff --git a/Engine/Plugins/Experimental/Mutable/Source/MutableRuntime/Private/MuR/OpImageProject.cpp b/Engine/Plugins/Experimental/Mutable/Source/MutableRuntime/Private/MuR/OpImageProject.cpp
--- a/Engine/Plugins/Experimental/Mutable/Source/MutableRuntime/Private/MuR/OpImageProject.cpp
+++ b/Engine/Plugins/Experimental/Mutable/Source/MutableRuntime/Private/MuR/OpImageProject.cpp
@@ -72,15 +72,16 @@ namespace Private
 	enum class EPixelProcessorFeatures
 	{
 		None = 0,
-		WithMask                 = 1 << 1,
-		SamplingPoint            = 1 << 2,
-		SamplingLinear			 = 1 << 3, 
-		ProjectionPlanarAndWrap  = 1 << 4,
-		ProjectionCylindrical    = 1 << 5,
-		FormatRGBA               = 1 << 6,
-		FormatRGB                = 1 << 7,
-		FormatL                  = 1 << 8,
-		VectorizedImpl			 = 1 << 9    
+		WithMask              = 1 << 1,
+		SamplingPoint         = 1 << 2,
+		SamplingLinear		  = 1 << 3, 
+		ProjectionPlanar	  = 1 << 4,
+		ProjectionCylindrical = 1 << 5,
+		ProjectionWrap		  = 1 << 6,
+		FormatRGBA            = 1 << 7,
+		FormatRGB             = 1 << 8,
+		FormatL               = 1 << 9,
+		VectorizedImpl		  = 1 << 10    
 	};
 
 	ENUM_CLASS_FLAGS(EPixelProcessorFeatures);
@@ -104,7 +105,8 @@ namespace Private
 	{
 		return
 			static_cast<int32>(EnumHasAnyFlags(Features, EPixelProcessorFeatures::ProjectionCylindrical)) +
-			static_cast<int32>(EnumHasAnyFlags(Features, EPixelProcessorFeatures::ProjectionPlanarAndWrap)) == 1;
+			static_cast<int32>(EnumHasAnyFlags(Features, EPixelProcessorFeatures::ProjectionPlanar)) +
+			static_cast<int32>(EnumHasAnyFlags(Features, EPixelProcessorFeatures::ProjectionWrap)) == 1;
 	}
 
 	constexpr int32 GetFormatNumChannels(EPixelProcessorFeatures Features)
@@ -122,7 +124,7 @@ namespace Private
 			return 4;
 		}
 
-		return 0;
+		return -1;
 	}
 
 	// The processor constant data fits into a 64 bytes cache line, align to it so potentially have less cache 
@@ -168,7 +170,7 @@ namespace Private
 		static_assert(CheckExactlyOneProjectionFlag(Features));
 
 		static constexpr int32 PIXEL_SIZE = GetFormatNumChannels(Features);
-		static_assert(PIXEL_SIZE != 0);
+		static_assert(PIXEL_SIZE != -1 && PIXEL_SIZE > 0);
 
 	public:
 		static FProjectedPixelProcessorContext MakeContext(
@@ -208,24 +210,32 @@ namespace Private
 			Context.RGBFadingEnabledMask = bInIsRGBFadingEnabled ? -1 : 0;
 			Context.AlphaFadingEnabledMask = bInIsAlphaFadingEnabled ? -1 : 0;
 
-			const float FadeStartCos = FMath::Cos(FadeStart);
-            Context.FadeEndCos = FMath::Cos(FadeEnd);
+			if constexpr (EnumHasAnyFlags(Features, EPixelProcessorFeatures::ProjectionPlanar))
+			{
+				const float FadeStartCos = FMath::Cos(FadeStart);
+				Context.FadeEndCos = FMath::Cos(FadeEnd);
 
-			const float FadeCosRangeSafeDiv = 
-				FMath::IsNearlyZero(FadeStartCos - Context.FadeEndCos, UE_KINDA_SMALL_NUMBER) 
-				? UE_KINDA_SMALL_NUMBER 
-				: FadeStartCos - Context.FadeEndCos;
+				const float FadeCosRangeSafeDiv =
+					FMath::IsNearlyZero(FadeStartCos - Context.FadeEndCos, UE_KINDA_SMALL_NUMBER)
+					? UE_KINDA_SMALL_NUMBER
+					: FadeStartCos - Context.FadeEndCos;
 
-			Context.OneOverFadeRangeTimes255 = 255.0f / FadeCosRangeSafeDiv;
+				Context.OneOverFadeRangeTimes255 = 255.0f / FadeCosRangeSafeDiv;
+			}
+			else
+			{
+				Context.FadeEndCos = 0.0f;
+				Context.OneOverFadeRangeTimes255 = 255.0f;
+			}
 
 			Context.OneOverProjectionAngle = FMath::IsNearlyZero(InProjectionAngle)
-				? 1.0f / UE_KINDA_SMALL_NUMBER
-				: 1.0f / InProjectionAngle;
+					? 1.0f / UE_KINDA_SMALL_NUMBER
+					: 1.0f / InProjectionAngle;
 
             Context.TargetData = InTargetData;
             Context.MaskData = InMaskData;
 
-            check(GetImageFormatData(Source->GetFormat() ).BytesPerBlock == PIXEL_SIZE);
+            check(GetImageFormatData(Source->GetFormat()).BytesPerBlock == PIXEL_SIZE);
 
 			return Context;
 		}
@@ -263,12 +273,16 @@ namespace Private
 				}
 			});
 
-            const float AngleCos = Varying[3];
+			float Factor = static_cast<float>(!bDepthClamp) * 255.0f;
+			float MaskFactor = 255.0f;
 
-			float Factor = FMath::Clamp((AngleCos - Context.FadeEndCos) * Context.OneOverFadeRangeTimes255, 0.0f, 255.0f);
-			Factor = Factor < 255.0f ? Factor : static_cast<float>(!bDepthClamp) * 255.0f;
+			if constexpr (EnumHasAnyFlags(Features, EPixelProcessorFeatures::ProjectionPlanar))
+			{
+				const float AngleCos = Varying[3];
 
-			float MaskFactor = AngleCos > Context.FadeEndCos ? 255.0f : 0.0f;
+				Factor = FMath::Min(Factor, (FMath::Clamp((AngleCos - Context.FadeEndCos) * Context.OneOverFadeRangeTimes255, 0.0f, 255.0f)));
+				MaskFactor = AngleCos > Context.FadeEndCos ? 255.0f : 0.0f;
+			}
 
 			if constexpr (EnumHasAnyFlags(Features, EPixelProcessorFeatures::WithMask))
 			{
@@ -456,12 +470,18 @@ namespace Private
 				}
 			});
 
-            const float AngleCos = Varying[3];
+			uint16 Factor = static_cast<uint16>(!bDepthClamp) * 255;
+			uint16 MaskFactor = 255;
 
-			uint16 Factor = static_cast<uint16>(FMath::Clamp((AngleCos - Context.FadeEndCos) * Context.OneOverFadeRangeTimes255, 0.0f, 255.0f));
-			Factor = Factor < 255 ? Factor : static_cast<uint16>(!bDepthClamp) * 255;
+			if constexpr (EnumHasAnyFlags(Features, EPixelProcessorFeatures::ProjectionPlanar))
+			{
+				const float AngleCos = Varying[3];
 
-			uint16 MaskFactor = AngleCos > Context.FadeEndCos ? 255 : 0;
+				Factor = FMath::Min(Factor, 
+						static_cast<uint16>(FMath::Clamp((AngleCos - Context.FadeEndCos) * Context.OneOverFadeRangeTimes255, 0.0f, 255.0f)));
+
+				MaskFactor = AngleCos > Context.FadeEndCos ? 255 : 0;
+			}
 
 			if constexpr (EnumHasAnyFlags(Features, EPixelProcessorFeatures::WithMask))
 			{
@@ -531,11 +551,6 @@ namespace Private
 					{
 						auto SampleImageBilinear = [&](FVector2f Uv, FUInt16Vector2 Size, const uint8* DataPtr)
 						{
-							auto ComputeInterpolator = [](float T) -> uint16
-							{
-								return static_cast<uint16>(255.0f * T);
-							};
-
 							const FVector2f SizeF(Size.X, Size.Y);
 
 							const FVector2f CoordsF = FVector2f(
@@ -543,8 +558,8 @@ namespace Private
 								FMath::Clamp(Uv.Y * SizeF.Y - 0.5f, 0.0f, SizeF.Y - 1.0f));
 
 							const FUInt16Vector2 Frac = FUInt16Vector2(
-								ComputeInterpolator(FMath::Frac(CoordsF.X)),
-								ComputeInterpolator(FMath::Frac(CoordsF.Y)));
+								static_cast<uint16>(FMath::Frac(CoordsF.X) * 255.0f),
+								static_cast<uint16>(FMath::Frac(CoordsF.Y) * 255.0f));
 
 							const FIntVector2 Coords = FIntVector2(CoordsF.X, CoordsF.Y);
 							const FIntVector2 CoordsPlusOne = FIntVector2(
@@ -730,7 +745,6 @@ namespace Private
 	}
 
 
-
 	//-------------------------------------------------------------------------------------------------
 	FORCENOINLINE void ImageRasterProjected_OptimisedWrapping(const Mesh* pMesh, Image* pImage, 
 		TTriangleRasterPixelProcRefType<4> PixelProc, 
@@ -772,7 +786,8 @@ namespace Private
 			const uint64* LayoutBlockIds = reinterpret_cast<const uint64*>(pMesh->GetVertexBuffers().GetBufferData(1));
 			for (int32 V = 0; V < VertexCount; ++V)
 			{
-				Scratch->CulledVertex[V] = pVertices[V].Normal[0] < FadeEndCos;
+				// Don't cull vertices based on Normal for wrapping projectors. 
+				//Scratch->CulledVertex[V] = pVertices[V].Normal[0] < FadeEndCos;
 				if (LayoutBlockIds[V] != BlockId)
 				{
 					Scratch->CulledVertex[V] = true;
@@ -786,7 +801,7 @@ namespace Private
 			const uint16* LayoutBlockIds = reinterpret_cast<const uint16*>(pMesh->GetVertexBuffers().GetBufferData(1));
 			for (int32 V = 0; V < VertexCount; ++V)
 			{
-				Scratch->CulledVertex[V] = pVertices[V].Normal[0] < FadeEndCos;
+				//Scratch->CulledVertex[V] = pVertices[V].Normal[0] < FadeEndCos;
 				if (LayoutBlockIds[V] != BlockId)
 				{
 					Scratch->CulledVertex[V] = true;
@@ -855,7 +870,7 @@ namespace Private
 	template <EPixelProcessorFeatures Features> 
 	FORCENOINLINE void InvokePlanarRasterizerImpl(const FImageRasterInvokeArgs& Args)
 	{
-		static_assert(EnumHasAnyFlags(Features, EPixelProcessorFeatures::ProjectionPlanarAndWrap));
+		static_assert(EnumHasAnyFlags(Features, EPixelProcessorFeatures::ProjectionPlanar));
 
 		const Private::FProjectedPixelProcessorContext Context = 
 			Private::TProjectedPixelProcessor<Features>::MakeContext(
@@ -893,7 +908,7 @@ namespace Private
 	template <EPixelProcessorFeatures Features> 
 	FORCENOINLINE void InvokeWrappingRasterizerImpl(const FImageRasterInvokeArgs& Args)
 	{
-		static_assert(EnumHasAnyFlags(Features, EPixelProcessorFeatures::ProjectionPlanarAndWrap));
+		static_assert(EnumHasAnyFlags(Features, EPixelProcessorFeatures::ProjectionWrap));
 
 		const Private::FProjectedPixelProcessorContext Context = 
 			Private::TProjectedPixelProcessor<Features>::MakeContext(
@@ -931,7 +946,7 @@ void ImageRasterProjectedPlanar(const Mesh* pMesh, Image* pImage,
 
 	MUTABLE_CPUPROFILER_SCOPE(ImageProject);
 
-	EPixelProcessorFeatures ProcessorFeatures = EPixelProcessorFeatures::ProjectionPlanarAndWrap;
+	EPixelProcessorFeatures ProcessorFeatures = EPixelProcessorFeatures::ProjectionPlanar;
 
 	// Disable vectorized implementation if Point sampling.
 	if (bUseVectorImplementation && SamplingMethod != ESamplingMethod::Point)
@@ -984,131 +999,131 @@ void ImageRasterProjectedPlanar(const Mesh* pMesh, Image* pImage,
 		// but maybe not needed. Mask permutations can be removed safely
 		switch (ProcessorFeatures)
 		{
-		case (EPPF::FormatL | EPPF::ProjectionPlanarAndWrap | EPPF::SamplingPoint):
+		case (EPPF::FormatL | EPPF::ProjectionPlanar | EPPF::SamplingPoint):
 		{
-			InvokePlanarRasterizerImpl<EPPF::FormatL | EPPF::ProjectionPlanarAndWrap | EPPF::SamplingPoint>(RasterArgs);
+			InvokePlanarRasterizerImpl<EPPF::FormatL | EPPF::ProjectionPlanar | EPPF::SamplingPoint>(RasterArgs);
 			break;
 		}
-		case (EPPF::FormatRGB | EPPF::ProjectionPlanarAndWrap | EPPF::SamplingPoint):
+		case (EPPF::FormatRGB | EPPF::ProjectionPlanar | EPPF::SamplingPoint):
 		{
-			InvokePlanarRasterizerImpl<EPPF::FormatRGB | EPPF::ProjectionPlanarAndWrap | EPPF::SamplingPoint>(RasterArgs);
+			InvokePlanarRasterizerImpl<EPPF::FormatRGB | EPPF::ProjectionPlanar | EPPF::SamplingPoint>(RasterArgs);
 			break;
 		}
-		case (EPPF::FormatRGBA | EPPF::ProjectionPlanarAndWrap | EPPF::SamplingPoint):
+		case (EPPF::FormatRGBA | EPPF::ProjectionPlanar | EPPF::SamplingPoint):
 		{
-			InvokePlanarRasterizerImpl<EPPF::FormatRGBA | EPPF::ProjectionPlanarAndWrap | EPPF::SamplingPoint>(RasterArgs);
+			InvokePlanarRasterizerImpl<EPPF::FormatRGBA | EPPF::ProjectionPlanar | EPPF::SamplingPoint>(RasterArgs);
 			break;
 		}
 
-		case (EPPF::FormatL | EPPF::ProjectionPlanarAndWrap | EPPF::SamplingLinear):
+		case (EPPF::FormatL | EPPF::ProjectionPlanar | EPPF::SamplingLinear):
 		{
-			InvokePlanarRasterizerImpl<EPPF::FormatL | EPPF::ProjectionPlanarAndWrap | EPPF::SamplingLinear>(RasterArgs);
+			InvokePlanarRasterizerImpl<EPPF::FormatL | EPPF::ProjectionPlanar | EPPF::SamplingLinear>(RasterArgs);
 			break;
 		}
-		case (EPPF::FormatRGB | EPPF::ProjectionPlanarAndWrap | EPPF::SamplingLinear):
+		case (EPPF::FormatRGB | EPPF::ProjectionPlanar | EPPF::SamplingLinear):
 		{
-			InvokePlanarRasterizerImpl<EPPF::FormatRGB | EPPF::ProjectionPlanarAndWrap | EPPF::SamplingLinear>(RasterArgs);
+			InvokePlanarRasterizerImpl<EPPF::FormatRGB | EPPF::ProjectionPlanar | EPPF::SamplingLinear>(RasterArgs);
 			break;
 		}
-		case (EPPF::FormatRGBA | EPPF::ProjectionPlanarAndWrap | EPPF::SamplingLinear):
+		case (EPPF::FormatRGBA | EPPF::ProjectionPlanar | EPPF::SamplingLinear):
 		{
-			InvokePlanarRasterizerImpl<EPPF::FormatRGBA | EPPF::ProjectionPlanarAndWrap | EPPF::SamplingLinear>(RasterArgs);
+			InvokePlanarRasterizerImpl<EPPF::FormatRGBA | EPPF::ProjectionPlanar | EPPF::SamplingLinear>(RasterArgs);
 			break;
 		}
 
-		case (EPPF::FormatL | EPPF::ProjectionPlanarAndWrap | EPPF::SamplingPoint | EPPF::WithMask):
+		case (EPPF::FormatL | EPPF::ProjectionPlanar | EPPF::SamplingPoint | EPPF::WithMask):
 		{
-			InvokePlanarRasterizerImpl<EPPF::FormatL | EPPF::ProjectionPlanarAndWrap | EPPF::SamplingPoint | EPPF::WithMask>(RasterArgs);
+			InvokePlanarRasterizerImpl<EPPF::FormatL | EPPF::ProjectionPlanar | EPPF::SamplingPoint | EPPF::WithMask>(RasterArgs);
 			break;
 		}
-		case (EPPF::FormatRGB | EPPF::ProjectionPlanarAndWrap | EPPF::SamplingPoint | EPPF::WithMask):
+		case (EPPF::FormatRGB | EPPF::ProjectionPlanar | EPPF::SamplingPoint | EPPF::WithMask):
 		{
-			InvokePlanarRasterizerImpl<EPPF::FormatRGB | EPPF::ProjectionPlanarAndWrap | EPPF::SamplingPoint | EPPF::WithMask>(RasterArgs);
+			InvokePlanarRasterizerImpl<EPPF::FormatRGB | EPPF::ProjectionPlanar | EPPF::SamplingPoint | EPPF::WithMask>(RasterArgs);
 			break;
 		}
-		case (EPPF::FormatRGBA | EPPF::ProjectionPlanarAndWrap | EPPF::SamplingPoint | EPPF::WithMask):
+		case (EPPF::FormatRGBA | EPPF::ProjectionPlanar | EPPF::SamplingPoint | EPPF::WithMask):
 		{
-			InvokePlanarRasterizerImpl<EPPF::FormatRGBA | EPPF::ProjectionPlanarAndWrap | EPPF::SamplingPoint>(RasterArgs);
+			InvokePlanarRasterizerImpl<EPPF::FormatRGBA | EPPF::ProjectionPlanar | EPPF::SamplingPoint>(RasterArgs);
 			break;
 		}
 
-		case (EPPF::FormatL | EPPF::ProjectionPlanarAndWrap | EPPF::SamplingLinear | EPPF::WithMask):
+		case (EPPF::FormatL | EPPF::ProjectionPlanar | EPPF::SamplingLinear | EPPF::WithMask):
 		{
-			InvokePlanarRasterizerImpl<EPPF::FormatL | EPPF::ProjectionPlanarAndWrap | EPPF::SamplingLinear | EPPF::WithMask>(RasterArgs);
+			InvokePlanarRasterizerImpl<EPPF::FormatL | EPPF::ProjectionPlanar | EPPF::SamplingLinear | EPPF::WithMask>(RasterArgs);
 			break;
 		}
-		case (EPPF::FormatRGB | EPPF::ProjectionPlanarAndWrap | EPPF::SamplingLinear | EPPF::WithMask):
+		case (EPPF::FormatRGB | EPPF::ProjectionPlanar | EPPF::SamplingLinear | EPPF::WithMask):
 		{
-			InvokePlanarRasterizerImpl<EPPF::FormatRGB | EPPF::ProjectionPlanarAndWrap | EPPF::SamplingLinear | EPPF::WithMask>(RasterArgs);
+			InvokePlanarRasterizerImpl<EPPF::FormatRGB | EPPF::ProjectionPlanar | EPPF::SamplingLinear | EPPF::WithMask>(RasterArgs);
 			break;
 		}
-		case (EPPF::FormatRGBA | EPPF::ProjectionPlanarAndWrap | EPPF::SamplingLinear | EPPF::WithMask):
+		case (EPPF::FormatRGBA | EPPF::ProjectionPlanar | EPPF::SamplingLinear | EPPF::WithMask):
 		{
-			InvokePlanarRasterizerImpl<EPPF::FormatRGBA | EPPF::ProjectionPlanarAndWrap | EPPF::SamplingLinear | EPPF::WithMask>(RasterArgs);
+			InvokePlanarRasterizerImpl<EPPF::FormatRGBA | EPPF::ProjectionPlanar | EPPF::SamplingLinear | EPPF::WithMask>(RasterArgs);
 			break;
 		}
 	
-		//case (EPPF::FormatL | EPPF::ProjectionPlanarAndWrap | EPPF::SamplingPoint | EPPF::VectorizedImpl):
+		//case (EPPF::FormatL | EPPF::ProjectionPlanar | EPPF::SamplingPoint | EPPF::VectorizedImpl):
 		//{
-		//	InvokePlanarRasterizerImpl<EPPF::FormatL | EPPF::ProjectionPlanarAndWrap | EPPF::SamplingPoint | EPPF::VectorizedImpl>(RasterArgs);
+		//	InvokePlanarRasterizerImpl<EPPF::FormatL | EPPF::ProjectionPlanar | EPPF::SamplingPoint | EPPF::VectorizedImpl>(RasterArgs);
 		//	break;
 		//}
-		//case (EPPF::FormatRGB | EPPF::ProjectionPlanarAndWrap | EPPF::SamplingPoint | EPPF::VectorizedImpl):
+		//case (EPPF::FormatRGB | EPPF::ProjectionPlanar | EPPF::SamplingPoint | EPPF::VectorizedImpl):
 		//{
-		//	InvokePlanarRasterizerImpl<EPPF::FormatRGB | EPPF::ProjectionPlanarAndWrap | EPPF::SamplingPoint | EPPF::VectorizedImpl>(RasterArgs);
+		//	InvokePlanarRasterizerImpl<EPPF::FormatRGB | EPPF::ProjectionPlanar | EPPF::SamplingPoint | EPPF::VectorizedImpl>(RasterArgs);
 		//	break;
 		//}
-		//case (EPPF::FormatRGBA | EPPF::ProjectionPlanarAndWrap | EPPF::SamplingPoint | EPPF::VectorizedImpl):
+		//case (EPPF::FormatRGBA | EPPF::ProjectionPlanar | EPPF::SamplingPoint | EPPF::VectorizedImpl):
 		//{
-		//	InvokePlanarRasterizerImpl<EPPF::FormatRGBA | EPPF::ProjectionPlanarAndWrap | EPPF::SamplingPoint | EPPF::VectorizedImpl>(RasterArgs);
+		//	InvokePlanarRasterizerImpl<EPPF::FormatRGBA | EPPF::ProjectionPlanar | EPPF::SamplingPoint | EPPF::VectorizedImpl>(RasterArgs);
 		//	break;
 		//}
 
-		case (EPPF::FormatL | EPPF::ProjectionPlanarAndWrap | EPPF::SamplingLinear | EPPF::VectorizedImpl):
+		case (EPPF::FormatL | EPPF::ProjectionPlanar | EPPF::SamplingLinear | EPPF::VectorizedImpl):
 		{
-			InvokePlanarRasterizerImpl<EPPF::FormatL | EPPF::ProjectionPlanarAndWrap | EPPF::SamplingLinear | EPPF::VectorizedImpl>(RasterArgs);
+			InvokePlanarRasterizerImpl<EPPF::FormatL | EPPF::ProjectionPlanar | EPPF::SamplingLinear | EPPF::VectorizedImpl>(RasterArgs);
 			break;
 		}
-		case (EPPF::FormatRGB | EPPF::ProjectionPlanarAndWrap | EPPF::SamplingLinear | EPPF::VectorizedImpl):
+		case (EPPF::FormatRGB | EPPF::ProjectionPlanar | EPPF::SamplingLinear | EPPF::VectorizedImpl):
 		{
-			InvokePlanarRasterizerImpl<EPPF::FormatRGB | EPPF::ProjectionPlanarAndWrap | EPPF::SamplingLinear | EPPF::VectorizedImpl>(RasterArgs);
+			InvokePlanarRasterizerImpl<EPPF::FormatRGB | EPPF::ProjectionPlanar | EPPF::SamplingLinear | EPPF::VectorizedImpl>(RasterArgs);
 			break;
 		}
-		case (EPPF::FormatRGBA | EPPF::ProjectionPlanarAndWrap | EPPF::SamplingLinear | EPPF::VectorizedImpl):
+		case (EPPF::FormatRGBA | EPPF::ProjectionPlanar | EPPF::SamplingLinear | EPPF::VectorizedImpl):
 		{
-			InvokePlanarRasterizerImpl<EPPF::FormatRGBA | EPPF::ProjectionPlanarAndWrap | EPPF::SamplingLinear | EPPF::VectorizedImpl>(RasterArgs);
+			InvokePlanarRasterizerImpl<EPPF::FormatRGBA | EPPF::ProjectionPlanar | EPPF::SamplingLinear | EPPF::VectorizedImpl>(RasterArgs);
 			break;
 		}
 
-		//case (EPPF::FormatL | EPPF::ProjectionPlanarAndWrap | EPPF::SamplingPoint | EPPF::WithMask | EPPF::VectorizedImpl):
+		//case (EPPF::FormatL | EPPF::ProjectionPlanar | EPPF::SamplingPoint | EPPF::WithMask | EPPF::VectorizedImpl):
 		//{
-		//	InvokePlanarRasterizerImpl<EPPF::FormatL | EPPF::ProjectionPlanarAndWrap | EPPF::SamplingPoint | EPPF::WithMask | EPPF::VectorizedImpl>(RasterArgs);
+		//	InvokePlanarRasterizerImpl<EPPF::FormatL | EPPF::ProjectionPlanar | EPPF::SamplingPoint | EPPF::WithMask | EPPF::VectorizedImpl>(RasterArgs);
 		//	break;
 		//}
-		//case (EPPF::FormatRGB | EPPF::ProjectionPlanarAndWrap | EPPF::SamplingPoint | EPPF::WithMask | EPPF::VectorizedImpl):
+		//case (EPPF::FormatRGB | EPPF::ProjectionPlanar | EPPF::SamplingPoint | EPPF::WithMask | EPPF::VectorizedImpl):
 		//{
-		//	InvokePlanarRasterizerImpl<EPPF::FormatRGB | EPPF::ProjectionPlanarAndWrap | EPPF::SamplingPoint | EPPF::WithMask | EPPF::VectorizedImpl>(RasterArgs);
+		//	InvokePlanarRasterizerImpl<EPPF::FormatRGB | EPPF::ProjectionPlanar | EPPF::SamplingPoint | EPPF::WithMask | EPPF::VectorizedImpl>(RasterArgs);
 		//	break;
 		//}
-		//case (EPPF::FormatRGBA | EPPF::ProjectionPlanarAndWrap | EPPF::SamplingPoint | EPPF::WithMask | EPPF::VectorizedImpl):
+		//case (EPPF::FormatRGBA | EPPF::ProjectionPlanar | EPPF::SamplingPoint | EPPF::WithMask | EPPF::VectorizedImpl):
 		//{
-		//	InvokePlanarRasterizerImpl<EPPF::FormatRGBA | EPPF::ProjectionPlanarAndWrap | EPPF::SamplingPoint | EPPF::VectorizedImpl>(RasterArgs);
+		//	InvokePlanarRasterizerImpl<EPPF::FormatRGBA | EPPF::ProjectionPlanar | EPPF::SamplingPoint | EPPF::VectorizedImpl>(RasterArgs);
 		//	break;
 		//}
 
-		case (EPPF::FormatL | EPPF::ProjectionPlanarAndWrap | EPPF::SamplingLinear | EPPF::WithMask | EPPF::VectorizedImpl):
+		case (EPPF::FormatL | EPPF::ProjectionPlanar | EPPF::SamplingLinear | EPPF::WithMask | EPPF::VectorizedImpl):
 		{
-			InvokePlanarRasterizerImpl<EPPF::FormatL | EPPF::ProjectionPlanarAndWrap | EPPF::SamplingLinear | EPPF::WithMask | EPPF::VectorizedImpl>(RasterArgs);
+			InvokePlanarRasterizerImpl<EPPF::FormatL | EPPF::ProjectionPlanar | EPPF::SamplingLinear | EPPF::WithMask | EPPF::VectorizedImpl>(RasterArgs);
 			break;
 		}
-		case (EPPF::FormatRGB | EPPF::ProjectionPlanarAndWrap | EPPF::SamplingLinear | EPPF::WithMask | EPPF::VectorizedImpl):
+		case (EPPF::FormatRGB | EPPF::ProjectionPlanar | EPPF::SamplingLinear | EPPF::WithMask | EPPF::VectorizedImpl):
 		{
-			InvokePlanarRasterizerImpl<EPPF::FormatRGB | EPPF::ProjectionPlanarAndWrap | EPPF::SamplingLinear | EPPF::WithMask>(RasterArgs);
+			InvokePlanarRasterizerImpl<EPPF::FormatRGB | EPPF::ProjectionPlanar | EPPF::SamplingLinear | EPPF::WithMask>(RasterArgs);
 			break;
 		}
-		case (EPPF::FormatRGBA | EPPF::ProjectionPlanarAndWrap | EPPF::SamplingLinear | EPPF::WithMask | EPPF::VectorizedImpl):
+		case (EPPF::FormatRGBA | EPPF::ProjectionPlanar | EPPF::SamplingLinear | EPPF::WithMask | EPPF::VectorizedImpl):
 		{
-			InvokePlanarRasterizerImpl<EPPF::FormatRGBA | EPPF::ProjectionPlanarAndWrap | EPPF::SamplingLinear | EPPF::WithMask | EPPF::VectorizedImpl>(RasterArgs);
+			InvokePlanarRasterizerImpl<EPPF::FormatRGBA | EPPF::ProjectionPlanar | EPPF::SamplingLinear | EPPF::WithMask | EPPF::VectorizedImpl>(RasterArgs);
 			break;
 		}
 
@@ -1146,7 +1161,7 @@ void ImageRasterProjectedWrapping( const Mesh* pMesh, Image* pImage,
 
 	MUTABLE_CPUPROFILER_SCOPE(ImageProjectWrapping);
 
-	EPixelProcessorFeatures ProcessorFeatures = EPixelProcessorFeatures::ProjectionPlanarAndWrap;
+	EPixelProcessorFeatures ProcessorFeatures = EPixelProcessorFeatures::ProjectionWrap;
 
 	if (bUseVectorImplementation && SamplingMethod != ESamplingMethod::Point)
 	{
@@ -1199,131 +1214,131 @@ void ImageRasterProjectedWrapping( const Mesh* pMesh, Image* pImage,
 		// but maybe not needed. Mask permutations can be removed safely
 		switch (ProcessorFeatures)
 		{
-		case (EPPF::FormatL | EPPF::ProjectionPlanarAndWrap | EPPF::SamplingPoint):
+		case (EPPF::FormatL | EPPF::ProjectionWrap | EPPF::SamplingPoint):
 		{
-			InvokeWrappingRasterizerImpl<EPPF::FormatL | EPPF::ProjectionPlanarAndWrap | EPPF::SamplingPoint>(RasterArgs);
+			InvokeWrappingRasterizerImpl<EPPF::FormatL | EPPF::ProjectionWrap | EPPF::SamplingPoint>(RasterArgs);
 			break;
 		}
-		case (EPPF::FormatRGB | EPPF::ProjectionPlanarAndWrap | EPPF::SamplingPoint):
+		case (EPPF::FormatRGB | EPPF::ProjectionWrap | EPPF::SamplingPoint):
 		{
-			InvokeWrappingRasterizerImpl<EPPF::FormatRGB | EPPF::ProjectionPlanarAndWrap | EPPF::SamplingPoint>(RasterArgs);
+			InvokeWrappingRasterizerImpl<EPPF::FormatRGB | EPPF::ProjectionWrap | EPPF::SamplingPoint>(RasterArgs);
 			break;
 		}
-		case (EPPF::FormatRGBA | EPPF::ProjectionPlanarAndWrap | EPPF::SamplingPoint):
+		case (EPPF::FormatRGBA | EPPF::ProjectionWrap | EPPF::SamplingPoint):
 		{
-			InvokeWrappingRasterizerImpl<EPPF::FormatRGBA | EPPF::ProjectionPlanarAndWrap | EPPF::SamplingPoint>(RasterArgs);
+			InvokeWrappingRasterizerImpl<EPPF::FormatRGBA | EPPF::ProjectionWrap | EPPF::SamplingPoint>(RasterArgs);
 			break;
 		}
 
-		case (EPPF::FormatL | EPPF::ProjectionPlanarAndWrap | EPPF::SamplingLinear):
+		case (EPPF::FormatL | EPPF::ProjectionWrap | EPPF::SamplingLinear):
 		{
-			InvokeWrappingRasterizerImpl<EPPF::FormatL | EPPF::ProjectionPlanarAndWrap | EPPF::SamplingLinear>(RasterArgs);
+			InvokeWrappingRasterizerImpl<EPPF::FormatL | EPPF::ProjectionWrap | EPPF::SamplingLinear>(RasterArgs);
 			break;
 		}
-		case (EPPF::FormatRGB | EPPF::ProjectionPlanarAndWrap | EPPF::SamplingLinear):
+		case (EPPF::FormatRGB | EPPF::ProjectionWrap | EPPF::SamplingLinear):
 		{
-			InvokeWrappingRasterizerImpl<EPPF::FormatRGB | EPPF::ProjectionPlanarAndWrap | EPPF::SamplingLinear>(RasterArgs);
+			InvokeWrappingRasterizerImpl<EPPF::FormatRGB | EPPF::ProjectionWrap | EPPF::SamplingLinear>(RasterArgs);
 			break;
 		}
-		case (EPPF::FormatRGBA | EPPF::ProjectionPlanarAndWrap | EPPF::SamplingLinear):
+		case (EPPF::FormatRGBA | EPPF::ProjectionWrap | EPPF::SamplingLinear):
 		{
-			InvokeWrappingRasterizerImpl<EPPF::FormatRGBA | EPPF::ProjectionPlanarAndWrap | EPPF::SamplingLinear>(RasterArgs);
+			InvokeWrappingRasterizerImpl<EPPF::FormatRGBA | EPPF::ProjectionWrap | EPPF::SamplingLinear>(RasterArgs);
 			break;
 		}
 
-		case (EPPF::FormatL | EPPF::ProjectionPlanarAndWrap | EPPF::SamplingPoint | EPPF::WithMask):
+		case (EPPF::FormatL | EPPF::ProjectionWrap | EPPF::SamplingPoint | EPPF::WithMask):
 		{
-			InvokeWrappingRasterizerImpl<EPPF::FormatL | EPPF::ProjectionPlanarAndWrap | EPPF::SamplingPoint | EPPF::WithMask>(RasterArgs);
+			InvokeWrappingRasterizerImpl<EPPF::FormatL | EPPF::ProjectionWrap | EPPF::SamplingPoint | EPPF::WithMask>(RasterArgs);
 			break;
 		}
-		case (EPPF::FormatRGB | EPPF::ProjectionPlanarAndWrap | EPPF::SamplingPoint | EPPF::WithMask):
+		case (EPPF::FormatRGB | EPPF::ProjectionWrap | EPPF::SamplingPoint | EPPF::WithMask):
 		{
-			InvokeWrappingRasterizerImpl<EPPF::FormatRGB | EPPF::ProjectionPlanarAndWrap | EPPF::SamplingPoint | EPPF::WithMask>(RasterArgs);
+			InvokeWrappingRasterizerImpl<EPPF::FormatRGB | EPPF::ProjectionWrap | EPPF::SamplingPoint | EPPF::WithMask>(RasterArgs);
 			break;
 		}
-		case (EPPF::FormatRGBA | EPPF::ProjectionPlanarAndWrap | EPPF::SamplingPoint | EPPF::WithMask):
+		case (EPPF::FormatRGBA | EPPF::ProjectionWrap | EPPF::SamplingPoint | EPPF::WithMask):
 		{
-			InvokeWrappingRasterizerImpl<EPPF::FormatRGBA | EPPF::ProjectionPlanarAndWrap | EPPF::SamplingPoint>(RasterArgs);
+			InvokeWrappingRasterizerImpl<EPPF::FormatRGBA | EPPF::ProjectionWrap | EPPF::SamplingPoint>(RasterArgs);
 			break;
 		}
 
-		case (EPPF::FormatL | EPPF::ProjectionPlanarAndWrap | EPPF::SamplingLinear | EPPF::WithMask):
+		case (EPPF::FormatL | EPPF::ProjectionWrap | EPPF::SamplingLinear | EPPF::WithMask):
 		{
-			InvokeWrappingRasterizerImpl<EPPF::FormatL | EPPF::ProjectionPlanarAndWrap | EPPF::SamplingLinear | EPPF::WithMask>(RasterArgs);
+			InvokeWrappingRasterizerImpl<EPPF::FormatL | EPPF::ProjectionWrap | EPPF::SamplingLinear | EPPF::WithMask>(RasterArgs);
 			break;
 		}
-		case (EPPF::FormatRGB | EPPF::ProjectionPlanarAndWrap | EPPF::SamplingLinear | EPPF::WithMask):
+		case (EPPF::FormatRGB | EPPF::ProjectionWrap | EPPF::SamplingLinear | EPPF::WithMask):
 		{
-			InvokeWrappingRasterizerImpl<EPPF::FormatRGB | EPPF::ProjectionPlanarAndWrap | EPPF::SamplingLinear | EPPF::WithMask>(RasterArgs);
+			InvokeWrappingRasterizerImpl<EPPF::FormatRGB | EPPF::ProjectionWrap | EPPF::SamplingLinear | EPPF::WithMask>(RasterArgs);
 			break;
 		}
-		case (EPPF::FormatRGBA | EPPF::ProjectionPlanarAndWrap | EPPF::SamplingLinear | EPPF::WithMask):
+		case (EPPF::FormatRGBA | EPPF::ProjectionWrap | EPPF::SamplingLinear | EPPF::WithMask):
 		{
-			InvokeWrappingRasterizerImpl<EPPF::FormatRGBA | EPPF::ProjectionPlanarAndWrap | EPPF::SamplingLinear | EPPF::WithMask>(RasterArgs);
+			InvokeWrappingRasterizerImpl<EPPF::FormatRGBA | EPPF::ProjectionWrap | EPPF::SamplingLinear | EPPF::WithMask>(RasterArgs);
 			break;
 		}
 	
-		//case (EPPF::FormatL | EPPF::ProjectionPlanarAndWrap | EPPF::SamplingPoint | EPPF::VectorizedImpl):
+		//case (EPPF::FormatL | EPPF::ProjectionWrap | EPPF::SamplingPoint | EPPF::VectorizedImpl):
 		//{
-		//	InvokeWrappingRasterizerImpl<EPPF::FormatL | EPPF::ProjectionPlanarAndWrap | EPPF::SamplingPoint | EPPF::VectorizedImpl>(RasterArgs);
+		//	InvokeWrappingRasterizerImpl<EPPF::FormatL | EPPF::ProjectionWrap | EPPF::SamplingPoint | EPPF::VectorizedImpl>(RasterArgs);
 		//	break;
 		//}
-		//case (EPPF::FormatRGB | EPPF::ProjectionPlanarAndWrap | EPPF::SamplingPoint | EPPF::VectorizedImpl):
+		//case (EPPF::FormatRGB | EPPF::ProjectionWrap | EPPF::SamplingPoint | EPPF::VectorizedImpl):
 		//{
-		//	InvokeWrappingRasterizerImpl<EPPF::FormatRGB | EPPF::ProjectionPlanarAndWrap | EPPF::SamplingPoint | EPPF::VectorizedImpl>(RasterArgs);
+		//	InvokeWrappingRasterizerImpl<EPPF::FormatRGB | EPPF::ProjectionWrap | EPPF::SamplingPoint | EPPF::VectorizedImpl>(RasterArgs);
 		//	break;
 		//}
-		//case (EPPF::FormatRGBA | EPPF::ProjectionPlanarAndWrap | EPPF::SamplingPoint | EPPF::VectorizedImpl):
+		//case (EPPF::FormatRGBA | EPPF::ProjectionWrap | EPPF::SamplingPoint | EPPF::VectorizedImpl):
 		//{
-		//	InvokeWrappingRasterizerImpl<EPPF::FormatRGBA | EPPF::ProjectionPlanarAndWrap | EPPF::SamplingPoint | EPPF::VectorizedImpl>(RasterArgs);
+		//	InvokeWrappingRasterizerImpl<EPPF::FormatRGBA | EPPF::ProjectionWrap | EPPF::SamplingPoint | EPPF::VectorizedImpl>(RasterArgs);
 		//	break;
 		//}
 
-		case (EPPF::FormatL | EPPF::ProjectionPlanarAndWrap | EPPF::SamplingLinear | EPPF::VectorizedImpl):
+		case (EPPF::FormatL | EPPF::ProjectionWrap | EPPF::SamplingLinear | EPPF::VectorizedImpl):
 		{
-			InvokeWrappingRasterizerImpl<EPPF::FormatL | EPPF::ProjectionPlanarAndWrap | EPPF::SamplingLinear | EPPF::VectorizedImpl>(RasterArgs);
+			InvokeWrappingRasterizerImpl<EPPF::FormatL | EPPF::ProjectionWrap | EPPF::SamplingLinear | EPPF::VectorizedImpl>(RasterArgs);
 			break;
 		}
-		case (EPPF::FormatRGB | EPPF::ProjectionPlanarAndWrap | EPPF::SamplingLinear | EPPF::VectorizedImpl):
+		case (EPPF::FormatRGB | EPPF::ProjectionWrap | EPPF::SamplingLinear | EPPF::VectorizedImpl):
 		{
-			InvokeWrappingRasterizerImpl<EPPF::FormatRGB | EPPF::ProjectionPlanarAndWrap | EPPF::SamplingLinear | EPPF::VectorizedImpl>(RasterArgs);
+			InvokeWrappingRasterizerImpl<EPPF::FormatRGB | EPPF::ProjectionWrap | EPPF::SamplingLinear | EPPF::VectorizedImpl>(RasterArgs);
 			break;
 		}
-		case (EPPF::FormatRGBA | EPPF::ProjectionPlanarAndWrap | EPPF::SamplingLinear | EPPF::VectorizedImpl):
+		case (EPPF::FormatRGBA | EPPF::ProjectionWrap | EPPF::SamplingLinear | EPPF::VectorizedImpl):
 		{
-			InvokeWrappingRasterizerImpl<EPPF::FormatRGBA | EPPF::ProjectionPlanarAndWrap | EPPF::SamplingLinear | EPPF::VectorizedImpl>(RasterArgs);
+			InvokeWrappingRasterizerImpl<EPPF::FormatRGBA | EPPF::ProjectionWrap | EPPF::SamplingLinear | EPPF::VectorizedImpl>(RasterArgs);
 			break;
 		}
 
-		//case (EPPF::FormatL | EPPF::ProjectionPlanarAndWrap | EPPF::SamplingPoint | EPPF::WithMask | EPPF::VectorizedImpl):
+		//case (EPPF::FormatL | EPPF::ProjectionWrap | EPPF::SamplingPoint | EPPF::WithMask | EPPF::VectorizedImpl):
 		//{
-		//	InvokeWrappingRasterizerImpl<EPPF::FormatL | EPPF::ProjectionPlanarAndWrap | EPPF::SamplingPoint | EPPF::WithMask | EPPF::VectorizedImpl>(RasterArgs);
+		//	InvokeWrappingRasterizerImpl<EPPF::FormatL | EPPF::ProjectionWrap | EPPF::SamplingPoint | EPPF::WithMask | EPPF::VectorizedImpl>(RasterArgs);
 		//	break;
 		//}
-		//case (EPPF::FormatRGB | EPPF::ProjectionPlanarAndWrap | EPPF::SamplingPoint | EPPF::WithMask | EPPF::VectorizedImpl):
+		//case (EPPF::FormatRGB | EPPF::ProjectionWrap | EPPF::SamplingPoint | EPPF::WithMask | EPPF::VectorizedImpl):
 		//{
-		//	InvokeWrappingRasterizerImpl<EPPF::FormatRGB | EPPF::ProjectionPlanarAndWrap | EPPF::SamplingPoint | EPPF::WithMask | EPPF::VectorizedImpl>(RasterArgs);
+		//	InvokeWrappingRasterizerImpl<EPPF::FormatRGB | EPPF::ProjectionWrap | EPPF::SamplingPoint | EPPF::WithMask | EPPF::VectorizedImpl>(RasterArgs);
 		//	break;
 		//}
-		//case (EPPF::FormatRGBA | EPPF::ProjectionPlanarAndWrap | EPPF::SamplingPoint | EPPF::WithMask | EPPF::VectorizedImpl):
+		//case (EPPF::FormatRGBA | EPPF::ProjectionWrap | EPPF::SamplingPoint | EPPF::WithMask | EPPF::VectorizedImpl):
 		//{
-		//	InvokeWrappingRasterizerImpl<EPPF::FormatRGBA | EPPF::ProjectionPlanarAndWrap | EPPF::SamplingPoint | EPPF::VectorizedImpl>(RasterArgs);
+		//	InvokeWrappingRasterizerImpl<EPPF::FormatRGBA | EPPF::ProjectionWrap | EPPF::SamplingPoint | EPPF::VectorizedImpl>(RasterArgs);
 		//	break;
 		//}
 
-		case (EPPF::FormatL | EPPF::ProjectionPlanarAndWrap | EPPF::SamplingLinear | EPPF::WithMask | EPPF::VectorizedImpl):
+		case (EPPF::FormatL | EPPF::ProjectionWrap | EPPF::SamplingLinear | EPPF::WithMask | EPPF::VectorizedImpl):
 		{
-			InvokeWrappingRasterizerImpl<EPPF::FormatL | EPPF::ProjectionPlanarAndWrap | EPPF::SamplingLinear | EPPF::WithMask | EPPF::VectorizedImpl>(RasterArgs);
+			InvokeWrappingRasterizerImpl<EPPF::FormatL | EPPF::ProjectionWrap | EPPF::SamplingLinear | EPPF::WithMask | EPPF::VectorizedImpl>(RasterArgs);
 			break;
 		}
-		case (EPPF::FormatRGB | EPPF::ProjectionPlanarAndWrap | EPPF::SamplingLinear | EPPF::WithMask | EPPF::VectorizedImpl):
+		case (EPPF::FormatRGB | EPPF::ProjectionWrap | EPPF::SamplingLinear | EPPF::WithMask | EPPF::VectorizedImpl):
 		{
-			InvokeWrappingRasterizerImpl<EPPF::FormatRGB | EPPF::ProjectionPlanarAndWrap | EPPF::SamplingLinear | EPPF::WithMask>(RasterArgs);
+			InvokeWrappingRasterizerImpl<EPPF::FormatRGB | EPPF::ProjectionWrap | EPPF::SamplingLinear | EPPF::WithMask>(RasterArgs);
 			break;
 		}
-		case (EPPF::FormatRGBA | EPPF::ProjectionPlanarAndWrap | EPPF::SamplingLinear | EPPF::WithMask | EPPF::VectorizedImpl):
+		case (EPPF::FormatRGBA | EPPF::ProjectionWrap | EPPF::SamplingLinear | EPPF::WithMask | EPPF::VectorizedImpl):
 		{
-			InvokeWrappingRasterizerImpl<EPPF::FormatRGBA | EPPF::ProjectionPlanarAndWrap | EPPF::SamplingLinear | EPPF::WithMask | EPPF::VectorizedImpl>(RasterArgs);
+			InvokeWrappingRasterizerImpl<EPPF::FormatRGBA | EPPF::ProjectionWrap | EPPF::SamplingLinear | EPPF::WithMask | EPPF::VectorizedImpl>(RasterArgs);
 			break;
 		}
 		
@@ -1620,17 +1635,17 @@ float ComputeProjectedFootprintBestMip(
 			const FOptimizedVertex& B = VerticesPtr[IndicesPtr[I + 1]];
 			const FOptimizedVertex& C = VerticesPtr[IndicesPtr[I + 2]];
 
-			const float UnwindedProjectorAngleRad = FMath::UnwindRadians(Projector.projectionAngle);
-			const float OneOverProjectionAngleSafe = !FMath::IsNearlyZero(UnwindedProjectorAngleRad) 
+			const float ClampedProjectorAngleRad = FMath::Clamp(Projector.projectionAngle, 0.0f, UE_TWO_PI);
+			const float OneOverProjectionAngleSafe = !FMath::IsNearlyZero(ClampedProjectorAngleRad) 
 				? 1.0f / Projector.projectionAngle
 				: 1.0f / KINDA_SMALL_NUMBER;
 
-			const FVector2f PosA = FVector2f(0.5f + FMath::Atan2(A.Position.Z, -A.Position.Y) * OneOverProjectionAngleSafe, A.Position.X);
-			const FVector2f PosB = FVector2f(0.5f + FMath::Atan2(B.Position.Z, -B.Position.Y) * OneOverProjectionAngleSafe, B.Position.X);
-			const FVector2f PosC = FVector2f(0.5f + FMath::Atan2(C.Position.Z, -C.Position.Y) * OneOverProjectionAngleSafe, C.Position.X);
+			const FVector2f PosA = FVector2f((FMath::Atan2(A.Position.Z, -A.Position.Y) + UE_PI) * OneOverProjectionAngleSafe, A.Position.X);
+			const FVector2f PosB = FVector2f((FMath::Atan2(B.Position.Z, -B.Position.Y) + UE_PI) * OneOverProjectionAngleSafe, B.Position.X);
+			const FVector2f PosC = FVector2f((FMath::Atan2(C.Position.Z, -C.Position.Y) + UE_PI) * OneOverProjectionAngleSafe, C.Position.X);
 
-			const float TriangleSourceArea = ComputeTriangleArea(PosA, PosB, PosC);
-			const float TriangleTargetArea = ComputeTriangleArea(A.Uv, B.Uv, C.Uv);
+			const float TriangleSourceArea = FMath::Abs(ComputeTriangleArea(PosA, PosB, PosC));
+			const float TriangleTargetArea = FMath::Abs(ComputeTriangleArea(A.Uv, B.Uv, C.Uv));
 
 			// Set weight to zero if source or target area are close to zero to remove outliers.
 			float TriangleWeight = static_cast<float>(!FMath::IsNearlyZero(TriangleSourceArea)) *
@@ -1813,8 +1828,8 @@ void PlanarlyProjectVertex(const FVector3f& unfoldedPosition, FVector4f& project
 
 FVector2f ChangeBase2D(const FVector2f& origPosition, const FVector2f& newOrigin, const FVector2f& newBaseX, const FVector2f& newBaseY)
 {
-	float x = FVector2f::DotProduct(origPosition - newOrigin, newBaseX) / powf(newBaseX.Length(), 2.f) + 0.5f;
-	float y = FVector2f::DotProduct(origPosition - newOrigin, newBaseY) / powf(newBaseY.Length(), 2.f) + 0.5f;
+	float x = FVector2f::DotProduct(origPosition - newOrigin, newBaseX) / newBaseX.SquaredLength() + 0.5f;
+	float y = FVector2f::DotProduct(origPosition - newOrigin, newBaseY) / newBaseY.SquaredLength() + 0.5f;
 	//y = 1.0f - y;
 
 	return FVector2f(x, y);
@@ -2145,7 +2160,6 @@ void MeshProject_Optimised_Cylindrical(const FOptimizedVertex* pVertices, int ve
     }
 }
 
-
 //-------------------------------------------------------------------------------------------------
 void MeshProject_Optimised_Wrapping(const Mesh* pMesh,
                                     const FVector3f& projectorPosition, const FVector3f& projectorDirection,
@@ -2891,14 +2905,14 @@ void MeshProject_Optimised_Wrapping(const Mesh* pMesh,
                     pResultVertices[currentVertex].Position[2] = projectedPositions[i].pos2;
 
                     // Normal is actually the fade factor
-                    int step = faceStep[f];
-                    const float maxGradient = 10.f;
-                    float stepGradient = step / maxGradient;
-                    stepGradient = stepGradient > maxGradient ? maxGradient : stepGradient;
-                    float angleCos = stepGradient; //1.f; // dot(pVertices[i].Normal, projectorDirection * -1.0f);
-                    pResultVertices[currentVertex].Normal[0] = angleCos;
-                    pResultVertices[currentVertex].Normal[1] = angleCos;
-                    pResultVertices[currentVertex].Normal[2] = angleCos;
+                    constexpr float MaxGradient = 10.f;
+                    const int32 Step = faceStep[f];
+                    float StepGradient = FMath::Min(Step, MaxGradient) / MaxGradient;
+                    float FadeFactor = 1.0f - StepGradient; //1.f; // dot(pVertices[i].Normal, projectorDirection * -1.0f);
+                    
+					pResultVertices[currentVertex].Normal[0] = FadeFactor;
+                    pResultVertices[currentVertex].Normal[1] = FadeFactor;
+                    pResultVertices[currentVertex].Normal[2] = FadeFactor;
 
                     oldToNewVertex[i] = currentVertex++;
                 }
@@ -2957,7 +2971,6 @@ void MeshProject_Optimised_Wrapping(const Mesh* pMesh,
 
 }
 
-
 //-------------------------------------------------------------------------------------------------
 void MeshProject_Optimised(Mesh* Result, const Mesh* pMesh, const FProjector& projector, bool& bOutSuccess)
 {
