[Mutable] Remove TODO.

#rnx

diff --git a/Engine/Plugins/Experimental/Mutable/Source/CustomizableObjectEditor/Private/MuCOE/CustomizableObjectInstanceEditor.cpp b/Engine/Plugins/Experimental/Mutable/Source/CustomizableObjectEditor/Private/MuCOE/CustomizableObjectInstanceEditor.cpp
--- a/Engine/Plugins/Experimental/Mutable/Source/CustomizableObjectEditor/Private/MuCOE/CustomizableObjectInstanceEditor.cpp
+++ b/Engine/Plugins/Experimental/Mutable/Source/CustomizableObjectEditor/Private/MuCOE/CustomizableObjectInstanceEditor.cpp
@@ -973,7 +973,7 @@ void FCustomizableObjectInstanceEditor::CreatePreviewActor()
 		CustomizableComponent->RegisterComponent();
 	}
 	
-	Viewport->SetPreviewActor(Actor.Get(), CustomizableObjectInstance, PreviewSkeletalMeshComponents); // TODO GMT
+	Viewport->SetPreviewActor(Actor.Get(), CustomizableObjectInstance, PreviewSkeletalMeshComponents);
 }
 
 
