[mutable] Redesign the object node hierarchy to better match unreal component and LOD concepts.

- Mutable used to have a hierarchy like LOD->Component->Material, now it has Component->LOD->Material more in line with UE. This doesn't affect existing assets.
- Remove some boilerplate from the nodes in MutableTools.
- Add explicit component support with a pass-through mesh.
- Fix bug in test commandlet only generating the minimal LODs.

#rnx

diff --git a/Engine/Plugins/Experimental/Mutable/Source/CustomizableObject/Internal/MuCO/CustomizableObjectInstancePrivate.h b/Engine/Plugins/Experimental/Mutable/Source/CustomizableObject/Internal/MuCO/CustomizableObjectInstancePrivate.h
--- a/Engine/Plugins/Experimental/Mutable/Source/CustomizableObject/Internal/MuCO/CustomizableObjectInstancePrivate.h
+++ b/Engine/Plugins/Experimental/Mutable/Source/CustomizableObject/Internal/MuCO/CustomizableObjectInstancePrivate.h
@@ -287,7 +287,6 @@ private:
 
 	bool DoComponentsNeedUpdate(UCustomizableObjectInstance* CustomizableObjectInstance, const TSharedRef<FUpdateContextPrivate>& OperationData, bool& bOutEmptyMesh);
 
-	mu::FResourceID GetLastMeshId(int32 ComponentIndex, int32 LODIndex) const;
 	void SetLastMeshId(int32 ComponentIndex, int32 LODIndex, mu::FResourceID MeshId);
 
 public:
diff --git a/Engine/Plugins/Experimental/Mutable/Source/CustomizableObject/Internal/MuCO/CustomizableObjectPrivate.h b/Engine/Plugins/Experimental/Mutable/Source/CustomizableObject/Internal/MuCO/CustomizableObjectPrivate.h
--- a/Engine/Plugins/Experimental/Mutable/Source/CustomizableObject/Internal/MuCO/CustomizableObjectPrivate.h
+++ b/Engine/Plugins/Experimental/Mutable/Source/CustomizableObject/Internal/MuCO/CustomizableObjectPrivate.h
@@ -664,6 +664,8 @@ struct FModelResources
 	 * properties for everything that Mutable doesn't create or modify. This struct stores the information used from
 	 * the Reference Skeletal Meshes to avoid having them loaded at all times. This includes data like LOD distances,
 	 * LOD render data settings, Mesh sockets, Bounding volumes, etc.
+	 * 
+	 * Index with Component index
 	 */
 	UPROPERTY()
 	TArray<FMutableRefSkeletalMeshData> ReferenceSkeletalMeshesData;
@@ -1075,6 +1077,10 @@ public:
 
 		AddedParameterThumbnailsToEditor,
 
+		ComponentsLODsRedesign,
+
+		ComponentsLODsRedesign2,
+
 		// -----<new versions can be added above this line>--------
 		LastCustomizableObjectVersion
 	};
diff --git a/Engine/Plugins/Experimental/Mutable/Source/CustomizableObject/Internal/MuCO/CustomizableObjectSystemPrivate.h b/Engine/Plugins/Experimental/Mutable/Source/CustomizableObject/Internal/MuCO/CustomizableObjectSystemPrivate.h
--- a/Engine/Plugins/Experimental/Mutable/Source/CustomizableObject/Internal/MuCO/CustomizableObjectSystemPrivate.h
+++ b/Engine/Plugins/Experimental/Mutable/Source/CustomizableObject/Internal/MuCO/CustomizableObjectSystemPrivate.h
@@ -356,15 +356,10 @@ struct FInstanceUpdateData
 		uint32 SurfaceId = 0;
 	};
 
-	struct FComponent
-	{
-		uint16 Id = 0;
-		
-		// True if the Mesh is valid
-		bool bGenerated = false;
-
+	struct FLOD
+	{		
 		mu::FResourceID MeshID;
-		mu::MeshPtrConst Mesh;
+		mu::Ptr<const mu::Mesh> Mesh;
 
 		/** Range in the Surfaces array */
 		uint16 FirstSurface = 0;
@@ -377,17 +372,22 @@ struct FInstanceUpdateData
 		/** Range in the external Bones array */
 		uint32 FirstBoneMap = 0;
 		uint32 BoneMapCount = 0;
+
+		// True if the Mesh is valid
+		bool bGenerated = false;
 	};
 
-	struct FLOD
+	struct FComponent
 	{
-		/** Range in the Components array */
-		uint16 FirstComponent = 0;
-		uint16 ComponentCount = 0;
+		uint16 Id = 0;
+		
+		/** Range in the LODs array */
+		uint16 FirstLOD = 0;
+		uint16 LODCount = 0;
 	};
 
-	TArray<FLOD> LODs;
 	TArray<FComponent> Components;
+	TArray<FLOD> LODs;
 	TArray<FSurface> Surfaces;
 	TArray<FImage> Images;
 	TArray<FVector> Vectors;
@@ -443,6 +443,7 @@ struct FInstanceUpdateData
 		TMap<mu::FBoneName, TPair<FName, uint16>> BoneInfoMap;
 	};
 
+	// Access by component index
 	TArray<FSkeletonData> Skeletons;
 
 	struct FNamedExtensionData
@@ -511,7 +512,7 @@ private:
 
 public:
 	/** Instance parameters at the time of the operation request. */
-	mu::ParametersPtr Parameters; 
+	mu::Ptr<mu::Parameters> Parameters; 
 	mu::Ptr<mu::System> MutableSystem;
 
 	bool bOnlyUpdateIfNotGenerated = false;
@@ -538,12 +539,16 @@ public:
 	int32 MipsToSkip = 0;
 
 	mu::Instance::ID InstanceID = 0; // Redundant
-	const mu::Instance* MutableInstance = nullptr;
+	mu::Ptr<const mu::Instance> MutableInstance;
 
 	TSharedPtr<mu::Model> Model;
 
+	// TODO: This is also available in InstanceUpdateData in MutableInstance and in NumLODsAvailablePerComponent
 	uint8 NumComponents = 0;
-	uint8 NumLODsAvailable = 0;
+
+	// TODO: Why is this cached?
+	TArray<uint8> NumLODsAvailablePerComponent;
+
 	uint8 FirstLODAvailable = 0;
 	uint8 FirstResidentLOD = 0;
 
@@ -553,7 +558,7 @@ public:
 
 	mu::FImageOperator::FImagePixelFormatFunc PixelFormatOverride;
 
-	/** Mutable Meshes required for each component. Outermost index is the component, inner index is the LOD. */
+	/** Mutable Meshes required for each component. Outermost index is the component index, inner index is the LOD. */
 	TArray<TArray<mu::FResourceID>> MeshDescriptors;
 
 	/** Used to know if the updated instances' meshes are different from the previous ones. 
diff --git a/Engine/Plugins/Experimental/Mutable/Source/CustomizableObject/Private/MuCO/BusyWaits_Deprecated.cpp b/Engine/Plugins/Experimental/Mutable/Source/CustomizableObject/Private/MuCO/BusyWaits_Deprecated.cpp
--- a/Engine/Plugins/Experimental/Mutable/Source/CustomizableObject/Private/MuCO/BusyWaits_Deprecated.cpp
+++ b/Engine/Plugins/Experimental/Mutable/Source/CustomizableObject/Private/MuCO/BusyWaits_Deprecated.cpp
@@ -191,7 +191,7 @@ namespace CustomizableObjectSystem::ImplDeprecated
 
 		// Main instance generation step
 		// LOD mask, set to all ones to build  all LODs
-		const mu::Instance* Instance = OperationData->MutableInstance; // TODO GMTFuture remove
+		const mu::Instance* Instance = OperationData->MutableInstance.get(); // TODO GMTFuture remove
 		if (!Instance)
 		{
 			UE_LOG(LogMutable, Warning, TEXT("An Instace update has failed."));
@@ -204,43 +204,45 @@ namespace CustomizableObjectSystem::ImplDeprecated
 		TArray<int32> SurfacesSharedId;
 
 		// Generate the mesh and gather all the required resource Ids
-		OperationData->InstanceUpdateData.LODs.SetNum(OperationData->NumLODsAvailable);
-		for (int32 MutableLODIndex = 0; MutableLODIndex < OperationData->NumLODsAvailable; ++MutableLODIndex)
+		int32 NumComponents = OperationData->NumComponents;
+		OperationData->InstanceUpdateData.Components.SetNum(NumComponents);
+		for (int32 ComponentIndex = 0; ComponentIndex < NumComponents; ++ComponentIndex)
 		{
-			// Skip LODs outside the range we want to generate
-			if (MutableLODIndex < OperationData->GetMinLOD())
+			FInstanceUpdateData::FComponent& Component = OperationData->InstanceUpdateData.Components[ComponentIndex];
+			Component.FirstLOD = OperationData->InstanceUpdateData.LODs.Num();
+			Component.LODCount = Instance->GetLODCount(ComponentIndex);
+			Component.Id = Instance->GetComponentId(ComponentIndex);
+
+			for (int32 MutableLODIndex = 0; MutableLODIndex < Component.LODCount; ++MutableLODIndex)
 			{
-				continue;
-			}
+				// If the LOD is not generated we still add an empty one to keep indexes aligned.
+				OperationData->InstanceUpdateData.LODs.Push(FInstanceUpdateData::FLOD());
 
-			FInstanceUpdateData::FLOD& LOD = OperationData->InstanceUpdateData.LODs[MutableLODIndex];
-			LOD.FirstComponent = OperationData->InstanceUpdateData.Components.Num();
-			LOD.ComponentCount = Instance->GetComponentCount(MutableLODIndex);
+				// Skip LODs outside the range we want to generate
+				if (MutableLODIndex < OperationData->GetMinLOD())
+				{
+					continue;
+				}
 
-			for (int32 ComponentIndex = 0; ComponentIndex < LOD.ComponentCount; ++ComponentIndex)
-			{
-				OperationData->InstanceUpdateData.Components.Push(FInstanceUpdateData::FComponent());
-				FInstanceUpdateData::FComponent& Component = OperationData->InstanceUpdateData.Components.Last();
-				Component.Id = Instance->GetComponentId(MutableLODIndex, ComponentIndex);
-				Component.FirstSurface = OperationData->InstanceUpdateData.Surfaces.Num();
-				Component.SurfaceCount = 0;
+				FInstanceUpdateData::FLOD& LOD = OperationData->InstanceUpdateData.LODs.Last();
+				LOD.FirstSurface = OperationData->InstanceUpdateData.Surfaces.Num();
+				LOD.SurfaceCount = 0;
 
-				const bool bGenerateLOD = RequestedLODs.IsValidIndex(Component.Id) ? RequestedLODs[Component.Id] <= MutableLODIndex : true;
+				const bool bGenerateLOD = RequestedLODs.IsValidIndex(ComponentIndex) ? RequestedLODs[ComponentIndex] <= MutableLODIndex : true;
 
 				// Mesh
-				if (Instance->GetMeshCount(MutableLODIndex, ComponentIndex) > 0)
 				{
 					MUTABLE_CPUPROFILER_SCOPE(GetMesh);
 
-					Component.MeshID = Instance->GetMeshId(MutableLODIndex, ComponentIndex, 0);
+					LOD.MeshID = Instance->GetMeshId(ComponentIndex, MutableLODIndex);
 
 					if (bGenerateLOD)
 					{
-						Component.Mesh = System->GetMeshInline(OperationData->InstanceID, Component.MeshID);
+						LOD.Mesh = System->GetMeshInline(OperationData->InstanceID, LOD.MeshID);
 
-						if (Component.Mesh->IsReference())
+						if (LOD.Mesh->IsReference())
 						{
-							uint32 ReferenceID = Component.Mesh->GetReferencedMesh();
+							uint32 ReferenceID = LOD.Mesh->GetReferencedMesh();
 
 							if (ModelResources.PassThroughMeshes.IsValidIndex(ReferenceID))
 							{
@@ -256,35 +258,28 @@ namespace CustomizableObjectSystem::ImplDeprecated
 					}
 				}
 
-				if (!Component.Mesh)
+				if (!LOD.Mesh)
 				{
 					continue;
 				}
 
-				Component.bGenerated = true;
-
+				LOD.bGenerated = true;
 
-				// Materials and images
-				const int32 SurfaceCount = Component.Mesh->GetSurfaceCount();
-				for (int32 MeshSurfaceIndex = 0; MeshSurfaceIndex < SurfaceCount; ++MeshSurfaceIndex)
-				{
-					const uint32 SurfaceId = Component.Mesh->GetSurfaceId(MeshSurfaceIndex);
-					const int32 InstanceSurfaceIndex = Instance->FindSurfaceById(MutableLODIndex, ComponentIndex, SurfaceId);
-					check(Component.Mesh->GetVertexCount() > 0 || InstanceSurfaceIndex >= 0);
 
-					int32 BaseSurfaceIndex = InstanceSurfaceIndex;
-					int32 BaseLODIndex = MutableLODIndex;
-
-					if (InstanceSurfaceIndex >= 0)
+				// This lambda does all the work to fill up the surface data
+				auto AddSurface = [&LOD, &SurfacesSharedId, &ModelResources, OperationData, Instance, System, CustomizableObject, CustomizableObjectInstancePrivateData, ComponentIndex, MutableLODIndex ]( uint32 SurfaceId, int32 InstanceSurfaceIndex)
 					{
+						int32 BaseSurfaceIndex = InstanceSurfaceIndex;
+						int32 BaseLODIndex = MutableLODIndex;
+
 						OperationData->InstanceUpdateData.Surfaces.Push({});
 						FInstanceUpdateData::FSurface& Surface = OperationData->InstanceUpdateData.Surfaces.Last();
-						++Component.SurfaceCount;
+						++LOD.SurfaceCount;
 
 						// Now Surface.MaterialIndex is decoded from a parameter at the end of this if()
 						Surface.SurfaceId = SurfaceId;
 
-						const int32 SharedSurfaceId = Instance->GetSharedSurfaceId(MutableLODIndex, ComponentIndex, InstanceSurfaceIndex);
+						const int32 SharedSurfaceId = Instance->GetSharedSurfaceId(ComponentIndex, MutableLODIndex, InstanceSurfaceIndex);
 						const int32 SharedSurfaceIndex = SurfacesSharedId.Find(SharedSurfaceId);
 
 						SurfacesSharedId.Add(SharedSurfaceId);
@@ -294,26 +289,26 @@ namespace CustomizableObjectSystem::ImplDeprecated
 							if (SharedSurfaceIndex >= 0)
 							{
 								Surface = OperationData->InstanceUpdateData.Surfaces[SharedSurfaceIndex];
-								continue;
+								return;
 							}
 
 							// Find the first LOD where this surface can be found
 							Instance->FindBaseSurfaceBySharedId(ComponentIndex, SharedSurfaceId, BaseSurfaceIndex, BaseLODIndex);
 
-							Surface.SurfaceId = Instance->GetSurfaceId(BaseLODIndex, ComponentIndex, BaseSurfaceIndex);
+							Surface.SurfaceId = Instance->GetSurfaceId(ComponentIndex, BaseLODIndex, BaseSurfaceIndex);
 						}
 
 						// Images
 						Surface.FirstImage = OperationData->InstanceUpdateData.Images.Num();
-						Surface.ImageCount = Instance->GetImageCount(BaseLODIndex, ComponentIndex, BaseSurfaceIndex);
+						Surface.ImageCount = Instance->GetImageCount(ComponentIndex, BaseLODIndex, BaseSurfaceIndex);
 						for (int32 ImageIndex = 0; ImageIndex < Surface.ImageCount; ++ImageIndex)
 						{
 							MUTABLE_CPUPROFILER_SCOPE(GetImageId);
 
 							OperationData->InstanceUpdateData.Images.Push({});
 							FInstanceUpdateData::FImage& Image = OperationData->InstanceUpdateData.Images.Last();
-							Image.Name = Instance->GetImageName(BaseLODIndex, ComponentIndex, BaseSurfaceIndex, ImageIndex);
-							Image.ImageID = Instance->GetImageId(BaseLODIndex, ComponentIndex, BaseSurfaceIndex, ImageIndex);
+							Image.Name = Instance->GetImageName(ComponentIndex, BaseLODIndex, BaseSurfaceIndex, ImageIndex);
+							Image.ImageID = Instance->GetImageId(ComponentIndex, BaseLODIndex, BaseSurfaceIndex, ImageIndex);
 							Image.FullImageSizeX = 0;
 							Image.FullImageSizeY = 0;
 							Image.BaseLOD = BaseLODIndex;
@@ -325,7 +320,7 @@ namespace CustomizableObjectSystem::ImplDeprecated
 							if (ImageKey >= 0 && ImageKey < ModelResources.ImageProperties.Num())
 							{
 								const FMutableModelImageProperties& Props = ModelResources.ImageProperties[ImageKey];
-								
+
 								// This has beed replicated to the non deperecated code path.
 								Image.bIsNonProgressive = Props.MipGenSettings == TMGS_NoMipmaps;
 
@@ -361,33 +356,33 @@ namespace CustomizableObjectSystem::ImplDeprecated
 
 						// Vectors
 						Surface.FirstVector = OperationData->InstanceUpdateData.Vectors.Num();
-						Surface.VectorCount = Instance->GetVectorCount(BaseLODIndex, ComponentIndex, BaseSurfaceIndex);
+						Surface.VectorCount = Instance->GetVectorCount(ComponentIndex, BaseLODIndex, BaseSurfaceIndex);
 						for (int32 VectorIndex = 0; VectorIndex < Surface.VectorCount; ++VectorIndex)
 						{
 							MUTABLE_CPUPROFILER_SCOPE(GetVector);
 							OperationData->InstanceUpdateData.Vectors.Push({});
 							FInstanceUpdateData::FVector& Vector = OperationData->InstanceUpdateData.Vectors.Last();
-							Vector.Name = Instance->GetVectorName(BaseLODIndex, ComponentIndex, BaseSurfaceIndex, VectorIndex);
-							Vector.Vector = Instance->GetVector(BaseLODIndex, ComponentIndex, BaseSurfaceIndex, VectorIndex);
+							Vector.Name = Instance->GetVectorName(ComponentIndex, BaseLODIndex, BaseSurfaceIndex, VectorIndex);
+							Vector.Vector = Instance->GetVector(ComponentIndex, BaseLODIndex, BaseSurfaceIndex, VectorIndex);
 						}
 
 						// Scalars
 						Surface.FirstScalar = OperationData->InstanceUpdateData.Scalars.Num();
-						Surface.ScalarCount = Instance->GetScalarCount(BaseLODIndex, ComponentIndex, BaseSurfaceIndex);
+						Surface.ScalarCount = Instance->GetScalarCount(ComponentIndex, BaseLODIndex, BaseSurfaceIndex);
 						for (int32 ScalarIndex = 0; ScalarIndex < Surface.ScalarCount; ++ScalarIndex)
 						{
 							MUTABLE_CPUPROFILER_SCOPE(GetScalar)
 
-							const FName ScalarName = Instance->GetScalarName(BaseLODIndex, ComponentIndex, BaseSurfaceIndex, ScalarIndex);
-							const float ScalarValue = Instance->GetScalar(BaseLODIndex, ComponentIndex, BaseSurfaceIndex, ScalarIndex);
-							
+								const FName ScalarName = Instance->GetScalarName(ComponentIndex, BaseLODIndex, BaseSurfaceIndex, ScalarIndex);
+							const float ScalarValue = Instance->GetScalar(ComponentIndex, BaseLODIndex, BaseSurfaceIndex, ScalarIndex);
+
 							FString EncodingMaterialIdString = "__MutableMaterialId";
-							
+
 							// Decoding Material Switch from Mutable parameter name
 							if (ScalarName.ToString().Equals(EncodingMaterialIdString))
 							{
 								Surface.MaterialIndex = static_cast<uint32>(ScalarValue);
-							
+
 								// This parameter is not needed in the final material instance
 								Surface.ScalarCount -= 1;
 							}
@@ -396,6 +391,38 @@ namespace CustomizableObjectSystem::ImplDeprecated
 								OperationData->InstanceUpdateData.Scalars.Push({ ScalarName, ScalarValue });
 							}
 						}
+
+					};
+
+				// Materials and images
+
+				// If the mesh is a reference mesh, it won't have the surface information in the mutable mesh. We need to get it from the instance
+				// and all defined surfaces will be present. 
+				if (LOD.Mesh->IsReference())
+				{
+					const int32 SurfaceCount = Instance->GetSurfaceCount(ComponentIndex, MutableLODIndex);
+					for (int32 SurfaceIndex = 0; SurfaceIndex < SurfaceCount; ++SurfaceIndex)
+					{
+						uint32 SurfaceId = Instance->GetSurfaceId(ComponentIndex, MutableLODIndex, SurfaceIndex);
+						AddSurface( SurfaceId, SurfaceIndex);
+					}
+				}
+
+				// If the mesh is a not a reference mesh, we have to add only the materials of the surfaces that appear in the actual final mesh. 
+				else
+				{
+					const int32 SurfaceCount = LOD.Mesh->GetSurfaceCount();
+					for (int32 MeshSurfaceIndex = 0; MeshSurfaceIndex < SurfaceCount; ++MeshSurfaceIndex)
+					{
+						const uint32 SurfaceId = LOD.Mesh->GetSurfaceId(MeshSurfaceIndex);
+
+						const int32 InstanceSurfaceIndex = Instance->FindSurfaceById(ComponentIndex, MutableLODIndex, SurfaceId);
+						check(LOD.Mesh->GetVertexCount() > 0 || InstanceSurfaceIndex >= 0);
+
+						if (InstanceSurfaceIndex >= 0)
+						{
+							AddSurface(SurfaceId, InstanceSurfaceIndex);
+						}
 					}
 				}
 			}
@@ -470,7 +497,7 @@ namespace CustomizableObjectMipDataProvider::ImplDeprecated
 
 		if (OperationData.IsValid())
 		{
-			mu::SystemPtr System = OperationData->UpdateContext->GetSystem();
+			mu::Ptr<mu::System> System = OperationData->UpdateContext->GetSystem();
 			const TSharedPtr<mu::Model, ESPMode::ThreadSafe> Model = OperationData->UpdateContext->GetModel();
 
 #if WITH_EDITOR
@@ -483,7 +510,7 @@ namespace CustomizableObjectMipDataProvider::ImplDeprecated
 				mu::Instance::ID InstanceID = System->NewInstance(Model);
 				UE_LOG(LogMutable, Verbose, TEXT("Creating Mutable instance with id [%d] for a single UpdateImage"), InstanceID)
 
-					const mu::Instance* Instance = nullptr;
+				mu::Ptr<const mu::Instance> Instance;
 
 				// Main instance generation step
 				{
@@ -501,12 +528,12 @@ namespace CustomizableObjectMipDataProvider::ImplDeprecated
 
 					const FMutableImageReference& ImageRef = OperationData->RequestedImage;
 
-					int32 SurfaceIndex = Instance->FindSurfaceById(ImageRef.LOD, ImageRef.Component, ImageRef.SurfaceId);
+					int32 SurfaceIndex = Instance->FindSurfaceById(ImageRef.Component, ImageRef.LOD, ImageRef.SurfaceId);
 					check(SurfaceIndex >= 0);
 
 					// This ID may be different than the ID obtained the first time the image was generated, because the mutable
 					// runtime cannot remember all the resources it has built, and only remembers a fixed amount.
-					mu::FResourceID MipImageID = Instance->GetImageId(ImageRef.LOD, ImageRef.Component, SurfaceIndex, ImageRef.Image);
+					mu::FResourceID MipImageID = Instance->GetImageId(ImageRef.Component, ImageRef.LOD, SurfaceIndex, ImageRef.Image);
 
 
 					mu::Ptr<const mu::Image> ResultImage;
@@ -645,7 +672,7 @@ namespace CustomizableObjectMeshUpdate::ImplDeprecated
 		const uint32 LODMask = 0xFFFFFFFF;
 
 		// Main instance generation step
-		const mu::Instance* Instance = System->BeginUpdate(InstanceID, OperationData->Parameters, OperationData->State, LODMask);
+		mu::Ptr<const mu::Instance> Instance = System->BeginUpdate(InstanceID, OperationData->Parameters, OperationData->State, LODMask);
 		check(Instance);
 
 		// Generate the required meshes
diff --git a/Engine/Plugins/Experimental/Mutable/Source/CustomizableObject/Private/MuCO/CustomizableObjectInstance.cpp b/Engine/Plugins/Experimental/Mutable/Source/CustomizableObject/Private/MuCO/CustomizableObjectInstance.cpp
--- a/Engine/Plugins/Experimental/Mutable/Source/CustomizableObject/Private/MuCO/CustomizableObjectInstance.cpp
+++ b/Engine/Plugins/Experimental/Mutable/Source/CustomizableObject/Private/MuCO/CustomizableObjectInstance.cpp
@@ -183,22 +183,6 @@ UTexture2D* UCustomizableInstancePrivate::CreateTexture()
 }
 
 
-mu::FResourceID UCustomizableInstancePrivate::GetLastMeshId(int32 ComponentIndex, int32 LODIndex) const
-{
-	const FCustomizableInstanceComponentData* ComponentData = GetComponentData(ComponentIndex);
-
-	if (ComponentData && ComponentData->LastMeshIdPerLOD.IsValidIndex(LODIndex))
-	{
-		return ComponentData->LastMeshIdPerLOD[LODIndex];
-	}
-
-	check(ComponentData);
-	check(ComponentData->LastMeshIdPerLOD.IsValidIndex(LODIndex));
-		
-	return MAX_uint64;
-}
-
-
 void UCustomizableInstancePrivate::SetLastMeshId(int32 ComponentIndex, int32 LODIndex, mu::FResourceID MeshId)
 {
 	FCustomizableInstanceComponentData* ComponentData = GetComponentData(ComponentIndex);
@@ -1104,6 +1088,12 @@ void UCustomizableInstancePrivate::UpdateInstanceIfNotGenerated(UCustomizableObj
 bool AreSkeletonsCompatible(const TArray<TObjectPtr<USkeleton>>& InSkeletons)
 {
 	MUTABLE_CPUPROFILER_SCOPE(AreSkeletonsCompatible);
+
+	if (InSkeletons.IsEmpty())
+	{
+		return true;
+	}
+
 	bool bCompatible = true;
 
 	struct FBoneToMergeInfo
@@ -1873,7 +1863,7 @@ bool UCustomizableInstancePrivate::DoComponentsNeedUpdate(UCustomizableObjectIns
 		return false;
 	}
 
-	const int32 NumComponents = CustomizableObject->GetComponentCount();
+	const int32 NumComponents = OperationData->InstanceUpdateData.Components.Num();
 
 	TArray<bool> ComponentWithMesh;
 	ComponentWithMesh.Init(false, NumComponents);
@@ -1882,33 +1872,34 @@ bool UCustomizableInstancePrivate::DoComponentsNeedUpdate(UCustomizableObjectIns
 	MeshIDs.Init(MAX_uint64, NumComponents * MAX_MESH_LOD_COUNT);
 
 	// Gather the Mesh Ids of all components, and validate the integrity of the meshes to generate. 
-	for (int32 LODIndex = OperationData->GetMinLOD(); LODIndex < OperationData->NumLODsAvailable; ++LODIndex)
+	for (int32 ComponentIndex = 0; ComponentIndex < NumComponents; ++ComponentIndex)
 	{
-		const FInstanceUpdateData::FLOD& LOD = OperationData->InstanceUpdateData.LODs[LODIndex];
-
-		for (int32 ComponentIndex = 0; ComponentIndex < LOD.ComponentCount; ++ComponentIndex)
+		const FInstanceUpdateData::FComponent& Component = OperationData->InstanceUpdateData.Components[ComponentIndex];
+		
+		for (int32 LODIndex = OperationData->GetMinLOD(); LODIndex < Component.LODCount; ++LODIndex)
 		{
-			const FInstanceUpdateData::FComponent& Component = OperationData->InstanceUpdateData.Components[LOD.FirstComponent + ComponentIndex];
+			const FInstanceUpdateData::FLOD& LOD = OperationData->InstanceUpdateData.LODs[Component.FirstLOD + LODIndex];
+
 
-			if (Component.Id >= NumComponents)
+			if (ComponentIndex >= NumComponents)
 			{
-				ensureMsgf(false, TEXT("Mutable: Failed to generate SkeletalMesh. Invalid ComponentIndex. Index [%d]. Valid range [0 .. %d)."), Component.Id, NumComponents);
+				ensureMsgf(false, TEXT("Mutable: Failed to generate SkeletalMesh. Invalid ComponentIndex. Index [%d]. Valid range [0 .. %d)."), ComponentIndex, NumComponents);
 				bHasInvalidMesh = true;
 				return false;
 			}
 
-			if (!Component.bGenerated || !Component.Mesh ) 
+			if (!LOD.bGenerated || !LOD.Mesh )
 			{
 				continue;
 			}
 
-			if (Component.SurfaceCount == 0 && !Component.Mesh->IsReference()) 
+			if (LOD.SurfaceCount == 0 && !LOD.Mesh->IsReference())
 			{
 				continue;
 			}
 
 			// Unreal does not support empty sections.
-			if (!Component.Mesh->IsReference() && Component.Mesh->GetVertexCount() == 0) // else
+			if (!LOD.Mesh->IsReference() && LOD.Mesh->GetVertexCount() == 0) // else
 			{
 				UE_LOG(LogMutable, Error, TEXT("Failed to generate SkeletalMesh for CO Instance [%s]. CO [%s] has invalid geometry for LOD [%d] Component [%d]."),
 					*Public->GetName(), *CustomizableObject->GetName(),
@@ -1917,8 +1908,8 @@ bool UCustomizableInstancePrivate::DoComponentsNeedUpdate(UCustomizableObjectIns
 				continue;
 			}
 
-			ComponentWithMesh[Component.Id] = true;
-			MeshIDs[Component.Id * MAX_MESH_LOD_COUNT + LODIndex] = Component.MeshID;
+			ComponentWithMesh[ComponentIndex] = true;
+			MeshIDs[ComponentIndex * MAX_MESH_LOD_COUNT + LODIndex] = LOD.MeshID;
 		}
 	}
 
@@ -1940,13 +1931,19 @@ bool UCustomizableInstancePrivate::DoComponentsNeedUpdate(UCustomizableObjectIns
 		}
 
 		// Components with mesh must have valid geometry at CurrentMaxLOD
-		if (ComponentWithMesh[ComponentIndex] && MeshIDs[ComponentIndex * MAX_MESH_LOD_COUNT + OperationData->NumLODsAvailable - 1] == MAX_uint64)
+		const FInstanceUpdateData::FComponent& Component = OperationData->InstanceUpdateData.Components[ComponentIndex];
+		const FInstanceUpdateData::FLOD& LOD = OperationData->InstanceUpdateData.LODs[Component.FirstLOD];
+		bool bIsReferenced = LOD.Mesh && LOD.Mesh->IsReference();
+		if (!bIsReferenced)
 		{
-			UE_LOG(LogMutable, Error, TEXT("Failed to generate SkeletalMesh for CO Instance [%s]. CO [%s] is missing geometry for LOD [%d] Component [%d]."),
-				*Public->GetName(), *CustomizableObject->GetName(),
-				OperationData->NumLODsAvailable - 1, ComponentIndex);
-			bHasInvalidMesh = true;
-			continue;
+			if (ComponentWithMesh[ComponentIndex] && MeshIDs[ComponentIndex * MAX_MESH_LOD_COUNT + OperationData->NumLODsAvailablePerComponent[ComponentIndex] - 1] == MAX_uint64)
+			{
+				UE_LOG(LogMutable, Error, TEXT("Failed to generate SkeletalMesh for CO Instance [%s]. CO [%s] is missing geometry for LOD [%d] Component [%d]."),
+					*Public->GetName(), *CustomizableObject->GetName(),
+					OperationData->NumLODsAvailablePerComponent[ComponentIndex] - 1, ComponentIndex);
+				bHasInvalidMesh = true;
+				continue;
+			}
 		}
 
 		// Update the component if there is a mesh and it shouldn't, or the other way around.
@@ -2025,23 +2022,17 @@ bool UCustomizableInstancePrivate::UpdateSkeletalMesh_PostBeginUpdate0(UCustomiz
 		const FInstanceUpdateData::FComponent& Component = OperationData->InstanceUpdateData.Components[ComponentIndex];
 
 		// If the component doesn't need an update copy the previously generated mesh.
-		if (!OperationData->MeshChanged[Component.Id])
+		if (!OperationData->MeshChanged[ComponentIndex])
 		{
-			if (OldSkeletalMeshes.IsValidIndex(Component.Id))
+			if (OldSkeletalMeshes.IsValidIndex(ComponentIndex))
 			{
-				SkeletalMeshes[Component.Id] = OldSkeletalMeshes[Component.Id];
+				SkeletalMeshes[ComponentIndex] = OldSkeletalMeshes[ComponentIndex];
 			}
 
 			continue;
 		}
 
-		// Check if we have initialized the component
-		if (SkeletalMeshes[Component.Id])
-		{
-			continue;
-		}
-
-		if (!ComponentsData.IsValidIndex(Component.Id))
+		if (!ComponentsData.IsValidIndex(ComponentIndex))
 		{
 			bSuccess = false;
 			ensure(false);
@@ -2050,78 +2041,80 @@ bool UCustomizableInstancePrivate::UpdateSkeletalMesh_PostBeginUpdate0(UCustomiz
 
 		if (OperationData->bUseMeshCache)
 		{
-			if (USkeletalMesh* CachedMesh = CustomizableObject->GetPrivate()->MeshCache.Get(OperationData->MeshDescriptors[Component.Id]))
+			if (USkeletalMesh* CachedMesh = CustomizableObject->GetPrivate()->MeshCache.Get(OperationData->MeshDescriptors[ComponentIndex]))
 			{
-				check(OperationData->MeshDescriptors[Component.Id].Num() == MAX_MESH_LOD_COUNT);
-				ComponentsData[Component.Id].LastMeshIdPerLOD = OperationData->MeshDescriptors[Component.Id];
-				SkeletalMeshes[Component.Id] = CachedMesh;
+				check(OperationData->MeshDescriptors[ComponentIndex].Num() == MAX_MESH_LOD_COUNT);
+				ComponentsData[ComponentIndex].LastMeshIdPerLOD = OperationData->MeshDescriptors[ComponentIndex];
+				SkeletalMeshes[ComponentIndex] = CachedMesh;
 				continue;
 			}
 		}
 
 		// Reset last mesh IDs.
-		ComponentsData[Component.Id].LastMeshIdPerLOD.Init(MAX_uint64, MAX_MESH_LOD_COUNT);
-
-		if (!Component.bGenerated || !Component.Mesh)
-		{
-			continue;
-		}
+		ComponentsData[ComponentIndex].LastMeshIdPerLOD.Init(MAX_uint64, MAX_MESH_LOD_COUNT);
 
-		if (Component.SurfaceCount == 0 && !Component.Mesh->IsReference())
+		// If it is a referenced resource, only the first LOD is relevant.
 		{
-			continue;
-		}
-
-
-		if (Component.Mesh && Component.Mesh->IsReference())
-		{
-			int32 ReferenceID = Component.Mesh->GetReferencedMesh();
-			TSoftObjectPtr<USkeletalMesh> Ref = ModelResources.PassThroughMeshes[ReferenceID];
-
-			if (Ref.IsValid())
+			mu::Ptr<const mu::Mesh> FirstMesh;
+			if (Component.LODCount > 0)
 			{
-				SkeletalMeshes[Component.Id] = Ref.Get();
+				FirstMesh = OperationData->InstanceUpdateData.LODs[Component.FirstLOD].Mesh;
 			}
-			else
+
+			if (FirstMesh && FirstMesh->IsReference())
 			{
-				// This shouldn't happen here synchronosuly. It should have been requested as an async load.
-				SkeletalMeshes[Component.Id] = Ref.LoadSynchronous();
-			}
+				int32 ReferenceID = FirstMesh->GetReferencedMesh();
+				TSoftObjectPtr<USkeletalMesh> Ref = ModelResources.PassThroughMeshes[ReferenceID];
 
-			break;
-		}		
+				if (!Ref.IsValid())
+				{
+					// This shouldn't happen here synchronosuly. It should have been requested as an async load.
+					UE_LOG(LogMutable, Error, TEXT("Referenced skeletal mesh [%s] was not pre-loaded. It will be sync-loaded probably causing a hitch. CO [%s]"), *Ref.ToString(), *GetNameSafe(CustomizableObject));
+				}
+
+				SkeletalMeshes[ComponentIndex] = Ref.LoadSynchronous();
+				continue;
+			}
+		}
 
-		if (!ModelResources.ReferenceSkeletalMeshesData.IsValidIndex(Component.Id))
+		if (!ModelResources.ReferenceSkeletalMeshesData.IsValidIndex(ComponentIndex))
 		{
 			bSuccess = false;
 			break;
 		}
 
-		const FMutableRefSkeletalMeshData& RefSkeletalMeshData = ModelResources.ReferenceSkeletalMeshesData[Component.Id];
+		const FMutableRefSkeletalMeshData& RefSkeletalMeshData = ModelResources.ReferenceSkeletalMeshesData[ComponentIndex];
 
 		// Create and initialize the SkeletalMesh for this component
 		MUTABLE_CPUPROFILER_SCOPE(ConstructMesh);
 
 		if (OperationData->bStreamMeshLODs)
 		{
-			SkeletalMeshes[Component.Id] = UCustomizableObjectSkeletalMesh::CreateSkeletalMesh(OperationData, *Public, *CustomizableObject, Component.Id);
+			SkeletalMeshes[ComponentIndex] = UCustomizableObjectSkeletalMesh::CreateSkeletalMesh(OperationData, *Public, *CustomizableObject, ComponentIndex);
 		}
 		else
 		{
 			FName SkeletalMeshName = GenerateUniqueNameFromCOInstance(*Public);
-			SkeletalMeshes[Component.Id] = NewObject<USkeletalMesh>(GetTransientPackage(), SkeletalMeshName, RF_Transient);
+			SkeletalMeshes[ComponentIndex] = NewObject<USkeletalMesh>(GetTransientPackage(), SkeletalMeshName, RF_Transient);
 		}
 
-		USkeletalMesh* SkeletalMesh = SkeletalMeshes[Component.Id];
+		USkeletalMesh* SkeletalMesh = SkeletalMeshes[ComponentIndex];
 		check(SkeletalMesh);
 
 		// Set up the default information any mesh from this component will have (LODArrayInfos, RenderData, Mesh settings, etc). 
-		InitSkeletalMeshData(OperationData, SkeletalMesh, RefSkeletalMeshData, *CustomizableObject, Component.Id);
+		InitSkeletalMeshData(OperationData, SkeletalMesh, RefSkeletalMeshData, *CustomizableObject, ComponentIndex);
+
+		// We need the first valid mesh. get it from the component, considering that some LOSs may have been skipped.
+		mu::Ptr<const mu::Mesh> ComponentMesh;
+
+		int32 FirstValidLODIndex = Component.FirstLOD + OperationData->FirstLODAvailable;
+		check (OperationData->FirstLODAvailable<Component.LODCount)
+		ComponentMesh = OperationData->InstanceUpdateData.LODs[FirstValidLODIndex].Mesh;
 
-		if (Component.Mesh)
+		if (ComponentMesh)
 		{
 			// Construct a new skeleton, fix up ActiveBones and Bonemap arrays and recompute the RefInvMatrices
-			bSuccess = BuildSkeletonData(OperationData, *SkeletalMesh, RefSkeletalMeshData, *CustomizableObject, Component.Id);
+			bSuccess = BuildSkeletonData(OperationData, *SkeletalMesh, RefSkeletalMeshData, *CustomizableObject, ComponentIndex);
 
 			if (!bSuccess)
 			{
@@ -2129,11 +2122,11 @@ bool UCustomizableInstancePrivate::UpdateSkeletalMesh_PostBeginUpdate0(UCustomiz
 			}
 
 			// Build PhysicsAsset merging physics assets coming from SubMeshes of the newly generated Mesh
-			if (mu::Ptr<const mu::PhysicsBody> MutablePhysics = Component.Mesh->GetPhysicsBody())
+			if (mu::Ptr<const mu::PhysicsBody> MutablePhysics = ComponentMesh->GetPhysicsBody())
 			{
 				constexpr bool bDisallowCollisionBetweenAssets = true;
 				UPhysicsAsset* PhysicsAssetResult = GetOrBuildMainPhysicsAsset(
-					OperationData, RefSkeletalMeshData.PhysicsAsset, MutablePhysics.get(), bDisallowCollisionBetweenAssets, Component.Id);
+					OperationData, RefSkeletalMeshData.PhysicsAsset, MutablePhysics.get(), bDisallowCollisionBetweenAssets, ComponentIndex);
 
 				SkeletalMesh->SetPhysicsAsset(PhysicsAssetResult);
 
@@ -2149,10 +2142,10 @@ bool UCustomizableInstancePrivate::UpdateSkeletalMesh_PostBeginUpdate0(UCustomiz
 				}
 			}
 
-			const int32 NumAdditionalPhysicsNum = Component.Mesh->AdditionalPhysicsBodies.Num();
+			const int32 NumAdditionalPhysicsNum = ComponentMesh->AdditionalPhysicsBodies.Num();
 			for (int32 I = 0; I < NumAdditionalPhysicsNum; ++I)
 			{
-				mu::Ptr<const mu::PhysicsBody> AdditionalPhysiscBody = Component.Mesh->AdditionalPhysicsBodies[I];
+				mu::Ptr<const mu::PhysicsBody> AdditionalPhysiscBody = ComponentMesh->AdditionalPhysicsBodies[I];
 				
 				check(AdditionalPhysiscBody);
 				if (!AdditionalPhysiscBody->bBodiesModified)
@@ -2160,7 +2153,7 @@ bool UCustomizableInstancePrivate::UpdateSkeletalMesh_PostBeginUpdate0(UCustomiz
 					continue;
 				}
 
-				const int32 PhysicsBodyExternalId = Component.Mesh->AdditionalPhysicsBodies[I]->CustomId;
+				const int32 PhysicsBodyExternalId = ComponentMesh->AdditionalPhysicsBodies[I]->CustomId;
 				
 				const FAnimBpOverridePhysicsAssetsInfo& Info = ModelResources.AnimBpOverridePhysiscAssetsInfo[PhysicsBodyExternalId];
 
@@ -2183,11 +2176,11 @@ bool UCustomizableInstancePrivate::UpdateSkeletalMesh_PostBeginUpdate0(UCustomiz
 
 				Entry.PropertyIndex = Info.PropertyIndex;
 				Entry.PhysicsAsset = MakePhysicsAssetFromTemplateAndMutableBody(
-					OperationData,PhysicsAssetTemplate, AdditionalPhysiscBody.get(), Component.Id);
+					OperationData,PhysicsAssetTemplate, AdditionalPhysiscBody.get(), ComponentIndex);
 			}
 
 			// Add sockets from the SkeletalMesh of reference and from the MutableMesh
-			BuildMeshSockets(OperationData, SkeletalMesh, ModelResources, RefSkeletalMeshData, Component.Mesh);
+			BuildMeshSockets(OperationData, SkeletalMesh, ModelResources, RefSkeletalMeshData, ComponentMesh);
 		}
 		else
 		{
@@ -2294,10 +2287,6 @@ bool UCustomizableInstancePrivate::UpdateSkeletalMesh_PostBeginUpdate0(UCustomiz
 		for (int32 ComponentIndex = 0; bSuccess && ComponentIndex < NumComponents; ++ComponentIndex)
 		{
 			const FInstanceUpdateData::FComponent& Component = OperationData->InstanceUpdateData.Components[ComponentIndex];
-			if (Component.Mesh && Component.Mesh->IsReference())
-			{
-				continue;
-			}
 
 			if (OperationData->bUseMeshCache)
 			{
@@ -2320,6 +2309,19 @@ bool UCustomizableInstancePrivate::UpdateSkeletalMesh_PostBeginUpdate0(UCustomiz
 				continue;
 			}
 
+			// Skip if it is a referenced mesh
+			{
+				mu::Ptr<const mu::Mesh> FirstMesh;
+				if (Component.LODCount > 0)
+				{
+					FirstMesh = OperationData->InstanceUpdateData.LODs[Component.FirstLOD].Mesh;
+				}
+				if (FirstMesh && FirstMesh->IsReference())
+				{
+					continue;
+				}
+			}
+
 			// Mesh to copy data from if possible. 
 			USkeletalMesh* OldSkeletalMesh = OldSkeletalMeshes.IsValidIndex(ComponentIndex) ? OldSkeletalMeshes[ComponentIndex] : nullptr;
 
@@ -3537,12 +3539,12 @@ void UCustomizableInstancePrivate::InitSkeletalMeshData(const TSharedRef<FUpdate
 		SkeletalMesh->AllocateResourceForRendering();
 
 		FSkeletalMeshRenderData* RenderData = SkeletalMesh->GetResourceForRendering();
-		RenderData->NumInlinedLODs = OperationData->NumLODsAvailable - OperationData->FirstResidentLOD;
-		RenderData->NumNonOptionalLODs = OperationData->NumLODsAvailable - OperationData->FirstLODAvailable;
+		RenderData->NumInlinedLODs = OperationData->NumLODsAvailablePerComponent[ComponentIndex] - OperationData->FirstResidentLOD;
+		RenderData->NumNonOptionalLODs = OperationData->NumLODsAvailablePerComponent[ComponentIndex] - OperationData->FirstLODAvailable;
 		RenderData->CurrentFirstLODIdx = OperationData->FirstResidentLOD;
 		RenderData->LODBiasModifier = 0;
 
-		for (int32 LODIndex = 0; LODIndex < OperationData->NumLODsAvailable; ++LODIndex)
+		for (int32 LODIndex = 0; LODIndex < OperationData->NumLODsAvailablePerComponent[ComponentIndex]; ++LODIndex)
 		{
 			RenderData->LODRenderData.Add(new FSkeletalMeshLODRenderData());
 			
@@ -3713,7 +3715,7 @@ bool UCustomizableInstancePrivate::BuildSkeletonData(const TSharedRef<FUpdateCon
 }
 
 
-void UCustomizableInstancePrivate::BuildMeshSockets(const TSharedRef<FUpdateContextPrivate>& OperationData, USkeletalMesh* SkeletalMesh, const FModelResources& ModelResources, const FMutableRefSkeletalMeshData& RefSkeletalMeshData, mu::MeshPtrConst MutableMesh)
+void UCustomizableInstancePrivate::BuildMeshSockets(const TSharedRef<FUpdateContextPrivate>& OperationData, USkeletalMesh* SkeletalMesh, const FModelResources& ModelResources, const FMutableRefSkeletalMeshData& RefSkeletalMeshData, mu::Ptr<const mu::Mesh> MutableMesh)
 {
 	// Build mesh sockets.
 	MUTABLE_CPUPROFILER_SCOPE(UCustomizableInstancePrivate::BuildMeshSockets);
@@ -3815,23 +3817,23 @@ void UCustomizableInstancePrivate::BuildOrCopyElementData(const TSharedRef<FUpda
 {
 	MUTABLE_CPUPROFILER_SCOPE(UCustomizableInstancePrivate::BuildOrCopyElementData);
 
-	for (int32 LODIndex = FirstLODAvailable; LODIndex < OperationData->NumLODsAvailable; ++LODIndex)
+	const FInstanceUpdateData::FComponent& Component = OperationData->InstanceUpdateData.Components[ComponentIndex];
+
+	for (int32 LODIndex = FirstLODAvailable; LODIndex < Component.LODCount; ++LODIndex)
 	{
-		const FInstanceUpdateData::FLOD& LOD = OperationData->InstanceUpdateData.LODs[LODIndex];
+		const FInstanceUpdateData::FLOD& LOD = OperationData->InstanceUpdateData.LODs[Component.FirstLOD+LODIndex];
 
-		if (!LOD.ComponentCount)
+		if (!LOD.SurfaceCount)
 		{
 			continue;
 		}
 
-		const FInstanceUpdateData::FComponent& Component = OperationData->InstanceUpdateData.Components[LOD.FirstComponent + ComponentIndex];
-
-		if (!Component.bGenerated)
+		if (!LOD.bGenerated)
 		{
 			continue;
 		}
 
-		for (int32 SurfaceIndex = 0; SurfaceIndex < Component.SurfaceCount; ++SurfaceIndex)
+		for (int32 SurfaceIndex = 0; SurfaceIndex < LOD.SurfaceCount; ++SurfaceIndex)
 		{
 			new(SkeletalMesh->GetResourceForRendering()->LODRenderData[LODIndex].RenderSections) FSkelMeshRenderSection();
 		}
@@ -3954,22 +3956,18 @@ void UCustomizableInstancePrivate::BuildOrCopyClothingData(const TSharedRef<FUpd
 	{
 		MUTABLE_CPUPROFILER_SCOPE(DiscoverSectionsWithCloth);
 
-		for (int32 LODIndex = OperationData->FirstLODAvailable; LODIndex < OperationData->NumLODsAvailable; ++LODIndex)
-		{
-			const FInstanceUpdateData::FLOD& LOD = OperationData->InstanceUpdateData.LODs[LODIndex];
-			if (ComponentIndex >= LOD.ComponentCount)
-			{
-				continue;
-			}
+		FInstanceUpdateData::FComponent& Component = OperationData->InstanceUpdateData.Components[ComponentIndex];
 
-			FInstanceUpdateData::FComponent& Component = OperationData->InstanceUpdateData.Components[LOD.FirstComponent + ComponentIndex];
+		for (int32 LODIndex = OperationData->FirstLODAvailable; LODIndex < Component.LODCount; ++LODIndex)
+		{
+			const FInstanceUpdateData::FLOD& LOD = OperationData->InstanceUpdateData.LODs[Component.FirstLOD+LODIndex];
 
-			if (!Component.bGenerated)
+			if (!LOD.bGenerated)
 			{
 				continue;
 			}
 
-			if (mu::MeshPtrConst MutableMesh = Component.Mesh)
+			if (mu::MeshPtrConst MutableMesh = LOD.Mesh)
 			{
 				const mu::FMeshBufferSet& MeshSet = MutableMesh->GetVertexBuffers();
 				const mu::FMeshBufferSet& IndicesSet = MutableMesh->GetIndexBuffers();
@@ -4097,7 +4095,7 @@ void UCustomizableInstancePrivate::BuildOrCopyClothingData(const TSharedRef<FUpd
 			// Only initilialize once, multiple sections with cloth could point to the same cloth asset.
 			if (!DstAssetData.LodMap.Num())
 			{
-				DstAssetData.LodMap.Init(INDEX_NONE, OperationData->NumLODsAvailable);
+				DstAssetData.LodMap.Init(INDEX_NONE, OperationData->NumLODsAvailablePerComponent[ComponentIndex]);
 
 				DstAssetData.UsedBoneNames = SrcAssetData.UsedBoneNames;
 				DstAssetData.UsedBoneIndices = SrcAssetData.UsedBoneIndices;
@@ -4685,15 +4683,12 @@ void UCustomizableInstancePrivate::BuildOrCopyClothingData(const TSharedRef<FUpd
 	{
 		int32 NumSectionsWithClothProcessed = 0;
 
-		for (int32 LODIndex = OperationData->GetMinLOD(); LODIndex < OperationData->NumLODsAvailable; ++LODIndex)
+		FInstanceUpdateData::FComponent& Component = OperationData->InstanceUpdateData.Components[ComponentIndex];
+		for (int32 LODIndex = OperationData->GetMinLOD(); LODIndex < Component.LODCount; ++LODIndex)
 		{
-			const FInstanceUpdateData::FLOD& LOD = OperationData->InstanceUpdateData.LODs[LODIndex];
-			if (ComponentIndex >= LOD.ComponentCount)
-			{
-				continue;
-			}
+			const FInstanceUpdateData::FLOD& LOD = OperationData->InstanceUpdateData.LODs[Component.FirstLOD+LODIndex];
 
-			if (mu::MeshPtrConst MutableMesh = OperationData->InstanceUpdateData.Components[LOD.FirstComponent+ComponentIndex].Mesh)
+			if (mu::MeshPtrConst MutableMesh = LOD.Mesh)
 			{
 				TArray<FMeshToMeshVertData>& MappingData = LodMappingData[LODIndex];
 				TArray<FClothBufferIndexMapping>& ClothingIndexMapping = LodClothingIndexMapping[LODIndex];
@@ -4738,7 +4733,7 @@ void UCustomizableInstancePrivate::BuildOrCopyClothingData(const TSharedRef<FUpd
 	{
 		MUTABLE_CPUPROFILER_SCOPE(InitClothRenderData)
 		// Based on FSkeletalMeshLODModel::GetClothMappingData().
-		for (int32 LODIndex = OperationData->GetMinLOD(); LODIndex < OperationData->NumLODsAvailable; ++LODIndex)
+		for (int32 LODIndex = OperationData->GetMinLOD(); LODIndex < OperationData->NumLODsAvailablePerComponent[ComponentIndex]; ++LODIndex)
 		{
 			FSkeletalMeshLODRenderData& LODModel = RenderResource->LODRenderData[LODIndex];
 	
@@ -4889,40 +4884,30 @@ bool UCustomizableInstancePrivate::BuildOrCopyRenderData(const TSharedRef<FUpdat
 	const FModelResources& ModelResources = CustomizableObject->GetPrivate()->GetModelResources();
 
 	const int32 FirstGeneratedLOD = FMath::Max((int32)OperationData->GetRequestedLODs()[ComponentIndex], OperationData->GetMinLOD());
-	for (int32 LODIndex = FirstGeneratedLOD; LODIndex < OperationData->NumLODsAvailable; ++LODIndex)
-	{
-		MUTABLE_CPUPROFILER_SCOPE(BuildRenderData);
-
-		const FInstanceUpdateData::FLOD& LOD = OperationData->InstanceUpdateData.LODs[LODIndex];
 
-		if (ComponentIndex >= LOD.ComponentCount)
-		{
-			// Interrupt the generation if the LOD is empty and it should have been generated.
-			UE_LOG(LogMutable, Warning, TEXT("Building instance: generated mesh [%s] has LOD [%d] with no component.")
-				, *SkeletalMesh->GetName()
-				, LODIndex);
+	FInstanceUpdateData::FComponent& Component = OperationData->InstanceUpdateData.Components[ComponentIndex];
 
-			// End with failure
-			return false;
-		}
+	for (int32 LODIndex = FirstGeneratedLOD; LODIndex < Component.LODCount; ++LODIndex)
+	{
+		MUTABLE_CPUPROFILER_SCOPE(BuildRenderData);
 
-		FInstanceUpdateData::FComponent& Component = OperationData->InstanceUpdateData.Components[LOD.FirstComponent + ComponentIndex];
+		const FInstanceUpdateData::FLOD& LOD = OperationData->InstanceUpdateData.LODs[Component.FirstLOD+LODIndex];
 
 		// There could be components without a mesh in LODs
-		if (!Component.bGenerated || !Component.Mesh || Component.SurfaceCount == 0)
+		if (!LOD.bGenerated || !LOD.Mesh || LOD.SurfaceCount == 0)
 		{
-			UE_LOG(LogMutable, Warning, TEXT("Building instance: generated mesh [%s] has LOD [%d] of component [%d] with no mesh.")
+			UE_LOG(LogMutable, Warning, TEXT("Building instance: generated mesh [%s] has LOD [%d] of component index [%d] with no mesh.")
 				, *SkeletalMesh->GetName()
 				, LODIndex
-				, Component.Id);
+				, ComponentIndex);
 
 			// End with failure
 			return false;
 		}
 
-		TRACE_CPUPROFILER_EVENT_SCOPE_TEXT(*FString::Printf(TEXT("BuildRenderData: Component %d, LOD %d"), ComponentIndex, LODIndex));
+		TRACE_CPUPROFILER_EVENT_SCOPE_TEXT(*FString::Printf(TEXT("BuildRenderData: Component index %d, LOD %d"), ComponentIndex, LODIndex));
 
-		SetLastMeshId(Component.Id, LODIndex, Component.MeshID);
+		SetLastMeshId(ComponentIndex, LODIndex, LOD.MeshID);
 
 		FSkeletalMeshLODRenderData& LODResource = RenderData->LODRenderData[LODIndex];
 
@@ -4931,11 +4916,11 @@ bool UCustomizableInstancePrivate::BuildOrCopyRenderData(const TSharedRef<FUpdat
 		// Set active and required bones
 		{
 			const TArray<mu::FBoneName>& ActiveBones = OperationData->InstanceUpdateData.ActiveBones;
-			LODResource.ActiveBoneIndices.Reserve(Component.ActiveBoneCount);
+			LODResource.ActiveBoneIndices.Reserve(LOD.ActiveBoneCount);
 
-			for (uint32 Index = 0; Index < Component.ActiveBoneCount; ++Index)
+			for (uint32 Index = 0; Index < LOD.ActiveBoneCount; ++Index)
 			{
-				const uint16 ActiveBoneIndex = BoneInfoMap[ActiveBones[Component.FirstActiveBone + Index]].Value;
+				const uint16 ActiveBoneIndex = BoneInfoMap[ActiveBones[LOD.FirstActiveBone + Index]].Value;
 				LODResource.ActiveBoneIndices.Add(ActiveBoneIndex);
 			}
 
@@ -4946,21 +4931,21 @@ bool UCustomizableInstancePrivate::BuildOrCopyRenderData(const TSharedRef<FUpdat
 		// Set RenderSections
 		UnrealConversionUtils::SetupRenderSections(
 			LODResource,
-			Component.Mesh,
+			LOD.Mesh,
 			OperationData->InstanceUpdateData.BoneMaps,
 			BoneInfoMap,
-			Component.FirstBoneMap);
+			LOD.FirstBoneMap);
 
 		if (LODResource.bStreamedDataInlined) // Non-streamable LOD
 		{
 			// Copy Vertices
 			UnrealConversionUtils::CopyMutableVertexBuffers(
 				LODResource,
-				Component.Mesh,
+				LOD.Mesh,
 				SkeletalMesh->GetLODInfo(LODIndex)->bAllowCPUAccess);
 
 			// Copy indices.
-			if (!UnrealConversionUtils::CopyMutableIndexBuffers(LODResource, Component.Mesh))
+			if (!UnrealConversionUtils::CopyMutableIndexBuffers(LODResource, LOD.Mesh))
 			{
 				// End with failure
 				return false;
@@ -4969,7 +4954,7 @@ bool UCustomizableInstancePrivate::BuildOrCopyRenderData(const TSharedRef<FUpdat
 			// Copy SkinWeightProfiles
 			if (!ModelResources.SkinWeightProfilesInfo.IsEmpty())
 			{
-				const mu::FMeshBufferSet& MutableMeshVertexBuffers = Component.Mesh->GetVertexBuffers();
+				const mu::FMeshBufferSet& MutableMeshVertexBuffers = LOD.Mesh->GetVertexBuffers();
 
 				bool bHasSkinWeightProfiles = false;
 			
@@ -5024,11 +5009,11 @@ bool UCustomizableInstancePrivate::BuildOrCopyRenderData(const TSharedRef<FUpdat
 			// Init VertexBuffers for streaming
 			UnrealConversionUtils::InitVertexBuffersWithDummyData(
 				LODResource,
-				Component.Mesh,
+				LOD.Mesh,
 				SkeletalMesh->GetLODInfo(LODIndex)->bAllowCPUAccess);
 
 			// Init IndexBuffers for streaming
-			UnrealConversionUtils::InitIndexBuffersWithDummyData(LODResource, Component.Mesh);
+			UnrealConversionUtils::InitIndexBuffersWithDummyData(LODResource, LOD.Mesh);
 
 			// SkinWeightProfilesInfo Not supported yet
 		}
@@ -5178,20 +5163,21 @@ UE::Tasks::FTask UCustomizableInstancePrivate::LoadAdditionalAssetsAndData(
 	// Load assets coming from SubMeshes of the newly generated Mesh
 	if (OperationData->InstanceUpdateData.LODs.Num())
 	{
-		for (int32 LODIndex = OperationData->FirstLODAvailable; LODIndex < OperationData->NumLODsAvailable; ++LODIndex)
+		for (int32 ComponentIndex = 0; ComponentIndex < OperationData->InstanceUpdateData.Components.Num(); ++ComponentIndex)
 		{
-			const FInstanceUpdateData::FLOD& LOD = OperationData->InstanceUpdateData.LODs[LODIndex];
-			for (int32 ComponentIndex = 0; ComponentIndex < LOD.ComponentCount; ++ComponentIndex)
+			const FInstanceUpdateData::FComponent& Component = Components[ComponentIndex];
+
+			for (int32 LODIndex = OperationData->FirstLODAvailable; LODIndex < Component.LODCount; ++LODIndex)
 			{
-				const FInstanceUpdateData::FComponent& Component = Components[LOD.FirstComponent + ComponentIndex];
-				mu::MeshPtrConst MutableMesh = Component.Mesh;
+				const FInstanceUpdateData::FLOD& LOD = OperationData->InstanceUpdateData.LODs[Component.FirstLOD + LODIndex];
 
+				mu::Ptr<const mu::Mesh> MutableMesh = LOD.Mesh;
 				if (!MutableMesh)
 				{
 					continue;
 				}
 
-				FCustomizableInstanceComponentData* ComponentData = GetComponentData(Component.Id);
+				FCustomizableInstanceComponentData* ComponentData = GetComponentData(ComponentIndex);
 
 				const TArray<uint64>& StreamedResources = MutableMesh->GetStreamedResources();
 
@@ -6018,12 +6004,14 @@ void UCustomizableInstancePrivate::BuildMaterials(const TSharedRef<FUpdateContex
 		MUTABLE_CPUPROFILER_SCOPE(BuildMaterials_LODLoop);
 
 		const int32 FirstGeneratedLOD = FMath::Max((int32)OperationData->GetRequestedLODs()[ComponentIndex], OperationData->GetMinLOD());
-		for (int32 LODIndex = FirstGeneratedLOD; LODIndex < OperationData->NumLODsAvailable; LODIndex++)
+
+		const FInstanceUpdateData::FComponent& Component = OperationData->InstanceUpdateData.Components[ComponentIndex];
+
+		for (int32 LODIndex = FirstGeneratedLOD; LODIndex < Component.LODCount; LODIndex++)
 		{
-			const FInstanceUpdateData::FLOD& LOD = OperationData->InstanceUpdateData.LODs[LODIndex];
-			const FInstanceUpdateData::FComponent& Component = OperationData->InstanceUpdateData.Components[LOD.FirstComponent + ComponentIndex];
+			const FInstanceUpdateData::FLOD& LOD = OperationData->InstanceUpdateData.LODs[Component.FirstLOD+LODIndex];
 
-			if (!Component.bGenerated)
+			if (!LOD.bGenerated)
 			{
 				continue;
 			}
@@ -6033,11 +6021,16 @@ void UCustomizableInstancePrivate::BuildMaterials(const TSharedRef<FUpdateContex
 				SkeletalMesh->GetLODInfo(LODIndex)->LODMaterialMap.Reset();
 			}
 
-			const FMutableRefSkeletalMeshData& RefSkeletalMeshData = ModelResources.ReferenceSkeletalMeshesData[Component.Id];
+			// Pass-through components will not have a reference mesh.
+			const FMutableRefSkeletalMeshData* RefSkeletalMeshData = nullptr;
+			if (ModelResources.ReferenceSkeletalMeshesData.IsValidIndex(ComponentIndex))
+			{
+				RefSkeletalMeshData = &ModelResources.ReferenceSkeletalMeshesData[ComponentIndex];
+			}
 
-			for (int32 SurfaceIndex = 0; SurfaceIndex < Component.SurfaceCount; ++SurfaceIndex)
+			for (int32 SurfaceIndex = 0; SurfaceIndex < LOD.SurfaceCount; ++SurfaceIndex)
 			{
-				const FInstanceUpdateData::FSurface& Surface = OperationData->InstanceUpdateData.Surfaces[Component.FirstSurface + SurfaceIndex];
+				const FInstanceUpdateData::FSurface& Surface = OperationData->InstanceUpdateData.Surfaces[LOD.FirstSurface + SurfaceIndex];
 
 				// Reuse MaterialSlot from the previous LOD.
 				if (const int32 MaterialIndex = SurfaceIdToMaterialIndex.Find(Surface.SurfaceId); MaterialIndex != INDEX_NONE)
@@ -6074,7 +6067,10 @@ void UCustomizableInstancePrivate::BuildMaterials(const TSharedRef<FUpdateContex
 					MaterialSlot.MaterialSlotName = ModelResources.MaterialSlotNames[Surface.MaterialIndex];
 				}
 
-				SetMeshUVChannelDensity(MaterialSlot.UVChannelData, RefSkeletalMeshData.Settings.DefaultUVChannelDensity);
+				if (RefSkeletalMeshData)
+				{
+					SetMeshUVChannelDensity(MaterialSlot.UVChannelData, RefSkeletalMeshData->Settings.DefaultUVChannelDensity);
+				}
 
 				if (!bUseOverrideMaterialsOnly)
 				{
diff --git a/Engine/Plugins/Experimental/Mutable/Source/CustomizableObject/Private/MuCO/CustomizableObjectMeshUpdate.cpp b/Engine/Plugins/Experimental/Mutable/Source/CustomizableObject/Private/MuCO/CustomizableObjectMeshUpdate.cpp
--- a/Engine/Plugins/Experimental/Mutable/Source/CustomizableObject/Private/MuCO/CustomizableObjectMeshUpdate.cpp
+++ b/Engine/Plugins/Experimental/Mutable/Source/CustomizableObject/Private/MuCO/CustomizableObjectMeshUpdate.cpp
@@ -177,7 +177,7 @@ namespace impl
 		const uint32 LODMask = 0xFFFFFFFF;
 
 		// Main instance generation step
-		const mu::Instance* Instance = System->BeginUpdate(InstanceID, OperationData->Parameters, OperationData->State, LODMask);
+		mu::Ptr<const mu::Instance> Instance = System->BeginUpdate(InstanceID, OperationData->Parameters, OperationData->State, LODMask);
 		check(Instance);
 
 		Task_Mutable_UpdateMesh_Loop(OperationData, Task, InstanceID, OperationData->PendingFirstLODIdx);
diff --git a/Engine/Plugins/Experimental/Mutable/Source/CustomizableObject/Private/MuCO/CustomizableObjectMipDataProvider.cpp b/Engine/Plugins/Experimental/Mutable/Source/CustomizableObject/Private/MuCO/CustomizableObjectMipDataProvider.cpp
--- a/Engine/Plugins/Experimental/Mutable/Source/CustomizableObject/Private/MuCO/CustomizableObjectMipDataProvider.cpp
+++ b/Engine/Plugins/Experimental/Mutable/Source/CustomizableObject/Private/MuCO/CustomizableObjectMipDataProvider.cpp
@@ -195,7 +195,7 @@ namespace Impl
 		mu::Instance::ID InstanceID = System->NewInstance(Model);
 		UE_LOG(LogMutable, Verbose, TEXT("Creating Mutable instance with id [%d] for a single UpdateImage"), InstanceID)
 
-		const mu::Instance* Instance = nullptr;
+		mu::Ptr<const mu::Instance> Instance;
 
 		// Main instance generation step
 		{
@@ -208,12 +208,12 @@ namespace Impl
 
 		const FMutableImageReference& ImageRef = OperationData->RequestedImage;
 
-		int32 SurfaceIndex = Instance->FindSurfaceById(ImageRef.LOD, ImageRef.Component, ImageRef.SurfaceId);
+		int32 SurfaceIndex = Instance->FindSurfaceById(ImageRef.Component, ImageRef.LOD, ImageRef.SurfaceId);
 		check(SurfaceIndex >= 0);
 
 		// This ID may be different than the ID obtained the first time the image was generated, because the mutable
 		// runtime cannot remember all the resources it has built, and only remembers a fixed amount.
-		mu::FResourceID MipImageID = Instance->GetImageId(ImageRef.LOD, ImageRef.Component, SurfaceIndex, ImageRef.Image);
+		mu::FResourceID MipImageID = Instance->GetImageId(ImageRef.Component, ImageRef.LOD, SurfaceIndex, ImageRef.Image);
 
 		UE::Tasks::TTask<mu::Ptr<const mu::Image>> GetImageTask = 
 				System->GetImage(InstanceID, MipImageID, ImageRef.BaseMip + OperationData->MipsToSkip, ImageRef.LOD);
diff --git a/Engine/Plugins/Experimental/Mutable/Source/CustomizableObject/Private/MuCO/CustomizableObjectSkeletalMesh.cpp b/Engine/Plugins/Experimental/Mutable/Source/CustomizableObject/Private/MuCO/CustomizableObjectSkeletalMesh.cpp
--- a/Engine/Plugins/Experimental/Mutable/Source/CustomizableObject/Private/MuCO/CustomizableObjectSkeletalMesh.cpp
+++ b/Engine/Plugins/Experimental/Mutable/Source/CustomizableObject/Private/MuCO/CustomizableObjectSkeletalMesh.cpp
@@ -38,17 +38,16 @@ UCustomizableObjectSkeletalMesh* UCustomizableObjectSkeletalMesh::CreateSkeletal
 	
 	OutSkeletalMesh->MeshIDs.Init(MAX_uint64, MAX_MESH_LOD_COUNT);
 	
-	for (int32 LODIndex = OperationData->FirstLODAvailable; LODIndex < OperationData->NumLODsAvailable; ++LODIndex)
+	for (int32 ComponentIndex = 0; ComponentIndex < OperationData->InstanceUpdateData.Components.Num(); ++ComponentIndex)
 	{
-		const FInstanceUpdateData::FLOD& LOD = OperationData->InstanceUpdateData.LODs[LODIndex];
+		const FInstanceUpdateData::FComponent& Component = OperationData->InstanceUpdateData.Components[ComponentIndex];
 
-		for (int32 ComponentIndex = 0; ComponentIndex < LOD.ComponentCount; ++ComponentIndex)
+		if (Component.Id == InComponentIndex)
 		{
-			const FInstanceUpdateData::FComponent& Component = OperationData->InstanceUpdateData.Components[LOD.FirstComponent + ComponentIndex];
-
-			if (Component.Id == InComponentIndex)
+			for (int32 LODIndex = OperationData->FirstLODAvailable; LODIndex < Component.LODCount; ++LODIndex)
 			{
-				OutSkeletalMesh->MeshIDs[LODIndex] = Component.MeshID;
+				const FInstanceUpdateData::FLOD& LOD = OperationData->InstanceUpdateData.LODs[Component.FirstLOD + LODIndex];
+				OutSkeletalMesh->MeshIDs[LODIndex] = LOD.MeshID;
 			}
 		}
 	}
diff --git a/Engine/Plugins/Experimental/Mutable/Source/CustomizableObject/Private/MuCO/CustomizableObjectSystem.cpp b/Engine/Plugins/Experimental/Mutable/Source/CustomizableObject/Private/MuCO/CustomizableObjectSystem.cpp
--- a/Engine/Plugins/Experimental/Mutable/Source/CustomizableObject/Private/MuCO/CustomizableObjectSystem.cpp
+++ b/Engine/Plugins/Experimental/Mutable/Source/CustomizableObject/Private/MuCO/CustomizableObjectSystem.cpp
@@ -1539,7 +1539,7 @@ namespace impl
 
 	struct FGetMeshData
 	{
-		int32 ComponentIndex;
+		int32 LODIndex;
 		mu::FResourceID MeshID;
 	};
 
@@ -1653,15 +1653,32 @@ namespace impl
 	
 	void FixLODs(const TSharedRef<FUpdateContextPrivate>& Operation)
 	{
-		Operation->NumLODsAvailable = Operation->MutableInstance->GetLODCount();
+		// This is set at context creation because for now the component count is a CO constant that includes all components
+		// regardles of what is generated.
+		//Operation->NumComponents = Operation->MutableInstance->GetComponentCount();
+		if (!Operation->NumComponents)
+		{
+			return;
+		}
+
+		// For now this should always match, since components are not optional.
+		check(Operation->NumComponents== Operation->MutableInstance->GetComponentCount());
+
+		Operation->NumLODsAvailablePerComponent.SetNumZeroed(Operation->NumComponents);
+
+		int32 NumComponentsInGeneratedInstance = Operation->MutableInstance->GetComponentCount();
+		for (int32 ComponentIndex = 0; ComponentIndex < NumComponentsInGeneratedInstance; ++ComponentIndex)
+		{
+			Operation->NumLODsAvailablePerComponent[ComponentIndex] = Operation->MutableInstance->GetLODCount(ComponentIndex);
+		}
 		
 		int32 CurrentMinLOD = Operation->bStreamMeshLODs ? 0 : Operation->GetMinLOD();
-		CurrentMinLOD = FMath::Clamp(CurrentMinLOD, Operation->FirstLODAvailable, Operation->NumLODsAvailable - 1);
+		CurrentMinLOD = FMath::Clamp(CurrentMinLOD, Operation->FirstLODAvailable, Operation->NumLODsAvailablePerComponent[0] - 1); 
 		Operation->SetMinLOD(CurrentMinLOD);
 
 		if (Operation->bStreamMeshLODs)
 		{
-			Operation->FirstResidentLOD = FMath::Clamp(Operation->FirstResidentLOD, Operation->FirstLODAvailable, Operation->NumLODsAvailable - 1);
+			Operation->FirstResidentLOD = FMath::Clamp(Operation->FirstResidentLOD, Operation->FirstLODAvailable, Operation->NumLODsAvailablePerComponent[0] - 1); 
 		}
 		else 
 		{
@@ -1681,7 +1698,7 @@ namespace impl
 			else
 			{
 				// Clamp value to the valid LOD range.
-				RequestedLODs[ComponentIndex] = FMath::Min(RequestedLODs[ComponentIndex], (uint16)(Operation->NumLODsAvailable - 1));
+				RequestedLODs[ComponentIndex] = FMath::Min(RequestedLODs[ComponentIndex], (uint16)(Operation->NumLODsAvailablePerComponent[0] - 1));
 			}
 		}
 
@@ -1733,43 +1750,50 @@ namespace impl
 
 		OperationData->InstanceUpdateData.Skeletons.SetNum(OperationData->NumComponents);
 
-		for (FInstanceUpdateData::FComponent& Component : OperationData->InstanceUpdateData.Components)
+		for (int32 ComponentIndex = 0; ComponentIndex<OperationData->InstanceUpdateData.Components.Num(); ++ComponentIndex)
 		{
-			mu::MeshPtrConst Mesh = Component.Mesh;
-			if (!Mesh || Mesh->IsReference())
-			{
-				continue;
-			}
-
-			FInstanceUpdateData::FSkeletonData& SkeletonData = OperationData->InstanceUpdateData.Skeletons[Component.Id];
+			const FInstanceUpdateData::FComponent& Component = OperationData->InstanceUpdateData.Components[ComponentIndex];
 
-			// Add SkeletonIds 
-			const int32 SkeletonIDsCount = Mesh->GetSkeletonIDsCount();
-			for (int32 SkeletonIndex = 0; SkeletonIndex < SkeletonIDsCount; ++SkeletonIndex)
+			for (int32 LODIndex = 0; LODIndex < Component.LODCount; ++LODIndex)
 			{
-				SkeletonData.SkeletonIds.AddUnique(Mesh->GetSkeletonID(SkeletonIndex));
-			}
+				FInstanceUpdateData::FLOD& LOD = OperationData->InstanceUpdateData.LODs[Component.FirstLOD+LODIndex];
+
+				mu::Ptr<const mu::Mesh> Mesh = LOD.Mesh;
+				if (!Mesh || Mesh->IsReference())
+				{
+					continue;
+				}
 
-			// Append BoneMap to the array of BoneMaps
-			const TArray<mu::FBoneName>& BoneMap = Mesh->GetBoneMap();
-			Component.FirstBoneMap = OperationData->InstanceUpdateData.BoneMaps.Num();
-			Component.BoneMapCount = BoneMap.Num();
-			OperationData->InstanceUpdateData.BoneMaps.Append(BoneMap);
+				FInstanceUpdateData::FSkeletonData& SkeletonData = OperationData->InstanceUpdateData.Skeletons[ComponentIndex];
 
-			// Add active bone indices and poses
-			Component.FirstActiveBone = OperationData->InstanceUpdateData.ActiveBones.Num();
-			Component.ActiveBoneCount = Mesh->GetBonePoseCount();
-			for (uint32 BoneIndex = 0; BoneIndex < Component.ActiveBoneCount; ++BoneIndex)
-			{
-				const mu::FBoneName& BoneId = Mesh->GetBonePoseId(BoneIndex);
+				// Add SkeletonIds 
+				const int32 SkeletonIDsCount = Mesh->GetSkeletonIDsCount();
+				for (int32 SkeletonIndex = 0; SkeletonIndex < SkeletonIDsCount; ++SkeletonIndex)
+				{
+					SkeletonData.SkeletonIds.AddUnique(Mesh->GetSkeletonID(SkeletonIndex));
+				}
 
-				OperationData->InstanceUpdateData.ActiveBones.Add(BoneId);
+				// Append BoneMap to the array of BoneMaps
+				const TArray<mu::FBoneName>& BoneMap = Mesh->GetBoneMap();
+				LOD.FirstBoneMap = OperationData->InstanceUpdateData.BoneMaps.Num();
+				LOD.BoneMapCount = BoneMap.Num();
+				OperationData->InstanceUpdateData.BoneMaps.Append(BoneMap);
 
-				if (!SkeletonData.BonePose.FindByKey(BoneId))
+				// Add active bone indices and poses
+				LOD.FirstActiveBone = OperationData->InstanceUpdateData.ActiveBones.Num();
+				LOD.ActiveBoneCount = Mesh->GetBonePoseCount();
+				for (uint32 BoneIndex = 0; BoneIndex < LOD.ActiveBoneCount; ++BoneIndex)
 				{
-					FTransform3f Transform;
-					Mesh->GetBonePoseTransform(BoneIndex, Transform);
-					SkeletonData.BonePose.Add({ BoneId,Transform.Inverse().ToMatrixWithScale() });
+					const mu::FBoneName& BoneId = Mesh->GetBonePoseId(BoneIndex);
+
+					OperationData->InstanceUpdateData.ActiveBones.Add(BoneId);
+
+					if (!SkeletonData.BonePose.FindByKey(BoneId))
+					{
+						FTransform3f Transform;
+						Mesh->GetBonePoseTransform(BoneIndex, Transform);
+						SkeletonData.BonePose.Add({ BoneId,Transform.Inverse().ToMatrixWithScale() });
+					}
 				}
 			}
 		}
@@ -1849,21 +1873,21 @@ namespace impl
 			MorphsData.SetNum(MorphTargetNames.Num());
 			for (int32 MorphIndex = 0; MorphIndex < NumMorphs; ++MorphIndex)
 			{
-				MorphsData[MorphIndex].SetNum(OperationData->NumLODsAvailable);
+				MorphsData[MorphIndex].SetNum(OperationData->NumLODsAvailablePerComponent[ComponentIndex]);
 			}
 
 			TArray<int32> SectionMorphTargetVerticesCount;
 			SectionMorphTargetVerticesCount.SetNumZeroed(ComponentMorphTargetsData.RealTimeMorphTargetNames.Num());
 
 			const int32 FirstGeneratedLOD = FMath::Max((int32)OperationData->GetRequestedLODs()[ComponentIndex], OperationData->GetMinLOD());
-			for (int32 LODIndex = FirstGeneratedLOD; LODIndex < OperationData->NumLODsAvailable; ++LODIndex)
+			for (int32 LODIndex = FirstGeneratedLOD; LODIndex < OperationData->NumLODsAvailablePerComponent[ComponentIndex]; ++LODIndex)
 			{
-				const FInstanceUpdateData::FLOD& LOD = UpdateData.LODs[LODIndex];
-				const FInstanceUpdateData::FComponent& Component = UpdateData.Components[LOD.FirstComponent + ComponentIndex];
-				check(Component.bGenerated);
-				check(Component.Mesh);
+				const FInstanceUpdateData::FComponent& Component = UpdateData.Components[ComponentIndex];
+				const FInstanceUpdateData::FLOD& LOD = UpdateData.LODs[Component.FirstLOD+LODIndex];
+				check(LOD.bGenerated);
+				check(LOD.Mesh);
 
-				const mu::FMeshBufferSet& MeshSet = Component.Mesh->GetVertexBuffers();
+				const mu::FMeshBufferSet& MeshSet = LOD.Mesh->GetVertexBuffers();
 
 				int32 VertexMorphsInfoIndexAndCountBufferIndex, VertexMorphsInfoIndexAndCountBufferChannel;
 				MeshSet.FindChannel(mu::MBS_OTHER, 0, &VertexMorphsInfoIndexAndCountBufferIndex, &VertexMorphsInfoIndexAndCountBufferChannel);
@@ -1882,7 +1906,7 @@ namespace impl
 				const uint32* const VertexMorphsResourceIdBuffer = reinterpret_cast<const uint32*>(MeshSet.GetBufferData(VertexMorphsResourceIdBufferIndex));
 				TArrayView<const uint32> VertexMorphsResourceIdView(VertexMorphsResourceIdBuffer, MeshSet.GetElementCount());
 
-				const int32 SurfaceCount = Component.Mesh->GetSurfaceCount();
+				const int32 SurfaceCount = LOD.Mesh->GetSurfaceCount();
 				for (int32 Section = 0; Section < SurfaceCount; ++Section)
 				{
 					// Reset SectionMorphTargets.
@@ -1892,7 +1916,7 @@ namespace impl
 					}
 
 					int32 FirstVertex, VerticesCount, FirstIndex, IndiciesCount;
-					Component.Mesh->GetSurface(Section, &FirstVertex, &VerticesCount, &FirstIndex, &IndiciesCount, nullptr, nullptr, nullptr);
+					LOD.Mesh->GetSurface(Section, &FirstVertex, &VerticesCount, &FirstIndex, &IndiciesCount, nullptr, nullptr, nullptr);
 
 					for (int32 VertexIdx = FirstVertex; VertexIdx < FirstVertex + VerticesCount;)
 					{
@@ -2124,54 +2148,48 @@ namespace impl
 		const UCustomizableObject* CustomizableObject = Instance->GetCustomizableObject();
 		const FModelResources& ModelResources = CustomizableObject->GetPrivate()->GetModelResources();
 
-		const mu::Instance* MutableInstance = OperationData->MutableInstance;
+		const mu::Instance* MutableInstance = OperationData->MutableInstance.get();
 
 		TArray<int32> SurfacesSharedId;
 
 		const TSharedRef<TArray<FGetImageData>> GetImagesData = MakeShared<TArray<FGetImageData>>();
 		
-		for (int32 MutableLODIndex = 0; MutableLODIndex < OperationData->NumLODsAvailable; ++MutableLODIndex)
+		for (int32 ComponentIndex = 0; ComponentIndex < OperationData->NumComponents; ++ComponentIndex)
 		{
-			// Skip LODs outside the range we want to generate
-			if (MutableLODIndex < OperationData->GetMinLOD())
+			FInstanceUpdateData::FComponent& Component = OperationData->InstanceUpdateData.Components[ComponentIndex];
+			for (int32 MutableLODIndex = 0; MutableLODIndex < Component.LODCount; ++MutableLODIndex)
 			{
-				continue;
-			}
+				// Skip LODs outside the range we want to generate
+				if (MutableLODIndex < OperationData->GetMinLOD())
+				{
+					continue;
+				}
 
-			const FInstanceUpdateData::FLOD& LOD = OperationData->InstanceUpdateData.LODs[MutableLODIndex];
+				FInstanceUpdateData::FLOD& LOD = OperationData->InstanceUpdateData.LODs[Component.FirstLOD + MutableLODIndex];
 
-			for (int32 ComponentIndex = 0; ComponentIndex < LOD.ComponentCount; ++ComponentIndex)
-			{
-				FInstanceUpdateData::FComponent& Component = OperationData->InstanceUpdateData.Components[LOD.FirstComponent + ComponentIndex];
-				Component.FirstSurface = OperationData->InstanceUpdateData.Surfaces.Num();
-				Component.SurfaceCount = 0;
+				LOD.FirstSurface = OperationData->InstanceUpdateData.Surfaces.Num();
+				LOD.SurfaceCount = 0;
 
-				if (!Component.Mesh)
+				if (!LOD.Mesh)
 				{
 					continue;
 				}
 				
-				// Materials and images
-				const int32 SurfaceCount = Component.Mesh->GetSurfaceCount();
-				for (int32 MeshSurfaceIndex = 0; MeshSurfaceIndex < SurfaceCount; ++MeshSurfaceIndex)
-				{
-					const uint32 SurfaceId = Component.Mesh->GetSurfaceId(MeshSurfaceIndex);
-					const int32 InstanceSurfaceIndex = MutableInstance->FindSurfaceById(MutableLODIndex, ComponentIndex, SurfaceId);
-					check(Component.Mesh->GetVertexCount() > 0 || InstanceSurfaceIndex >= 0);
 
-					if (InstanceSurfaceIndex >= 0)
+				// This lambda does all the work to fill up the surface data
+				auto AddSurface = [&LOD, &SurfacesSharedId, &ModelResources, &GetImagesData, OperationData, MutableInstance, CustomizableObject, ComponentIndex, MutableLODIndex](uint32 SurfaceId, int32 InstanceSurfaceIndex)
 					{
 						int32 BaseSurfaceIndex = InstanceSurfaceIndex;
 						int32 BaseLODIndex = MutableLODIndex;
 
 						OperationData->InstanceUpdateData.Surfaces.Push({});
 						FInstanceUpdateData::FSurface& Surface = OperationData->InstanceUpdateData.Surfaces.Last();
-						++Component.SurfaceCount;
+						++LOD.SurfaceCount;
 
 						// Now Surface.MaterialIndex is decoded from a parameter at the end of this if()
 						Surface.SurfaceId = SurfaceId;
 
-						const int32 SharedSurfaceId = MutableInstance->GetSharedSurfaceId(MutableLODIndex, ComponentIndex, InstanceSurfaceIndex);
+						const int32 SharedSurfaceId = MutableInstance->GetSharedSurfaceId(ComponentIndex, MutableLODIndex, InstanceSurfaceIndex);
 						const int32 SharedSurfaceIndex = SurfacesSharedId.Find(SharedSurfaceId);
 
 						SurfacesSharedId.Add(SharedSurfaceId);
@@ -2181,44 +2199,44 @@ namespace impl
 							if (SharedSurfaceIndex >= 0)
 							{
 								Surface = OperationData->InstanceUpdateData.Surfaces[SharedSurfaceIndex];
-								continue;
+								return;
 							}
 
 							// Find the first LOD where this surface can be found
 							MutableInstance->FindBaseSurfaceBySharedId(ComponentIndex, SharedSurfaceId, BaseSurfaceIndex, BaseLODIndex);
 
-							Surface.SurfaceId = MutableInstance->GetSurfaceId(BaseLODIndex, ComponentIndex, BaseSurfaceIndex);
+							Surface.SurfaceId = MutableInstance->GetSurfaceId(ComponentIndex, BaseLODIndex, BaseSurfaceIndex);
 						}
 
 						// Vectors
 						Surface.FirstVector = OperationData->InstanceUpdateData.Vectors.Num();
-						Surface.VectorCount = MutableInstance->GetVectorCount(BaseLODIndex, ComponentIndex, BaseSurfaceIndex);
+						Surface.VectorCount = MutableInstance->GetVectorCount(ComponentIndex, BaseLODIndex, BaseSurfaceIndex);
 						for (int32 VectorIndex = 0; VectorIndex < Surface.VectorCount; ++VectorIndex)
 						{
 							MUTABLE_CPUPROFILER_SCOPE(GetVector);
 							OperationData->InstanceUpdateData.Vectors.Push({});
 							FInstanceUpdateData::FVector& Vector = OperationData->InstanceUpdateData.Vectors.Last();
-							Vector.Name = MutableInstance->GetVectorName(BaseLODIndex, ComponentIndex, BaseSurfaceIndex, VectorIndex);
-							Vector.Vector = MutableInstance->GetVector(BaseLODIndex, ComponentIndex, BaseSurfaceIndex, VectorIndex);
+							Vector.Name = MutableInstance->GetVectorName(ComponentIndex, BaseLODIndex, BaseSurfaceIndex, VectorIndex);
+							Vector.Vector = MutableInstance->GetVector(ComponentIndex, BaseLODIndex, BaseSurfaceIndex, VectorIndex);
 						}
 
 						// Scalars
 						Surface.FirstScalar = OperationData->InstanceUpdateData.Scalars.Num();
-						Surface.ScalarCount = MutableInstance->GetScalarCount(BaseLODIndex, ComponentIndex, BaseSurfaceIndex);
+						Surface.ScalarCount = MutableInstance->GetScalarCount(ComponentIndex, BaseLODIndex, BaseSurfaceIndex);
 						for (int32 ScalarIndex = 0; ScalarIndex < Surface.ScalarCount; ++ScalarIndex)
 						{
 							MUTABLE_CPUPROFILER_SCOPE(GetScalar)
 
-							const FName ScalarName = MutableInstance->GetScalarName(BaseLODIndex, ComponentIndex, BaseSurfaceIndex, ScalarIndex);
-							const float ScalarValue = MutableInstance->GetScalar(BaseLODIndex, ComponentIndex, BaseSurfaceIndex, ScalarIndex);
-							
+								const FName ScalarName = MutableInstance->GetScalarName(ComponentIndex, BaseLODIndex, BaseSurfaceIndex, ScalarIndex);
+							const float ScalarValue = MutableInstance->GetScalar(ComponentIndex, BaseLODIndex, BaseSurfaceIndex, ScalarIndex);
+
 							FString EncodingMaterialIdString = "__MutableMaterialId";
-							
+
 							// Decoding Material Switch from Mutable parameter name
 							if (ScalarName.ToString().Equals(EncodingMaterialIdString))
 							{
 								Surface.MaterialIndex = static_cast<uint32>(ScalarValue);
-							
+
 								// This parameter is not needed in the final material instance
 								Surface.ScalarCount -= 1;
 							}
@@ -2230,15 +2248,15 @@ namespace impl
 
 						// Images
 						Surface.FirstImage = OperationData->InstanceUpdateData.Images.Num();
-						Surface.ImageCount = MutableInstance->GetImageCount(BaseLODIndex, ComponentIndex, BaseSurfaceIndex);
+						Surface.ImageCount = MutableInstance->GetImageCount(ComponentIndex, BaseLODIndex, BaseSurfaceIndex);
 						for (int32 ImageIndex = 0; ImageIndex < Surface.ImageCount; ++ImageIndex)
 						{
 							MUTABLE_CPUPROFILER_SCOPE(GetImageId);
 
 							const int32 UpdateDataImageIndex = OperationData->InstanceUpdateData.Images.AddDefaulted();
 							FInstanceUpdateData::FImage& Image = OperationData->InstanceUpdateData.Images.Last();
-							Image.Name = MutableInstance->GetImageName(BaseLODIndex, ComponentIndex, BaseSurfaceIndex, ImageIndex);
-							Image.ImageID = MutableInstance->GetImageId(BaseLODIndex, ComponentIndex, BaseSurfaceIndex, ImageIndex);
+							Image.Name = MutableInstance->GetImageName(ComponentIndex, BaseLODIndex, BaseSurfaceIndex, ImageIndex);
+							Image.ImageID = MutableInstance->GetImageId(ComponentIndex, BaseLODIndex, BaseSurfaceIndex, ImageIndex);
 							Image.FullImageSizeX = 0;
 							Image.FullImageSizeY = 0;
 							Image.BaseLOD = BaseLODIndex;
@@ -2267,8 +2285,41 @@ namespace impl
 								UE_LOG(LogMutable, Error, TEXT("CustomizableObject derived data out of sync with asset for [%s]. Try recompiling it."), *CustomizableObject->GetName());
 							}
 						}
+
+					};
+
+				// Materials and images
+
+				// If the mesh is a reference mesh, it won't have the surface information in the mutable mesh. We need to get it from the instance
+				// and all defined surfaces will be present. 
+				if (LOD.Mesh->IsReference())
+				{
+					const int32 SurfaceCount = MutableInstance->GetSurfaceCount(ComponentIndex, MutableLODIndex);
+					for (int32 SurfaceIndex = 0; SurfaceIndex < SurfaceCount; ++SurfaceIndex)
+					{
+						uint32 SurfaceId = MutableInstance->GetSurfaceId(ComponentIndex, MutableLODIndex, SurfaceIndex);
+						AddSurface(SurfaceId, SurfaceIndex);
 					}
 				}
+
+				// If the mesh is a not a reference mesh, we have to add only the materials of the surfaces that appear in the actual final mesh. 
+				else
+				{
+					const int32 SurfaceCount = LOD.Mesh->GetSurfaceCount();
+					for (int32 MeshSurfaceIndex = 0; MeshSurfaceIndex < SurfaceCount; ++MeshSurfaceIndex)
+					{
+						const uint32 SurfaceId = LOD.Mesh->GetSurfaceId(MeshSurfaceIndex);
+
+						const int32 InstanceSurfaceIndex = MutableInstance->FindSurfaceById(ComponentIndex, MutableLODIndex, SurfaceId);
+						check(LOD.Mesh->GetVertexCount() > 0 || InstanceSurfaceIndex >= 0);
+
+						if (InstanceSurfaceIndex >= 0)
+						{
+							AddSurface(SurfaceId, InstanceSurfaceIndex);
+						}
+					}
+				}
+
 			}	
 		}
 
@@ -2287,18 +2338,18 @@ namespace impl
 	{
 		MUTABLE_CPUPROFILER_SCOPE(Task_MutableGetMeshes_GetMesh_Post)
 
-		const int32 ComponentIndex = (*GetMeshesData)[GetMeshIndex].ComponentIndex;
-		FInstanceUpdateData::FComponent& Component = OperationData->InstanceUpdateData.Components[ComponentIndex];
+		const int32 LODIndex = (*GetMeshesData)[GetMeshIndex].LODIndex;
+		FInstanceUpdateData::FLOD& LOD = OperationData->InstanceUpdateData.LODs[LODIndex];
 
-		Component.Mesh = GetMeshTask.GetResult();
+		LOD.Mesh = GetMeshTask.GetResult();
 
-		if (Component.Mesh->IsReference())
+		if (LOD.Mesh->IsReference())
 		{
 			const UCustomizableObjectInstance* Instance = OperationData->Instance.Get();
 			const UCustomizableObject* CustomizableObject = Instance->GetCustomizableObject();
 			const FModelResources& ModelResources = CustomizableObject->GetPrivate()->GetModelResources();
 
-			uint32 ReferenceID = Component.Mesh->GetReferencedMesh();
+			uint32 ReferenceID = LOD.Mesh->GetReferencedMesh();
 
 			if (ModelResources.PassThroughMeshes.IsValidIndex(ReferenceID))
 			{
@@ -2378,8 +2429,7 @@ namespace impl
 			}
 
 			// Main instance generation step
-			// LOD mask, set to all ones to build  all LODs
-			const mu::Instance* Instance = OperationData->MutableInstance; // TODO GMTFuture remove
+			const mu::Instance* Instance = OperationData->MutableInstance.get(); // TODO GMTFuture remove
 			if (!Instance)
 			{
 				UE_LOG(LogMutable, Warning, TEXT("An Instace update has failed."));
@@ -2391,39 +2441,40 @@ namespace impl
 
 			const TSharedRef<TArray<FGetMeshData>> GetMeshesData = MakeShared<TArray<FGetMeshData>>();
 		
-			OperationData->InstanceUpdateData.LODs.SetNum(OperationData->NumLODsAvailable);
-			for (int32 MutableLODIndex = 0; MutableLODIndex < OperationData->NumLODsAvailable; ++MutableLODIndex)
+			OperationData->InstanceUpdateData.Components.SetNum(OperationData->NumComponents);
+			for (int32 ComponentIndex = 0; ComponentIndex < Instance->GetComponentCount(); ++ComponentIndex)
 			{
-				// Skip LODs outside the range we want to generate
-				if (MutableLODIndex < OperationData->GetMinLOD())
+				FInstanceUpdateData::FComponent& Component = OperationData->InstanceUpdateData.Components[ComponentIndex];
+				Component.FirstLOD = OperationData->InstanceUpdateData.LODs.Num();
+				Component.LODCount = OperationData->NumLODsAvailablePerComponent[ComponentIndex];
+				Component.Id = Instance->GetComponentId(ComponentIndex);
+
+				for (int32 MutableLODIndex = 0; MutableLODIndex < OperationData->NumLODsAvailablePerComponent[ComponentIndex]; ++MutableLODIndex)
 				{
-					continue;
-				}
+					// If the LOD is not generated we still add an empty one to keep indexes aligned.
+					const int32 UpdateDataLODIndex = OperationData->InstanceUpdateData.LODs.AddDefaulted();
 
-				FInstanceUpdateData::FLOD& LOD = OperationData->InstanceUpdateData.LODs[MutableLODIndex];
-				LOD.FirstComponent = OperationData->InstanceUpdateData.Components.Num();
-				LOD.ComponentCount = Instance->GetComponentCount(MutableLODIndex);
+					// Skip LODs outside the range we want to generate
+					if (MutableLODIndex < OperationData->GetMinLOD())
+					{
+						continue;
+					}
 
-				for (int32 ComponentIndex = 0; ComponentIndex < LOD.ComponentCount; ++ComponentIndex)
-				{
-					const int32 UpdateDataComponentIndex = OperationData->InstanceUpdateData.Components.AddDefaulted();
-					FInstanceUpdateData::FComponent& Component = OperationData->InstanceUpdateData.Components.Last();
-					Component.Id = Instance->GetComponentId(MutableLODIndex, ComponentIndex);
+					FInstanceUpdateData::FLOD& LOD = OperationData->InstanceUpdateData.LODs.Last();
 
-					const bool bGenerateLOD = RequestedLODs.IsValidIndex(Component.Id) ? RequestedLODs[Component.Id] <= MutableLODIndex : true;
+					const bool bGenerateLOD = RequestedLODs.IsValidIndex(ComponentIndex) ? RequestedLODs[ComponentIndex] <= MutableLODIndex : true;
 
 					// Mesh
-					if (Instance->GetMeshCount(MutableLODIndex, ComponentIndex) > 0)
 					{
 						MUTABLE_CPUPROFILER_SCOPE(GetMesh);
 
-						Component.MeshID = Instance->GetMeshId(MutableLODIndex, ComponentIndex, 0);
+						LOD.MeshID = Instance->GetMeshId(ComponentIndex, MutableLODIndex);
 
 						if (bGenerateLOD)
 						{
-							Component.bGenerated = true;
+							LOD.bGenerated = true;
 
-							GetMeshesData->Add({ UpdateDataComponentIndex, Component.MeshID});
+							GetMeshesData->Add({ UpdateDataLODIndex, LOD.MeshID});
 						}
 					}
 				}
@@ -3145,12 +3196,12 @@ namespace impl
 			MeshId.Init(MAX_uint64, MAX_MESH_LOD_COUNT);
 
 			
-			for (int32 LODIndex = Operation->GetMinLOD(); LODIndex < Operation->NumLODsAvailable; ++LODIndex)
+			for (int32 LODIndex = Operation->GetMinLOD(); LODIndex < Operation->NumLODsAvailablePerComponent[ComponentIndex]; ++LODIndex)
 			{
 				const bool bGenerateLOD = RequestedLODs.IsValidIndex(ComponentIndex) ? RequestedLODs[ComponentIndex] <= LODIndex : true;
 				if (bGenerateLOD)
 				{
-					MeshId[LODIndex] = Operation->MutableInstance->GetMeshId(LODIndex, ComponentIndex, 0);
+					MeshId[LODIndex] = Operation->MutableInstance->GetMeshId(ComponentIndex, LODIndex);
 				}
 			}
 		}
diff --git a/Engine/Plugins/Experimental/Mutable/Source/CustomizableObjectEditor/Private/MuCOE/CustomizableObjectCompiler.cpp b/Engine/Plugins/Experimental/Mutable/Source/CustomizableObjectEditor/Private/MuCOE/CustomizableObjectCompiler.cpp
--- a/Engine/Plugins/Experimental/Mutable/Source/CustomizableObjectEditor/Private/MuCOE/CustomizableObjectCompiler.cpp
+++ b/Engine/Plugins/Experimental/Mutable/Source/CustomizableObjectEditor/Private/MuCOE/CustomizableObjectCompiler.cpp
@@ -727,13 +727,14 @@ mu::NodeObjectPtr FCustomizableObjectCompiler::GenerateMutableRoot(
 	return MutableRoot;
 }
 
+
 void FCustomizableObjectCompiler::LaunchMutableCompile()
 {
 	AddCompileNotification(LOCTEXT("CustomizableObjectCompileInProgress", "Compiling"));
 
 	// Even for async build, we spawn a thread, so that we can set a large stack. 
 	// Thread names need to be unique, apparently.
-	static int ThreadCount = 0;
+	static int32 ThreadCount = 0;
 	FString ThreadName = FString::Printf(TEXT("MutableCompile-%03d"), ++ThreadCount);
 	CompileThread = MakeShareable(FRunnableThread::Create(CompileTask.Get(), *ThreadName, 16 * 1024 * 1024, TPri_Normal));
 }
@@ -1102,7 +1103,7 @@ void FCustomizableObjectCompiler::CompileInternal(bool bAsync)
 		CurrentObject->GetPrivate()->CustomizableObjectPathMap = GenerationContext.CustomizableObjectPathMap;
 #endif
 
-		ModelResources.NumComponents = GenerationContext.NumMeshComponentsInRoot;
+		ModelResources.NumComponents = GenerationContext.NumMeshComponentsInRoot + GenerationContext.NumExplicitMeshComponents;
 		ModelResources.NumLODs = GenerationContext.NumLODsInRoot;
 		ModelResources.NumLODsToStream = GenerationContext.bEnableLODStreaming ? GenerationContext.NumMaxLODsToStream : 0;
 		ModelResources.FirstLODAvailable = GenerationContext.FirstLODAvailable;
diff --git a/Engine/Plugins/Experimental/Mutable/Source/CustomizableObjectEditor/Private/MuCOE/EdGraphSchema_CustomizableObject.cpp b/Engine/Plugins/Experimental/Mutable/Source/CustomizableObjectEditor/Private/MuCOE/EdGraphSchema_CustomizableObject.cpp
--- a/Engine/Plugins/Experimental/Mutable/Source/CustomizableObjectEditor/Private/MuCOE/EdGraphSchema_CustomizableObject.cpp
+++ b/Engine/Plugins/Experimental/Mutable/Source/CustomizableObjectEditor/Private/MuCOE/EdGraphSchema_CustomizableObject.cpp
@@ -47,10 +47,10 @@
 #include "MuCOE/Nodes/CustomizableObjectNodeMeshReshape.h"
 #include "MuCOE/Nodes/CustomizableObjectNodeMeshSwitch.h"
 #include "MuCOE/Nodes/CustomizableObjectNodeMeshVariation.h"
-#include "MuCOE/Nodes/CustomizableObjectNodePassThroughMesh.h"
 #include "MuCOE/Nodes/CustomizableObjectNodeMorphMaterial.h"
 #include "MuCOE/Nodes/CustomizableObjectNodeObjectChild.h"
 #include "MuCOE/Nodes/CustomizableObjectNodeObjectGroup.h"
+#include "MuCOE/Nodes/CustomizableObjectNodeComponentMesh.h"
 #include "MuCOE/Nodes/CustomizableObjectNodeProjectorConstant.h"
 #include "MuCOE/Nodes/CustomizableObjectNodeRemoveMesh.h"
 #include "MuCOE/Nodes/CustomizableObjectNodeRemoveMeshBlocks.h"
@@ -220,6 +220,7 @@ UEdGraphNode* FCustomizableObjectSchemaAction_Paste::PerformAction(class UEdGrap
 //////////////////////////////////////////////////////////////////////////
 // DO NOT change the values because it will break the external pin nodes!
 const FName UEdGraphSchema_CustomizableObject::PC_Object("object");
+const FName UEdGraphSchema_CustomizableObject::PC_Component("component");
 const FName UEdGraphSchema_CustomizableObject::PC_Material("material");
 const FName UEdGraphSchema_CustomizableObject::PC_Mesh("mesh");
 const FName UEdGraphSchema_CustomizableObject::PC_PassThroughMesh("passThroughMesh");
@@ -394,12 +395,20 @@ void UEdGraphSchema_CustomizableObject::GetGraphContextActions(FGraphContextMenu
 		AddNewNodeCategoryActionsFiltered(ObjectTemplateNodes, ContextMenuBuilder, TEXT("Object"), GeneralGrouping, Filter);
 	}
 
+	{
+		UCustomizableObjectNode* ComponentTemplateNodes[]
+		{
+			ContextMenuBuilder.CreateTemplateNode<UCustomizableObjectNodeComponentMesh>(),
+		};
+
+		AddNewNodeCategoryActionsFiltered(ComponentTemplateNodes, ContextMenuBuilder, TEXT("Component"), GeneralGrouping, Filter);
+	}
+
 	{
 		UCustomizableObjectNode* MeshTemplateNodes[]
 		{
 			ContextMenuBuilder.CreateTemplateNode<UCustomizableObjectNodeSkeletalMesh>(),
 			ContextMenuBuilder.CreateTemplateNode<UCustomizableObjectNodeStaticMesh>(),
-			ContextMenuBuilder.CreateTemplateNode<UCustomizableObjectNodePassThroughMesh>(),
 			ContextMenuBuilder.CreateTemplateNode<UCustomizableObjectNodeLayoutBlocks>(),
 			ContextMenuBuilder.CreateTemplateNode<UCustomizableObjectNodeMeshMorph>(),
 			ContextMenuBuilder.CreateTemplateNode<UCustomizableObjectNodeMeshSwitch>(),
@@ -691,6 +700,10 @@ FLinearColor UEdGraphSchema_CustomizableObject::GetPinTypeColor(const FName& Typ
 	{
 		return FLinearColor(0.000000f, 0.400000f, 0.910000f, 1.000000f); // Light blue
 	}
+	else if (TypeString == PC_Component)
+	{
+		return FLinearColor(FColorList::MediumOrchid); 
+	}
 	else if (TypeString == PC_Stack)
 	{
 		return FLinearColor(1.000000f, 0.000000f, 0.800000f, 1.000000f); // Pink
diff --git a/Engine/Plugins/Experimental/Mutable/Source/CustomizableObjectEditor/Private/MuCOE/GenerateMutableSource/GenerateMutableSource.cpp b/Engine/Plugins/Experimental/Mutable/Source/CustomizableObjectEditor/Private/MuCOE/GenerateMutableSource/GenerateMutableSource.cpp
--- a/Engine/Plugins/Experimental/Mutable/Source/CustomizableObjectEditor/Private/MuCOE/GenerateMutableSource/GenerateMutableSource.cpp
+++ b/Engine/Plugins/Experimental/Mutable/Source/CustomizableObjectEditor/Private/MuCOE/GenerateMutableSource/GenerateMutableSource.cpp
@@ -810,7 +810,7 @@ mu::NodeMeshApplyPosePtr CreateNodeMeshApplyPose(FMutableGraphGenerationContext&
 
 
 // Convert a CustomizableObject Source Graph into a mutable source graph  
-mu::NodeObjectPtr GenerateMutableSource(const UEdGraphPin * Pin, FMutableGraphGenerationContext & GenerationContext, bool bPartialCompilation)
+mu::Ptr<mu::NodeObject> GenerateMutableSource(const UEdGraphPin * Pin, FMutableGraphGenerationContext & GenerationContext, bool bPartialCompilation)
 {
 	check(Pin)
 	RETURN_ON_CYCLE(*Pin, GenerationContext)
@@ -851,6 +851,10 @@ mu::NodeObjectPtr GenerateMutableSource(const UEdGraphPin * Pin, FMutableGraphGe
 		}
 		ObjectNode->SetUid(FinalGuid.ToString());
 
+
+		// Process the pins defining components implicitly (original method)
+		//-------------------------------------------------------------------
+
 		// LOD
 		const int32 NumLODs = TypedNodeObj->GetNumLODPins();
 
@@ -1008,89 +1012,80 @@ mu::NodeObjectPtr GenerateMutableSource(const UEdGraphPin * Pin, FMutableGraphGe
 
 		const int32 NumLODsInRoot = GenerationContext.NumLODsInRoot;
 		const int32 NumMeshComponentsInRoot = GenerationContext.NumMeshComponentsInRoot;
-		while (GenerationContext.ComponentNewNode.Last().ComponentsPerLOD.Num() < NumLODsInRoot)
+		if (GenerationContext.ComponentNewNode.Last().Components.Num() < NumMeshComponentsInRoot)
 		{
-			GenerationContext.ComponentNewNode.Last().ComponentsPerLOD.Emplace_GetRef()
-				.Init(nullptr, NumMeshComponentsInRoot);
+			GenerationContext.ComponentNewNode.Last().Components.SetNum(NumMeshComponentsInRoot);
 		}
 
-		int32 FirstLOD = -1;
-		ObjectNode->SetLODCount(NumLODsInRoot);
-		for (int32 CurrentLOD = 0; CurrentLOD < NumLODsInRoot; ++CurrentLOD)
-		{
-			GenerationContext.CurrentLOD = CurrentLOD;
-
-			mu::NodeLODPtr LODNode = new mu::NodeLOD();
-			ObjectNode->SetLOD(CurrentLOD, LODNode);
+		// Number of components defined in the traditional way: in the object node, and referenced in graph materials
+		int32 NumImplicitComponents = GenerationContext.NumMeshComponentsInRoot;
 
-			LODNode->SetMessageContext(Node);
-			LODNode->SetComponentCount(GenerationContext.NumMeshComponentsInRoot);
-		
-			TArray<mu::NodeComponentPtr> ComponentNodes;
-			ComponentNodes.Init(nullptr, NumMeshComponentsInRoot);
-
-			// Generate a component node for every component in root regardless of if will be populated. 
-			// Not sure this is what we we want, but is what it was done before.
-			// TODO: Review if this is the behaviour we want.
-			for (int32 MeshComponentIndex = 0; MeshComponentIndex < NumMeshComponentsInRoot; ++MeshComponentIndex)
-			{
-				mu::NodeComponentPtr& ComponentNode = ComponentNodes[MeshComponentIndex];
+		// Add implicit components
+		for (int32 CurrentComponent = 0; CurrentComponent < NumImplicitComponents; ++CurrentComponent)
+		{
+			GenerationContext.CurrentMeshComponent = CurrentComponent;
 
-				ComponentNode = Invoke([&GenerationContext, CurrentLOD, MeshComponentIndex]() -> mu::NodeComponentPtr
+			mu::Ptr<mu::NodeComponent> ComponentNode = Invoke([&GenerationContext, CurrentComponent]() -> mu::Ptr<mu::NodeComponent>
 				{
-					mu::NodeComponentNewPtr& ParentComponent =
-							GenerationContext.ComponentNewNode.Last().ComponentsPerLOD[CurrentLOD][MeshComponentIndex];
+					mu::Ptr<mu::NodeComponentNew>& ParentComponent =
+						GenerationContext.ComponentNewNode.Last().Components[CurrentComponent];
 
 					if (!ParentComponent)
 					{
 						ParentComponent = new mu::NodeComponentNew();
-						ParentComponent->SetId(MeshComponentIndex);
-				
+						ParentComponent->Id = CurrentComponent;
+
 						return ParentComponent;
 					}
-					
-					mu::NodeComponentEditPtr EditComponent = new mu::NodeComponentEdit();
+
+					mu::Ptr<mu::NodeComponentEdit> EditComponent = new mu::NodeComponentEdit();
 					EditComponent->SetParent(ParentComponent.get());
 
 					return EditComponent;
 				});
 
-				ComponentNode->SetMessageContext(Node);
-				LODNode->SetComponent(MeshComponentIndex, ComponentNode);
-			}
+			ComponentNode->SetMessageContext(Node);
+			ObjectNode->Components.Add(ComponentNode);
 
-			const bool bUseAutomaticLods = 
-					GenerationContext.CurrentAutoLODStrategy == ECustomizableObjectAutomaticLODStrategy::AutomaticFromMesh;
-			FirstLOD = (CurrentLOD < NumLODs) && (FirstLOD == INDEX_NONE || !bUseAutomaticLods) ? CurrentLOD : FirstLOD;
 
-			if (FirstLOD < 0)
+			int32 FirstLOD = -1;
+			for (int32 CurrentLOD = 0; CurrentLOD < NumLODsInRoot; ++CurrentLOD)
 			{
-				continue;
-			}
+				GenerationContext.CurrentLOD = CurrentLOD;
 
-			if (GenerationContext.CurrentLOD < GenerationContext.FirstLODAvailable)
-			{
-				continue;
-			} 
-			
-			// Generate all relevant LODs for this object up until the current LODIndex.
-			for (int32 LODIndex = FirstLOD; LODIndex <= CurrentLOD; ++LODIndex)
-			{
-				const UEdGraphPin* LODPin = TypedNodeObj->LODPin(LODIndex);
-				if (!LODPin)
+				mu::Ptr<mu::NodeLOD> LODNode = new mu::NodeLOD();
+				ComponentNode->LODs.Add(LODNode);
+				LODNode->SetMessageContext(Node);
+
+		
+				const bool bUseAutomaticLods = GenerationContext.CurrentAutoLODStrategy == ECustomizableObjectAutomaticLODStrategy::AutomaticFromMesh;
+				FirstLOD = (CurrentLOD < NumLODs) && (FirstLOD == INDEX_NONE || !bUseAutomaticLods) ? CurrentLOD : FirstLOD;
+
+				if (FirstLOD < 0)
 				{
 					continue;
 				}
 
-				GenerationContext.FromLOD = LODIndex;
+				if (GenerationContext.CurrentLOD < GenerationContext.FirstLODAvailable)
+				{
+					continue;
+				} 
+			
+				// Generate all relevant LODs for this object up until the current LODIndex.
+				for (int32 LODIndex = FirstLOD; LODIndex <= CurrentLOD; ++LODIndex)
+				{
+
+					const UEdGraphPin* LODPin = TypedNodeObj->LODPin(LODIndex);
+					if (!LODPin)
+					{
+						continue;
+					}
 
-				TArray<UEdGraphPin*> ConnectedLODPins = FollowInputPinArray(*LODPin);
+					GenerationContext.FromLOD = LODIndex;
 
-				// Proccess non modifier material nodes.
-				for (int32 MeshComponentIndex = 0; MeshComponentIndex < NumMeshComponentsInRoot; ++MeshComponentIndex)
-				{
-					GenerationContext.CurrentMeshComponent = MeshComponentIndex;
+					TArray<UEdGraphPin*> ConnectedLODPins = FollowInputPinArray(*LODPin);
 
+					// Proccess non modifier nodes.
 					for (UEdGraphPin* const ChildNodePin : ConnectedLODPins)
 					{
 						// Modifiers are shared for all components and are processed per LOD and not component.
@@ -1105,37 +1100,55 @@ mu::NodeObjectPtr GenerateMutableSource(const UEdGraphPin * Pin, FMutableGraphGe
 						}
 
 						mu::NodeSurfacePtr SurfaceNode = GenerateMutableSourceSurface(ChildNodePin, GenerationContext);
+						LODNode->Surfaces.Add(SurfaceNode);
+					}
+
+					// Process modifiers.
+					for (UEdGraphPin* const ChildNodePin : ConnectedLODPins)
+					{
+						if (!Cast<UCustomizableObjectNodeModifierBase>(ChildNodePin->GetOwningNode()))
+						{
+							continue;
+						}
+
+						// Set it to -1 to indicate we don't care about component id.
+						GenerationContext.CurrentMeshComponent = -1;
 
-						mu::NodeComponentPtr& ComponentNode = ComponentNodes[MeshComponentIndex];
+						mu::NodeModifierPtr ModifierNode = GenerateMutableSourceModifier(ChildNodePin, GenerationContext);
+						LODNode->Modifiers.Add(ModifierNode);
 
-						const int32 SurfaceCount = ComponentNode->GetSurfaceCount();
-						ComponentNode->SetSurfaceCount(SurfaceCount + 1);
-						ComponentNode->SetSurface(SurfaceCount, SurfaceNode.get());
-						ComponentNode->SetMessageContext(Node);
+						GenerationContext.CurrentMeshComponent = CurrentComponent;
 					}
 				}
+			}
+		}
 
-				// Process modfiers. Those are shared between different components in a lod.	
-				for (UEdGraphPin* const ChildNodePin : ConnectedLODPins)
-				{
-					// Set it to -1 to indicate we don't care about component id.
-					GenerationContext.CurrentMeshComponent = -1;
 
-					if (!Cast<UCustomizableObjectNodeModifierBase>(ChildNodePin->GetOwningNode()))
-					{
-						continue;
-					}
+		// Process the pins defining explicit components
+		//-------------------------------------------------------------------
+		TArray<mu::Ptr<mu::NodeComponent>> ExplicitComponents;
 
-					mu::NodeModifierPtr ModifierNode = GenerateMutableSourceModifier(ChildNodePin, GenerationContext);
+		const UEdGraphPin* ComponentsPin = TypedNodeObj->ComponentsPin();
+		TArray<UEdGraphPin*> ConnectedComponentPins = FollowInputPinArray(*ComponentsPin);
+		for (const UEdGraphPin* ComponentNodePin : ConnectedComponentPins)
+		{		
+			// \TODO: Needed?
+			//GenerationContext.ComponentInfos.Add(RefSkeletalMesh);
+			//GenerationContext.ReferencedSkeletons.Add(RefSkeleton);
 
-					const int32 ModifierCount = LODNode->GetModifierCount();
-					LODNode->SetModifierCount(ModifierCount + 1);
-					LODNode->SetModifier(ModifierCount, ModifierNode.get());
-				}
-			}
+			++GenerationContext.NumExplicitMeshComponents;
+
+			GenerationContext.CurrentMeshComponent = ObjectNode->Components.Num();
+
+			mu::Ptr<mu::NodeComponent> ComponentNode = GenerateMutableSourceComponent(ComponentNodePin, GenerationContext);
+			ObjectNode->Components.Add(ComponentNode);
+
+			GenerationContext.CurrentMeshComponent = -1;
 		}
 
+
 		// Generate inputs to Object node pins added by extensions
+		//-------------------------------------------------------------------
 		for (const FRegisteredObjectNodeInputPin& ExtensionInputPin : ICustomizableObjectModule::Get().GetAdditionalObjectNodePins())
 		{
 			const UEdGraphPin* GraphPin = TypedNodeObj->FindPin(ExtensionInputPin.GlobalPinName, EGPD_Input);
@@ -1162,7 +1175,7 @@ mu::NodeObjectPtr GenerateMutableSource(const UEdGraphPin * Pin, FMutableGraphGe
 				
 				if (const ICustomizableObjectExtensionNode* ExtensionNode = Cast<ICustomizableObjectExtensionNode>(ConnectedNode))
 				{
-					if (mu::NodeExtensionDataPtr GeneratedNode = ExtensionNode->GenerateMutableNode(GenerationContext.ExtensionDataCompilerInterface))
+					if (mu::Ptr<mu::NodeExtensionData> GeneratedNode = ExtensionNode->GenerateMutableNode(GenerationContext.ExtensionDataCompilerInterface))
 					{
 						ObjectNode->AddExtensionDataNode(GeneratedNode, TCHAR_TO_ANSI(*ExtensionInputPin.GlobalPinName.ToString()));
 					}
@@ -1171,12 +1184,13 @@ mu::NodeObjectPtr GenerateMutableSource(const UEdGraphPin * Pin, FMutableGraphGe
 		}
 
 		// Children
+		//-------------------------------------------------------------------
 		TArray<UEdGraphPin*> ConnectedChildrenPins = FollowInputPinArray(*TypedNodeObj->ChildrenPin());
-		ObjectNode->SetChildCount(ConnectedChildrenPins.Num());
+		ObjectNode->Children.Reserve(ConnectedChildrenPins.Num());
 		for (int32 ChildIndex = 0; ChildIndex < ConnectedChildrenPins.Num(); ++ChildIndex)
 		{
-			mu::NodeObjectPtr ChildNode = GenerateMutableSource(ConnectedChildrenPins[ChildIndex], GenerationContext, bPartialCompilation);
-			ObjectNode->SetChild(ChildIndex, ChildNode.get());
+			mu::Ptr<mu::NodeObject> ChildNode = GenerateMutableSource(ConnectedChildrenPins[ChildIndex], GenerationContext, bPartialCompilation);
+			ObjectNode->Children.Add(ChildNode);
 		}
 
 		// Remove from the current processing hierarchy
@@ -1561,7 +1575,7 @@ bool AffectsCurrentComponent(const UEdGraphPin* Pin, FMutableGraphGenerationCont
 	}
 	else if (const UCustomizableObjectNodeModifierBase* TypedNodeModifier = Cast<UCustomizableObjectNodeModifierBase>(Node))
 	{
-		// Because of the current implementation, modifiers affect all componeents at lod level. If there is only one component it is ok, but otherwise rise an error.
+		// Because of the current implementation, modifiers affect all components at lod level. If there is only one component it is ok, but otherwise rise an error.
 		if (GenerationContext.NumMeshComponentsInRoot == 1)
 		{
 			ComponentIndex = 0;
