[mutable] Redesign the object node hierarchy to better match unreal component and LOD concepts.

diff --git a/Engine/Plugins/Experimental/Mutable/Source/CustomizableObjectEditor/Private/MuCOE/GenerateMutableSource/GenerateMutableSource.h b/Engine/Plugins/Experimental/Mutable/Source/CustomizableObjectEditor/Private/MuCOE/GenerateMutableSource/GenerateMutableSource.h
--- a/Engine/Plugins/Experimental/Mutable/Source/CustomizableObjectEditor/Private/MuCOE/GenerateMutableSource/GenerateMutableSource.h
+++ b/Engine/Plugins/Experimental/Mutable/Source/CustomizableObjectEditor/Private/MuCOE/GenerateMutableSource/GenerateMutableSource.h
@@ -12,10 +12,7 @@
 #include "MuCOE/ExtensionDataCompilerInterface.h"
 #include "MuCOE/Nodes/CustomizableObjectNodeMaterialBase.h"
 #include "MuCOE/Nodes/CustomizableObjectNodeObject.h"
-#include "MuR/MutableMemory.h"
-#include "MuR/Skeleton.h"
 #include "MuT/NodeComponentNew.h"
-#include "MuR/Ptr.h"
 #include "MuT/NodeImageConstant.h"
 #include "MuT/NodeMeshApplyPose.h"
 #include "MuT/NodeModifierMeshClipWithMesh.h"
@@ -25,6 +22,10 @@
 #include "MuT/NodeScalarParameter.h"
 #include "MuT/NodeSurfaceNew.h"
 #include "MuT/Table.h"
+#include "MuR/MutableMemory.h"
+#include "MuR/Skeleton.h"
+#include "MuR/Mesh.h"
+#include "MuR/Ptr.h"
 #include "UObject/Package.h"
 
 class UCustomizableObjectNodeMaterialBase;
@@ -757,10 +758,10 @@ struct FMutableGraphGenerationContext
 
 	TArray<FAnimBpOverridePhysicsAssetsInfo> AnimBpOverridePhysicsAssetsInfo;
 
-	// Hierarchy of current ComponentNew nodes, each stored for every LOD
+	// Hierarchy of current ComponentNew nodes
 	struct ObjectParent
 	{
-		TArray<TArray<mu::NodeComponentNewPtr>> ComponentsPerLOD;
+		TArray<mu::Ptr<mu::NodeComponentNew>> Components;
 	};
 	TArray< ObjectParent > ComponentNewNode;
 
@@ -768,8 +769,15 @@ struct FMutableGraphGenerationContext
 	uint8 CurrentLOD = 0;
 	int32 CurrentMeshComponent = 0;
 
+	/** If this is set, we are genreating materials for a "passthrough" component, with a fixed mesh. */
+	mu::Ptr<mu::NodeMesh> ComponentMeshOverride;
+
 	uint8 NumLODsInRoot = 0;
 	uint8 NumMeshComponentsInRoot = 0;
+
+	/** Number of additional mesh components added with explict Component nodes to the root or another object node. */
+	uint8 NumExplicitMeshComponents = 0;
+
 	uint8 FirstLODAvailable = MAX_MESH_LOD_COUNT;
 	uint8 NumMaxLODsToStream = MAX_MESH_LOD_COUNT;
 
@@ -876,7 +884,8 @@ private:
 
 
 //
-mu::NodeObjectPtr GenerateMutableSource(const class UEdGraphPin* Pin, FMutableGraphGenerationContext& GenerationContext, bool bPartialCompilation);
+mu::Ptr<mu::NodeObject> GenerateMutableSource(const class UEdGraphPin* Pin, FMutableGraphGenerationContext& GenerationContext, bool bPartialCompilation);
+mu::Ptr<mu::NodeComponent> GenerateMutableSourceComponent(const UEdGraphPin* Pin, FMutableGraphGenerationContext& GenerationContext);
 
 
 /** Populate an array with all the information related to the reference skeletal meshes we might need in-game to generate instances */
diff --git a/Engine/Plugins/Experimental/Mutable/Source/CustomizableObjectEditor/Private/MuCOE/GenerateMutableSource/GenerateMutableSourceComponent.cpp b/Engine/Plugins/Experimental/Mutable/Source/CustomizableObjectEditor/Private/MuCOE/GenerateMutableSource/GenerateMutableSourceComponent.cpp
new file mode 100644
--- /dev/null
+++ b/Engine/Plugins/Experimental/Mutable/Source/CustomizableObjectEditor/Private/MuCOE/GenerateMutableSource/GenerateMutableSourceComponent.cpp
@@ -0,0 +1,119 @@
+// Copyright Epic Games, Inc. All Rights Reserved.
+
+#include "MuCOE/GenerateMutableSource/GenerateMutableSource.h"
+
+#include "MuCOE/CustomizableObjectCompiler.h"
+#include "MuCOE/GenerateMutableSource/GenerateMutableSourceMesh.h"
+#include "MuCOE/GenerateMutableSource/GenerateMutableSourceSurface.h"
+#include "MuCOE/GraphTraversal.h"
+#include "MuCOE/MutableUtils.h"
+#include "MuCOE/Nodes/CustomizableObjectNodeComponentMesh.h"
+#include "MuT/NodeComponent.h"
+#include "MuT/NodeComponentNew.h"
+#include "MuT/NodeLOD.h"
+#include "MuT/NodeSurfaceNew.h"
+#include "MuT/NodeMeshConstant.h"
+
+#include "Rendering/SkeletalMeshLODModel.h"
+#include "Rendering/SkeletalMeshModel.h"
+#include "Rendering/SkeletalMeshRenderData.h"
+
+
+#define LOCTEXT_NAMESPACE "CustomizableObjectEditor"
+
+
+mu::Ptr<mu::NodeComponent> GenerateMutableSourceComponent(const UEdGraphPin * Pin, FMutableGraphGenerationContext& GenerationContext)
+{
+	check(Pin)
+	RETURN_ON_CYCLE(*Pin, GenerationContext)
+
+	CheckNumOutputs(*Pin, GenerationContext);
+	
+	UCustomizableObjectNode* Node = CastChecked<UCustomizableObjectNode>(Pin->GetOwningNode());
+
+	const FGeneratedKey Key(reinterpret_cast<void*>(&GenerateMutableSourceComponent), *Pin, *Node, GenerationContext, true);
+	if (const FGeneratedData* Generated = GenerationContext.Generated.Find(Key))
+	{
+		return static_cast<mu::NodeComponent*>(Generated->Node.get());
+	}
+	
+	mu::Ptr<mu::NodeComponent> Result;
+	
+	if (const UCustomizableObjectNodeComponentMesh* TypedComponentMesh = Cast<UCustomizableObjectNodeComponentMesh>(Node))
+	{
+		if (!TypedComponentMesh->Mesh.IsValid())
+		{
+			FString Msg = FString::Printf(TEXT("No mesh set for component node."));
+			GenerationContext.Compiler->CompilerLog(FText::FromString(Msg), TypedComponentMesh, EMessageSeverity::Warning);
+			return nullptr;
+		}
+
+		USkeletalMesh* SkeletalMesh = Cast<USkeletalMesh>(TypedComponentMesh->Mesh.TryLoad());
+		if (!SkeletalMesh)
+		{
+			FString Msg = FString::Printf(TEXT("Only SkeletalMeshes are supported in this node, for now."));
+			GenerationContext.Compiler->CompilerLog(FText::FromString(Msg), TypedComponentMesh, EMessageSeverity::Warning);
+			return nullptr;
+		}
+
+		// Create the referenced mesh node.
+		mu::Ptr<mu::NodeMeshConstant> MeshNode;
+		{
+			MeshNode = new mu::NodeMeshConstant();
+
+			FString MeshUniqueTags;
+			constexpr bool bIsReference = true;
+			TSoftClassPtr<UAnimInstance> AnimInstance;
+			mu::Ptr<mu::Mesh> MutableMesh = GenerateMutableMesh(SkeletalMesh, AnimInstance, 0, 0, 0, 0, MeshUniqueTags, GenerationContext, TypedComponentMesh, bIsReference);
+
+			MeshNode->SetValue(MutableMesh);
+		}
+
+		// Create the component node
+		mu::Ptr<mu::NodeComponentNew> ComponentNode = new mu::NodeComponentNew;
+		ComponentNode->Id = GenerationContext.CurrentMeshComponent;
+
+		// Create a LOD for each pass-through mesh LOD.
+		const FSkeletalMeshModel* Model = SkeletalMesh->GetImportedModel();
+		int32 SkeletalMeshLODCount = Model->LODModels.Num();
+		for (int32 LODIndex=0; LODIndex<SkeletalMeshLODCount; ++LODIndex)
+		{
+			mu::Ptr<mu::NodeLOD> LODNode = new mu::NodeLOD;
+			ComponentNode->LODs.Add(LODNode);
+
+			const FSkeletalMeshLODModel& LODModel = Model->LODModels[LODIndex];
+			int32 SectionCount = LODModel.Sections.Num();
+			for (int32 SectionIndex = 0; SectionIndex < SectionCount; ++SectionIndex)
+			{
+				// Is there a pin in the unreal node for this section?
+				UEdGraphPin* InMaterialPin = TypedComponentMesh->GetMaterialPin(LODIndex,SectionIndex);
+				UEdGraphPin* ConnectedMaterialPin = FollowInputPin(*InMaterialPin);
+
+				if (ConnectedMaterialPin)
+				{
+					GenerationContext.ComponentMeshOverride = MeshNode;
+					
+					mu::Ptr<mu::NodeSurface> SurfaceNode = GenerateMutableSourceSurface(ConnectedMaterialPin, GenerationContext);
+					LODNode->Surfaces.Add(SurfaceNode);
+
+					GenerationContext.ComponentMeshOverride = nullptr;
+				}
+			}
+		}
+
+		Result = ComponentNode;
+	}
+
+	else
+	{
+		GenerationContext.Compiler->CompilerLog(LOCTEXT("UnimplementedNode", "Node type not implemented yet."), Node);
+	}
+
+	GenerationContext.Generated.Add(Key, FGeneratedData(Node, Result));
+	GenerationContext.GeneratedNodes.Add(Node);
+	
+	return Result;
+}
+
+#undef LOCTEXT_NAMESPACE
+
diff --git a/Engine/Plugins/Experimental/Mutable/Source/CustomizableObjectEditor/Private/MuCOE/GenerateMutableSource/GenerateMutableSourceMesh.cpp b/Engine/Plugins/Experimental/Mutable/Source/CustomizableObjectEditor/Private/MuCOE/GenerateMutableSource/GenerateMutableSourceMesh.cpp
--- a/Engine/Plugins/Experimental/Mutable/Source/CustomizableObjectEditor/Private/MuCOE/GenerateMutableSource/GenerateMutableSourceMesh.cpp
+++ b/Engine/Plugins/Experimental/Mutable/Source/CustomizableObjectEditor/Private/MuCOE/GenerateMutableSource/GenerateMutableSourceMesh.cpp
@@ -36,7 +36,6 @@
 #include "MuCOE/Nodes/CustomizableObjectNodeSkeletalMesh.h"
 #include "MuCOE/Nodes/CustomizableObjectNodeStaticMesh.h"
 #include "MuCOE/Nodes/CustomizableObjectNodeTable.h"
-#include "MuCOE/Nodes/CustomizableObjectNodePassThroughMesh.h"
 #include "MuCOE/UnrealEditorPortabilityHelpers.h"
 #include "MuT/NodeMeshConstant.h"
 #include "MuT/NodeMeshGeometryOperation.h"
@@ -54,7 +53,6 @@
 
 #define LOCTEXT_NAMESPACE "CustomizableObjectEditor"
 
-
 void GetLODAndSectionForAutomaticLODs(const FMutableGraphGenerationContext& Context, const UCustomizableObjectNode& Node, const USkeletalMesh& SkeletalMesh,
 	const int32 LODIndexConnected, const int32 SectionIndexConnected, int32& OutLODIndex, int32& OutSectionIndex, const bool bOnlyConnectedLOD)
 {
@@ -3478,27 +3476,6 @@ mu::NodeMeshPtr GenerateMutableSourceMesh(const UEdGraphPin* Pin,
 		}
 	}
 
-	else if (const UCustomizableObjectNodePassThroughMesh* TypedNodePassThroughMesh = Cast<UCustomizableObjectNodePassThroughMesh>(Node))
-	{
-		UStreamableRenderAsset* BaseMesh = TypedNodePassThroughMesh->GetMesh();
-		if (BaseMesh)
-		{
-			mu::Ptr<mu::NodeMeshConstant> MeshNode = new mu::NodeMeshConstant();
-			Result = MeshNode;
-
-			FString MeshUniqueTags;
-			constexpr bool bIsReference = true;
-			TSoftClassPtr<UAnimInstance> AnimInstance;
-			mu::Ptr<mu::Mesh> MutableMesh = GenerateMutableMesh(BaseMesh, AnimInstance, 0, 0, 0, 0, MeshUniqueTags, GenerationContext, TypedNodeSkel, bIsReference);
-
-			MeshNode->SetValue(MutableMesh);
-		}
-		else
-		{
-			GenerationContext.Compiler->CompilerLog(LOCTEXT("MissingImagePassThrough", "Missing image in pass-through texture node."), Node);
-		}
-	}
-
 	else if (UCustomizableObjectNodeMeshMorph* TypedNodeMorph = Cast<UCustomizableObjectNodeMeshMorph>(Node))
 	{
 		if (const UEdGraphPin* ConnectedPin = FollowInputPin(*TypedNodeMorph->MeshPin()))
diff --git a/Engine/Plugins/Experimental/Mutable/Source/CustomizableObjectEditor/Private/MuCOE/GenerateMutableSource/GenerateMutableSourceSurface.cpp b/Engine/Plugins/Experimental/Mutable/Source/CustomizableObjectEditor/Private/MuCOE/GenerateMutableSource/GenerateMutableSourceSurface.cpp
--- a/Engine/Plugins/Experimental/Mutable/Source/CustomizableObjectEditor/Private/MuCOE/GenerateMutableSource/GenerateMutableSourceSurface.cpp
+++ b/Engine/Plugins/Experimental/Mutable/Source/CustomizableObjectEditor/Private/MuCOE/GenerateMutableSource/GenerateMutableSourceSurface.cpp
@@ -228,7 +228,7 @@ void AddModifierToSharedSurface(FMutableGraphGenerationContext& GenerationContex
 }
 
 
-mu::NodeSurfacePtr GenerateMutableSourceSurface(const UEdGraphPin * Pin, FMutableGraphGenerationContext & GenerationContext)
+mu::Ptr<mu::NodeSurface> GenerateMutableSourceSurface(const UEdGraphPin * Pin, FMutableGraphGenerationContext & GenerationContext)
 {
 	check(Pin)
 	RETURN_ON_CYCLE(*Pin, GenerationContext)
@@ -243,7 +243,7 @@ mu::NodeSurfacePtr GenerateMutableSourceSurface(const UEdGraphPin * Pin, FMutabl
 		return static_cast<mu::NodeSurface*>(Generated->Node.get());
 	}
 	
-	mu::NodeSurfacePtr Result;
+	mu::Ptr<mu::NodeSurface> Result;
 
 	const int32 LOD = Node->IsAffectedByLOD() ? GenerationContext.CurrentLOD : 0;
 	
@@ -257,7 +257,10 @@ mu::NodeSurfacePtr GenerateMutableSourceSurface(const UEdGraphPin * Pin, FMutabl
 
 	if (UCustomizableObjectNodeMaterialBase* TypedNodeMat = Cast<UCustomizableObjectNodeMaterialBase>(Node))
 	{
-		if (TypedNodeMat->GetMeshComponentIndex() != GenerationContext.CurrentMeshComponent)
+		bool bGeneratingImplicitComponent = GenerationContext.ComponentMeshOverride.get()!=nullptr;
+
+		bool bValidForCurrentComponent = bGeneratingImplicitComponent || (TypedNodeMat->GetMeshComponentIndex() == GenerationContext.CurrentMeshComponent);
+		if (!bValidForCurrentComponent)
 		{
 			return Result;
 		}
@@ -274,7 +277,7 @@ mu::NodeSurfacePtr GenerateMutableSourceSurface(const UEdGraphPin * Pin, FMutabl
 			}
 		}
 
-		mu::NodeSurfaceNewPtr SurfNode = new mu::NodeSurfaceNew();
+		mu::Ptr<mu::NodeSurfaceNew> SurfNode = new mu::NodeSurfaceNew();
 		Result = SurfNode;
 
 		// Add to the list of surfaces that could be reused between LODs for this NodeMaterial.
@@ -339,34 +342,45 @@ mu::NodeSurfacePtr GenerateMutableSourceSurface(const UEdGraphPin * Pin, FMutabl
 			}
 		}
 
-		mu::NodeMeshPtr MeshNode;
+		mu::Ptr<mu::NodeMesh> MeshNode;
 		
-		if (const UEdGraphPin* ConnectedPin = FollowInputPin(*TypedNodeMat->GetMeshPin()))
+		if (bGeneratingImplicitComponent)
 		{
-			FMutableGraphMeshGenerationData MeshData;
-			MeshNode = GenerateMutableSourceMesh(ConnectedPin, GenerationContext, MeshData, false, false);
+			MeshNode = GenerationContext.ComponentMeshOverride;
+			SurfNode->SetMesh(MeshNode);
 
-			if (MeshNode)
+			if (const UEdGraphPin* ConnectedPin = FollowInputPin(*TypedNodeMat->GetMeshPin()))
+			{
+				GenerationContext.Compiler->CompilerLog(LOCTEXT("MeshIgnored", "The mesh nodes connected to a material node will be ignored because it is part of an explicit mesh component."), Node);
+			}
+		}
+		else
+		{
+			if (const UEdGraphPin* ConnectedPin = FollowInputPin(*TypedNodeMat->GetMeshPin()))
 			{
-				SurfNode->SetMeshCount(1);
+				FMutableGraphMeshGenerationData MeshData;
+				MeshNode = GenerateMutableSourceMesh(ConnectedPin, GenerationContext, MeshData, false, false);
 
-				mu::Ptr<mu::NodeMeshFormat> MeshFormatNode = new mu::NodeMeshFormat();
-				MeshFormatNode->SetSource(MeshNode.get());
-				SetSurfaceFormat(GenerationContext,
+				if (MeshNode)
+				{
+					mu::Ptr<mu::NodeMeshFormat> MeshFormatNode = new mu::NodeMeshFormat();
+					MeshFormatNode->SetSource(MeshNode.get());
+					SetSurfaceFormat(GenerationContext,
 						MeshFormatNode->GetVertexBuffers(), MeshFormatNode->GetIndexBuffers(), MeshData,
 						GenerationContext.Options.CustomizableObjectNumBoneInfluences,
 						GenerationContext.Options.b16BitBoneWeightsEnabled);
 
-				// \TODO: Make it an option?
-				MeshFormatNode->SetOptimizeBuffers(true);
+					// \TODO: Make it an option?
+					MeshFormatNode->SetOptimizeBuffers(true);
 
-				MeshFormatNode->SetMessageContext(Node);
+					MeshFormatNode->SetMessageContext(Node);
 
-				SurfNode->SetMesh(0, MeshFormatNode);
-			}
-			else
-			{
-				GenerationContext.Compiler->CompilerLog(LOCTEXT("MeshFailed", "Mesh generation failed."), Node);
+					SurfNode->SetMesh(MeshFormatNode);
+				}
+				else
+				{
+					GenerationContext.Compiler->CompilerLog(LOCTEXT("MeshFailed", "Mesh generation failed."), Node);
+				}
 			}
 		}
 		
@@ -377,19 +391,19 @@ mu::NodeSurfacePtr GenerateMutableSourceSurface(const UEdGraphPin * Pin, FMutabl
 		FString TableColumnName;
 
 		// Checking if we should not use the material of the table node even if it is linked to the material node
-		const UEdGraphPin* TempConnectedPin = nullptr;
+		const UEdGraphPin* MaterialAssetConnectedPin = nullptr;
 		if (TypedNodeMat->GetMaterialAssetPin())
 		{
-			TempConnectedPin = FollowInputPin(*TypedNodeMat->GetMaterialAssetPin());
+			MaterialAssetConnectedPin = FollowInputPin(*TypedNodeMat->GetMaterialAssetPin());
 		}
 
-		if (TempConnectedPin)
+		if (MaterialAssetConnectedPin)
 		{
-			if (const UCustomizableObjectNodeTable* TypedNodeTable = Cast< UCustomizableObjectNodeTable >(TempConnectedPin->GetOwningNode()))
+			if (const UCustomizableObjectNodeTable* TypedNodeTable = Cast< UCustomizableObjectNodeTable >(MaterialAssetConnectedPin->GetOwningNode()))
 			{
-				TableColumnName = TempConnectedPin->PinFriendlyName.ToString();
+				TableColumnName = MaterialAssetConnectedPin->PinFriendlyName.ToString();
 
-				if (UMaterialInstance * TableMaterial = TypedNodeTable->GetColumnDefaultAssetByType<UMaterialInstance>(TempConnectedPin))
+				if (UMaterialInstance * TableMaterial = TypedNodeTable->GetColumnDefaultAssetByType<UMaterialInstance>(MaterialAssetConnectedPin))
 				{
 					// Checking if the reference material of the Table Node has the same parent as the material of the Material Node 
 					if (!TypedNodeMat->GetMaterial() || TableMaterial->GetMaterial() != TypedNodeMat->GetMaterial()->GetMaterial())
@@ -810,7 +824,13 @@ mu::NodeSurfacePtr GenerateMutableSourceSurface(const UEdGraphPin * Pin, FMutabl
 						const FString LayerEncoding = LayerIndex != INDEX_NONE ? "-MutableLayerParam:" + FString::FromInt(LayerIndex) : "";
 						
 						SurfNode->SetImageName(ImageIndex, SurfNodeImageName + LayerEncoding);
-						const int32 UVLayout = TypedNodeMat->GetImageUVLayout(ImageIndex);
+
+						// If we are generating an implicit component (with a passthrough mesh) we don't apply any layout.
+						int32 UVLayout = -1;
+						if (!bGeneratingImplicitComponent)
+						{
+							UVLayout = TypedNodeMat->GetImageUVLayout(ImageIndex);;
+						}
 						SurfNode->SetImageLayoutIndex(ImageIndex, UVLayout);
 						SurfNode->SetImageAdditionalNames(ImageIndex, TypedNodeMat->GetMaterial()->GetName(), ImageName);
 
@@ -951,8 +971,7 @@ mu::NodeSurfacePtr GenerateMutableSourceSurface(const UEdGraphPin * Pin, FMutabl
 
 			SurfNode2->SetCustomID(ReferencedMaterialsIndex);
 
-			SurfNode2->SetMeshCount(SurfNode->GetMeshCount());
-			SurfNode2->SetMesh(0, SurfNode->GetMesh(0));
+			SurfNode2->SetMesh(SurfNode->GetMesh());
 
 			SurfNode2->SetVectorCount(SurfNode->GetVectorCount());
 
@@ -1623,4 +1642,5 @@ mu::NodeSurfacePtr GenerateMutableSourceSurface(const UEdGraphPin * Pin, FMutabl
 	return Result;
 }
 
+
 #undef LOCTEXT_NAMESPACE
diff --git a/Engine/Plugins/Experimental/Mutable/Source/CustomizableObjectEditor/Private/MuCOE/GraphTraversal.cpp b/Engine/Plugins/Experimental/Mutable/Source/CustomizableObjectEditor/Private/MuCOE/GraphTraversal.cpp
--- a/Engine/Plugins/Experimental/Mutable/Source/CustomizableObjectEditor/Private/MuCOE/GraphTraversal.cpp
+++ b/Engine/Plugins/Experimental/Mutable/Source/CustomizableObjectEditor/Private/MuCOE/GraphTraversal.cpp
@@ -18,7 +18,6 @@
 #include "MuCOE/Nodes/CustomizableObjectNodeMeshReshape.h"
 #include "MuCOE/Nodes/CustomizableObjectNodeMeshSwitch.h"
 #include "MuCOE/Nodes/CustomizableObjectNodeMeshVariation.h"
-#include "MuCOE/Nodes/CustomizableObjectNodePassThroughMesh.h"
 #include "MuCOE/Nodes/CustomizableObjectNodeObject.h"
 #include "MuCOE/Nodes/CustomizableObjectNodeSkeletalMesh.h"
 #include "MuCOE/Nodes/CustomizableObjectNodeStaticMesh.h"
@@ -382,14 +381,6 @@ const UEdGraphPin* FindMeshBaseSource(const UEdGraphPin& Pin, const bool bOnlyLo
 		}
 	}
 
-	else if (Cast<UCustomizableObjectNodePassThroughMesh>(Node))
-	{
-		if (!bOnlyLookForStaticMesh)
-		{
-			return &Pin;
-		}
-	}
-
 	else if (Cast<UCustomizableObjectNodeStaticMesh>(Node))
 	{
 		return &Pin;
diff --git a/Engine/Plugins/Experimental/Mutable/Source/CustomizableObjectEditor/Private/MuCOE/Nodes/CustomizableObjectNodeComponentMesh.cpp b/Engine/Plugins/Experimental/Mutable/Source/CustomizableObjectEditor/Private/MuCOE/Nodes/CustomizableObjectNodeComponentMesh.cpp
new file mode 100644
--- /dev/null
+++ b/Engine/Plugins/Experimental/Mutable/Source/CustomizableObjectEditor/Private/MuCOE/Nodes/CustomizableObjectNodeComponentMesh.cpp
@@ -0,0 +1,283 @@
+// Copyright Epic Games, Inc. All Rights Reserved.
+
+#include "MuCOE/Nodes/CustomizableObjectNodeComponentMesh.h"
+
+#include "Engine/SkeletalMesh.h"
+#include "Modules/ModuleManager.h"
+#include "MuCO/CustomizableObjectCustomVersion.h"
+#include "MuCOE/CustomizableObjectEditor.h"
+#include "MuCOE/CustomizableObjectEditorStyle.h"
+#include "MuCOE/CustomizableObjectLayout.h"
+#include "MuCOE/EdGraphSchema_CustomizableObject.h"
+#include "MuCOE/GraphTraversal.h"
+#include "MuCOE/MutableUtils.h"
+#include "MuCOE/Nodes/CustomizableObjectNodeLayoutBlocks.h"
+#include "MuCOE/RemapPins/CustomizableObjectNodeRemapPinsByNameDefaultPin.h"
+#include "MuCOE/UnrealEditorPortabilityHelpers.h"
+
+class UCustomizableObjectNodeRemapPinsByName;
+class UObject;
+
+
+#define LOCTEXT_NAMESPACE "CustomizableObjectEditor"
+
+const FName UCustomizableObjectNodeComponentMesh::OutputPinName(TEXT("Component"));
+
+
+void UCustomizableObjectNodeComponentMeshPinDataSection::Init(int32 InLODIndex, int32 InSectionIndex)
+{
+	LODIndex = InLODIndex;
+	SectionIndex = InSectionIndex;
+}
+
+
+int32 UCustomizableObjectNodeComponentMeshPinDataSection::GetLODIndex() const
+{
+	return LODIndex;
+}
+
+
+int32 UCustomizableObjectNodeComponentMeshPinDataSection::GetSectionIndex() const
+{
+	return SectionIndex;
+}
+
+
+bool UCustomizableObjectNodeComponentMeshPinDataSection::Equals(const UCustomizableObjectNodePinData& Other) const
+{
+	if (GetClass() != Other.GetClass())
+	{
+		return false;
+	}
+
+	const UCustomizableObjectNodeComponentMeshPinDataSection& OtherTyped = static_cast<const UCustomizableObjectNodeComponentMeshPinDataSection&>(Other);
+	if (LODIndex != OtherTyped.LODIndex ||
+		SectionIndex != OtherTyped.SectionIndex)
+	{
+		return false;
+	}
+
+	return Super::Equals(Other);
+}
+
+
+void UCustomizableObjectNodeComponentMesh::PostEditChangeProperty(FPropertyChangedEvent& PropertyChangedEvent)
+{
+	Super::PostEditChangeProperty(PropertyChangedEvent);
+
+	FProperty* PropertyThatChanged = PropertyChangedEvent.Property;
+	if (PropertyThatChanged && PropertyThatChanged->GetName() == TEXT("SkeletalMesh"))
+	{
+		ReconstructNode();
+	}
+}
+
+
+void UCustomizableObjectNodeComponentMesh::AllocateDefaultPins(UCustomizableObjectNodeRemapPins* RemapPins)
+{
+	const UEdGraphSchema_CustomizableObject* Schema = GetDefault<UEdGraphSchema_CustomizableObject>();
+	if (!Schema)
+	{
+		return;
+	}
+
+	UEdGraphPin* OutputPin = CustomCreatePin(EGPD_Output, Schema->PC_Component, OutputPinName);
+
+	if (!Mesh.IsValid())
+	{
+		return;
+	}
+
+	// Support SkeletalMeshes only, for now.
+	USkeletalMesh* SkeletalMesh = Cast<USkeletalMesh>( Mesh.TryLoad() );
+	if (!SkeletalMesh)
+	{
+		return;
+	}
+	
+	if (const FSkeletalMeshModel* ImportedModel = SkeletalMesh->GetImportedModel())
+	{
+		const int32 NumLODs = SkeletalMesh->GetLODNum();
+		for (int32 LODIndex = 0; LODIndex < NumLODs; ++LODIndex)
+		{
+			const int32 NumSections = ImportedModel->LODModels[LODIndex].Sections.Num();
+			for (int32 SectionIndex = 0; SectionIndex < NumSections; ++SectionIndex)
+			{
+				// Ignore disabled sections.
+				const bool bIsSectionDisabled = ImportedModel ->LODModels[LODIndex].Sections[SectionIndex].bDisabled;
+				if (bIsSectionDisabled)
+				{
+					continue;
+				}
+
+				UMaterialInterface* MaterialInterface = GetMaterialInterfaceFor(LODIndex, SectionIndex, ImportedModel);
+				
+				// Material for the section
+				{
+					UCustomizableObjectNodeComponentMeshPinDataMaterial* PinData = NewObject<UCustomizableObjectNodeComponentMeshPinDataMaterial>(this);
+					PinData->Init(LODIndex, SectionIndex);
+					
+					FString SectionFriendlyName = MaterialInterface ? MaterialInterface->GetName() : FString::Printf(TEXT("Section %i"), SectionIndex);
+					FString MaterialPinName = FString::Printf(TEXT("LOD %i - Section %i - Material"), LODIndex, SectionIndex);
+
+					UEdGraphPin* Pin = CustomCreatePin(EGPD_Input, Schema->PC_Material, FName(*MaterialPinName), PinData);
+					Pin->PinFriendlyName = FText::FromString(FString::Printf(TEXT("LOD %i - Section %i - %s"), LODIndex, SectionIndex, *SectionFriendlyName));
+					Pin->PinToolTip = MaterialPinName;
+				}
+				
+			}
+		}
+	}
+}
+
+
+FText UCustomizableObjectNodeComponentMesh::GetNodeTitle(ENodeTitleType::Type TitleType) const
+{
+	USkeletalMesh* SkeletalMesh = Cast<USkeletalMesh>(Mesh.TryLoad());
+	if (SkeletalMesh)
+	{
+		FFormatNamedArguments Args;
+		Args.Add(TEXT("MeshName"), FText::FromString(SkeletalMesh->GetName()));
+
+		return FText::Format(LOCTEXT("ComponentMesh_Title", "{MeshName}\nMesh Component"), Args);
+	}
+	else
+	{
+		return LOCTEXT("Component_Mesh", "Mesh Component");
+	}
+}
+
+
+FLinearColor UCustomizableObjectNodeComponentMesh::GetNodeTitleColor() const
+{
+	const UEdGraphSchema_CustomizableObject* Schema = GetDefault<UEdGraphSchema_CustomizableObject>();
+	return Schema->GetPinTypeColor(Schema->PC_Component);
+}
+
+
+void UCustomizableObjectNodeComponentMesh::GetPinSection(const UEdGraphPin& Pin, int32& OutLODIndex, int32& OutSectionIndex) const
+{
+	if (const UCustomizableObjectNodeComponentMeshPinDataSection* PinData = Cast<UCustomizableObjectNodeComponentMeshPinDataSection>(GetPinData(Pin)))
+	{
+		OutLODIndex = PinData->GetLODIndex();
+		OutSectionIndex = PinData->GetSectionIndex();
+		return;
+	}
+
+	OutLODIndex = -1;
+	OutSectionIndex = -1;
+}
+
+
+UMaterialInterface* UCustomizableObjectNodeComponentMesh::GetMaterialFor(const UEdGraphPin* Pin) const
+{
+	USkeletalMesh* SkeletalMesh = Cast<USkeletalMesh>(Mesh.TryLoad());
+	if (SkeletalMesh)
+	{
+		if (FSkeletalMaterial* SkeletalMaterial = GetSkeletalMaterialFor(*Pin))
+		{
+			return SkeletalMaterial->MaterialInterface;
+		}
+	}
+
+	return nullptr;
+}
+
+
+FSkeletalMaterial* UCustomizableObjectNodeComponentMesh::GetSkeletalMaterialFor(const UEdGraphPin& Pin) const
+{
+	int32 LODIndex;
+	int32 SectionIndex;
+	GetPinSection(Pin, LODIndex, SectionIndex);
+
+	return GetSkeletalMaterialFor(LODIndex, SectionIndex);
+}
+
+
+FText UCustomizableObjectNodeComponentMesh::GetTooltipText() const
+{
+	return LOCTEXT("Component_Mesh_Tooltip", "Define a new object component based on a Skeletal Mesh.");
+}
+
+
+UMaterialInterface* UCustomizableObjectNodeComponentMesh::GetMaterialInterfaceFor(const int32 LODIndex, const int32 SectionIndex, const FSkeletalMeshModel* ImportedModel) const
+{
+	if (FSkeletalMaterial* SkeletalMaterial = GetSkeletalMaterialFor(LODIndex, SectionIndex, ImportedModel))
+	{
+		return SkeletalMaterial->MaterialInterface;
+	}
+
+	return nullptr;
+}
+
+
+FSkeletalMaterial* UCustomizableObjectNodeComponentMesh::GetSkeletalMaterialFor(const int32 LODIndex, const int32 SectionIndex, const FSkeletalMeshModel* ImportedModel) const
+{
+	USkeletalMesh* SkeletalMesh = Cast<USkeletalMesh>(Mesh.TryLoad());
+	if (!SkeletalMesh)
+	{
+		return nullptr;
+	}
+
+	// We assume that LODIndex and MaterialIndex are valid for the imported model
+	int32 SkeletalMeshMaterialIndex = INDEX_NONE;
+	
+	// Check if we have lod info map to get the correct material index
+	if (const FSkeletalMeshLODInfo* LodInfo = SkeletalMesh->GetLODInfo(LODIndex))
+	{
+		if (LodInfo->LODMaterialMap.IsValidIndex(SectionIndex))
+		{
+			SkeletalMeshMaterialIndex = LodInfo->LODMaterialMap[SectionIndex];
+		}
+	}
+
+	// Only deduce index when the explicit mapping is not found or there is no remap
+	if (SkeletalMeshMaterialIndex == INDEX_NONE)
+	{
+	if (ImportedModel && ImportedModel->LODModels.IsValidIndex(LODIndex) && ImportedModel->LODModels[LODIndex].Sections.IsValidIndex(SectionIndex))
+	{
+		SkeletalMeshMaterialIndex = ImportedModel->LODModels[LODIndex].Sections[SectionIndex].MaterialIndex;
+	}
+	else
+	{
+		FSkeletalMeshModel* AuxImportedModel = SkeletalMesh->GetImportedModel();
+
+		if (AuxImportedModel)
+		{
+			if (AuxImportedModel->LODModels.IsValidIndex(LODIndex) && AuxImportedModel->LODModels[LODIndex].Sections.IsValidIndex(SectionIndex))
+			{
+				SkeletalMeshMaterialIndex = AuxImportedModel->LODModels[LODIndex].Sections[SectionIndex].MaterialIndex;
+			}
+		}
+	}
+	}
+	
+	if (SkeletalMesh->GetMaterials().IsValidIndex(SkeletalMeshMaterialIndex))
+	{
+		return &SkeletalMesh->GetMaterials()[SkeletalMeshMaterialIndex];
+	}
+	
+	return nullptr;
+}
+
+
+
+UEdGraphPin* UCustomizableObjectNodeComponentMesh::GetMaterialPin(const int32 LODIndex, const int32 SectionIndex) const
+{
+	for (UEdGraphPin* Pin : GetAllNonOrphanPins())
+	{
+		if (const UCustomizableObjectNodeComponentMeshPinDataSection* PinData = Cast<UCustomizableObjectNodeComponentMeshPinDataSection>(GetPinData(*Pin)))
+		{
+			if (PinData->GetLODIndex() == LODIndex &&
+				PinData->GetSectionIndex() == SectionIndex)
+			{
+				return Pin;
+			}
+		}
+	}
+
+	return nullptr;
+}
+
+
+#undef LOCTEXT_NAMESPACE
diff --git a/Engine/Plugins/Experimental/Mutable/Source/CustomizableObjectEditor/Private/MuCOE/Nodes/CustomizableObjectNodeExtensionDataSwitch.cpp b/Engine/Plugins/Experimental/Mutable/Source/CustomizableObjectEditor/Private/MuCOE/Nodes/CustomizableObjectNodeExtensionDataSwitch.cpp
--- a/Engine/Plugins/Experimental/Mutable/Source/CustomizableObjectEditor/Private/MuCOE/Nodes/CustomizableObjectNodeExtensionDataSwitch.cpp
+++ b/Engine/Plugins/Experimental/Mutable/Source/CustomizableObjectEditor/Private/MuCOE/Nodes/CustomizableObjectNodeExtensionDataSwitch.cpp
@@ -26,7 +26,7 @@ FString UCustomizableObjectNodeExtensionDataSwitch::GetPinPrefix() const
 	return FString::Format(TEXT("{0} "), { GetOutputPinName() });
 }
 
-mu::NodeExtensionDataPtr UCustomizableObjectNodeExtensionDataSwitch::GenerateMutableNode(FExtensionDataCompilerInterface& InCompilerInterface) const
+mu::Ptr<mu::NodeExtensionData> UCustomizableObjectNodeExtensionDataSwitch::GenerateMutableNode(FExtensionDataCompilerInterface& InCompilerInterface) const
 {
 	const int32 NumParameters = FollowInputPinArray(*SwitchParameter()).Num();
 
@@ -35,7 +35,7 @@ mu::NodeExtensionDataPtr UCustomizableObjectNodeExtensionDataSwitch::GenerateMut
 
 	const int32 NumSwitchOptions = GetNumElements();
 
-	mu::NodeExtensionDataSwitchPtr SwitchNode = new mu::NodeExtensionDataSwitch;
+	mu::Ptr<mu::NodeExtensionDataSwitch> SwitchNode = new mu::NodeExtensionDataSwitch;
 	SwitchNode->SetParameter(SwitchParam);
 	SwitchNode->SetOptionCount(NumSwitchOptions);
 
@@ -47,7 +47,7 @@ mu::NodeExtensionDataPtr UCustomizableObjectNodeExtensionDataSwitch::GenerateMut
 			{
 				if (const ICustomizableObjectExtensionNode* ExtensionNode = Cast<ICustomizableObjectExtensionNode>(ConnectedPin->GetOwningNode()))
 				{
-					if (mu::NodeExtensionDataPtr ExtensionMutableNode = ExtensionNode->GenerateMutableNode(InCompilerInterface))
+					if (mu::Ptr<mu::NodeExtensionData> ExtensionMutableNode = ExtensionNode->GenerateMutableNode(InCompilerInterface))
 					{
 						SwitchNode->SetOption(OptionIndex, ExtensionMutableNode);
 					}
diff --git a/Engine/Plugins/Experimental/Mutable/Source/CustomizableObjectEditor/Private/MuCOE/Nodes/CustomizableObjectNodeExtensionDataVariation.cpp b/Engine/Plugins/Experimental/Mutable/Source/CustomizableObjectEditor/Private/MuCOE/Nodes/CustomizableObjectNodeExtensionDataVariation.cpp
--- a/Engine/Plugins/Experimental/Mutable/Source/CustomizableObjectEditor/Private/MuCOE/Nodes/CustomizableObjectNodeExtensionDataVariation.cpp
+++ b/Engine/Plugins/Experimental/Mutable/Source/CustomizableObjectEditor/Private/MuCOE/Nodes/CustomizableObjectNodeExtensionDataVariation.cpp
@@ -57,15 +57,15 @@ void UCustomizableObjectNodeExtensionDataVariation::AllocateDefaultPins(UCustomi
 	DefaultVariation->bDefaultValueIsIgnored = true;
 }
 
-mu::NodeExtensionDataPtr UCustomizableObjectNodeExtensionDataVariation::GenerateMutableNode(FExtensionDataCompilerInterface& InCompilerInterface) const
+mu::Ptr<mu::NodeExtensionData> UCustomizableObjectNodeExtensionDataVariation::GenerateMutableNode(FExtensionDataCompilerInterface& InCompilerInterface) const
 {
-	mu::NodeExtensionDataVariationPtr VariationNode = new mu::NodeExtensionDataVariation;
+	mu::Ptr<mu::NodeExtensionDataVariation> VariationNode = new mu::NodeExtensionDataVariation;
 
 	if (const UEdGraphPin* ConnectedPin = FollowInputPin(*GetDefaultPin()))
 	{
 		if (const ICustomizableObjectExtensionNode* ExtensionNode = Cast<ICustomizableObjectExtensionNode>(ConnectedPin->GetOwningNode()))
 		{
-			if (mu::NodeExtensionDataPtr DefaultValueExtensionData = ExtensionNode->GenerateMutableNode(InCompilerInterface))
+			if (mu::Ptr<mu::NodeExtensionData> DefaultValueExtensionData = ExtensionNode->GenerateMutableNode(InCompilerInterface))
 			{
 				VariationNode->SetDefaultValue(DefaultValueExtensionData);
 			}
@@ -86,7 +86,7 @@ mu::NodeExtensionDataPtr UCustomizableObjectNodeExtensionDataVariation::Generate
 			{
 				if (const ICustomizableObjectExtensionNode* ExtensionNode = Cast<ICustomizableObjectExtensionNode>(ConnectedPin->GetOwningNode()))
 				{
-					if (mu::NodeExtensionDataPtr VariationExtensionData = ExtensionNode->GenerateMutableNode(InCompilerInterface))
+					if (mu::Ptr<mu::NodeExtensionData> VariationExtensionData = ExtensionNode->GenerateMutableNode(InCompilerInterface))
 					{
 						VariationNode->SetVariationValue(VariationIndex, VariationExtensionData);
 					}
diff --git a/Engine/Plugins/Experimental/Mutable/Source/CustomizableObjectEditor/Private/MuCOE/Nodes/CustomizableObjectNodeObject.cpp b/Engine/Plugins/Experimental/Mutable/Source/CustomizableObjectEditor/Private/MuCOE/Nodes/CustomizableObjectNodeObject.cpp
--- a/Engine/Plugins/Experimental/Mutable/Source/CustomizableObjectEditor/Private/MuCOE/Nodes/CustomizableObjectNodeObject.cpp
+++ b/Engine/Plugins/Experimental/Mutable/Source/CustomizableObjectEditor/Private/MuCOE/Nodes/CustomizableObjectNodeObject.cpp
@@ -23,6 +23,7 @@ class UCustomizableObjectNodeRemapPins;
 #define LOCTEXT_NAMESPACE "CustomizableObjectEditor"
 
 const FName UCustomizableObjectNodeObject::ChildrenPinName(TEXT("Children"));
+const FName UCustomizableObjectNodeObject::ComponentsPinName(TEXT("Components"));
 const FName UCustomizableObjectNodeObject::OutputPinName(TEXT("Object"));
 const TCHAR* UCustomizableObjectNodeObject::LODPinNamePrefix = TEXT("LOD ");
 
@@ -151,7 +152,10 @@ void UCustomizableObjectNodeObject::AllocateDefaultPins(UCustomizableObjectNodeR
 		Pin->bDefaultValueIsIgnored = true;
 	}
 
-	UEdGraphPin* ChildrenPin = CustomCreatePin(EGPD_Input, Schema->PC_Object, ChildrenPinName, true );
+	UEdGraphPin* ComponentsPin = CustomCreatePin(EGPD_Input, Schema->PC_Component, ComponentsPinName, true);
+	ComponentsPin->bDefaultValueIsIgnored = true;
+
+	UEdGraphPin* ChildrenPin = CustomCreatePin(EGPD_Input, Schema->PC_Object, ChildrenPinName, true);
 	ChildrenPin->bDefaultValueIsIgnored = true;
 
 	for (const FRegisteredObjectNodeInputPin& Pin : ICustomizableObjectModule::Get().GetAdditionalObjectNodePins())
@@ -359,6 +363,7 @@ bool UCustomizableObjectNodeObject::IsSingleOutputNode() const
 bool UCustomizableObjectNodeObject::IsBuiltInPin(FName PinName)
 {
 	return PinName == ChildrenPinName
+		|| PinName == ComponentsPinName
 		|| PinName == OutputPinName
 		|| PinName.ToString().StartsWith(LODPinNamePrefix);
 }
diff --git a/Engine/Plugins/Experimental/Mutable/Source/CustomizableObjectEditor/Private/MuCOE/SMutableCodeViewer.cpp b/Engine/Plugins/Experimental/Mutable/Source/CustomizableObjectEditor/Private/MuCOE/SMutableCodeViewer.cpp
--- a/Engine/Plugins/Experimental/Mutable/Source/CustomizableObjectEditor/Private/MuCOE/SMutableCodeViewer.cpp
+++ b/Engine/Plugins/Experimental/Mutable/Source/CustomizableObjectEditor/Private/MuCOE/SMutableCodeViewer.cpp
@@ -920,7 +920,6 @@ void SMutableCodeViewer::GenerateNavigationDropdownElements()
 	const int32 OperationTypesCount = ModelOperationTypes.Num();
 	
 	// It must have at least one type, if not may be because we are running this before filling ModelOperationTypes
-	check(OperationTypesCount);
 	FoundModelOperationTypeElements.Empty(OperationTypesCount);
 	
 	for	(int32 OperationTypeIndex = 0; OperationTypeIndex < OperationTypesCount;  OperationTypeIndex++)
diff --git a/Engine/Plugins/Experimental/Mutable/Source/CustomizableObjectEditor/Private/MuCOE/SMutableGraphViewer.cpp b/Engine/Plugins/Experimental/Mutable/Source/CustomizableObjectEditor/Private/MuCOE/SMutableGraphViewer.cpp
--- a/Engine/Plugins/Experimental/Mutable/Source/CustomizableObjectEditor/Private/MuCOE/SMutableGraphViewer.cpp
+++ b/Engine/Plugins/Experimental/Mutable/Source/CustomizableObjectEditor/Private/MuCOE/SMutableGraphViewer.cpp
@@ -55,21 +55,14 @@
 #include "MuT/NodeScalarCurve.h"
 #include "MuT/NodeScalarSwitch.h"
 #include "MuT/NodeScalarTable.h"
+#include "MuT/NodeObjectNew.h"
 
-
-// This is necessary because of problems with rtti information in other platforms. In any case, this part of the debugger is only useful in the standard editor.
-#if PLATFORM_WINDOWS
-#include "MuT/NodeObjectNewPrivate.h"
 #include "MuT/NodeObjectGroupPrivate.h"
 #include "MuT/NodeSurfaceNewPrivate.h"
 #include "MuT/NodeSurfaceEditPrivate.h"
 #include "MuT/NodeSurfaceSwitchPrivate.h"
 #include "MuT/NodeSurfaceVariationPrivate.h"
-#include "MuT/NodeLODPrivate.h"
-#include "MuT/NodeComponentPrivate.h"
 #include "MuT/NodeModifierPrivate.h"
-#include "MuT/NodeComponentNewPrivate.h"
-#include "MuT/NodeComponentEditPrivate.h"
 #include "MuT/NodeImageFormatPrivate.h"
 #include "MuT/NodeMeshFormatPrivate.h"
 #include "MuT/NodePatchImagePrivate.h"
@@ -104,7 +97,6 @@
 #include "MuT/NodeMeshTablePrivate.h"
 #include "MuT/NodeColourFromScalarsPrivate.h"
 #include "MuT/NodeScalarSwitchPrivate.h"
-#endif
 
 class FExtender;
 class FReferenceCollector;
@@ -134,16 +126,16 @@ public:
 		FText MainLabel = FText::GetEmpty();
 		if (RowItem->MutableNode)
 		{
-			const char* TypeName = RowItem->MutableNode->GetType()->m_strName;
+			mu::Node::EType MutableType = RowItem->MutableNode->GetType()->Type;
 			
 			const FString LabelString = RowItem->Prefix.IsEmpty() 
-				? StringCast<TCHAR>(TypeName).Get() 
-				: FString::Printf( TEXT("%s : %s"), *RowItem->Prefix, StringCast<TCHAR>(TypeName).Get() );
+				? FString::Printf(TEXT("%d"), int32(MutableType) )
+				: FString::Printf( TEXT("%s : %d"), *RowItem->Prefix, int32(MutableType) );
 
 			MainLabel = FText::FromString(LabelString);
 			if (RowItem->DuplicatedOf)
 			{
-				MainLabel = FText::FromString( FString::Printf(TEXT("%s (Duplicated)"), StringCast<TCHAR>(TypeName).Get()));
+				MainLabel = FText::FromString( FString::Printf(TEXT("%d (Duplicated)"), int32(MutableType)));
 			}
 		}
 		else
@@ -338,15 +330,14 @@ void SMutableGraphViewer::GetChildrenForInfo(TSharedPtr<FMutableGraphTreeElement
 	if (ParentNode->GetType() == mu::NodeObjectNew::GetStaticType())
 	{
 		mu::NodeObjectNew* ObjectNew = StaticCast<mu::NodeObjectNew*>(ParentNode);
-		mu::NodeObjectNew::Private* Private = ObjectNew->GetPrivate();
-		for (int32 l = 0; l < Private->m_lods.Num(); ++l)
+		for (int32 l = 0; l < ObjectNew->Components.Num(); ++l)
 		{
-			AddChildFunc(Private->m_lods[l].get(), TEXT("LOD") );
+			AddChildFunc(ObjectNew->Components[l].get(), TEXT("COMP") );
 		}
 
-		for (int32 l = 0; l < Private->m_children.Num(); ++l)
+		for (int32 l = 0; l < ObjectNew->Children.Num(); ++l)
 		{
-			AddChildFunc(Private->m_children[l].get(), TEXT("CHILD"));
+			AddChildFunc(ObjectNew->Children[l].get(), TEXT("CHILD"));
 		}
 	}
 
@@ -364,10 +355,7 @@ void SMutableGraphViewer::GetChildrenForInfo(TSharedPtr<FMutableGraphTreeElement
 	{
 		mu::NodeSurfaceNew* SurfaceNew = StaticCast<mu::NodeSurfaceNew*>(ParentNode);
 		mu::NodeSurfaceNew::Private* Private = SurfaceNew->GetPrivate();
-		for (int32 l = 0; l < Private->m_meshes.Num(); ++l)
-		{
-			AddChildFunc(Private->m_meshes[l].m_pMesh.get(), TEXT("MESH"));
-		}
+		AddChildFunc(Private->Mesh.get(), TEXT("MESH"));
 
 		for (int32 l = 0; l < Private->m_images.Num(); ++l)
 		{
@@ -441,38 +429,34 @@ void SMutableGraphViewer::GetChildrenForInfo(TSharedPtr<FMutableGraphTreeElement
 	else if (ParentNode->GetType() == mu::NodeLOD::GetStaticType())
 	{
 		mu::NodeLOD* LodVar = StaticCast<mu::NodeLOD*>(ParentNode);
-		mu::NodeLOD::Private* Private = LodVar->GetPrivate();
 
-		for (int32 Component = 0; Component < Private->m_components.Num(); Component++)
+		for (int32 SurfaceIndex = 0; SurfaceIndex < LodVar->Surfaces.Num(); SurfaceIndex++)
 		{
-			AddChildFunc(Private->m_components[Component].get(), FString::Printf(TEXT("COMP [%d]"),  Component));
+			AddChildFunc(LodVar->Surfaces[SurfaceIndex].get(), FString::Printf(TEXT("SURFACE [%d]"), SurfaceIndex));
 		}
-		for (int32 Modifier = 0; Modifier < Private->m_modifiers.Num(); Modifier++)
+		for (int32 Modifier = 0; Modifier < LodVar->Modifiers.Num(); Modifier++)
 		{
-			AddChildFunc(Private->m_modifiers[Modifier].get(), FString::Printf(TEXT("MOD [%d]"), Modifier));
+			AddChildFunc(LodVar->Modifiers[Modifier].get(), FString::Printf(TEXT("MOD [%d]"), Modifier));
 		}
 	}
 	
 	else if (ParentNode->GetType() == mu::NodeComponentNew::GetStaticType())
 	{
 		mu::NodeComponentNew* ComponentVar = StaticCast<mu::NodeComponentNew*>(ParentNode);
-		mu::NodeComponentNew::Private* Private = ComponentVar->GetPrivate();
-		for (int32 Surface = 0; Surface < Private->m_surfaces.Num(); Surface++)
+		for (int32 LODIndex = 0; LODIndex < ComponentVar->LODs.Num(); LODIndex++)
 		{
-			AddChildFunc(Private->m_surfaces[Surface].get(), FString::Printf(TEXT("SURF [%d]"), Surface));
+			AddChildFunc(ComponentVar->LODs[LODIndex].get(), FString::Printf(TEXT("LOD [%d]"), LODIndex));
 		}
 	}
 
 	else if (ParentNode->GetType() == mu::NodeComponentEdit::GetStaticType())
 	{
-		mu::NodeComponentEdit* ComponentEditVar = StaticCast<mu::NodeComponentEdit*>(ParentNode);
-		mu::NodeComponentEdit::Private* Private = ComponentEditVar->GetPrivate();
-		for (int32 Surface = 0; Surface < Private->m_surfaces.Num(); Surface++)
+		mu::NodeComponentEdit* ComponentVar = StaticCast<mu::NodeComponentEdit*>(ParentNode);
+		for (int32 LODIndex = 0; LODIndex < ComponentVar->LODs.Num(); LODIndex++)
 		{
-			AddChildFunc(Private->m_surfaces[Surface].get(), FString::Printf(TEXT("SURF [%d]"), Surface));
+			AddChildFunc(ComponentVar->LODs[LODIndex].get(), FString::Printf(TEXT("LOD [%d]"), LODIndex));
 		}
 	}
-	
 
 	else if (ParentNode->GetType() == mu::NodeMeshConstant::GetStaticType())
 	{
@@ -801,11 +785,10 @@ void SMutableGraphViewer::GetChildrenForInfo(TSharedPtr<FMutableGraphTreeElement
 	
 	else
 	{
-		const FString ParentNodeTypeString = ANSI_TO_TCHAR(ParentNode->GetType()->m_strName);
-		UE_LOG(LogMutable,Error,TEXT("The node of type %s has not been implemented, so its children won't be added to the tree."), *ParentNodeTypeString);
+		UE_LOG(LogMutable,Error,TEXT("The node of type %d has not been implemented, so its children won't be added to the tree."), int32(ParentNode->GetType()->Type));
 
 		// Add a placeholder to the tree
-		const FString Prefix =  FString::Printf(TEXT("[%s] NODE TYPE NOT IMPLEMENTED"), *ParentNodeTypeString);
+		const FString Prefix =  FString::Printf(TEXT("[%d] NODE TYPE NOT IMPLEMENTED"), int32(ParentNode->GetType()->Type));
 		AddChildFunc(nullptr, Prefix);
 	}
 #endif
diff --git a/Engine/Plugins/Experimental/Mutable/Source/CustomizableObjectEditor/Public/MuCOE/EdGraphSchema_CustomizableObject.h b/Engine/Plugins/Experimental/Mutable/Source/CustomizableObjectEditor/Public/MuCOE/EdGraphSchema_CustomizableObject.h
--- a/Engine/Plugins/Experimental/Mutable/Source/CustomizableObjectEditor/Public/MuCOE/EdGraphSchema_CustomizableObject.h
+++ b/Engine/Plugins/Experimental/Mutable/Source/CustomizableObjectEditor/Public/MuCOE/EdGraphSchema_CustomizableObject.h
@@ -89,6 +89,7 @@ public:
 
 	// Allowable PinType.PinCategory values
 	static const FName PC_Object;
+	static const FName PC_Component;
 	static const FName PC_Material;
 	static const FName PC_Mesh;
 	static const FName PC_PassThroughMesh;
diff --git a/Engine/Plugins/Experimental/Mutable/Source/CustomizableObjectEditor/Public/MuCOE/ICustomizableObjectExtensionNode.h b/Engine/Plugins/Experimental/Mutable/Source/CustomizableObjectEditor/Public/MuCOE/ICustomizableObjectExtensionNode.h
--- a/Engine/Plugins/Experimental/Mutable/Source/CustomizableObjectEditor/Public/MuCOE/ICustomizableObjectExtensionNode.h
+++ b/Engine/Plugins/Experimental/Mutable/Source/CustomizableObjectEditor/Public/MuCOE/ICustomizableObjectExtensionNode.h
@@ -22,5 +22,5 @@ class ICustomizableObjectExtensionNode
 
 public:
 	/** Generate a mu::Node that produces ExtensionData */
-	virtual mu::NodeExtensionDataPtr GenerateMutableNode(FExtensionDataCompilerInterface& CompilerInterface) const = 0;
+	virtual mu::Ptr<mu::NodeExtensionData> GenerateMutableNode(FExtensionDataCompilerInterface& CompilerInterface) const = 0;
 };
diff --git a/Engine/Plugins/Experimental/Mutable/Source/CustomizableObjectEditor/Public/MuCOE/Nodes/CustomizableObjectNodeComponent.h b/Engine/Plugins/Experimental/Mutable/Source/CustomizableObjectEditor/Public/MuCOE/Nodes/CustomizableObjectNodeComponent.h
new file mode 100644
--- /dev/null
+++ b/Engine/Plugins/Experimental/Mutable/Source/CustomizableObjectEditor/Public/MuCOE/Nodes/CustomizableObjectNodeComponent.h
@@ -0,0 +1,20 @@
+// Copyright Epic Games, Inc. All Rights Reserved.
+
+#pragma once
+
+#include "MuCOE/Nodes/CustomizableObjectNode.h"
+
+#include "CustomizableObjectNodeComponent.generated.h"
+
+class UEdGraphPin;
+
+
+UCLASS(abstract)
+class CUSTOMIZABLEOBJECTEDITOR_API UCustomizableObjectNodeComponent : public UCustomizableObjectNode
+{
+public:
+	GENERATED_BODY()
+
+
+};
+
diff --git a/Engine/Plugins/Experimental/Mutable/Source/CustomizableObjectEditor/Public/MuCOE/Nodes/CustomizableObjectNodeComponentMesh.h b/Engine/Plugins/Experimental/Mutable/Source/CustomizableObjectEditor/Public/MuCOE/Nodes/CustomizableObjectNodeComponentMesh.h
new file mode 100644
--- /dev/null
+++ b/Engine/Plugins/Experimental/Mutable/Source/CustomizableObjectEditor/Public/MuCOE/Nodes/CustomizableObjectNodeComponentMesh.h
@@ -0,0 +1,105 @@
+// Copyright Epic Games, Inc. All Rights Reserved.
+
+#pragma once
+
+#include "MuCOE/Nodes/CustomizableObjectNodeComponent.h"
+#include "MuCOE/RemapPins/CustomizableObjectNodeRemapPinsByNameDefaultPin.h"
+#include "SGraphNode.h"
+#include "GameplayTagContainer.h"
+
+#include "CustomizableObjectNodeComponentMesh.generated.h"
+
+namespace ENodeTitleType { enum Type : int; }
+
+class FArchive;
+class FAssetThumbnail;
+class FAssetThumbnailPool;
+class FSkeletalMeshModel;
+class ISinglePropertyView;
+class UAnimInstance;
+class UCustomizableObjectLayout;
+class UCustomizableObjectNodeRemapPins;
+class UCustomizableObjectNodeRemapPinsByName;
+class UMaterialInterface;
+class UObject;
+class USkeletalMesh;
+class UTexture2D;
+struct FPropertyChangedEvent;
+struct FSkeletalMaterial;
+
+
+/** PinData of a pin that belongs to a Skeletal Mesh Section. */
+UCLASS()
+class UCustomizableObjectNodeComponentMeshPinDataSection : public UCustomizableObjectNodePinData
+{
+	GENERATED_BODY()
+
+public:
+	void Init(int32 InLODIndex, int32 InSectionIndex);
+
+	int32 GetLODIndex() const;
+
+	int32 GetSectionIndex() const;
+
+protected:
+	virtual bool Equals(const UCustomizableObjectNodePinData& Other) const override;
+
+private:
+	UPROPERTY()
+	int32 LODIndex = -1;
+
+	UPROPERTY()
+	int32 SectionIndex = -1;
+};
+
+
+/** PinData of a Mesh pin. */
+UCLASS()
+class UCustomizableObjectNodeComponentMeshPinDataMaterial : public UCustomizableObjectNodeComponentMeshPinDataSection
+{
+	GENERATED_BODY()
+};
+
+
+UCLASS()
+class CUSTOMIZABLEOBJECTEDITOR_API UCustomizableObjectNodeComponentMesh : public UCustomizableObjectNodeComponent
+{
+public:
+	GENERATED_BODY()
+		
+	UPROPERTY(EditAnywhere, Category = Mesh, Meta = (DisplayName = Mesh, AllowedClasses = "/Script/Engine.StaticMesh, /Script/Engine.SkeletalMesh"))
+	FSoftObjectPath Mesh;
+
+	// UObject interface.
+	virtual void PostEditChangeProperty(FPropertyChangedEvent& PropertyChangedEvent) override;
+	
+	// UEdGraphNode interface
+	virtual FText GetNodeTitle(ENodeTitleType::Type TitleType) const override;
+	virtual FLinearColor GetNodeTitleColor() const override;
+	virtual FText GetTooltipText() const override;
+
+	// UCustomizableObjectNode interface
+	virtual void AllocateDefaultPins(UCustomizableObjectNodeRemapPins* RemapPins) override;
+	//virtual UCustomizableObjectNodeRemapPins* CreateRemapPinsDefault() const override;
+
+	// UCustomizableObjectNodeComponent interface
+
+
+	// Own interface
+	
+	void GetPinSection(const UEdGraphPin& Pin, int32& OutLODIndex, int32& OutSectionIndex) const;
+
+	/** Find the pin for a given lod and section. */
+	UEdGraphPin* GetMaterialPin(const int32 LODIndex, const int32 SectionIndex) const;
+
+	/** Returns the material associated to the given output pin. */
+	UMaterialInterface* GetMaterialFor(const UEdGraphPin* Pin) const;
+	FSkeletalMaterial* GetSkeletalMaterialFor(const UEdGraphPin& Pin) const;
+
+private:
+	UMaterialInterface* GetMaterialInterfaceFor(const int32 LODIndex, const int32 MaterialIndex, const FSkeletalMeshModel* ImportedModel = nullptr) const;
+	FSkeletalMaterial* GetSkeletalMaterialFor(const int32 LODIndex, const int32 MaterialIndex, const FSkeletalMeshModel* ImportedModel = nullptr) const;
+
+	static const FName OutputPinName;
+
+};
diff --git a/Engine/Plugins/Experimental/Mutable/Source/CustomizableObjectEditor/Public/MuCOE/Nodes/CustomizableObjectNodeExtensionDataSwitch.h b/Engine/Plugins/Experimental/Mutable/Source/CustomizableObjectEditor/Public/MuCOE/Nodes/CustomizableObjectNodeExtensionDataSwitch.h
--- a/Engine/Plugins/Experimental/Mutable/Source/CustomizableObjectEditor/Public/MuCOE/Nodes/CustomizableObjectNodeExtensionDataSwitch.h
+++ b/Engine/Plugins/Experimental/Mutable/Source/CustomizableObjectEditor/Public/MuCOE/Nodes/CustomizableObjectNodeExtensionDataSwitch.h
@@ -26,6 +26,6 @@ public:
 	//~End UCustomizableObjectNodeSwitchBase interface
 
 	//~Begin ICustomizableObjectExtensionNode interface
-	mu::NodeExtensionDataPtr GenerateMutableNode(class FExtensionDataCompilerInterface& InCompilerInterface) const override;
+	mu::Ptr<mu::NodeExtensionData> GenerateMutableNode(class FExtensionDataCompilerInterface& InCompilerInterface) const override;
 	//~End ICustomizableObjectExtensionNode interface
 };
diff --git a/Engine/Plugins/Experimental/Mutable/Source/CustomizableObjectEditor/Public/MuCOE/Nodes/CustomizableObjectNodeExtensionDataVariation.h b/Engine/Plugins/Experimental/Mutable/Source/CustomizableObjectEditor/Public/MuCOE/Nodes/CustomizableObjectNodeExtensionDataVariation.h
--- a/Engine/Plugins/Experimental/Mutable/Source/CustomizableObjectEditor/Public/MuCOE/Nodes/CustomizableObjectNodeExtensionDataVariation.h
+++ b/Engine/Plugins/Experimental/Mutable/Source/CustomizableObjectEditor/Public/MuCOE/Nodes/CustomizableObjectNodeExtensionDataVariation.h
@@ -40,7 +40,7 @@ public:
 	//~End UCustomizableObjectNode interface
 
 	//~Begin ICustomizableObjectExtensionNode interface
-	virtual mu::NodeExtensionDataPtr GenerateMutableNode(FExtensionDataCompilerInterface& InCompilerInterface) const override;
+	virtual mu::Ptr<mu::NodeExtensionData> GenerateMutableNode(FExtensionDataCompilerInterface& InCompilerInterface) const override;
 	//~End ICustomizableObjectExtensionNode interface
 
 	virtual FName GetCategory() const PURE_VIRTUAL(GetCategory(), return FName();)
diff --git a/Engine/Plugins/Experimental/Mutable/Source/CustomizableObjectEditor/Public/MuCOE/Nodes/CustomizableObjectNodeObject.h b/Engine/Plugins/Experimental/Mutable/Source/CustomizableObjectEditor/Public/MuCOE/Nodes/CustomizableObjectNodeObject.h
--- a/Engine/Plugins/Experimental/Mutable/Source/CustomizableObjectEditor/Public/MuCOE/Nodes/CustomizableObjectNodeObject.h
+++ b/Engine/Plugins/Experimental/Mutable/Source/CustomizableObjectEditor/Public/MuCOE/Nodes/CustomizableObjectNodeObject.h
@@ -250,6 +250,11 @@ public:
 		return Count;
 	}
 
+	UEdGraphPin* ComponentsPin() const
+	{
+		return FindPin(ComponentsPinName);
+	}
+
 	UEdGraphPin* ChildrenPin() const
 	{
 		return FindPin(ChildrenPinName);
@@ -260,7 +265,7 @@ public:
 		return FindPin(OutputPinName);
 	}
 
-	/** Return the LOD which a LOD pin references to. Retrun -1 if a pin does not belong to any LOD. */
+	/** Return the LOD which a LOD pin references to. Return -1 if a pin does not belong to any LOD. */
 	int32 GetLOD(UEdGraphPin* Pin) const;
 
 	virtual bool CanUserDeleteNode() const override;
@@ -288,6 +293,7 @@ public:
 
 private:
 	static const FName ChildrenPinName;
+	static const FName ComponentsPinName;
 	static const FName OutputPinName;
 	static const TCHAR* LODPinNamePrefix;
 
diff --git a/Engine/Plugins/Experimental/Mutable/Source/MutableRuntime/Private/MuR/CodeRunner.cpp b/Engine/Plugins/Experimental/Mutable/Source/MutableRuntime/Private/MuR/CodeRunner.cpp
--- a/Engine/Plugins/Experimental/Mutable/Source/MutableRuntime/Private/MuR/CodeRunner.cpp
+++ b/Engine/Plugins/Experimental/Mutable/Source/MutableRuntime/Private/MuR/CodeRunner.cpp
@@ -562,8 +562,7 @@ namespace mu
 
                 if ( args.value )
                 {
-                    Ptr<const String> value =
-                        LoadString( FCacheAddress( args.value, item ) );
+                    Ptr<const String> value = LoadString( FCacheAddress( args.value, item ) );
 
                     OP::ADDRESS nameAd = args.name;
                     check( nameAd < (uint32)Program.m_constantStrings.Num() );
@@ -595,8 +594,8 @@ namespace mu
 
             case 1:
             {
-                InstancePtrConst pBase = LoadInstance( FCacheAddress(args.instance,item) );
-                InstancePtr pResult;
+				Ptr<const Instance> pBase = LoadInstance( FCacheAddress(args.instance,item) );
+				Ptr<Instance> pResult;
                 if (!pBase)
                 {
                     pResult = new Instance();
@@ -608,26 +607,16 @@ namespace mu
 
                 if ( args.value )
                 {
-                    InstancePtrConst pComp = LoadInstance( FCacheAddress(args.value,item) );
+                    Ptr<const Instance> pComp = LoadInstance( FCacheAddress(args.value,item) );
 
-                    int cindex = pResult->GetPrivate()->AddComponent( 0 );
+                    int32 NewComponentIndex = pResult->GetPrivate()->AddComponent();
 
-                    if ( !pComp->GetPrivate()->Lods.IsEmpty()
-                         &&
-                         !pResult->GetPrivate()->Lods.IsEmpty()
-                         &&
-                         !pComp->GetPrivate()->Lods[0].Components.IsEmpty() )
+                    if ( !pComp->GetPrivate()->Components.IsEmpty() )
                     {
-                        pResult->GetPrivate()->Lods[0].Components[cindex] =
-                                pComp->GetPrivate()->Lods[0].Components[0];
-
-                    	pResult->GetPrivate()->Lods[0].Components[cindex].Id = args.id;
-                    	
-                        // Name
-                        OP::ADDRESS nameAd = args.name;
-                        check( nameAd < (uint32)Program.m_constantStrings.Num() );
-                        const FString& Name = Program.m_constantStrings[ nameAd ];
-                        pResult->GetPrivate()->SetComponentName( 0, cindex, FName(Name) );
+                        pResult->GetPrivate()->Components[NewComponentIndex] = pComp->GetPrivate()->Components[0];
+
+						// Id
+                    	pResult->GetPrivate()->Components[NewComponentIndex].Id = args.id;
                     }
                 }
                 StoreInstance( item, pResult );
@@ -654,9 +643,9 @@ namespace mu
 
             case 1:
             {
-                InstancePtrConst pBase = LoadInstance( FCacheAddress(args.instance,item) );
+				Ptr<const Instance> pBase = LoadInstance( FCacheAddress(args.instance,item) );
 
-                InstancePtr pResult;
+				Ptr<Instance> pResult;
 				if (pBase)
 				{
 					pResult = mu::CloneOrTakeOver<Instance>(pBase.get());
@@ -670,24 +659,21 @@ namespace mu
                 // additional information like internal or external IDs
                 //if ( args.value )
                 {
-                    InstancePtrConst pSurf = LoadInstance( FCacheAddress(args.value,item) );
+					Ptr<const Instance> pSurf = LoadInstance( FCacheAddress(args.value,item) );
 
                     int sindex = pResult->GetPrivate()->AddSurface( 0, 0 );
 
                     // Surface data
                     if (pSurf
                             &&
-                            pSurf->GetPrivate()->Lods.Num()
+                            pSurf->GetPrivate()->Components.Num()
                             &&
-                            pSurf->GetPrivate()->Lods[0].Components.Num()
+                            pSurf->GetPrivate()->Components[0].LODs.Num()
                             &&
-                            pSurf->GetPrivate()->Lods[0].Components[0].Surfaces.Num())
+                            pSurf->GetPrivate()->Components[0].LODs[0].Surfaces.Num())
                     {
-                        pResult->GetPrivate()->Lods[0].Components[0].Surfaces[sindex] =
-                            pSurf->GetPrivate()->Lods[0].Components[0].Surfaces[0];
-
-                        // Meshes must be added later.
-                        check(!pSurf->GetPrivate()->Lods[0].Components[0].Meshes.Num());
+                        pResult->GetPrivate()->Components[0].LODs[0].Surfaces[sindex] =
+                            pSurf->GetPrivate()->Components[0].LODs[0].Surfaces[0];
                     }
 
                     // Name
@@ -697,9 +683,9 @@ namespace mu
                     pResult->GetPrivate()->SetSurfaceName( 0, 0, sindex, FName(Name) );
 
                     // IDs
-                    pResult->GetPrivate()->Lods[0].Components[0].Surfaces[sindex].InternalId = args.id;
-                    pResult->GetPrivate()->Lods[0].Components[0].Surfaces[sindex].ExternalId = args.ExternalId;
-                    pResult->GetPrivate()->Lods[0].Components[0].Surfaces[sindex].SharedId = args.SharedSurfaceId;
+                    pResult->GetPrivate()->Components[0].LODs[0].Surfaces[sindex].InternalId = args.id;
+                    pResult->GetPrivate()->Components[0].LODs[0].Surfaces[sindex].ExternalId = args.ExternalId;
+                    pResult->GetPrivate()->Components[0].LODs[0].Surfaces[sindex].SharedId = args.SharedSurfaceId;
                 }
                 StoreInstance( item, pResult );
                 break;
@@ -741,31 +727,30 @@ namespace mu
             case 1:
             {
                 // Assemble result
-                InstancePtr pResult = new Instance();
+				Ptr<Instance> pResult = new Instance();
+				int32 ComponentIndex = pResult->GetPrivate()->AddComponent();
 
-                for ( int i=0; i<MUTABLE_OP_MAX_ADD_COUNT; ++i )
+                for ( int32 i=0; i<MUTABLE_OP_MAX_ADD_COUNT; ++i )
                 {
                     if ( args.lod[i] )
                     {
-                        bool selectedLod = ( (1<<i) & lodMask ) != 0;
+                        bool bIsSelectedLod = ( (1<<i) & lodMask ) != 0;
 
-                        if ( selectedLod )
+						// Add an empty LOD even if not selected.
+						int32 LODIndex = pResult->GetPrivate()->AddLOD(ComponentIndex);
+						
+						if (bIsSelectedLod)
                         {
-                            InstancePtrConst pLOD = LoadInstance( FCacheAddress(args.lod[i],item) );
-
-                            int LODIndex = pResult->GetPrivate()->AddLOD();
+							Ptr<const Instance> pLOD = LoadInstance( FCacheAddress(args.lod[i],item) );
 
                             // In a degenerated case, the returned pLOD may not have an LOD inside
-                            if ( pLOD && !pLOD->GetPrivate()->Lods.IsEmpty() )
-                            {
-                                pResult->GetPrivate()->Lods[LODIndex] = pLOD->GetPrivate()->Lods[0];
-                            }
-                        }
-                        else
-                        {
-                            // LOD not selected. Add an empty one
-                            pResult->GetPrivate()->AddLOD();
-                        }
+ 							if (!pLOD->GetPrivate()->Components.IsEmpty()
+								&&
+								!pLOD->GetPrivate()->Components[0].LODs.IsEmpty())
+							{
+								pResult->GetPrivate()->Components[ComponentIndex].LODs[LODIndex] = pLOD->GetPrivate()->Components[0].LODs[0];
+							}
+						}
                     }
                 }
 
@@ -855,11 +840,8 @@ namespace mu
             switch (item.Stage)
             {
             case 0:
-                    AddOp( FScheduledOp( item.At, item, 1),
-                           FScheduledOp( args.instance, item) );
-
-                // We don't build the resources when building instance: just store ids for them.
-                //PushIfNotVisited(args.value, item);
+				// We don't build the resources when building instance: just store ids for them.
+				AddOp( FScheduledOp( item.At, item, 1), FScheduledOp( args.instance, item) );
                 break;
 
             case 1:
@@ -881,7 +863,7 @@ namespace mu
 					OP::ADDRESS NameAd = args.name;
 					check(NameAd < (uint32)Program.m_constantStrings.Num());
 					const FString& Name = Program.m_constantStrings[NameAd];
-					pResult->GetPrivate()->AddMesh(0, 0, MeshId, FName(Name));
+					pResult->GetPrivate()->SetMesh(0, 0, MeshId, FName(Name));
                 }
                 StoreInstance( item, pResult );
                 break;
diff --git a/Engine/Plugins/Experimental/Mutable/Source/MutableRuntime/Private/MuR/Instance.cpp b/Engine/Plugins/Experimental/Mutable/Source/MutableRuntime/Private/MuR/Instance.cpp
--- a/Engine/Plugins/Experimental/Mutable/Source/MutableRuntime/Private/MuR/Instance.cpp
+++ b/Engine/Plugins/Experimental/Mutable/Source/MutableRuntime/Private/MuR/Instance.cpp
@@ -53,7 +53,7 @@ namespace mu
 	//---------------------------------------------------------------------------------------------
 	int32 Instance::GetDataSize() const
 	{
-		return 16 + sizeof(Private) + m_pD->Lods.GetAllocatedSize() + m_pD->ExtensionData.GetAllocatedSize();
+		return 16 + sizeof(Private) + m_pD->Components.GetAllocatedSize() + m_pD->ExtensionData.GetAllocatedSize();
 	}
 
 
@@ -65,32 +65,30 @@ namespace mu
 
 
     //---------------------------------------------------------------------------------------------
-    int32 Instance::GetLODCount() const
+    int32 Instance::GetComponentCount() const
     {
-		return m_pD->Lods.Num();
+		return m_pD->Components.Num();
     }
 
 
 	//---------------------------------------------------------------------------------------------
-	int32 Instance::GetComponentCount( int32 LODIndex ) const
+	int32 Instance::GetLODCount( int32 ComponentIndex ) const
 	{
-		check(LODIndex >= 0 && LODIndex < m_pD->Lods.Num());
-		if (LODIndex >= 0 && LODIndex < m_pD->Lods.Num())
+		if (m_pD->Components.IsValidIndex(ComponentIndex))
 		{
-			return m_pD->Lods[LODIndex].Components.Num();
+			return m_pD->Components[ComponentIndex].LODs.Num();
 		}
-
+		check(false);
 		return 0;
 	}
 
 	
 	//---------------------------------------------------------------------------------------------
-	uint16 Instance::GetComponentId( int32 LODIndex, int32 ComponentIndex ) const
+	uint16 Instance::GetComponentId( int32 ComponentIndex ) const
 	{
-		if ( LODIndex>=0 && LODIndex<m_pD->Lods.Num() &&
-			 ComponentIndex>=0 && ComponentIndex<m_pD->Lods[LODIndex].Components.Num() )
+		if (m_pD->Components.IsValidIndex(ComponentIndex))
 		{
-			return m_pD->Lods[LODIndex].Components[ComponentIndex].Id;
+			return m_pD->Components[ComponentIndex].Id;
 		}
 		else
 		{
@@ -102,12 +100,12 @@ namespace mu
 
 
     //---------------------------------------------------------------------------------------------
-    int32 Instance::GetSurfaceCount( int32 LODIndex, int32 ComponentIndex ) const
+    int32 Instance::GetSurfaceCount( int32 ComponentIndex, int32 LODIndex ) const
     {
-		if (LODIndex >= 0 && LODIndex < m_pD->Lods.Num() &&
-			ComponentIndex >= 0 && ComponentIndex < m_pD->Lods[LODIndex].Components.Num())
+		if (m_pD->Components.IsValidIndex(ComponentIndex) &&
+			m_pD->Components[ComponentIndex].LODs.IsValidIndex(LODIndex))
 		{
-			return m_pD->Lods[LODIndex].Components[ComponentIndex].Surfaces.Num();
+			return m_pD->Components[ComponentIndex].LODs[LODIndex].Surfaces.Num();
 		}
 		else
 		{
@@ -119,13 +117,13 @@ namespace mu
 
 
     //---------------------------------------------------------------------------------------------
-    uint32 Instance::GetSurfaceId( int32 LODIndex, int32 ComponentIndex, int32 SurfaceIndex ) const
+    uint32 Instance::GetSurfaceId(int32 ComponentIndex, int32 LODIndex, int32 SurfaceIndex ) const
     {
-        if ( LODIndex>=0 && LODIndex<m_pD->Lods.Num() &&
-             ComponentIndex>=0 && ComponentIndex<m_pD->Lods[LODIndex].Components.Num() &&
-             SurfaceIndex>=0 && SurfaceIndex<m_pD->Lods[LODIndex].Components[ComponentIndex].Surfaces.Num() )
+        if (m_pD->Components.IsValidIndex(ComponentIndex) &&
+			m_pD->Components[ComponentIndex].LODs.IsValidIndex(LODIndex) &&
+			m_pD->Components[ComponentIndex].LODs[LODIndex].Surfaces.IsValidIndex(SurfaceIndex) )
         {
-            return m_pD->Lods[LODIndex].Components[ComponentIndex].Surfaces[SurfaceIndex].InternalId;
+            return m_pD->Components[ComponentIndex].LODs[LODIndex].Surfaces[SurfaceIndex].InternalId;
         }
 		else
 		{
@@ -137,14 +135,14 @@ namespace mu
 
 
     //---------------------------------------------------------------------------------------------
-    int32 Instance::FindSurfaceById( int32 LODIndex, int32 ComponentIndex, uint32 id ) const
+    int32 Instance::FindSurfaceById(int32 ComponentIndex, int32 LODIndex, uint32 id ) const
     {
-		if (LODIndex >= 0 && LODIndex < m_pD->Lods.Num() &&
-			ComponentIndex >= 0 && ComponentIndex < m_pD->Lods[LODIndex].Components.Num())
+		if (m_pD->Components.IsValidIndex(ComponentIndex) &&
+			m_pD->Components[ComponentIndex].LODs.IsValidIndex(LODIndex))
 		{
-			for (int32 i = 0; i < m_pD->Lods[LODIndex].Components[ComponentIndex].Surfaces.Num(); ++i)
+			for (int32 i = 0; i < m_pD->Components[ComponentIndex].LODs[LODIndex].Surfaces.Num(); ++i)
 			{
-				if (m_pD->Lods[LODIndex].Components[ComponentIndex].Surfaces[i].InternalId == id)
+				if (m_pD->Components[ComponentIndex].LODs[LODIndex].Surfaces[i].InternalId == id)
 				{
 					return i;
 				}
@@ -162,13 +160,14 @@ namespace mu
 	//---------------------------------------------------------------------------------------------
 	void Instance::FindBaseSurfaceBySharedId(int32 CompIndex, int32 SharedId, int32& OutSurfaceIndex, int32& OutLODIndex) const
 	{
-		for (int32 LodIndex = 0; LodIndex < m_pD->Lods.Num(); LodIndex++)
+		if (m_pD->Components.IsValidIndex(CompIndex))
 		{
-			if (m_pD->Lods[LodIndex].Components.IsValidIndex(CompIndex))
+			for (int32 LodIndex = 0; LodIndex < m_pD->Components[CompIndex].LODs.Num(); LodIndex++)
 			{
-				for (int32 SurfaceIndex = 0; SurfaceIndex < m_pD->Lods[LodIndex].Components[CompIndex].Surfaces.Num(); ++SurfaceIndex)
+				FInstanceLOD& LOD = m_pD->Components[CompIndex].LODs[LodIndex];
+				for (int32 SurfaceIndex = 0; SurfaceIndex < LOD.Surfaces.Num(); ++SurfaceIndex)
 				{
-					if (m_pD->Lods[LodIndex].Components[CompIndex].Surfaces[SurfaceIndex].SharedId == SharedId)
+					if (LOD.Surfaces[SurfaceIndex].SharedId == SharedId)
 					{
 						OutSurfaceIndex = SurfaceIndex;
 						OutLODIndex = LodIndex;
@@ -185,13 +184,13 @@ namespace mu
 
 
 	//---------------------------------------------------------------------------------------------
-	int32 Instance::GetSharedSurfaceId(int32 LodIndex, int32 CompIndex, int32 SurfaceIndex) const
+	int32 Instance::GetSharedSurfaceId(int32 ComponentIndex, int32 LODIndex, int32 SurfaceIndex) const
 	{
-		if (LodIndex >= 0 && LodIndex < m_pD->Lods.Num() &&
-			CompIndex >= 0 && CompIndex < m_pD->Lods[LodIndex].Components.Num() &&
-			SurfaceIndex >= 0 && SurfaceIndex < m_pD->Lods[LodIndex].Components[CompIndex].Surfaces.Num())
+		if (m_pD->Components.IsValidIndex(ComponentIndex) &&
+			m_pD->Components[ComponentIndex].LODs.IsValidIndex(LODIndex) &&
+			m_pD->Components[ComponentIndex].LODs[LODIndex].Surfaces.IsValidIndex(SurfaceIndex))
 		{
-			return m_pD->Lods[LodIndex].Components[CompIndex].Surfaces[SurfaceIndex].SharedId;
+			return m_pD->Components[ComponentIndex].LODs[LODIndex].Surfaces[SurfaceIndex].SharedId;
 		}
 		else
 		{
@@ -203,13 +202,13 @@ namespace mu
 	
 
     //---------------------------------------------------------------------------------------------
-    uint32 Instance::GetSurfaceCustomId( int32 LODIndex, int32 ComponentIndex, int32 SurfaceIndex ) const
+    uint32 Instance::GetSurfaceCustomId(int32 ComponentIndex, int32 LODIndex, int32 SurfaceIndex ) const
     {
-        if ( LODIndex>=0 && LODIndex<m_pD->Lods.Num() &&
-             ComponentIndex>=0 && ComponentIndex<m_pD->Lods[LODIndex].Components.Num() &&
-             SurfaceIndex>=0 && SurfaceIndex<m_pD->Lods[LODIndex].Components[ComponentIndex].Surfaces.Num() )
+        if (m_pD->Components.IsValidIndex(ComponentIndex) &&
+			m_pD->Components[ComponentIndex].LODs.IsValidIndex(LODIndex) &&
+			m_pD->Components[ComponentIndex].LODs[LODIndex].Surfaces.IsValidIndex(SurfaceIndex))
         {
-            return m_pD->Lods[LODIndex].Components[ComponentIndex].Surfaces[SurfaceIndex].ExternalId;
+            return m_pD->Components[ComponentIndex].LODs[LODIndex].Surfaces[SurfaceIndex].ExternalId;
         }
 		else
 		{
@@ -221,164 +220,143 @@ namespace mu
 
 
 	//---------------------------------------------------------------------------------------------
-    int32 Instance::GetMeshCount( int32 LODIndex, int32 ComponentIndex ) const
+    int32 Instance::GetImageCount(int32 ComponentIndex, int32 LODIndex, int32 SurfaceIndex ) const
 	{
-		check(LODIndex >= 0 && LODIndex < m_pD->Lods.Num());
-		check(ComponentIndex >= 0 && ComponentIndex < m_pD->Lods[LODIndex].Components.Num());
+		check(m_pD->Components.IsValidIndex(ComponentIndex));
+		check(m_pD->Components[ComponentIndex].LODs.IsValidIndex(LODIndex));
+		check(m_pD->Components[ComponentIndex].LODs[LODIndex].Surfaces.IsValidIndex(SurfaceIndex));
 
-		return m_pD->Lods[LODIndex].Components[ComponentIndex].Meshes.Num();
+		return m_pD->Components[ComponentIndex].LODs[LODIndex].Surfaces[SurfaceIndex].Images.Num();
 	}
 
 
 	//---------------------------------------------------------------------------------------------
-    int32 Instance::GetImageCount( int32 LODIndex, int32 ComponentIndex, int32 SurfaceIndex ) const
+    int32 Instance::GetVectorCount(int32 ComponentIndex, int32 LODIndex, int32 SurfaceIndex ) const
 	{
-		check(LODIndex >= 0 && LODIndex < m_pD->Lods.Num());
-		check(ComponentIndex >= 0 && ComponentIndex < m_pD->Lods[LODIndex].Components.Num());
-		check(SurfaceIndex >= 0 && SurfaceIndex < m_pD->Lods[LODIndex].Components[ComponentIndex].Surfaces.Num());
+		check(m_pD->Components.IsValidIndex(ComponentIndex));
+		check(m_pD->Components[ComponentIndex].LODs.IsValidIndex(LODIndex));
+		check(m_pD->Components[ComponentIndex].LODs[LODIndex].Surfaces.IsValidIndex(SurfaceIndex));
 
-		return m_pD->Lods[LODIndex].Components[ComponentIndex].Surfaces[SurfaceIndex].Images.Num();
+		return m_pD->Components[ComponentIndex].LODs[LODIndex].Surfaces[SurfaceIndex].Vectors.Num();
 	}
 
 
 	//---------------------------------------------------------------------------------------------
-    int32 Instance::GetVectorCount( int32 LODIndex, int32 ComponentIndex, int32 SurfaceIndex ) const
+    int32 Instance::GetScalarCount(int32 ComponentIndex, int32 LODIndex, int32 SurfaceIndex ) const
 	{
-		check(LODIndex >= 0 && LODIndex < m_pD->Lods.Num());
-		check(ComponentIndex >= 0 && ComponentIndex < m_pD->Lods[LODIndex].Components.Num());
-		check(SurfaceIndex >= 0 && SurfaceIndex < m_pD->Lods[LODIndex].Components[ComponentIndex].Surfaces.Num());
+		check(m_pD->Components.IsValidIndex(ComponentIndex));
+		check(m_pD->Components[ComponentIndex].LODs.IsValidIndex(LODIndex));
+		check(m_pD->Components[ComponentIndex].LODs[LODIndex].Surfaces.IsValidIndex(SurfaceIndex));
 
-		return m_pD->Lods[LODIndex].Components[ComponentIndex].Surfaces[SurfaceIndex].Vectors.Num();
-	}
-
-
-	//---------------------------------------------------------------------------------------------
-    int32 Instance::GetScalarCount( int32 LODIndex, int32 ComponentIndex, int32 SurfaceIndex ) const
-	{
-		check(LODIndex >= 0 && LODIndex < m_pD->Lods.Num());
-		check(ComponentIndex >= 0 && ComponentIndex < m_pD->Lods[LODIndex].Components.Num());
-		check(SurfaceIndex >= 0 && SurfaceIndex < m_pD->Lods[LODIndex].Components[ComponentIndex].Surfaces.Num());
-
-		return m_pD->Lods[LODIndex].Components[ComponentIndex].Surfaces[SurfaceIndex].Scalars.Num();
+		return m_pD->Components[ComponentIndex].LODs[LODIndex].Surfaces[SurfaceIndex].Scalars.Num();
 	}
 
 
     //---------------------------------------------------------------------------------------------
-    int32 Instance::GetStringCount( int32 LODIndex, int32 ComponentIndex, int32 SurfaceIndex ) const
+    int32 Instance::GetStringCount(int32 ComponentIndex, int32 LODIndex, int32 SurfaceIndex ) const
     {
-		check(LODIndex >= 0 && LODIndex < m_pD->Lods.Num());
-		check(ComponentIndex >= 0 && ComponentIndex < m_pD->Lods[LODIndex].Components.Num());
-		check(SurfaceIndex >= 0 && SurfaceIndex < m_pD->Lods[LODIndex].Components[ComponentIndex].Surfaces.Num());
+		check(m_pD->Components.IsValidIndex(ComponentIndex));
+		check(m_pD->Components[ComponentIndex].LODs.IsValidIndex(LODIndex));
+		check(m_pD->Components[ComponentIndex].LODs[LODIndex].Surfaces.IsValidIndex(SurfaceIndex));
 
-		return m_pD->Lods[LODIndex].Components[ComponentIndex].Surfaces[SurfaceIndex].Strings.Num();
+		return m_pD->Components[ComponentIndex].LODs[LODIndex].Surfaces[SurfaceIndex].Strings.Num();
 	}
 
 
     //---------------------------------------------------------------------------------------------
-	FResourceID Instance::GetMeshId( int32 LODIndex, int32 ComponentIndex, int32 mesh ) const
+	FResourceID Instance::GetMeshId(int32 ComponentIndex, int32 LODIndex ) const
     {
-        check( LODIndex>=0 && LODIndex<m_pD->Lods.Num() );
-        check( ComponentIndex>=0 && ComponentIndex<m_pD->Lods[LODIndex].Components.Num() );
-        check( mesh>=0 && mesh<m_pD->Lods[LODIndex].Components[ComponentIndex].Meshes.Num() );
+        check(m_pD->Components.IsValidIndex(ComponentIndex));
+        check(m_pD->Components[ComponentIndex].LODs.IsValidIndex(LODIndex));
 
-		FResourceID result = m_pD->Lods[LODIndex].Components[ComponentIndex].Meshes[mesh].Id;
-        return result;
+		return m_pD->Components[ComponentIndex].LODs[LODIndex].MeshId;
     }
 
 
 	//---------------------------------------------------------------------------------------------
-	FResourceID Instance::GetImageId( int32 LODIndex, int32 ComponentIndex, int32 SurfaceIndex, int32 img ) const
+	FResourceID Instance::GetImageId(int32 ComponentIndex, int32 LODIndex, int32 SurfaceIndex, int32 ImageIndex) const
 	{
-		check( LODIndex>=0 && LODIndex<m_pD->Lods.Num() );
-		check( ComponentIndex>=0 && ComponentIndex<m_pD->Lods[LODIndex].Components.Num() );
-        check( SurfaceIndex>=0 && SurfaceIndex<m_pD->Lods[LODIndex].Components[ComponentIndex].Surfaces.Num() );
-        check( img>=0 && img<m_pD->Lods[LODIndex].Components[ComponentIndex].Surfaces[SurfaceIndex].Images.Num() );
+		check(m_pD->Components.IsValidIndex(ComponentIndex));
+		check(m_pD->Components[ComponentIndex].LODs.IsValidIndex(LODIndex));
+        check(m_pD->Components[ComponentIndex].LODs[LODIndex].Surfaces.IsValidIndex(SurfaceIndex));
+        check(m_pD->Components[ComponentIndex].LODs[LODIndex].Surfaces[SurfaceIndex].Images.IsValidIndex(ImageIndex));
 
-		FResourceID result = m_pD->Lods[LODIndex].Components[ComponentIndex].Surfaces[SurfaceIndex].Images[img].Id;
-        return result;
+		return m_pD->Components[ComponentIndex].LODs[LODIndex].Surfaces[SurfaceIndex].Images[ImageIndex].Id;
 	}
 
 
 	//---------------------------------------------------------------------------------------------
-    FName Instance::GetImageName( int32 LODIndex, int32 ComponentIndex, int32 SurfaceIndex, int32 img ) const
+    FName Instance::GetImageName(int32 ComponentIndex, int32 LODIndex, int32 SurfaceIndex, int32 ImageIndex ) const
 	{
-		check( LODIndex>=0 && LODIndex<m_pD->Lods.Num() );
-		check( ComponentIndex>=0 && ComponentIndex<m_pD->Lods[LODIndex].Components.Num() );
-        check( SurfaceIndex>=0 && SurfaceIndex<m_pD->Lods[LODIndex].Components[ComponentIndex].Surfaces.Num() );
-        check( img>=0 && img<m_pD->Lods[LODIndex].Components[ComponentIndex].Surfaces[SurfaceIndex].Images.Num() );
+		check(m_pD->Components.IsValidIndex(ComponentIndex));
+		check(m_pD->Components[ComponentIndex].LODs.IsValidIndex(LODIndex));
+        check(m_pD->Components[ComponentIndex].LODs[LODIndex].Surfaces.IsValidIndex(SurfaceIndex));
+		check(m_pD->Components[ComponentIndex].LODs[LODIndex].Surfaces[SurfaceIndex].Images.IsValidIndex(ImageIndex));
 
-        return m_pD->Lods[LODIndex].Components[ComponentIndex].Surfaces[SurfaceIndex].Images[img].Name;
+        return m_pD->Components[ComponentIndex].LODs[LODIndex].Surfaces[SurfaceIndex].Images[ImageIndex].Name;
 	}
 
 
 	//---------------------------------------------------------------------------------------------
-	FVector4f Instance::GetVector( int32 LODIndex, int32 ComponentIndex, int32 SurfaceIndex, int32 vec ) const
+	FVector4f Instance::GetVector(int32 ComponentIndex, int32 LODIndex, int32 SurfaceIndex, int32 VectorIndex) const
 	{
-		check( LODIndex>=0 && LODIndex<m_pD->Lods.Num() );
-		check( ComponentIndex>=0 && ComponentIndex<m_pD->Lods[LODIndex].Components.Num() );
-        check( SurfaceIndex>=0 && SurfaceIndex<m_pD->Lods[LODIndex].Components[ComponentIndex].Surfaces.Num() );
-        check( vec>=0 && vec<m_pD->Lods[LODIndex].Components[ComponentIndex].Surfaces[SurfaceIndex].Vectors.Num() );
+		check(m_pD->Components.IsValidIndex(ComponentIndex));
+		check(m_pD->Components[ComponentIndex].LODs.IsValidIndex(LODIndex));
+        check(m_pD->Components[ComponentIndex].LODs[LODIndex].Surfaces.IsValidIndex(SurfaceIndex));
+		check(m_pD->Components[ComponentIndex].LODs[LODIndex].Surfaces[SurfaceIndex].Vectors.IsValidIndex(VectorIndex));
 
-        return m_pD->Lods[LODIndex].Components[ComponentIndex].Surfaces[SurfaceIndex].Vectors[vec].Value;
+        return m_pD->Components[ComponentIndex].LODs[LODIndex].Surfaces[SurfaceIndex].Vectors[VectorIndex].Value;
 	}
 
 
 	//---------------------------------------------------------------------------------------------
-	FName Instance::GetVectorName( int32 LODIndex, int32 ComponentIndex, int32 SurfaceIndex, int32 vec ) const
+	FName Instance::GetVectorName(int32 ComponentIndex, int32 LODIndex, int32 SurfaceIndex, int32 VectorIndex) const
 	{
-		check( LODIndex>=0 && LODIndex<m_pD->Lods.Num() );
-		check( ComponentIndex>=0 && ComponentIndex<m_pD->Lods[LODIndex].Components.Num() );
-        check( SurfaceIndex>=0 && SurfaceIndex<m_pD->Lods[LODIndex].Components[ComponentIndex].Surfaces.Num() );
-        check( vec>=0 && vec<m_pD->Lods[LODIndex].Components[ComponentIndex].Surfaces[SurfaceIndex].Vectors.Num() );
+		check(m_pD->Components.IsValidIndex(ComponentIndex));
+		check(m_pD->Components[ComponentIndex].LODs.IsValidIndex(LODIndex));
+        check(m_pD->Components[ComponentIndex].LODs[LODIndex].Surfaces.IsValidIndex(SurfaceIndex));
+		check(m_pD->Components[ComponentIndex].LODs[LODIndex].Surfaces[SurfaceIndex].Vectors.IsValidIndex(VectorIndex));
 
-        return m_pD->Lods[LODIndex].Components[ComponentIndex].Surfaces[SurfaceIndex].Vectors[vec].Name;
+        return m_pD->Components[ComponentIndex].LODs[LODIndex].Surfaces[SurfaceIndex].Vectors[VectorIndex].Name;
 	}
 
 
 	//---------------------------------------------------------------------------------------------
-    float Instance::GetScalar( int32 LODIndex, int32 ComponentIndex, int32 SurfaceIndex, int32 sca ) const
+    float Instance::GetScalar(int32 ComponentIndex, int32 LODIndex, int32 SurfaceIndex, int32 ScalarIndex ) const
 	{
-		check( LODIndex>=0 && LODIndex<m_pD->Lods.Num() );
-		check( ComponentIndex>=0 && ComponentIndex<m_pD->Lods[LODIndex].Components.Num() );
-        check( SurfaceIndex>=0 && SurfaceIndex<m_pD->Lods[LODIndex].Components[ComponentIndex].Surfaces.Num() );
-        check( sca>=0 && sca<m_pD->Lods[LODIndex].Components[ComponentIndex].Surfaces[SurfaceIndex].Scalars.Num() );
+		check(m_pD->Components.IsValidIndex(ComponentIndex));
+		check(m_pD->Components[ComponentIndex].LODs.IsValidIndex(LODIndex));
+        check(m_pD->Components[ComponentIndex].LODs[LODIndex].Surfaces.IsValidIndex(SurfaceIndex));
+		check(m_pD->Components[ComponentIndex].LODs[LODIndex].Surfaces[SurfaceIndex].Scalars.IsValidIndex(ScalarIndex));
 
-        return m_pD->Lods[LODIndex].Components[ComponentIndex].Surfaces[SurfaceIndex].Scalars[sca].Value;
+        return m_pD->Components[ComponentIndex].LODs[LODIndex].Surfaces[SurfaceIndex].Scalars[ScalarIndex].Value;
 	}
 
 
 	//---------------------------------------------------------------------------------------------
-	FName Instance::GetScalarName( int32 LODIndex, int32 ComponentIndex, int32 SurfaceIndex, int32 sca ) const
+	FName Instance::GetScalarName(int32 ComponentIndex, int32 LODIndex, int32 SurfaceIndex, int32 ScalarIndex) const
 	{
-		check( LODIndex>=0 && LODIndex<m_pD->Lods.Num() );
-		check( ComponentIndex>=0 && ComponentIndex<m_pD->Lods[LODIndex].Components.Num() );
-        check( SurfaceIndex>=0 && SurfaceIndex<m_pD->Lods[LODIndex].Components[ComponentIndex].Surfaces.Num() );
-        check( sca>=0 && sca<m_pD->Lods[LODIndex].Components[ComponentIndex].Surfaces[SurfaceIndex].Scalars.Num() );
+		check(m_pD->Components.IsValidIndex(ComponentIndex));
+		check(m_pD->Components[ComponentIndex].LODs.IsValidIndex(LODIndex));
+        check(m_pD->Components[ComponentIndex].LODs[LODIndex].Surfaces.IsValidIndex(SurfaceIndex));
+		check(m_pD->Components[ComponentIndex].LODs[LODIndex].Surfaces[SurfaceIndex].Scalars.IsValidIndex(ScalarIndex));
 
-        return m_pD->Lods[LODIndex].Components[ComponentIndex].Surfaces[SurfaceIndex].Scalars[sca].Name;
+        return m_pD->Components[ComponentIndex].LODs[LODIndex].Surfaces[SurfaceIndex].Scalars[ScalarIndex].Name;
 	}
 
 
     //---------------------------------------------------------------------------------------------
-    FString Instance::GetString( int32 LODIndex, int32 ComponentIndex, int32 SurfaceIndex, int32 str ) const
+    FString Instance::GetString(int32 ComponentIndex, int32 LODIndex, int32 SurfaceIndex, int32 StringIndex ) const
     {
-        check( LODIndex >= 0 && LODIndex < m_pD->Lods.Num() );
-        check( ComponentIndex >= 0 && ComponentIndex < m_pD->Lods[LODIndex].Components.Num() );
-        check( SurfaceIndex >= 0 &&
-                        SurfaceIndex < m_pD->Lods[LODIndex].Components[ComponentIndex].Surfaces.Num() );
+        check(m_pD->Components.IsValidIndex(ComponentIndex));
+        check(m_pD->Components[ComponentIndex].LODs.IsValidIndex(LODIndex));
+        check(m_pD->Components[ComponentIndex].LODs[LODIndex].Surfaces.IsValidIndex(SurfaceIndex));
+		check(m_pD->Components[ComponentIndex].LODs[LODIndex].Surfaces[SurfaceIndex].Strings.IsValidIndex(StringIndex));
 
-        bool valid =
-            str >= 0 &&
-            str < m_pD->Lods[LODIndex].Components[ComponentIndex].Surfaces[SurfaceIndex].Strings.Num();
-        check(valid);
-
-        if (valid)
+		bool bValid = m_pD->Components[ComponentIndex].LODs[LODIndex].Surfaces[SurfaceIndex].Strings.IsValidIndex(StringIndex);
+		if (bValid)
         {
-            return m_pD->Lods[LODIndex]
-				.Components[ComponentIndex]
-				.Surfaces[SurfaceIndex]
-				.Strings[str]
-				.Value;
+            return m_pD->Components[ComponentIndex].LODs[LODIndex].Surfaces[SurfaceIndex].Strings[StringIndex].Value;
         }
 
         return "";
@@ -386,20 +364,16 @@ namespace mu
 
 
     //---------------------------------------------------------------------------------------------
-	FName Instance::GetStringName( int32 LODIndex, int32 ComponentIndex, int32 SurfaceIndex, int32 str ) const
+	FName Instance::GetStringName(int32 ComponentIndex, int32 LODIndex, int32 SurfaceIndex, int32 StringIndex) const
     {
-        check( LODIndex >= 0 && LODIndex < m_pD->Lods.Num() );
-        check( ComponentIndex >= 0 && ComponentIndex < m_pD->Lods[LODIndex].Components.Num() );
-        check( SurfaceIndex >= 0 &&
-                        SurfaceIndex < m_pD->Lods[LODIndex].Components[ComponentIndex].Surfaces.Num() );
-        bool valid =
-            str >= 0 &&
-            str < m_pD->Lods[LODIndex].Components[ComponentIndex].Surfaces[SurfaceIndex].Strings.Num();
-        check( valid );
-
-        if (valid)
-        {
-            return m_pD->Lods[LODIndex].Components[ComponentIndex].Surfaces[SurfaceIndex].Strings[str].Name;
+        check(m_pD->Components.IsValidIndex(ComponentIndex));
+        check(m_pD->Components[ComponentIndex].LODs.IsValidIndex(LODIndex));
+        check(m_pD->Components[ComponentIndex].LODs[LODIndex].Surfaces.IsValidIndex(SurfaceIndex));
+
+		bool bValid = m_pD->Components[ComponentIndex].LODs[LODIndex].Surfaces[SurfaceIndex].Strings.IsValidIndex(StringIndex);
+		if (bValid)
+		{
+            return m_pD->Components[ComponentIndex].LODs[LODIndex].Surfaces[SurfaceIndex].Strings[StringIndex].Name;
         }
 
         return NAME_None;
@@ -424,206 +398,169 @@ namespace mu
 
 
     //---------------------------------------------------------------------------------------------
-	int32 Instance::Private::AddLOD()
+	int32 Instance::Private::AddComponent()
 	{
-		int32 result = Lods.Num();
-		Lods.Add( INSTANCE_LOD() );
+		int32 result = Components.Emplace();
 		return result;
 	}
 
 
     //---------------------------------------------------------------------------------------------
-    int32 Instance::Private::AddComponent( int32 LODIndex )
-    {
-        // Automatically create the necessary lods and components
-        while (LODIndex >= Lods.Num())
-        {
-            AddLOD();
-        }
-
-        int32 result = Lods[LODIndex].Components.Num();
-        Lods[LODIndex].Components.Add( INSTANCE_COMPONENT() );
-        return result;
-    }
-
-
-    //---------------------------------------------------------------------------------------------
-    int32 Instance::Private::AddSurface( int32 LODIndex, int32 ComponentIndex )
+    int32 Instance::Private::AddLOD( int32 ComponentIndex )
     {
         // Automatically create the necessary lods and components
-        while (LODIndex >= Lods.Num())
+        while (ComponentIndex >= Components.Num())
         {
-            AddLOD();
-        }
-        while (ComponentIndex >= Lods[LODIndex].Components.Num())
-        {
-            AddComponent(LODIndex);
+            AddComponent();
         }
 
-        int32 result = Lods[LODIndex].Components[ComponentIndex].Surfaces.Num();
-        Lods[LODIndex].Components[ComponentIndex].Surfaces.Add( INSTANCE_SURFACE() );
-        return result;
+        return Components[ComponentIndex].LODs.Emplace();
     }
 
 
     //---------------------------------------------------------------------------------------------
-    void Instance::Private::SetComponentName( int32 LODIndex, int32 ComponentIndex, FName Name)
+    int32 Instance::Private::AddSurface( int32 ComponentIndex, int32 LODIndex )
     {
         // Automatically create the necessary lods and components
-        while (LODIndex >= Lods.Num())
+        while (ComponentIndex >= Components.Num())
         {
-            AddLOD();
+            AddComponent();
         }
-        while (ComponentIndex >= Lods[LODIndex].Components.Num())
+        while (LODIndex >= Components[ComponentIndex].LODs.Num())
         {
-            AddComponent( LODIndex );
+            AddLOD(ComponentIndex);
         }
 
-        INSTANCE_COMPONENT& component = Lods[LODIndex].Components[ComponentIndex];
-        component.Name = Name;
+        return Components[ComponentIndex].LODs[LODIndex].Surfaces.Emplace();
     }
 
 
     //---------------------------------------------------------------------------------------------
-    void Instance::Private::SetSurfaceName( int32 LODIndex, int32 ComponentIndex, int32 SurfaceIndex, FName Name)
+    void Instance::Private::SetSurfaceName( int32 ComponentIndex, int32 LODIndex, int32 SurfaceIndex, FName Name)
     {
         // Automatically create the necessary lods and components
-        while ( LODIndex>=Lods.Num() )
-        {
-            AddLOD();
-        }
-        while (ComponentIndex >= Lods[LODIndex].Components.Num())
-        {
-            AddComponent( LODIndex );
-        }
-        while ( SurfaceIndex>=Lods[LODIndex].Components[ComponentIndex].Surfaces.Num() )
+		while (ComponentIndex >= Components.Num())
+		{
+			AddComponent();
+		}
+		while (LODIndex >= Components[ComponentIndex].LODs.Num())
+		{
+			AddLOD(ComponentIndex);
+		}
+		while ( SurfaceIndex>=Components[ComponentIndex].LODs[LODIndex].Surfaces.Num() )
         {
             AddSurface( LODIndex, ComponentIndex );
         }
 
-        INSTANCE_SURFACE& surface = Lods[LODIndex].Components[ComponentIndex].Surfaces[SurfaceIndex];
+        FInstanceSurface& surface = Components[ComponentIndex].LODs[LODIndex].Surfaces[SurfaceIndex];
         surface.Name = Name;
     }
 
 
 	//---------------------------------------------------------------------------------------------
-    int32 Instance::Private::AddMesh( int32 LODIndex, int32 ComponentIndex, FResourceID meshId, FName Name)
+	void Instance::Private::SetMesh(int32 ComponentIndex, int32 LODIndex, FResourceID meshId, FName Name)
 	{
 		// Automatically create the necessary lods and components
-		while (LODIndex >= Lods.Num())
+		while (ComponentIndex >= Components.Num())
 		{
-			AddLOD();
+			AddComponent();
 		}
-		while (ComponentIndex >= Lods[LODIndex].Components.Num())
+		while (LODIndex >= Components[ComponentIndex].LODs.Num())
 		{
-			AddComponent( LODIndex );
+			AddLOD(ComponentIndex);
 		}
 
-		INSTANCE_COMPONENT& component = Lods[LODIndex].Components[ComponentIndex];
-        int32 result = component.Meshes.Num();
-        component.Meshes.Emplace( meshId, Name );
-
-		return result;
+		FInstanceLOD& LOD = Components[ComponentIndex].LODs[LODIndex];
+		LOD.MeshId = meshId;
+		LOD.MeshName = Name;
 	}
 
 
 	//---------------------------------------------------------------------------------------------
-    int32 Instance::Private::AddImage( int32 LODIndex, int32 ComponentIndex, int32 SurfaceIndex, FResourceID imageId, FName Name)
+    int32 Instance::Private::AddImage( int32 ComponentIndex, int32 LODIndex, int32 SurfaceIndex, FResourceID imageId, FName Name)
 	{
 		// Automatically create the necessary lods and components
-		while (LODIndex >= Lods.Num())
+		while (ComponentIndex >= Components.Num())
 		{
-			AddLOD();
+			AddComponent();
 		}
-		while (ComponentIndex >= Lods[LODIndex].Components.Num())
+		while (LODIndex >= Components[ComponentIndex].LODs.Num())
 		{
-			AddComponent(LODIndex);
+			AddLOD(ComponentIndex);
 		}
-		while (SurfaceIndex >= Lods[LODIndex].Components[ComponentIndex].Surfaces.Num())
+		while (SurfaceIndex >= Components[ComponentIndex].LODs[LODIndex].Surfaces.Num())
 		{
 			AddSurface(LODIndex, ComponentIndex);
 		}
 
-        INSTANCE_SURFACE& surface = Lods[LODIndex].Components[ComponentIndex].Surfaces[SurfaceIndex];
-        int32 result = surface.Images.Num();
-        surface.Images.Emplace( imageId, Name);
-
-        return result;
+		FInstanceSurface& Surface = Components[ComponentIndex].LODs[LODIndex].Surfaces[SurfaceIndex];
+		return Surface.Images.Add({ imageId, Name });
 	}
 
 
 	//---------------------------------------------------------------------------------------------
-    int32 Instance::Private::AddVector( int32 LODIndex, int32 ComponentIndex, int32 SurfaceIndex, const FVector4f& vec, FName Name)
+    int32 Instance::Private::AddVector( int32 ComponentIndex, int32 LODIndex, int32 SurfaceIndex, const FVector4f& vec, FName Name)
 	{
 		// Automatically create the necessary lods and components
-		while (LODIndex >= Lods.Num())
+		while (ComponentIndex >= Components.Num())
 		{
-			AddLOD();
+			AddComponent();
 		}
-		while (ComponentIndex >= Lods[LODIndex].Components.Num())
+		while (LODIndex >= Components[ComponentIndex].LODs.Num())
 		{
-			AddComponent(LODIndex);
+			AddLOD(ComponentIndex);
 		}
-		while (SurfaceIndex >= Lods[LODIndex].Components[ComponentIndex].Surfaces.Num())
+		while (SurfaceIndex >= Components[ComponentIndex].LODs[LODIndex].Surfaces.Num())
 		{
 			AddSurface(LODIndex, ComponentIndex);
 		}
 
-        INSTANCE_SURFACE& surface = Lods[LODIndex].Components[ComponentIndex].Surfaces[SurfaceIndex];
-        int32 result = surface.Vectors.Num();
-        surface.Vectors.Emplace( vec, Name);
-
-        return result;
+		FInstanceSurface& Surface = Components[ComponentIndex].LODs[LODIndex].Surfaces[SurfaceIndex];
+		return Surface.Vectors.Add({ vec, Name } );
 	}
 
 
     //---------------------------------------------------------------------------------------------
-    int32 Instance::Private::AddScalar( int32 LODIndex, int32 ComponentIndex, int32 SurfaceIndex, float sca, FName Name)
+    int32 Instance::Private::AddScalar( int32 ComponentIndex, int32 LODIndex, int32 SurfaceIndex, float sca, FName Name)
     {
         // Automatically create the necessary lods and components
-		while (LODIndex >= Lods.Num())
+		while (ComponentIndex >= Components.Num())
 		{
-			AddLOD();
+			AddComponent();
 		}
-		while (ComponentIndex >= Lods[LODIndex].Components.Num())
+		while (LODIndex >= Components[ComponentIndex].LODs.Num())
 		{
-			AddComponent(LODIndex);
+			AddLOD(ComponentIndex);
 		}
-		while (SurfaceIndex >= Lods[LODIndex].Components[ComponentIndex].Surfaces.Num())
+		while (SurfaceIndex >= Components[ComponentIndex].LODs[LODIndex].Surfaces.Num())
 		{
 			AddSurface(LODIndex, ComponentIndex);
 		}
 
-        INSTANCE_SURFACE& surface = Lods[LODIndex].Components[ComponentIndex].Surfaces[SurfaceIndex];
-        int32 result = surface.Scalars.Num();
-        surface.Scalars.Emplace( sca, Name );
-
-        return result;
+		FInstanceSurface& Surface = Components[ComponentIndex].LODs[LODIndex].Surfaces[SurfaceIndex];
+		return Surface.Scalars.Add({ sca, Name });
     }
 
 
     //---------------------------------------------------------------------------------------------
-    int32 Instance::Private::AddString( int32 LODIndex, int32 ComponentIndex, int32 SurfaceIndex, const FString& Value, FName Name)
+    int32 Instance::Private::AddString( int32 ComponentIndex, int32 LODIndex, int32 SurfaceIndex, const FString& Value, FName Name)
     {
         // Automatically create the necessary lods and components
-		while (LODIndex >= Lods.Num())
+		while (ComponentIndex >= Components.Num())
 		{
-			AddLOD();
+			AddComponent();
 		}
-		while (ComponentIndex >= Lods[LODIndex].Components.Num())
+		while (LODIndex >= Components[ComponentIndex].LODs.Num())
 		{
-			AddComponent(LODIndex);
+			AddLOD(ComponentIndex);
 		}
-		while (SurfaceIndex >= Lods[LODIndex].Components[ComponentIndex].Surfaces.Num())
+		while (SurfaceIndex >= Components[ComponentIndex].LODs[LODIndex].Surfaces.Num())
 		{
 			AddSurface(LODIndex, ComponentIndex);
 		}
 
-        INSTANCE_SURFACE& surface = Lods[LODIndex].Components[ComponentIndex].Surfaces[SurfaceIndex];
-        int32 result = surface.Strings.Num();
-        surface.Strings.Emplace(Value, Name );
-
-        return result;
+		FInstanceSurface& Surface = Components[ComponentIndex].LODs[LODIndex].Surfaces[SurfaceIndex];
+		return Surface.Strings.Add({ Value, Name });
     }
 
 
diff --git a/Engine/Plugins/Experimental/Mutable/Source/MutableRuntime/Private/MuR/InstancePrivate.h b/Engine/Plugins/Experimental/Mutable/Source/MutableRuntime/Private/MuR/InstancePrivate.h
--- a/Engine/Plugins/Experimental/Mutable/Source/MutableRuntime/Private/MuR/InstancePrivate.h
+++ b/Engine/Plugins/Experimental/Mutable/Source/MutableRuntime/Private/MuR/InstancePrivate.h
@@ -23,98 +23,62 @@ namespace mu
 	}
 
 	/** */
-    struct INSTANCE_SURFACE
+    struct FInstanceSurface
 	{
 		FName Name;
         uint32 InternalId=0;
         uint32 ExternalId =0;
         uint32 SharedId =0;
 
-		struct IMAGE
+		struct FInstanceImage
 		{
-            IMAGE(FResourceID InId, FName InName)
-			{
-				Id = InId;
-				Name = InName;
-			}
-
 			FResourceID Id;
 			FName Name;
 		};
 
-		TArray<IMAGE, TInlineAllocator<4>> Images;
+		TArray<FInstanceImage, TInlineAllocator<4>> Images;
 
-		struct VECTOR
+		struct FInstanceVector
 		{
-			VECTOR( const FVector4f& v, FName InName)
-			{
-				Value = v;
-				Name = InName;
-			}
-
 			FVector4f Value;
 			FName Name;
 		};
 
-		TArray<VECTOR> Vectors;
+		TArray<FInstanceVector> Vectors;
 
-        struct SCALAR
+        struct FInstanceScalar
         {
-            SCALAR( float v, FName InName )
-            {
-                Value = v;
-                Name = InName;
-            }
-
             float Value;
 			FName Name;
 		};
 
-		TArray<SCALAR> Scalars;
+		TArray<FInstanceScalar> Scalars;
 
-        struct STRING
+        struct FInstanceString
         {
-            STRING(const FString& InValue, FName InName)
-            {
-				Value = InValue;
-				Name = InName;
-			}
-
 			FString Value;
 			FName Name;
 		};
 
-		TArray<STRING> Strings;
+		TArray<FInstanceString> Strings;
     };
 
 
-    struct INSTANCE_COMPONENT
+    struct FInstanceLOD
     {
-        FName Name;
-
-    	uint16 Id;
-
-		struct MESH
-		{
-            MESH(FResourceID InId, FName InName)
-			{
-				Id = InId;
-				Name = InName;
-			}
-
-			FResourceID Id;
-			FName Name;
-		};
-		TArray<MESH, TInlineAllocator<2>> Meshes;
+		FResourceID MeshId;
+		FName MeshName;
 
 		// The order must match the meshes surfaces
-		TArray<INSTANCE_SURFACE, TInlineAllocator<4>> Surfaces;
+		TArray<FInstanceSurface, TInlineAllocator<4>> Surfaces;
 	};
 
 
-    struct INSTANCE_LOD
+    struct FInstanceComponent
     {
-		TArray<INSTANCE_COMPONENT, TInlineAllocator<4>> Components;
+		uint16 Id;
+
+		TArray<FInstanceLOD, TInlineAllocator<4>> LODs;
     };
 
 	struct NamedExtensionData
@@ -131,21 +95,20 @@ namespace mu
         Instance::ID Id = 0;
 
 		//!
-		TArray<INSTANCE_LOD,TInlineAllocator<4>> Lods;
+		TArray<FInstanceComponent,TInlineAllocator<4>> Components;
 
 		// Every entry must have a valid ExtensionData and name
 		TArray<NamedExtensionData> ExtensionData;
 
-		int32 AddLOD();
-		int32 AddComponent(int32 lod );
-        void SetComponentName( int32 lod, int32 comp, FName Name );
-		int32 AddMesh(int32 lod, int32 comp, FResourceID, FName Name);
-		int32 AddSurface(int32 lod, int32 comp );
-        void SetSurfaceName( int32 lod, int32 comp, int32 surf, FName Name);
-		int32 AddImage( int32 lod, int32 comp, int32 surf, FResourceID, FName Name);
-        int32 AddVector( int32 lod, int32 comp, int32 surf, const FVector4f&, FName Name);
-        int32 AddScalar( int32 lod, int32 comp, int32 surf, float, FName Name);
-        int32 AddString( int32 lod, int32 comp, int32 surf, const FString& Value, FName Name);
+		int32 AddComponent();
+		int32 AddLOD(int32 ComponentIndex);
+		void SetMesh(int32 ComponentIndex, int32 LODIndex, FResourceID, FName Name);
+		int32 AddSurface(int32 ComponentIndex, int32 LODIndex);
+        void SetSurfaceName( int32 ComponentIndex, int32 LODIndex, int32 SurfaceIndex, FName Name);
+		int32 AddImage( int32 ComponentIndex, int32 LODIndex, int32 SurfaceIndex, FResourceID, FName Name);
+        int32 AddVector( int32 ComponentIndex, int32 LODIndex, int32 SurfaceIndex, const FVector4f&, FName Name);
+        int32 AddScalar( int32 ComponentIndex, int32 LODIndex, int32 SurfaceIndex, float, FName Name);
+        int32 AddString( int32 ComponentIndex, int32 LODIndex, int32 SurfaceIndex, const FString& Value, FName Name);
 		
 		// Data must be non-null
 		void AddExtensionData(const Ptr<const class ExtensionData>& Data, FName Name);
