[mutable] Redesign the object node hierarchy to better match unreal component and LOD concepts.

diff --git a/Engine/Plugins/Experimental/Mutable/Source/MutableRuntime/Private/MuR/Mesh.cpp b/Engine/Plugins/Experimental/Mutable/Source/MutableRuntime/Private/MuR/Mesh.cpp
--- a/Engine/Plugins/Experimental/Mutable/Source/MutableRuntime/Private/MuR/Mesh.cpp
+++ b/Engine/Plugins/Experimental/Mutable/Source/MutableRuntime/Private/MuR/Mesh.cpp
@@ -83,6 +83,7 @@ MeshPtr Mesh::Clone() const
 
     pResult->InternalId = InternalId;
 	pResult->Flags = Flags;
+	pResult->ReferenceID = ReferenceID;
 	pResult->Surfaces = Surfaces;
 	pResult->Skeleton = Skeleton;
 	pResult->PhysicsBody = PhysicsBody;
@@ -127,6 +128,7 @@ Ptr<Mesh> Mesh::Clone(EMeshCopyFlags InFlags) const
     pResult->InternalId = InternalId;
 	pResult->MeshIDPrefix = MeshIDPrefix;
 	pResult->Flags = Flags;
+	pResult->ReferenceID = ReferenceID;
 
 	if (EnumHasAnyFlags(InFlags, EMeshCopyFlags::WithSurfaces))
 	{
@@ -213,6 +215,8 @@ void Mesh::CopyFrom(const Mesh& From, EMeshCopyFlags InFlags)
 
     InternalId = From.InternalId;
 	Flags = From.Flags;
+	ReferenceID = From.ReferenceID;
+
 	MeshIDPrefix = From.MeshIDPrefix;
 
 	if (EnumHasAnyFlags(InFlags, EMeshCopyFlags::WithSurfaces))
@@ -1057,9 +1061,6 @@ void UnserialiseLegacySurfaces(InputArchive& arch, TArray<FMeshSurface>& OutMesh
 
 void FMeshSurface::Serialise(OutputArchive& arch) const
 {
-	const int32 ver = 1;
-	arch << ver;
-
 	arch << FirstVertex;
 	arch << VertexCount;
 	arch << FirstIndex;
@@ -1067,38 +1068,25 @@ void FMeshSurface::Serialise(OutputArchive& arch) const
 	arch << BoneMapIndex;
 	arch << BoneMapCount;
 	arch << bCastShadow;
-
 	arch << Id;
 }
 
 
 void FMeshSurface::Unserialise(InputArchive& arch)
 {
-	int32 ver = 0;
-	arch >> ver;
-	check(ver <= 1);
-
 	arch >> FirstVertex;
 	arch >> VertexCount;
 	arch >> FirstIndex;
 	arch >> IndexCount;
 	arch >> BoneMapIndex;
 	arch >> BoneMapCount;
-
-	if (ver >= 1)
-	{
-		arch >> bCastShadow;
-	}
-
+	arch >> bCastShadow;
 	arch >> Id;
 }
 
 
 void Mesh::FBonePose::Serialise(OutputArchive& arch) const
 {
-	const int32 ver = 2;
-	arch << ver;
-
 	arch << BoneId;
 	arch << BoneUsageFlags;
 	arch << BoneTransform;
@@ -1107,40 +1095,15 @@ void Mesh::FBonePose::Serialise(OutputArchive& arch) const
 
 void Mesh::FBonePose::Unserialise(InputArchive& arch)
 {
-	int32 ver = 0;
-	arch >> ver;
-	check(ver <= 2);
-
-	if (ver <= 1)
-	{
-		std::string DeprecatedBoneName;
-		arch >> DeprecatedBoneName;
-
-		BoneId = FBoneName(0);
-	}
-	else
-	{
-		arch >> BoneId;
-	}
-
-	if (ver == 0)
-	{
-		uint8 Skinned = 0;
-		arch >> Skinned;
-		BoneUsageFlags = Skinned ? EBoneUsageFlags::Skinning : EBoneUsageFlags::None;
-	}
-	else
-	{
-		arch >> BoneUsageFlags;
-	}
-
+	arch >> BoneId;
+	arch >> BoneUsageFlags;
 	arch >> BoneTransform;
 }
 
 
 void Mesh::Serialise(OutputArchive& arch) const
 {
-	uint32 ver = 20;
+	uint32 ver = 22;
 	arch << ver;
 
 	arch << IndexBuffers;
@@ -1165,6 +1128,7 @@ void Mesh::Serialise(OutputArchive& arch) const
 	arch << AdditionalPhysicsBodies;
 
 	arch << MeshIDPrefix;
+	arch << ReferenceID;
 }
 
 
@@ -1172,7 +1136,7 @@ void Mesh::Unserialise(InputArchive& arch)
 {
 	uint32 ver;
 	arch >> ver;
-	check(ver == 20);
+	check(ver == 22);
 
 	arch >> IndexBuffers;
 	arch >> VertexBuffers;
@@ -1199,6 +1163,7 @@ void Mesh::Unserialise(InputArchive& arch)
 	arch >> AdditionalPhysicsBodies;
 
 	arch >> MeshIDPrefix;
+	arch >> ReferenceID;
 }
 
 
@@ -1212,6 +1177,7 @@ bool Mesh::IsSimilar(const Mesh& o, bool bCompareLayouts) const
 	}
 
 	bool equal = IndexBuffers == o.IndexBuffers;
+	if (equal) equal = (ReferenceID==o.ReferenceID);
 	if (equal && bCompareLayouts) equal = (Layouts.Num() == o.Layouts.Num());
 	if (equal && Skeleton != o.Skeleton)
 	{
diff --git a/Engine/Plugins/Experimental/Mutable/Source/MutableRuntime/Private/MuR/System.cpp b/Engine/Plugins/Experimental/Mutable/Source/MutableRuntime/Private/MuR/System.cpp
--- a/Engine/Plugins/Experimental/Mutable/Source/MutableRuntime/Private/MuR/System.cpp
+++ b/Engine/Plugins/Experimental/Mutable/Source/MutableRuntime/Private/MuR/System.cpp
@@ -244,7 +244,7 @@ namespace mu
 
 
     //---------------------------------------------------------------------------------------------
-    const Instance* System::BeginUpdate( Instance::ID InInstanceID,
+    Ptr<const Instance> System::BeginUpdate( Instance::ID InInstanceID,
                                      const Ptr<const Parameters>& InParams,
                                      int32 InStateIndex,
                                      uint32 InLodMask )
@@ -320,12 +320,17 @@ namespace mu
 		{
 			Result->GetPrivate()->Id = pLiveInstance->InstanceID;
 		}
+		else
+		{
+			// In case of failure return an empty instance, to prevent following code to have to check it every time
+			Result = new Instance;
+		}
 
 		m_pD->WorkingMemoryManager.EndRunnerThread();
 
 		m_pD->WorkingMemoryManager.CurrentInstanceCache = nullptr;
 
-		return Result.get();
+		return Result;
 	}
 
 
diff --git a/Engine/Plugins/Experimental/Mutable/Source/MutableRuntime/Public/MuR/Instance.h b/Engine/Plugins/Experimental/Mutable/Source/MutableRuntime/Public/MuR/Instance.h
--- a/Engine/Plugins/Experimental/Mutable/Source/MutableRuntime/Public/MuR/Instance.h
+++ b/Engine/Plugins/Experimental/Mutable/Source/MutableRuntime/Public/MuR/Instance.h
@@ -50,156 +50,141 @@ namespace mu
         //! or release the instance.
         Instance::ID GetId() const;
 
-		//! Get the number of levels-of-detail of this instance.
-		int32 GetLODCount() const;
+		//! Get the number of components of this instance.
+		int32 GetComponentCount() const;
 
-        //! Get the number of components in a level-of-detail
-        //! \param lod Index of the level of detail, from 0 to GetLODCount()-1
-        int32 GetComponentCount( int32 lod ) const;
-		
 		//! Get the Id of a component
-		//! \param lod Index of the level of detail, from 0 to GetLODCount()-1
-		//! \param comp Index of the component, from 0 to GetComponentCount(lod)-1
-		uint16 GetComponentId( int32 lod, int32 comp ) const;
+		//! \param ComponentIndex Index of the component, from 0 to GetComponentCount()-1
+		uint16 GetComponentId(int32 ComponentIndex) const;
+
+        //! Get the number of LODs in a component
+        //! \param ComponentIndex Index of the component, from 0 to GetComponentCount()-1
+        int32 GetLODCount( int32 ComponentIndex ) const;
 
         //! Get the number of surfaces in a component
-        //! \param lod Index of the level of detail, from 0 to GetLODCount()-1
-        //! \param comp Index of the component, from 0 to GetComponentCount(lod)-1
-        int32 GetSurfaceCount( int32 lod, int32 comp ) const;
+        //! \param ComponentIndex Index of the component, from 0 to GetComponentCount()-1
+		//! \param LODIndex Index of the level of detail, from 0 to GetLODCount()-1
+		int32 GetSurfaceCount( int32 ComponentIndex, int32 LODIndex) const;
 
         //! Get an id that can be used to match the surface data with the mesh surface data.
-        //! \param lod Index of the level of detail, from 0 to GetLODCount()-1
-        //! \param comp Index of the component, from 0 to GetComponentCount(lod)-1
-        //! \param surf Index of the surface, from 0 to GetSurfaceCount(lod,comp)-1
-        uint32 GetSurfaceId( int32 lod, int32 comp, int32 surf ) const;
+        //! \param ComponentIndex Index of the component, from 0 to GetComponentCount()-1
+		//! \param LODIndex Index of the level of detail, from 0 to GetLODCount()-1
+		//! \param SurfaceIndex Index of the surface, from 0 to GetSurfaceCount()-1
+        uint32 GetSurfaceId( int32 ComponentIndex, int32 LODIndex, int32 SurfaceIndex) const;
 
         //! Find a surface index from the internal id (as returned by GetSurfaceId).
-        //! \param lod Index of the level of detail, from 0 to GetLODCount()-1
-        //! \param comp Index of the component, from 0 to GetComponentCount(lod)-1
-        //! \param id ID of the surface to look for.
-        int32 FindSurfaceById( int32 lod, int32 comp, uint32 id ) const;
+        //! \param ComponentIndex Index of the component, from 0 to GetComponentCount()-1
+		//! \param LODIndex Index of the level of detail, from 0 to GetLODCount()-1
+		//! \param id ID of the surface to look for.
+        int32 FindSurfaceById( int32 ComponentIndex, int32 LODIndex, uint32 id ) const;
 
 		//! Find the base surface index and Lod index when reusing surfaces between LODs. Return the surface index
 		//! and the LOD it belongs to.
-		//! \param Comp - Index of the component, from 0 to GetComponentCount(lod)-1
+		//! \param ComponentIndex - Index of the component, from 0 to GetComponentCount()-1
 		//! \param SharedSurfaceId - Id of the surface to look for (as returned by GetSharedSurfaceId).
 		//! \param OutSurfaceIndex - Index of the surface in the OutLODIndex lod. 
 		//! \param OutLODIndex - Index of the first LOD where the surface can be found. 
-		void FindBaseSurfaceBySharedId(int32 CompIndex, int32 SharedId, int32& OutSurfaceIndex, int32& OutLODIndex) const;
+		void FindBaseSurfaceBySharedId(int32 ComponentIndex, int32 SharedId, int32& OutSurfaceIndex, int32& OutLODIndex) const;
 
 		//! Get an id that can be used to find the same surface on other LODs
+		//! \param ComponentIndex Index of the component, from 0 to GetComponentCount()-1
 		//! \param lod Index of the level of detail, from 0 to GetLODCount()-1
-		//! \param comp Index of the component, from 0 to GetComponentCount(lod)-1
-		//! \param surf Index of the surface, from 0 to GetSurfaceCount(lod,comp)-1
-		int32 GetSharedSurfaceId(int32 lod, int32 comp, int32 surf) const;
+		//! \param SurfaceIndex Index of the surface, from 0 to GetSurfaceCount()-1
+		int32 GetSharedSurfaceId(int32 ComponentIndex, int32 LODIndex, int32 SurfaceIndex) const;
 
         //! Get an optional, opaque application-defined identifier for this surface. The meaning of
         //! this ID depends on each application, and it is specified when creating the source data
         //! that generates this surface.
         //! See NodeSurfaceNew::SetCustomID.
-		//! \param lod Index of the level of detail, from 0 to GetLODCount()-1
-		//! \param comp Index of the component, from 0 to GetComponentCount(lod)-1
-        //! \param surf Index of the surface, from 0 to GetSurfaceCount(lod,comp)-1
-        uint32 GetSurfaceCustomId( int32 lod, int32 comp, int32 surf ) const;
+		//! \param ComponentIndex Index of the component, from 0 to GetComponentCount()-1
+		//! \param LODIndex Index of the level of detail, from 0 to GetLODCount()-1
+		//! \param SurfaceIndex Index of the surface, from 0 to GetSurfaceCount()-1
+        uint32 GetSurfaceCustomId( int32 ComponentIndex, int32 LODIndex, int32 SurfaceIndex) const;
 
-        //! Get the number of meshes in a surface
-		//! \param lod Index of the level of detail, from 0 to GetLODCount()-1
-		//! \param comp Index of the component, from 0 to GetComponentCount(lod)-1
-        int32 GetMeshCount( int32 lod, int32 comp ) const;
-
-        //! Get a mesh resource id from a component
-        //! \param lod Index of the level of detail, from 0 to GetLODCount()-1
-        //! \param comp Index of the component, from 0 to GetComponentCount(lod)-1
-        //! \param mesh Index of the mesh, from 0 to GetMeshCount(lod,comp,surf)-1
-        //! \return an ummodifiable pointer to the requested mesh. The returned object is guaranteed
-        //! to be alive only while this instance is alive. The ownership of the returned object
-        //! remains in the instance, so it should not be deleted.
-		FResourceID GetMeshId( int32 lod, int32 comp, int32 mesh ) const;
+        //! Get the mesh resource id from a component
+        //! \param ComponentIndex Index of the component, from 0 to GetComponentCount()-1
+		//! \param LODIndex Index of the level of detail, from 0 to GetLODCount()-1
+  		FResourceID GetMeshId( int32 ComponentIndex, int32 LODIndex ) const;
 
 		//! Get the number of images in a component
-		//! \param lod Index of the level of detail, from 0 to GetLODCount()
-		//! \param comp Index of the component, from 0 to GetComponentCount(lod)
-        //! \param surf Index of the surface, from 0 to GetSurfaceCount(lod,comp)-1
-        int32 GetImageCount( int32 lod, int32 comp, int32 surf ) const;
+		//! \param ComponentIndex Index of the component, from 0 to GetComponentCount()-1
+		//! \param LODIndex Index of the level of detail, from 0 to GetLODCount()-1
+		//! \param SurfaceIndex Index of the surface, from 0 to GetSurfaceCount()-1
+        int32 GetImageCount( int32 ComponentIndex, int32 LODIndex, int32 SurfaceIndex) const;
 
         //! Get an image resource id from a component
-		//! \param lod Index of the level of detail, from 0 to GetLODCount()-1
-		//! \param comp Index of the component, from 0 to GetComponentCount(lod)-1
-        //! \param surf Index of the surface, from 0 to GetSurfaceCount(lod,comp)-1
-        //! \param img Index of the image, from 0 to GetImageCount(lod,comp,surf)-1
-        //! \return an ummodifiable pointer to the requested image. The returned object is
-        //! guaranteed to be alive only while this instance is alive. The ownership of the returned
-        //! object remains in the instance, so it should not be deleted.
-		FResourceID GetImageId( int32 lod, int32 comp, int32 surf, int32 img ) const;
+		//! \param ComponentIndex Index of the component, from 0 to GetComponentCount()-1
+		//! \param LODIndex Index of the level of detail, from 0 to GetLODCount()-1
+		//! \param SurfaceIndex Index of the surface, from 0 to GetSurfaceCount()-1
+        //! \param ImageIndex Index of the image, from 0 to GetImageCount()-1
+ 		FResourceID GetImageId( int32 ComponentIndex, int32 LODIndex, int32 SurfaceIndex, int32 ImageIndex) const;
 
 		//! Get the name of an image in a component
-		//! \param lod Index of the level of detail, from 0 to GetLODCount()-1
-		//! \param comp Index of the component, from 0 to GetComponentCount(lod)-1
-        //! \param surf Index of the surface, from 0 to GetSurfaceCount(lod,comp)-1
-        //! \param img Index of the image, from 0 to GetImageCount(lod,comp,surf)-1
-		FName GetImageName( int32 lod, int32 comp, int32 surf, int32 img ) const;
+		//! \param ComponentIndex Index of the component, from 0 to GetComponentCount()-1
+		//! \param LODIndex Index of the level of detail, from 0 to GetLODCount()-1
+		//! \param SurfaceIndex Index of the surface, from 0 to GetSurfaceCount()-1
+        //! \param ImageIndex Index of the image, from 0 to GetImageCount()-1
+		FName GetImageName( int32 ComponentIndex, int32 LODIndex, int32 SurfaceIndex, int32 ImageIndex ) const;
 
 		//! Get the number of vectors in a component
-		//! \param lod Index of the level of detail, from 0 to GetLODCount()
-		//! \param comp Index of the component, from 0 to GetComponentCount(lod)
-        //! \param surf Index of the surface, from 0 to GetSurfaceCount(lod,comp)-1
-        int32 GetVectorCount( int32 lod, int32 comp, int32 surf ) const;
+		//! \param ComponentIndex Index of the component, from 0 to GetComponentCount()-1
+		//! \param LODIndex Index of the level of detail, from 0 to GetLODCount()-1
+		//! \param SurfaceIndex Index of the surface, from 0 to GetSurfaceCount()-1
+        int32 GetVectorCount( int32 ComponentIndex, int32 LODIndex, int32 SurfaceIndex) const;
 
 		//! Get a vector from a component
-		//! \param lod Index of the level of detail, from 0 to GetLODCount()-1
-		//! \param comp Index of the component, from 0 to GetComponentCount(lod)-1
-        //! \param surf Index of the surface, from 0 to GetSurfaceCount(lod,comp)-1
-        //! \param vec Index of the vector, from 0 to GetVectorCount(lod,comp,surf)-1
-		//! \param pX, pY, pZ, pW are optional to floats to store the components of the vector
-		//!		   value. They can be null.
-        FVector4f GetVector( int32 lod, int32 comp, int32 surf, int32 vec ) const;
+		//! \param ComponentIndex Index of the component, from 0 to GetComponentCount()-1
+		//! \param LODIndex Index of the level of detail, from 0 to GetLODCount()-1
+		//! \param SurfaceIndex Index of the surface, from 0 to GetSurfaceCount()-1
+        //! \param VectorIndex Index of the vector, from 0 to GetVectorCount()-1
+        FVector4f GetVector( int32 ComponentIndex, int32 LODIndex, int32 SurfaceIndex, int32 VectorIndex) const;
 
 		//! Get the name of a vector in a component
-		//! \param lod Index of the level of detail, from 0 to GetLODCount()-1
-		//! \param comp Index of the component, from 0 to GetComponentCount(lod)-1
-        //! \param surf Index of the surface, from 0 to GetSurfaceCount(lod,comp)-1
-        //! \param vec Index of the vector, from 0 to GetVectorCount(lod,comp)-1
-		FName GetVectorName( int32 lod, int32 comp, int32 surf, int32 vec ) const;
+		//! \param ComponentIndex Index of the component, from 0 to GetComponentCount()-1
+		//! \param LODIndex Index of the level of detail, from 0 to GetLODCount()-1
+		//! \param SurfaceIndex Index of the surface, from 0 to GetSurfaceCount()-1
+        //! \param VectorIndex Index of the vector, from 0 to GetVectorCount()-1
+		FName GetVectorName( int32 ComponentIndex, int32 LODIndex, int32 SurfaceIndex, int32 VectorIndex ) const;
 
         //! Get the number of scalar values in a component
-        //! \param lod Index of the level of detail, from 0 to GetLODCount()
-        //! \param comp Index of the component, from 0 to GetComponentCount(lod)
-        //! \param surf Index of the surface, from 0 to GetSurfaceCount(lod,comp)-1
-        int32 GetScalarCount( int32 lod, int32 comp, int32 surf ) const;
+        //! \param ComponentIndex Index of the component, from 0 to GetComponentCount()-1
+		//! \param LODIndex Index of the level of detail, from 0 to GetLODCount()-1
+		//! \param SurfaceIndex Index of the surface, from 0 to GetSurfaceCount()-1
+        int32 GetScalarCount( int32 ComponentIndex, int32 LODIndex, int32 SurfaceIndex) const;
 
         //! Get a scalar value from a component
-        //! \param lod Index of the level of detail, from 0 to GetLODCount()-1
-        //! \param comp Index of the component, from 0 to GetComponentCount(lod)-1
-        //! \param surf Index of the surface, from 0 to GetSurfaceCount(lod,comp)-1
-        //! \param sca Index of the scalar, from 0 to GetScalarCount(lod,comp)-1
-        float GetScalar( int32 lod, int32 comp, int32 surf, int32 sca ) const;
+        //! \param ComponentIndex Index of the component, from 0 to GetComponentCount()-1
+		//! \param LODIndex Index of the level of detail, from 0 to GetLODCount()-1
+		//! \param SurfaceIndex Index of the surface, from 0 to GetSurfaceCount()-1
+        //! \param ScalarIndex Index of the scalar, from 0 to GetScalarCount()-1
+        float GetScalar( int32 ComponentIndex, int32 LODIndex, int32 SurfaceIndex, int32 ScalarIndex) const;
 
         //! Get the name of a scalar from a component
-        //! \param lod Index of the level of detail, from 0 to GetLODCount()-1
-        //! \param comp Index of the component, from 0 to GetComponentCount(lod)-1
-        //! \param surf Index of the surface, from 0 to GetSurfaceCount(lod,comp)-1
-        //! \param sca Index of the scalar, from 0 to GetScalarCount(lod,comp)-1
-		FName GetScalarName( int32 lod, int32 comp, int32 surf, int32 sca ) const;
+        //! \param ComponentIndex Index of the component, from 0 to GetComponentCount()-1
+		//! \param LODIndex Index of the level of detail, from 0 to GetLODCount()-1
+		//! \param SurfaceIndex Index of the surface, from 0 to GetSurfaceCount()-1
+        //! \param ScalarIndex Index of the scalar, from 0 to GetScalarCount()-1
+		FName GetScalarName( int32 ComponentIndex, int32 LODIndex, int32 SurfaceIndex, int32 ScalarIndex ) const;
 
         //! Get the number of string values in a component
-        //! \param lod Index of the level of detail, from 0 to GetLODCount()
-        //! \param comp Index of the component, from 0 to GetComponentCount(lod)
-        //! \param surf Index of the surface, from 0 to GetSurfaceCount(lod,comp)-1
-        int32 GetStringCount( int32 lod, int32 comp, int32 surf ) const;
+        //! \param ComponentIndex Index of the component, from 0 to GetComponentCount()-1
+		//! \param LODIndex Index of the level of detail, from 0 to GetLODCount()-1
+		//! \param SurfaceIndex Index of the surface, from 0 to GetSurfaceCount()-1
+        int32 GetStringCount( int32 ComponentIndex, int32 LODIndex, int32 SurfaceIndex) const;
 
         //! Get a string value from a component
-        //! \param lod Index of the level of detail, from 0 to GetLODCount()-1
-        //! \param comp Index of the component, from 0 to GetComponentCount(lod)-1
-        //! \param surf Index of the surface, from 0 to GetSurfaceCount(lod,comp)-1
-        //! \param sca Index of the string, from 0 to GetStringCount(lod,comp)-1
-        FString GetString( int32 lod, int32 comp, int32 surf, int32 str ) const;
+        //! \param ComponentIndex Index of the component, from 0 to GetComponentCount()-1
+		//! \param LODIndex Index of the level of detail, from 0 to GetLODCount()-1
+		//! \param SurfaceIndex Index of the surface, from 0 to GetSurfaceCount()-1
+        //! \param StringIndex Index of the string, from 0 to GetStringCount()-1
+        FString GetString( int32 ComponentIndex, int32 LODIndex, int32 SurfaceIndex, int32 StringIndex) const;
 
         //! Get the name of a string from a component
-        //! \param lod Index of the level of detail, from 0 to GetLODCount()-1
-        //! \param comp Index of the component, from 0 to GetComponentCount(lod)-1
-        //! \param surf Index of the surface, from 0 to GetSurfaceCount(lod,comp)-1
-        //! \param sca Index of the string, from 0 to GetStringCount(lod,comp)-1
-        FName GetStringName( int32 lod, int32 comp, int32 surf, int32 str ) const;
+        //! \param ComponentIndex Index of the component, from 0 to GetComponentCount()-1
+		//! \param LODIndex Index of the level of detail, from 0 to GetLODCount()-1
+		//! \param SurfaceIndex Index of the surface, from 0 to GetSurfaceCount()-1
+        //! \param StringIndex Index of the string, from 0 to GetStringCount()-1
+        FName GetStringName( int32 ComponentIndex, int32 LODIndex, int32 SurfaceIndex, int32 StringIndex ) const;
 
 		//! Get the number of ExtensionData values in a component
 		int32 GetExtensionDataCount() const;
diff --git a/Engine/Plugins/Experimental/Mutable/Source/MutableRuntime/Public/MuR/Mesh.h b/Engine/Plugins/Experimental/Mutable/Source/MutableRuntime/Public/MuR/Mesh.h
--- a/Engine/Plugins/Experimental/Mutable/Source/MutableRuntime/Public/MuR/Mesh.h
+++ b/Engine/Plugins/Experimental/Mutable/Source/MutableRuntime/Public/MuR/Mesh.h
@@ -225,7 +225,7 @@ namespace mu
 
         //! Get the number of surfaces defined in this mesh. Surfaces are buffer-contiguous mesh
         //! fragments that share common properties (usually material)
-        int GetSurfaceCount() const;
+		int32 GetSurfaceCount() const;
         void GetSurface( int32 surfaceIndex,
                          int32* FirstVertex, int32* VertexCount,
                          int32* FirstIndex, int32* IndexCount,
@@ -234,7 +234,7 @@ namespace mu
 
         //! Return an internal id that can be used to match mesh surfaces and instance surfaces.
         //! Only valid for meshes that are part of instances.
-        uint32 GetSurfaceId( int surfaceIndex ) const;
+        uint32 GetSurfaceId( int32 surfaceIndex ) const;
 
         //! \}
 
@@ -255,13 +255,13 @@ namespace mu
         void AddLayout( Ptr<const Layout> pLayout );
 
         //!
-        int GetLayoutCount() const;
+        int32 GetLayoutCount() const;
 
         //!
-        const Layout* GetLayout( int i ) const;
+        const Layout* GetLayout( int32 i ) const;
 
         //!
-        void SetLayout( int i, Ptr<const Layout> );
+        void SetLayout( int32 i, Ptr<const Layout> );
         //! \}
 
         //! \name Skeleton information
@@ -288,16 +288,16 @@ namespace mu
         //! \{
 
         //!
-        void SetTagCount( int count );
+        void SetTagCount( int32 count );
 
         //!
-        int GetTagCount() const;
+        int32 GetTagCount() const;
 
         //!
-        const FString& GetTag( int tagIndex ) const;
+        const FString& GetTag( int32 tagIndex ) const;
 
         //!
-        void SetTag( int tagIndex, const FString& Name );
+        void SetTag( int32 tagIndex, const FString& Name );
 
 		//!
 		void AddStreamedResource(uint64 ResourceId);
@@ -444,6 +444,7 @@ namespace mu
 		{
 			bool bEqual = true;
 
+			if (bEqual) bEqual = (ReferenceID == o.ReferenceID);
 			if (bEqual) bEqual = (MeshIDPrefix == o.MeshIDPrefix);
 			if (bEqual) bEqual = (IndexBuffers == o.IndexBuffers);
 			if (bEqual) bEqual = (VertexBuffers == o.VertexBuffers);
@@ -516,7 +517,7 @@ namespace mu
 		void GetVertexMap( const Mesh& other, FVertexMatchMap& vertexMap, float tolerance = 1e-3f) const;
 
 		//! Compare the vertex attributes to check if they match.
-		UE::Math::TIntVector3<uint32> GetFaceVertexIndices(int f) const;
+		UE::Math::TIntVector3<uint32> GetFaceVertexIndices(int32 f) const;
 
 		//! Return true if the given mesh has the same vertex and index formats, and in the same
 		//! buffer structure.
diff --git a/Engine/Plugins/Experimental/Mutable/Source/MutableRuntime/Public/MuR/System.h b/Engine/Plugins/Experimental/Mutable/Source/MutableRuntime/Public/MuR/System.h
--- a/Engine/Plugins/Experimental/Mutable/Source/MutableRuntime/Public/MuR/System.h
+++ b/Engine/Plugins/Experimental/Mutable/Source/MutableRuntime/Public/MuR/System.h
@@ -152,18 +152,15 @@ namespace mu
         //!
         //! \warning a call to BeginUpdate must be paired with a call to EndUpdate once the returned
         //! data has been processed.
-        //! \param instanceID The id of the instance to update, as created by a NewInstance call.
-        //! \param pParams The parameters that customise this instance.
-        //! \param stateIndex The index of the state this instance will be set to. The states range
+        //! \param InstanceID The id of the instance to update, as created by a NewInstance call.
+        //! \param Params The parameters that customise this instance.
+        //! \param StateIndex The index of the state this instance will be set to. The states range
         //! from 0 to Model::GetStateCount-1
-        //! \param lodMask Bitmask selecting the levels of detail to build (i-th bit selects i-th lod).
+        //! \param LodMask Bitmask selecting the levels of detail to build (i-th bit selects i-th lod).
         //! \return the instance data with all the LOD, components, and ids to generate the meshes 
-        //! and images.  The returned Instance is only valid until the next call to EndUpdate with 
+        //! and images. The returned Instance is only valid until the next call to EndUpdate with 
         //! the same instanceID parameter.
-        const Instance* BeginUpdate(Instance::ID InstanceID,
-                                    const Ptr<const Parameters>& Params,
-                                    int32 StateIndex,
-                                    uint32 LodMask);
+        Ptr<const Instance> BeginUpdate(Instance::ID InstanceID, const Ptr<const Parameters>& Params, int32 StateIndex, uint32 LodMask);
 
 		//! Only valid between BeginUpdate and EndUpdate
 		//! Calculate the description of an image, without generating it.
diff --git a/Engine/Plugins/Experimental/Mutable/Source/MutableTools/Internal/MuT/NodeMeshConstantPrivate.h b/Engine/Plugins/Experimental/Mutable/Source/MutableTools/Internal/MuT/NodeMeshConstantPrivate.h
--- a/Engine/Plugins/Experimental/Mutable/Source/MutableTools/Internal/MuT/NodeMeshConstantPrivate.h
+++ b/Engine/Plugins/Experimental/Mutable/Source/MutableTools/Internal/MuT/NodeMeshConstantPrivate.h
@@ -21,9 +21,6 @@ namespace mu
 
 		TArray<Ptr<NodeLayout>> Layouts;
 
-		// NodeMesh::Private interface
-		Ptr<NodeLayout> GetLayout(int32 index) const override;
-
 	};
 
 }
diff --git a/Engine/Plugins/Experimental/Mutable/Source/MutableTools/Internal/MuT/NodeMeshFormatPrivate.h b/Engine/Plugins/Experimental/Mutable/Source/MutableTools/Internal/MuT/NodeMeshFormatPrivate.h
--- a/Engine/Plugins/Experimental/Mutable/Source/MutableTools/Internal/MuT/NodeMeshFormatPrivate.h
+++ b/Engine/Plugins/Experimental/Mutable/Source/MutableTools/Internal/MuT/NodeMeshFormatPrivate.h
@@ -4,6 +4,7 @@
 
 #include "MuT/NodeMeshPrivate.h"
 #include "MuT/NodeMeshFormat.h"
+#include "MuR/Mesh.h"
 
 namespace mu
 {
@@ -24,9 +25,6 @@ namespace mu
 		
 		/** */
 		bool bOptimizeBuffers = false;
-
-		// NodeMesh::Private interface
-        Ptr<NodeLayout> GetLayout( int32 index ) const override;
 	};
 
 
diff --git a/Engine/Plugins/Experimental/Mutable/Source/MutableTools/Internal/MuT/NodeMeshFragmentPrivate.h b/Engine/Plugins/Experimental/Mutable/Source/MutableTools/Internal/MuT/NodeMeshFragmentPrivate.h
--- a/Engine/Plugins/Experimental/Mutable/Source/MutableTools/Internal/MuT/NodeMeshFragmentPrivate.h
+++ b/Engine/Plugins/Experimental/Mutable/Source/MutableTools/Internal/MuT/NodeMeshFragmentPrivate.h
@@ -27,10 +27,6 @@ namespace mu
 		TArray<uint64> Blocks;
 
         FRAGMENT_TYPE m_fragmentType;
-
-
-		// NodeMesh::Private interface
-        Ptr<NodeLayout> GetLayout( int32 index ) const override;
 	};
 
 
diff --git a/Engine/Plugins/Experimental/Mutable/Source/MutableTools/Internal/MuT/NodeMeshMakeMorphPrivate.h b/Engine/Plugins/Experimental/Mutable/Source/MutableTools/Internal/MuT/NodeMeshMakeMorphPrivate.h
--- a/Engine/Plugins/Experimental/Mutable/Source/MutableTools/Internal/MuT/NodeMeshMakeMorphPrivate.h
+++ b/Engine/Plugins/Experimental/Mutable/Source/MutableTools/Internal/MuT/NodeMeshMakeMorphPrivate.h
@@ -19,11 +19,6 @@ namespace mu
 		NodeMeshPtr m_pTarget;
 		
 		bool bOnlyPositionAndNormal = false;
-
-
-		// NodeMesh::Private interface
-        NodeLayoutPtr GetLayout( int index ) const override;
-
 	};
 
 
diff --git a/Engine/Plugins/Experimental/Mutable/Source/MutableTools/Internal/MuT/NodeMeshMorphPrivate.h b/Engine/Plugins/Experimental/Mutable/Source/MutableTools/Internal/MuT/NodeMeshMorphPrivate.h
--- a/Engine/Plugins/Experimental/Mutable/Source/MutableTools/Internal/MuT/NodeMeshMorphPrivate.h
+++ b/Engine/Plugins/Experimental/Mutable/Source/MutableTools/Internal/MuT/NodeMeshMorphPrivate.h
@@ -26,10 +26,6 @@ namespace mu
 		
 		TArray<FBoneName> BonesToDeform;
 		TArray<FBoneName> PhysicsToDeform;
-
-		// NodeMesh::Private interface
-        Ptr<NodeLayout> GetLayout( int index ) const override;
-
 	};
 
 
diff --git a/Engine/Plugins/Experimental/Mutable/Source/MutableTools/Internal/MuT/NodeMeshPrivate.h b/Engine/Plugins/Experimental/Mutable/Source/MutableTools/Internal/MuT/NodeMeshPrivate.h
--- a/Engine/Plugins/Experimental/Mutable/Source/MutableTools/Internal/MuT/NodeMeshPrivate.h
+++ b/Engine/Plugins/Experimental/Mutable/Source/MutableTools/Internal/MuT/NodeMeshPrivate.h
@@ -5,19 +5,12 @@
 #include "MuT/NodePrivate.h"
 
 #include "MuT/NodeMesh.h"
-#include "MuT/NodeLayout.h"
 
 
 namespace mu
 {
-
-
 	class NodeMesh::Private : public Node::Private
 	{
-	public:
-
-		virtual NodeLayoutPtr GetLayout( int index ) const = 0;
-
 	};
 
 }
diff --git a/Engine/Plugins/Experimental/Mutable/Source/MutableTools/Internal/MuT/NodeMeshTablePrivate.h b/Engine/Plugins/Experimental/Mutable/Source/MutableTools/Internal/MuT/NodeMeshTablePrivate.h
--- a/Engine/Plugins/Experimental/Mutable/Source/MutableTools/Internal/MuT/NodeMeshTablePrivate.h
+++ b/Engine/Plugins/Experimental/Mutable/Source/MutableTools/Internal/MuT/NodeMeshTablePrivate.h
@@ -25,10 +25,6 @@ namespace mu
 		FString DefaultRowName;
 
 		TArray<NodeLayoutPtr> Layouts;
-
-		// NodeMesh::Private interface
-        NodeLayoutPtr GetLayout( int index ) const override;
-
 	};
 
 }
diff --git a/Engine/Plugins/Experimental/Mutable/Source/MutableTools/Internal/MuT/NodeModifierPrivate.h b/Engine/Plugins/Experimental/Mutable/Source/MutableTools/Internal/MuT/NodeModifierPrivate.h
--- a/Engine/Plugins/Experimental/Mutable/Source/MutableTools/Internal/MuT/NodeModifierPrivate.h
+++ b/Engine/Plugins/Experimental/Mutable/Source/MutableTools/Internal/MuT/NodeModifierPrivate.h
@@ -14,15 +14,6 @@ namespace mu
 	{
 	public:
 
-		/** Tags that target surface need to have enabled to receive this modifier. */
-		TArray<FString> RequiredTags;
-
-		/** In case of multiple tags in RequiredTags: are they all required, or one is enough? */
-		EMutableMultipleTagPolicy MultipleTagsPolicy = EMutableMultipleTagPolicy::OnlyOneRequired;
-
-		// Wether the modifier has to be applied after the normal node operations or before
-		bool bApplyBeforeNormalOperations = true;
-
 	};
 
 }
diff --git a/Engine/Plugins/Experimental/Mutable/Source/MutableTools/Internal/MuT/NodeObjectGroupPrivate.h b/Engine/Plugins/Experimental/Mutable/Source/MutableTools/Internal/MuT/NodeObjectGroupPrivate.h
--- a/Engine/Plugins/Experimental/Mutable/Source/MutableTools/Internal/MuT/NodeObjectGroupPrivate.h
+++ b/Engine/Plugins/Experimental/Mutable/Source/MutableTools/Internal/MuT/NodeObjectGroupPrivate.h
@@ -30,9 +30,6 @@ namespace mu
 		TArray<NodeObjectPtr> m_children;
 		int32 DefaultValue;
 
-        // NodeObject::Private interface
-        NodeLayoutPtr GetLayout( int lod, int component, int surface, int texture ) const override;
-
-	};
+ 	};
 
 }
diff --git a/Engine/Plugins/Experimental/Mutable/Source/MutableTools/Internal/MuT/NodeObjectPrivate.h b/Engine/Plugins/Experimental/Mutable/Source/MutableTools/Internal/MuT/NodeObjectPrivate.h
--- a/Engine/Plugins/Experimental/Mutable/Source/MutableTools/Internal/MuT/NodeObjectPrivate.h
+++ b/Engine/Plugins/Experimental/Mutable/Source/MutableTools/Internal/MuT/NodeObjectPrivate.h
@@ -5,7 +5,6 @@
 #include "MuT/NodePrivate.h"
 
 #include "MuT/NodeObject.h"
-#include "MuT/NodeLayout.h"
 
 namespace mu
 {
@@ -13,9 +12,6 @@ namespace mu
 	class NodeObject::Private : public Node::Private
 	{
 	public:
-
-        virtual NodeLayoutPtr GetLayout( int lod, int component, int surface, int texture ) const = 0;
-
 	};
 
 }
diff --git a/Engine/Plugins/Experimental/Mutable/Source/MutableTools/Internal/MuT/NodePatchImage.h b/Engine/Plugins/Experimental/Mutable/Source/MutableTools/Internal/MuT/NodePatchImage.h
--- a/Engine/Plugins/Experimental/Mutable/Source/MutableTools/Internal/MuT/NodePatchImage.h
+++ b/Engine/Plugins/Experimental/Mutable/Source/MutableTools/Internal/MuT/NodePatchImage.h
@@ -77,7 +77,6 @@ namespace mu
         //-----------------------------------------------------------------------------------------
         class Private;
         Private* GetPrivate() const;
-        Node::Private* GetBasePrivate() const override;
 
     protected:
 
diff --git a/Engine/Plugins/Experimental/Mutable/Source/MutableTools/Internal/MuT/NodePrivate.h b/Engine/Plugins/Experimental/Mutable/Source/MutableTools/Internal/MuT/NodePrivate.h
--- a/Engine/Plugins/Experimental/Mutable/Source/MutableTools/Internal/MuT/NodePrivate.h
+++ b/Engine/Plugins/Experimental/Mutable/Source/MutableTools/Internal/MuT/NodePrivate.h
@@ -16,49 +16,39 @@ namespace mu
 		/** Force a virtual destructor. */
 		virtual ~Private() = default; 
 
-        /** This is an opaque context used to attach to reported error messages. */
-		const void* m_errorContext = nullptr;
-
 		/** Generic pointer to the node owning this private. */
 		const Node* m_pNode = nullptr;
     };
 
 
 
-#define MUTABLE_IMPLEMENT_NODE( N, T, PN, PT )				\
-    N::N()													\
-    {														\
-		Type = T;											\
-        m_pD = new Private();								\
-		m_pD->m_pNode = this;								\
-		PN::Type = PT;										\
-	}														\
-                                                            \
-    N::~N()													\
-    {														\
-        check( m_pD );										\
-        delete m_pD;										\
-        m_pD = nullptr;										\
-    }														\
-                                                            \
-    N::Private* N::GetPrivate() const						\
-    {														\
-        return m_pD;										\
-    }														\
-                                                            \
-    Node::Private* N::GetBasePrivate() const				\
-    {														\
-        return m_pD;										\
-    }														\
-                                                            \
-    const FNodeType* N::GetType() const						\
-    {														\
-        return GetStaticType();								\
-    }														\
-                                                            \
-    const FNodeType* N::GetStaticType()						\
-    {														\
-        return &Private::s_type;							\
+#define MUTABLE_IMPLEMENT_NODE( N )				\
+    N::N()										\
+    {											\
+        m_pD = new Private();					\
+		m_pD->m_pNode = this;					\
+	}											\
+                                                \
+    N::~N()										\
+    {											\
+        check( m_pD );							\
+        delete m_pD;							\
+        m_pD = nullptr;							\
+    }											\
+                                                \
+    N::Private* N::GetPrivate() const			\
+    {											\
+        return m_pD;							\
+    }											\
+                                                \
+    const FNodeType* N::GetType() const			\
+    {											\
+        return GetStaticType();					\
+    }											\
+                                                \
+    const FNodeType* N::GetStaticType()			\
+    {											\
+        return &Private::s_type;				\
     }
 
 }
diff --git a/Engine/Plugins/Experimental/Mutable/Source/MutableTools/Internal/MuT/NodeSurfaceNewPrivate.h b/Engine/Plugins/Experimental/Mutable/Source/MutableTools/Internal/MuT/NodeSurfaceNewPrivate.h
--- a/Engine/Plugins/Experimental/Mutable/Source/MutableTools/Internal/MuT/NodeSurfaceNewPrivate.h
+++ b/Engine/Plugins/Experimental/Mutable/Source/MutableTools/Internal/MuT/NodeSurfaceNewPrivate.h
@@ -29,13 +29,7 @@ namespace mu
         uint32 ExternalId =0;
         int32 SharedSurfaceId =INDEX_NONE;
 
-		struct FMeshData
-		{
-			FString m_name;
-			NodeMeshPtr m_pMesh;
-		};
-
-		TArray<FMeshData> m_meshes;
+		Ptr<NodeMesh> Mesh;
 
 		struct FImageData
 		{
@@ -80,9 +74,6 @@ namespace mu
 		//! Find an image node index by name or return -1
 		int32 FindImage( const FString& strName ) const;
 
-		//! Find a mesh node index by name or return -1
-		int32 FindMesh(const FString& strName) const;
-
         //! Find a vector node index by name or return -1
         int32 FindVector(const FString& strName) const;
 
diff --git a/Engine/Plugins/Experimental/Mutable/Source/MutableTools/Private/MuT/CodeGenerator.cpp b/Engine/Plugins/Experimental/Mutable/Source/MutableTools/Private/MuT/CodeGenerator.cpp
--- a/Engine/Plugins/Experimental/Mutable/Source/MutableTools/Private/MuT/CodeGenerator.cpp
+++ b/Engine/Plugins/Experimental/Mutable/Source/MutableTools/Private/MuT/CodeGenerator.cpp
@@ -48,14 +48,12 @@
 #include "MuT/NodeColour.h"
 #include "MuT/NodeColourConstant.h"
 #include "MuT/NodeComponent.h"
-#include "MuT/NodeComponentNewPrivate.h"
 #include "MuT/NodeImage.h"
 #include "MuT/NodeImageFormat.h"
 #include "MuT/NodeImageFormatPrivate.h"
 #include "MuT/NodeImageMipmap.h"
 #include "MuT/NodeImageMipmapPrivate.h"
 #include "MuT/NodeImageSwizzlePrivate.h"
-#include "MuT/NodeLODPrivate.h"
 #include "MuT/NodeMesh.h"
 #include "MuT/NodeMeshClipMorphPlane.h"
 #include "MuT/NodeMeshClipWithMesh.h"
@@ -74,7 +72,7 @@
 #include "MuT/NodeModifierPrivate.h"
 #include "MuT/NodeObject.h"
 #include "MuT/NodeObjectGroupPrivate.h"
-#include "MuT/NodeObjectNewPrivate.h"
+#include "MuT/NodeObjectNew.h"
 #include "MuT/NodePatchImagePrivate.h"
 #include "MuT/NodePatchMesh.h"
 #include "MuT/NodePrivate.h"
@@ -118,7 +116,7 @@ namespace mu
 
 		// Second pass
 		SecondPassGenerator SecondPass(&m_firstPass, m_compilerOptions);
-		bool bSuccess = SecondPass.Generate(m_pErrorLog, pNode->GetBasePrivate());
+		bool bSuccess = SecondPass.Generate(m_pErrorLog, pNode.get());
 		if (!bSuccess)
 		{
 			return;
@@ -129,14 +127,14 @@ namespace mu
 			MUTABLE_CPUPROFILER_SCOPE(MainPass);
 
             int32 CurrentStateIndex = 0;
-			for (const TPair<FObjectState, const Node::Private*>& s : m_firstPass.m_states)
+			for (const TPair<FObjectState, const Node*>& s : m_firstPass.m_states)
 			{
 				MUTABLE_CPUPROFILER_SCOPE(MainPassState);
 
 				FGenericGenerationOptions Options;
 				Options.State = CurrentStateIndex;
 
-				Ptr<ASTOp> stateRoot = Generate(pNode, Options);
+				Ptr<ASTOp> stateRoot = Generate_Generic(pNode, Options);
 				m_states.Emplace(s.Key, stateRoot);
 
 				++CurrentStateIndex;
@@ -145,7 +143,7 @@ namespace mu
 	}
 
 
-	Ptr<ASTOp> CodeGenerator::Generate(const Ptr<const Node> pNode, const FGenericGenerationOptions& Options )
+	Ptr<ASTOp> CodeGenerator::Generate_Generic(const Ptr<const Node> pNode, const FGenericGenerationOptions& Options )
 	{
 		if (!pNode)
 		{
@@ -184,7 +182,8 @@ namespace mu
 			// This happens only if we generate a node graph that has a NodeSurfaceNew at the root.
 			FSurfaceGenerationResult surfResult;
 			const TArray<FirstPassGenerator::FSurface::FEdit> edits;
-			GenerateSurface(surfResult, Options, surfNode, edits);
+			FSurfaceGenerationOptions SurfaceOptions(Options);
+			GenerateSurface(surfResult, SurfaceOptions, surfNode, edits);
 			return surfResult.surfaceOp;
 		}
 
@@ -212,6 +211,14 @@ namespace mu
 			return nullptr;
 		}
 
+		else if (pNode->GetType()->IsA(NodeComponent::GetStaticType()))
+		{
+			const NodeComponent* ComponentNode = static_cast<const NodeComponent*>(pNode.get());
+			FGenericGenerationResult Result;
+			GenerateComponent(Options, Result, ComponentNode);
+			return Result.op;
+		}
+
 
 		Ptr<ASTOp> ResultOp;
 
@@ -230,15 +237,7 @@ namespace mu
 			FGenericGenerationResult Result;
 
 			// Generate for each different type of node
-			if (pNode->GetType() == NodeComponentNew::GetStaticType())
-			{
-				Generate_ComponentNew(Options, Result, static_cast<const NodeComponentNew*>(pNode.get()));
-			}
-			else if (pNode->GetType() == NodeComponentEdit::GetStaticType())
-			{
-				// Nothing to do: processed elsewhere.
-			}
-			else if (pNode->GetType() == NodeObjectNew::GetStaticType())
+			if (pNode->GetType() == NodeObjectNew::GetStaticType())
 			{
 				Generate_ObjectNew(Options, Result, static_cast<const NodeObjectNew*>(pNode.get()));
 			}
@@ -443,7 +442,7 @@ namespace mu
 			}
 			else
 			{
-				blend = GenerateMissingImageCode(TEXT("Blend top image"), EImageFormat::IF_RGB_UBYTE, pPatch->GetPrivate()->m_errorContext, ImageOptions);
+				blend = GenerateMissingImageCode(TEXT("Blend top image"), EImageFormat::IF_RGB_UBYTE, pPatch->GetMessageContext(), ImageOptions);
 			}
 			blend = GenerateImageFormat(blend, blockAd->GetImageDesc().m_format);
 			blend = GenerateImageSize(blend, ImageOptions.RectSize);
@@ -495,78 +494,76 @@ namespace mu
 	}
 
 
-	void CodeGenerator::Generate_LOD(const FGenericGenerationOptions& Options, FGenericGenerationResult& Result, const NodeLOD* InNode)
-	{
-		const NodeLOD::Private& node = *InNode->GetPrivate();
 
-		// Build a series of operations to assemble all the LOD components
-		Ptr<ASTOp> lastCompOp;
+	//---------------------------------------------------------------------------------------------
+	void CodeGenerator::GenerateComponent(const FGenericGenerationOptions& InOptions, FGenericGenerationResult& OutResult, const NodeComponent* InUntypedNode)
+	{
+		if (!InUntypedNode)
+		{
+			OutResult = FGenericGenerationResult();
+			return;
+		}
 
-		// Create the expression for each component in this object
-		// TODO: More components per operation
-		for (int32 t = 0; t < node.m_components.Num(); ++t)
+		// See if it was already generated
+		FGeneratedComponentCacheKey Key;
+		Key.Node = InUntypedNode;
+		Key.Options = InOptions;
+		GeneratedComponentMap::ValueType* it = GeneratedComponents.Find(Key);
+		if (it)
 		{
-			if (const NodeComponent* pComponentNode = node.m_components[t].get())
-			{
-				m_currentParents.Last().m_component = (int)t;
+			OutResult = *it;
+			return;
+		}
 
-				Ptr<ASTOp> componentOp = Generate(pComponentNode, Options);
+		// Generate for each different type of node
+		const FNodeType* Type = InUntypedNode->GetType();
+		if (Type == NodeComponentNew::GetStaticType())
+		{
+			GenerateComponent_New(InOptions, OutResult, static_cast<const NodeComponentNew*>(InUntypedNode));
+		}
+		else if (Type == NodeComponentEdit::GetStaticType())
+		{
+			// Nothing to do because it is all preprocessed in the first code generator stage
+			//GenerateComponent_Edit(InOptions, OutResult, static_cast<const NodeComponentEdit*>(InUntypedNode));
+		}
+		else
+		{
+			check(false);
+		}
 
-				if (componentOp)
-				{
-					check(componentOp->GetOpType() == OP_TYPE::IN_ADDCOMPONENT);
+		// Cache the result
+		GeneratedComponents.Add(Key, OutResult);
+	}
 
-					ASTOpInstanceAdd* typedOp = static_cast<ASTOpInstanceAdd*>(componentOp.get());
 
-					// Complete the instruction adding the base
-					typedOp->instance = lastCompOp;
+	void CodeGenerator::GenerateComponent_New(const FGenericGenerationOptions& Options, FGenericGenerationResult& Result, const NodeComponentNew* InNode)
+	{
+		const NodeComponentNew& node = *InNode;
 
-					lastCompOp = componentOp;
-				}
-			}
-		}
+		// Create the expression for each component in this object
+		Ptr<ASTOpAddLOD> LODsOp = new ASTOpAddLOD();
 
-		// Add components from child objects
-		FAdditionalComponentKey thisKey;
-		thisKey.m_lod = m_currentParents.Last().m_lod;
-		thisKey.m_pObject = m_currentParents.Last().m_pObject;
-		TArray<Ptr<ASTOp>>* addIt = AdditionalComponents.Find(thisKey);
-		if (addIt)
+		for (int32 LODIndex = 0; LODIndex < node.LODs.Num(); ++LODIndex)
 		{
-			for (const auto& cop : *addIt)
+			if (const NodeLOD* LODNode = node.LODs[LODIndex].get())
 			{
-				// Add the additional components after the main ones, this means higher up in the
-				// op tree.
+				m_currentParents.Last().m_lod = LODIndex;
 
-                // First find the last op in the chain of IN_ADDCOMPONENT operations
-                check( cop->GetOpType() == OP_TYPE::IN_ADDCOMPONENT );
-				ASTOpInstanceAdd* typedOp = static_cast<ASTOpInstanceAdd*>(cop.get());
-				ASTOpInstanceAdd* bottomOp = typedOp;
-				while (bottomOp->instance)
-				{
-					// Step down
-					check(bottomOp->instance->GetOpType() == OP_TYPE::IN_ADDCOMPONENT);
-					bottomOp = static_cast<ASTOpInstanceAdd*>(bottomOp->instance.child().get());
-				}
+				FLODGenerationOptions LODOptions(Options, LODIndex, InNode );
+				FGenericGenerationResult LODResult;
+				Generate_LOD(LODOptions, LODResult, LODNode);
 
-				// Chain
-				bottomOp->instance = lastCompOp;
-				lastCompOp = typedOp;
+				LODsOp->lods.Emplace(LODsOp, LODResult.op);
 			}
 		}
 
-		// Store for possible parent objects if necessary
-		// 2 is because there must be a parent and there is always a null element as well.
-		if (lastCompOp && m_currentParents.Num() > 2)
-		{
-			const auto& parentObjectKey = m_currentParents[m_currentParents.Num() - 2];
-			FAdditionalComponentKey parentKey;
-			parentKey.m_lod = m_currentParents.Last().m_lod;
-			parentKey.m_pObject = parentObjectKey.m_pObject;
-			AdditionalComponents.FindOrAdd(parentKey).Add(lastCompOp);
-		}
+		Ptr<ASTOpInstanceAdd> InstanceOp = new ASTOpInstanceAdd();
+		InstanceOp->type = OP_TYPE::IN_ADDCOMPONENT;
+		InstanceOp->instance = nullptr;
+		InstanceOp->value = LODsOp;
+		InstanceOp->id = node.Id;
 
-		Result.op = lastCompOp;
+		Result.op = InstanceOp;
 	}
 
 
@@ -628,7 +625,7 @@ namespace mu
 
     //---------------------------------------------------------------------------------------------
     void CodeGenerator::GenerateSurface( FSurfaceGenerationResult& result, 
-										 const FGenericGenerationOptions& Options,
+										 const FSurfaceGenerationOptions& Options,
                                          Ptr<const NodeSurfaceNew> surfaceNode,
                                          const TArray<FirstPassGenerator::FSurface::FEdit>& edits )
     {
@@ -636,31 +633,13 @@ namespace mu
 
         const NodeSurfaceNew::Private& node = *surfaceNode->GetPrivate();
 
-        // Clear the surface generation state
-//        m_generatedLayouts.Empty();
-//        m_absoluteLayoutIndex = 0;
-//        m_currentLayoutMesh = nullptr;
-//        m_currentLayoutChannel = 0;
-//        m_generatedMeshes.Empty();
-
         // Build a series of operations to assemble the surface
         Ptr<ASTOp> lastSurfOp;
 
         // Generate the mesh
         //------------------------------------------------------------------------
-        NodeMeshPtr pMesh;
-        if ( node.m_meshes.Num()>0 )
-        {
-            pMesh = node.m_meshes[0].m_pMesh;
-
-            // TODO: Support more meshes
-            check( node.m_meshes.Num()==1 );
-        }
-
-
         FMeshGenerationResult meshResults;
 
-
 		// We don't add the mesh here, since it will be added directly at the top of the
 		// component expression in the NodeComponentNew visitor with the right merges
 		// and conditions.
@@ -707,7 +686,7 @@ namespace mu
 		// If this is true, we will reuse the surface properties from a higher LOD, se we can skip the generation of material properties and images.
 		const bool bShareSurface = node.SharedSurfaceId != INDEX_NONE && !bIsBaseForSharedSurface;
 
-        if ( pMesh )
+        if (node.Mesh)
         {
             MUTABLE_CPUPROFILER_SCOPE(SurfaceMesh);
 
@@ -742,7 +721,7 @@ namespace mu
 			// that may cause them to fall on a different block.
 			MeshOptions.bClampUVIslands = bShareSurface && bNormalizeUVs;
 
-            GenerateMesh(MeshOptions, meshResults, pMesh );
+            GenerateMesh(MeshOptions, meshResults, node.Mesh);
             lastMeshOp = meshResults.meshOp;
             meshResults.extraMeshLayouts.SetNum( edits.Num() );
 
@@ -750,7 +729,7 @@ namespace mu
 
             // Let's remember the base mesh of each added mesh, so that we can use them for the
             // "remove mesh" operations that apply to them instead of the base.
-            std::map<const Node::Private*,Ptr<ASTOp>> baseMeshesForEachAddedMesh;
+            std::map<const Node*,Ptr<ASTOp>> baseMeshesForEachAddedMesh;
 
 
             // Apply mesh merges from child objects "edit surface" nodes
@@ -758,19 +737,19 @@ namespace mu
             {
                 const FirstPassGenerator::FSurface::FEdit& e = edits[editIndex];
 
-                if ( e.node->m_pMesh )
+                if ( e.node->GetPrivate()->m_pMesh)
                 {
-                    if ( NodeMeshPtr pAdd = e.node->m_pMesh->GetAdd() )
+                    if ( NodeMeshPtr pAdd = e.node->GetPrivate()->m_pMesh->GetAdd() )
                     {
 						// Store the data necessary to apply modifiers for the pre-normal operations stage.
-						m_activeTags.Add(e.node->m_tags);
+						m_activeTags.Add(e.node->GetPrivate()->m_tags);
 						
 						FMeshGenerationOptions MergedMeshOptions;
 						MergedMeshOptions.bLayouts = true;
 						MergedMeshOptions.bClampUVIslands = bShareSurface && bNormalizeUVs;
 						MergedMeshOptions.bNormalizeUVs = bNormalizeUVs;
 						MergedMeshOptions.State = Options.State;
-						MergedMeshOptions.ActiveTags = e.node->m_tags;
+						MergedMeshOptions.ActiveTags = e.node->GetPrivate()->m_tags;
 
 						if (SharedMeshResults)
 						{
@@ -788,8 +767,8 @@ namespace mu
 						// Apply the modifier for the post-normal operations stage to the added mesh
 						bool bModifiersForBeforeOperations = false;
 						FGenericGenerationOptions ModifierOptions(Options);
-						ModifierOptions.ActiveTags = e.node->m_tags;
-						lastMeshOp = ApplyMeshModifiers(Options, lastMeshOp, bModifiersForBeforeOperations, node.m_errorContext);
+						ModifierOptions.ActiveTags = e.node->GetPrivate()->m_tags;
+						lastMeshOp = ApplyMeshModifiers(Options, lastMeshOp, bModifiersForBeforeOperations, surfaceNode->GetMessageContext());
 
                         FMeshGenerationResult::FExtraLayouts data;
 						data.GeneratedLayouts = addResults.GeneratedLayouts;
@@ -830,15 +809,15 @@ namespace mu
             Ptr<ASTOpMeshRemoveMask> rop;
             for ( const FirstPassGenerator::FSurface::FEdit& e: edits )
             {
-                if ( e.node->m_pMesh )
+                if ( e.node->GetPrivate()->m_pMesh )
                 {
-                    if ( NodeMeshPtr pRemove = e.node->m_pMesh->GetRemove() )
+                    if ( NodeMeshPtr pRemove = e.node->GetPrivate()->m_pMesh->GetRemove() )
                     {
                         FMeshGenerationResult removeResults;
 						FMeshGenerationOptions RemoveMeshOptions;
 						RemoveMeshOptions.bLayouts = false;
 						RemoveMeshOptions.State = Options.State;
-						RemoveMeshOptions.ActiveTags = e.node->m_tags;
+						RemoveMeshOptions.ActiveTags = e.node->GetPrivate()->m_tags;
 
                         GenerateMesh(RemoveMeshOptions, removeResults, pRemove );
 
@@ -848,9 +827,9 @@ namespace mu
                         // By default, remove from the base
                         Ptr<ASTOp> removeFrom = meshResults.baseMeshOp;
                         // See if we want to remove from an added mesh instead.
-                        if ( e.node->m_pParent )
+                        if ( e.node->GetPrivate()->m_pParent )
                         {
-                            auto addedBaseMeshIt = baseMeshesForEachAddedMesh.find(e.node->m_pParent->GetBasePrivate());
+                            auto addedBaseMeshIt = baseMeshesForEachAddedMesh.find(e.node->GetPrivate()->m_pParent.get());
                             if (addedBaseMeshIt!=baseMeshesForEachAddedMesh.end())
                             {
                                 removeFrom = addedBaseMeshIt->second;
@@ -879,7 +858,7 @@ namespace mu
             // Apply mesh morphs from child objects "edit surface" nodes
             for ( const FirstPassGenerator::FSurface::FEdit& e: edits )
             {
-                if ( NodeMeshPtr pMorph = e.node->m_pMorph )
+                if ( NodeMeshPtr pMorph = e.node->GetPrivate()->m_pMorph )
                 {
 					// Not needed because it has been generated before already.
 					// Base mesh
@@ -915,10 +894,10 @@ namespace mu
                         Ptr<ASTOpMeshMorph> op = new ASTOpMeshMorph();
 
 						// Factor
-						if (e.node->m_pFactor)
+						if (e.node->GetPrivate()->m_pFactor)
 						{
 							FScalarGenerationResult ChildResult;
-							GenerateScalar(ChildResult, Options, e.node->m_pFactor);
+							GenerateScalar(ChildResult, Options, e.node->GetPrivate()->m_pFactor);
 							op->Factor = ChildResult.op;
 						}
 						else
@@ -960,7 +939,7 @@ namespace mu
 			bool bModifiersForBeforeOperations = false;
 			FGenericGenerationOptions ModifierOptions(Options);
 			ModifierOptions.ActiveTags = node.m_tags;
-			lastMeshOp = ApplyMeshModifiers(Options, lastMeshOp, bModifiersForBeforeOperations, node.m_errorContext);
+			lastMeshOp = ApplyMeshModifiers(Options, lastMeshOp, bModifiersForBeforeOperations, surfaceNode->GetMessageContext());
 
             // Layouts
             for ( int32 LayoutIndex=0; LayoutIndex <meshResults.GeneratedLayouts.Num(); ++LayoutIndex)
@@ -1000,7 +979,7 @@ namespace mu
 
 							if (data.GeneratedLayouts.Num() != meshResults.GeneratedLayouts.Num())
 							{
-								m_pErrorLog->GetPrivate()->Add(TEXT("Merged layout has been ignored because the number of layouts is different."), ELMT_ERROR, node.m_errorContext);
+								m_pErrorLog->GetPrivate()->Add(TEXT("Merged layout has been ignored because the number of layouts is different."), ELMT_ERROR, surfaceNode->GetMessageContext());
 							}
 							else
 							{
@@ -1163,9 +1142,9 @@ namespace mu
 						for (int32 editIndex = 0; editIndex < edits.Num(); ++editIndex)
 						{
 							const FirstPassGenerator::FSurface::FEdit& e = edits[editIndex];
-							if (t < e.node->m_textures.Num())
+							if (t < e.node->GetPrivate()->m_textures.Num())
 							{
-								if (const NodePatchImage* pPatch = e.node->m_textures[t].m_pPatch.get())
+								if (const NodePatchImage* pPatch = e.node->GetPrivate()->m_textures[t].m_pPatch.get())
 								{
 									imageAd = GenerateImageBlockPatch(imageAd, pPatch, e.condition, ImageOptions);
 								}
@@ -1227,7 +1206,7 @@ namespace mu
 						if (LayoutIndex >= meshResults.GeneratedLayouts.Num() ||
 							LayoutIndex >= meshResults.layoutOps.Num())
 						{
-							m_pErrorLog->GetPrivate()->Add("Missing layout in object, or its parent.", ELMT_ERROR, node.m_errorContext);
+							m_pErrorLog->GetPrivate()->Add("Missing layout in object, or its parent.", ELMT_ERROR, surfaceNode->GetMessageContext());
 						}
 						else
 						{
@@ -1263,7 +1242,7 @@ namespace mu
 								imageAd = BlankImageOp;
 							}
 
-							auto UpdateBlockSize = [&BlockPixelsX, &BlockPixelsY, &FinalFormat, &bBlocksHaveMips, &bImageSizeWarning, &formatNode, &BlankImageOp, &node, &t, this]( FImageDesc BlockDesc, UE::Math::TIntVector2<uint16> LayoutCellSize )
+							auto UpdateBlockSize = [&BlockPixelsX, &BlockPixelsY, &FinalFormat, &bBlocksHaveMips, &bImageSizeWarning, &formatNode, &BlankImageOp, &node, surfaceNode, &t, this]( FImageDesc BlockDesc, UE::Math::TIntVector2<uint16> LayoutCellSize )
 							{
 								if (BlockPixelsX == 0)
 								{
@@ -1280,7 +1259,7 @@ namespace mu
 												*node.m_images[t].m_materialName,
 												*node.m_images[t].m_materialParameterName,
 												currentLOD);
-											m_pErrorLog->GetPrivate()->Add(Msg, ELMT_INFO, node.m_errorContext);
+											m_pErrorLog->GetPrivate()->Add(Msg, ELMT_INFO, surfaceNode->GetMessageContext());
 										}
 									}
 
@@ -1351,9 +1330,9 @@ namespace mu
 								for (int32 editIndex = 0; editIndex < edits.Num(); ++editIndex)
 								{
 									const FirstPassGenerator::FSurface::FEdit& e = edits[editIndex];
-									if (t < e.node->m_textures.Num())
+									if (t < e.node->GetPrivate()->m_textures.Num())
 									{
-										if (const NodePatchImage* pPatch = e.node->m_textures[t].m_pPatch.get())
+										if (const NodePatchImage* pPatch = e.node->GetPrivate()->m_textures[t].m_pPatch.get())
 										{
 											// Is the current block to be patched?
 											if (pPatch->GetPrivate()->BlockIndices.Contains(BlockIndex))
@@ -1392,9 +1371,9 @@ namespace mu
 							for (int32 editIndex = 0; editIndex < edits.Num(); ++editIndex)
 							{
 								const FirstPassGenerator::FSurface::FEdit& e = edits[editIndex];
-								if (t < e.node->m_textures.Num())
+								if (t < e.node->GetPrivate()->m_textures.Num())
 								{
-									Ptr<NodeImage> pExtend = e.node->m_textures[t].m_pExtend;
+									Ptr<NodeImage> pExtend = e.node->GetPrivate()->m_textures[t].m_pExtend;
 									if (pExtend)
 									{
 										if (LayoutIndex >= meshResults.extraMeshLayouts[editIndex].GeneratedLayouts.Num()
@@ -1402,10 +1381,10 @@ namespace mu
 											!meshResults.extraMeshLayouts[editIndex].GeneratedLayouts[LayoutIndex])
 										{
 											FString Msg = FString::Printf(TEXT("Trying to extend a layout that doesn't exist in object [%s]."),
-												*m_currentParents.Last().m_pObject->m_name
+												*m_currentParents.Last().m_pObject->Name
 											);
 
-											m_pErrorLog->GetPrivate()->Add(Msg, ELMT_ERROR, node.m_errorContext);
+											m_pErrorLog->GetPrivate()->Add(Msg, ELMT_ERROR, surfaceNode->GetMessageContext());
 										}
 										else
 										{
@@ -1682,12 +1661,11 @@ namespace mu
 
 
     //---------------------------------------------------------------------------------------------
-    //Ptr<ASTOp> CodeGenerator::Visit( const NodeComponentNew::Private& node )
-	void CodeGenerator::Generate_ComponentNew(const FGenericGenerationOptions& Options, FGenericGenerationResult& Result, const NodeComponentNew* InNode)
+	void CodeGenerator::Generate_LOD(const FLODGenerationOptions& Options, FGenericGenerationResult& Result, const NodeLOD* InNode)
 	{
-		const NodeComponentNew::Private& node = *InNode->GetPrivate();
+		const NodeLOD& node = *InNode;
 
-		MUTABLE_CPUPROFILER_SCOPE(NodeComponentNew);
+		MUTABLE_CPUPROFILER_SCOPE(Generate_LOD);
 
 		// Build a series of operations to assemble the component
         Ptr<ASTOp> lastCompOp;
@@ -1697,13 +1675,15 @@ namespace mu
         // This generates a different ID for each surface. It can be used to match it to the
         // mesh surface, or for debugging. It cannot be 0 because it is a special case for the
         // merge operation.
-        int surfaceID=1;
+        int32 surfaceID=1;
 
         // Look for all surfaces that belong to this component
 		for (int32 i = 0; i<m_firstPass.surfaces.Num(); ++i, ++surfaceID)
 		{
 			const FirstPassGenerator::FSurface& its = m_firstPass.surfaces[i];
-			if (its.component==&node)
+			if (its.Component==Options.Component
+				&&
+				its.LOD==Options.LODIndex)
 			{
                 // Apply state conditions: only generate it if it enabled in this state
                 {
@@ -1726,8 +1706,9 @@ namespace mu
                 sop->name = its.node->GetPrivate()->m_name;
                 sop->instance = lastCompOp;
 
+				FSurfaceGenerationOptions SurfaceOptions(Options);
 				FSurfaceGenerationResult surfaceGenerationResult;
-                GenerateSurface( surfaceGenerationResult, Options, its.node, its.edits );
+                GenerateSurface( surfaceGenerationResult, SurfaceOptions, its.node, its.edits );
                 sop->value = surfaceGenerationResult.surfaceOp;
 
                 sop->id = surfaceID;
@@ -1757,16 +1738,10 @@ namespace mu
                 }
 
                 // Add the mesh with its condition
-                Ptr<ASTOp> mergeAd;
-                if (!lastMeshOp && its.node->GetPrivate()->m_meshes.Num())
-                {
-                    // \todo: mesh index in case of multiple meshes.
-                    const int meshIndex = 0;
-                    lastMeshName = its.node->GetPrivate()->m_meshes[ meshIndex ].m_name;
-                }
-
+ 
                 // We add the merge op even for the first mesh, so that we set the surface id.
-                {
+				Ptr<ASTOp> mergeAd;
+				{
                     Ptr<ASTOp> added = its.resultMeshOp;
  
                     Ptr<ASTOpFixed> mop = new ASTOpFixed();
@@ -1807,21 +1782,9 @@ namespace mu
             iop->instance = lastCompOp;
             iop->value = lastMeshOp;
 
-            // Name
-            // \todo: do mesh names make any sense?
-            iop->name = lastMeshName;
-
             lastCompOp = iop;
         }
 
-
-        Ptr<ASTOpInstanceAdd> iop = new ASTOpInstanceAdd();
-        iop->type = OP_TYPE::IN_ADDCOMPONENT;
-        iop->value = lastCompOp;
-        iop->name = node.m_name;
-    	iop->id = node.m_id;
-        lastCompOp = iop;
-
         Result.op = lastCompOp;
     }
 
@@ -1830,17 +1793,15 @@ namespace mu
    // Ptr<ASTOp> CodeGenerator::Visit( const NodeObjectNew::Private& node )
 	void CodeGenerator::Generate_ObjectNew(const FGenericGenerationOptions& Options, FGenericGenerationResult& Result, const NodeObjectNew* InNode)
 	{
-		const NodeObjectNew::Private& node = *InNode->GetPrivate();
-		
 		MUTABLE_CPUPROFILER_SCOPE(NodeObjectNew);
 
         m_currentParents.Add( FParentKey() );
-        m_currentParents.Last().m_pObject = &node;
+        m_currentParents.Last().m_pObject = InNode;
 
         // Parse the child objects first, which will accumulate operations in the patching lists
-        for ( int32 t=0; t<node.m_children.Num(); ++t )
+        for ( int32 t=0; t< InNode->Children.Num(); ++t )
         {
-            if ( const NodeObject* pChildNode = node.m_children[t].get() )
+            if ( const NodeObject* pChildNode = InNode->Children[t].get() )
             {
                 Ptr<ASTOp> paramOp;
 
@@ -1863,31 +1824,81 @@ namespace mu
                 data.m_condition = paramOp;
                 m_currentObject.Add( data );
 
-                // This op is ignored: everything is stored as patches to apply to the parent when
-                // it is compiled.
-                Generate( pChildNode, Options );
+                // This op is ignored: everything is stored as patches to apply to the parent when it is compiled.
+                Generate_Generic( pChildNode, Options );
 
                 m_currentObject.Pop();
             }
         }
 
-        // Create the expression for each lod
-        Ptr<ASTOpAddLOD> lodsOp = new ASTOpAddLOD();
-        for ( int32 t=0; t<node.m_lods.Num(); ++t )
+        // Create the expression adding all the components
+		Ptr<ASTOp> LastCompOp;
+
+		// Add the components in this node
+        for ( int32 t=0; t< InNode->Components.Num(); ++t )
         {
-            if ( const NodeLOD* pLODNode = node.m_lods[t].get() )
+			const NodeComponent* ComponentNode = InNode->Components[t].get();
+            if (ComponentNode)
             {
-                m_currentParents.Last().m_lod = t;
+				FGenericGenerationResult ComponentResult;
+				GenerateComponent(Options, ComponentResult, ComponentNode);
 
-				FGenericGenerationResult LODResult;
-				Generate_LOD(Options, LODResult, pLODNode);
-                lodsOp->lods.Emplace( lodsOp, LODResult.op);
+				if (ComponentResult.op)
+				{
+					check(ComponentResult.op->GetOpType() == OP_TYPE::IN_ADDCOMPONENT);
+					ASTOpInstanceAdd* typedOp = static_cast<ASTOpInstanceAdd*>(ComponentResult.op.get());
+
+					check(!typedOp->instance.child());
+					typedOp->instance = LastCompOp;
+					LastCompOp = typedOp;
+				}
             }
         }
-        Ptr<ASTOp> rootOp = lodsOp;
+
+		// Add the components from child objects
+		FAdditionalComponentKey thisKey;
+		thisKey.m_lod = m_currentParents.Last().m_lod;
+		thisKey.m_pObject = m_currentParents.Last().m_pObject;
+		TArray<Ptr<ASTOp>>* addIt = AdditionalComponents.Find(thisKey);
+		if (addIt)
+		{
+			for (const Ptr<ASTOp>& cop : *addIt)
+			{
+				// Add the additional components after the main ones, this means higher up in the op tree.
+
+				// First find the last op in the chain of IN_ADDCOMPONENT operations
+				check(cop->GetOpType() == OP_TYPE::IN_ADDCOMPONENT);
+				ASTOpInstanceAdd* typedOp = static_cast<ASTOpInstanceAdd*>(cop.get());
+				ASTOpInstanceAdd* bottomOp = typedOp;
+				while (bottomOp->instance)
+				{
+					// Step down
+					check(bottomOp->instance->GetOpType() == OP_TYPE::IN_ADDCOMPONENT);
+					bottomOp = static_cast<ASTOpInstanceAdd*>(bottomOp->instance.child().get());
+				}
+
+				// Chain
+				bottomOp->instance = LastCompOp;
+				LastCompOp = typedOp;
+			}
+		}
+
+		// Store for possible parent objects if necessary
+		// 2 is because there must be a parent and there is always a null element as well.
+		if (LastCompOp && m_currentParents.Num() > 2)
+		{
+			const auto& parentObjectKey = m_currentParents[m_currentParents.Num() - 2];
+			FAdditionalComponentKey parentKey;
+			parentKey.m_lod = m_currentParents.Last().m_lod;
+			parentKey.m_pObject = parentObjectKey.m_pObject;
+			AdditionalComponents.FindOrAdd(parentKey).Add(LastCompOp);
+		}
+
+
+        Ptr<ASTOp> rootOp = LastCompOp;
 
 		// Add an ASTOpAddExtensionData for each connected ExtensionData node
-		for (const NodeObjectNew::Private::NamedExtensionDataNode& NamedNode : node.m_extensionDataNodes)
+		for (const NodeObjectNew::FNamedExtensionDataNode& NamedNode : InNode->ExtensionDataNodes)
 		{
 			if (!NamedNode.Node.get())
 			{
@@ -1967,7 +1978,7 @@ namespace mu
                 for( int32 i = 0; !found && i != m_firstPass.objects.Num(); i++ )
 				{
 					FirstPassGenerator::FObject& it = m_firstPass.objects[i];
-					if (it.node == pChildNode->GetBasePrivate())
+					if (it.node == pChildNode)
 					{
 						found = true;
 						conditionOp = it.condition;
@@ -1985,7 +1996,7 @@ namespace mu
 
                 // This op is ignored: everything is stored as patches to apply to the parent when
                 // it is compiled.
-                Generate( pChildNode, Options );
+                Generate_Generic( pChildNode, Options );
 
                 m_currentObject.Pop();
 
@@ -1994,7 +2005,7 @@ namespace mu
 				if (usedNames.Contains(strChildName))
 				{
 					FString Msg = FString::Printf(TEXT("Object group has more than one children with the same name [%s]."), *strChildName );
-					m_pErrorLog->GetPrivate()->Add(Msg, ELMT_WARNING, node.m_errorContext);
+					m_pErrorLog->GetPrivate()->Add(Msg, ELMT_WARNING, InNode->GetMessageContext());
 				}
 				else
 				{
@@ -2039,7 +2050,7 @@ namespace mu
 		for (const FirstPassGenerator::FModifier& m: m_firstPass.modifiers)
 		{
 			// Correct LOD?
-			if (m.lod != LOD)
+			if (m.LOD != LOD)
 			{
 				continue;
 			}
@@ -2123,9 +2134,9 @@ namespace mu
 		Ptr<ASTOpMeshRemoveMask> removeOp;
 		for (const FirstPassGenerator::FModifier& m : modifiers)
 		{
-			if (m.node->m_pNode->GetType()== NodeModifierMeshClipWithMesh::GetStaticType())
+			if (m.node->GetType()== NodeModifierMeshClipWithMesh::GetStaticType())
 			{
-				const NodeModifierMeshClipWithMesh::Private* TypedClipNode = static_cast<const NodeModifierMeshClipWithMesh::Private*>(m.node);
+				const NodeModifierMeshClipWithMesh* TypedClipNode = static_cast<const NodeModifierMeshClipWithMesh*>(m.node);
 				Ptr<ASTOpMeshMaskClipMesh> op = new ASTOpMeshMaskClipMesh();
 				op->source = preModifiersMesh;
 
@@ -2135,7 +2146,7 @@ namespace mu
 				ClipOptions.State = Options.State;
 
 				FMeshGenerationResult clipResult;
-				GenerateMesh(ClipOptions, clipResult, TypedClipNode->ClipMesh);
+				GenerateMesh(ClipOptions, clipResult, TypedClipNode->GetPrivate()->ClipMesh);
 				op->clip = clipResult.meshOp;
 
 				if (!op->clip)
@@ -2166,12 +2177,12 @@ namespace mu
 		// Process clip-with-mask modifiers
 		for (const FirstPassGenerator::FModifier& m : modifiers)
 		{
-			if (m.node->m_pNode->GetType() == NodeModifierMeshClipWithUVMask::GetStaticType())
+			if (m.node->GetType() == NodeModifierMeshClipWithUVMask::GetStaticType())
 			{
-				const NodeModifierMeshClipWithUVMask::Private* TypedClipNode = static_cast<const NodeModifierMeshClipWithUVMask::Private*>(m.node);
+				const NodeModifierMeshClipWithUVMask* TypedClipNode = static_cast<const NodeModifierMeshClipWithUVMask*>(m.node);
 				Ptr<ASTOpMeshMaskClipUVMask> op = new ASTOpMeshMaskClipUVMask();
 				op->Source = preModifiersMesh;
-				op->LayoutIndex = TypedClipNode->LayoutIndex;
+				op->LayoutIndex = TypedClipNode->GetPrivate()->LayoutIndex;
 
 				// Parameters
 				FImageGenerationOptions ClipOptions;
@@ -2180,7 +2191,7 @@ namespace mu
 				ClipOptions.State = Options.State;
 
 				FImageGenerationResult ClipMaskResult;
-				GenerateImage(ClipOptions, ClipMaskResult, TypedClipNode->ClipMask);
+				GenerateImage(ClipOptions, ClipMaskResult, TypedClipNode->GetPrivate()->ClipMask);
 
 				// It could be IF_L_UBIT, but since this should be optimized out at compile time, leave the most cpu efficient.
 				op->Mask = GenerateImageFormat(ClipMaskResult.op, mu::EImageFormat::IF_L_UBYTE);
@@ -2215,9 +2226,9 @@ namespace mu
 		{
 			Ptr<ASTOp> modifiedMeshOp;
 
-			if (m.node->m_pNode->GetType() == NodeModifierMeshClipMorphPlane::GetStaticType())
+			if (m.node->GetType() == NodeModifierMeshClipMorphPlane::GetStaticType())
 			{
-				const NodeModifierMeshClipMorphPlane::Private* TypedNode = static_cast<const NodeModifierMeshClipMorphPlane::Private*>(m.node);
+				const NodeModifierMeshClipMorphPlane::Private* TypedNode = static_cast<const NodeModifierMeshClipMorphPlane*>(m.node)->GetPrivate();
 				Ptr<ASTOpMeshClipMorphPlane> op = new ASTOpMeshClipMorphPlane();
 				op->source = lastMeshOp;
 
@@ -2296,9 +2307,9 @@ namespace mu
 		{
 			Ptr<ASTOp> ModifiedMeshOp;
 
-			if (M.node->m_pNode->GetType()==NodeModifierMeshClipDeform::GetStaticType())
+			if (M.node->GetType()==NodeModifierMeshClipDeform::GetStaticType())
 			{
-				const NodeModifierMeshClipDeform::Private* TypedClipNode = static_cast<const NodeModifierMeshClipDeform::Private*>(M.node);
+				const NodeModifierMeshClipDeform::Private* TypedClipNode = static_cast<const NodeModifierMeshClipDeform*>(M.node)->GetPrivate();
 				Ptr<ASTOpMeshBindShape>  BindOp = new ASTOpMeshBindShape();
 				Ptr<ASTOpMeshClipDeform> ClipOp = new ASTOpMeshClipDeform();
 
diff --git a/Engine/Plugins/Experimental/Mutable/Source/MutableTools/Private/MuT/CodeGenerator.h b/Engine/Plugins/Experimental/Mutable/Source/MutableTools/Private/MuT/CodeGenerator.h
--- a/Engine/Plugins/Experimental/Mutable/Source/MutableTools/Private/MuT/CodeGenerator.h
+++ b/Engine/Plugins/Experimental/Mutable/Source/MutableTools/Private/MuT/CodeGenerator.h
@@ -144,6 +144,42 @@ namespace mu
 			TArray<FString> ActiveTags;
 		};
 
+		struct FLODGenerationOptions : public FGenericGenerationOptions
+		{
+			FLODGenerationOptions(const FGenericGenerationOptions& BaseOptions, int32 InLODIndex, const NodeComponentNew* InComponent)
+			{
+				Component = InComponent;
+				LODIndex = InLODIndex;
+				State = BaseOptions.State;
+				ActiveTags = BaseOptions.ActiveTags;
+			}
+
+			const NodeComponentNew* Component = nullptr;
+			int32 LODIndex;
+		};
+
+		struct FSurfaceGenerationOptions : public FGenericGenerationOptions
+		{
+			explicit FSurfaceGenerationOptions(const FGenericGenerationOptions& BaseOptions)
+			{
+				Component = nullptr;
+				LODIndex = -1;
+				State = BaseOptions.State;
+				ActiveTags = BaseOptions.ActiveTags;
+			}
+
+			explicit FSurfaceGenerationOptions(const FLODGenerationOptions& BaseOptions)
+			{
+				Component = BaseOptions.Component;
+				LODIndex = BaseOptions.LODIndex;
+				State = BaseOptions.State;
+				ActiveTags = BaseOptions.ActiveTags;
+			}
+
+			const NodeComponentNew* Component = nullptr;
+			int32 LODIndex = -1;
+		};
+
 		struct FGenericGenerationResult
 		{
 			Ptr<ASTOp> op;
@@ -171,9 +207,8 @@ namespace mu
 		typedef TMap<FGeneratedCacheKey, FGenericGenerationResult> FGeneratedGenericNodesMap;
 		FGeneratedGenericNodesMap GeneratedGenericNodes;
 
-		Ptr<ASTOp> Generate(const Ptr<const Node>, const FGenericGenerationOptions& );
-		void Generate_ComponentNew(const FGenericGenerationOptions&, FGenericGenerationResult&, const NodeComponentNew*);
-		void Generate_LOD(const FGenericGenerationOptions&, FGenericGenerationResult&, const NodeLOD*);
+		Ptr<ASTOp> Generate_Generic(const Ptr<const Node>, const FGenericGenerationOptions& );
+		void Generate_LOD(const FLODGenerationOptions&, FGenericGenerationResult&, const NodeLOD*);
 		void Generate_ObjectNew(const FGenericGenerationOptions&, FGenericGenerationResult&, const NodeObjectNew*);
 		void Generate_ObjectGroup(const FGenericGenerationOptions&, FGenericGenerationResult&, const NodeObjectGroup*);
 
@@ -214,14 +249,9 @@ namespace mu
 
         struct FParentKey
         {
-            const NodeObjectNew::Private* m_pObject = nullptr;
-            int32 m_state = -1;
+            const NodeObjectNew* m_pObject = nullptr;
             int32 m_lod = -1;
-            int32 m_component = -1;
-            int32 m_surface = -1;
-            int32 m_texture = -1;
-            int32 m_block = -1;
-        };
+         };
 
 		TArray< FParentKey > m_currentParents;
 
@@ -235,7 +265,7 @@ namespace mu
                 m_lod = -1;
             }
 
-            const NodeObjectNew::Private* m_pObject;
+            const NodeObjectNew* m_pObject;
             int32 m_lod;
 
 			FORCEINLINE bool operator==(const FAdditionalComponentKey& Other) const
@@ -294,6 +324,31 @@ namespace mu
 
 		TArray<FConditionalExtensionDataOp> ConditionalExtensionDataOps;
 
+		struct FGeneratedComponentCacheKey
+		{
+			Ptr<const Node> Node;
+			FGenericGenerationOptions Options;
+
+			friend FORCEINLINE uint32 GetTypeHash(const FGeneratedComponentCacheKey& InKey)
+			{
+				uint32 KeyHash = 0;
+				KeyHash = HashCombineFast(KeyHash, ::GetTypeHash(InKey.Node.get()));
+				KeyHash = HashCombineFast(KeyHash, GetTypeHash(InKey.Options));
+				return KeyHash;
+			}
+
+			FORCEINLINE bool operator==(const FGeneratedComponentCacheKey& Other) const
+			{
+				return Node == Other.Node && Options == Other.Options;
+			}
+		};
+
+		typedef TMap<FGeneratedComponentCacheKey, FGenericGenerationResult> GeneratedComponentMap;
+		GeneratedComponentMap GeneratedComponents;
+
+		void GenerateComponent(const FGenericGenerationOptions&, FGenericGenerationResult&, const NodeComponent*);
+		void GenerateComponent_New(const FGenericGenerationOptions&, FGenericGenerationResult&, const NodeComponentNew*);
+
 		//-----------------------------------------------------------------------------------------
 
 		// Get the modifiers that have to be applied to elements with a specific tag.
@@ -582,7 +637,7 @@ namespace mu
 		typedef TMap<FGeneratedExtensionDataCacheKey, FExtensionDataGenerationResult> FGeneratedExtensionDataMap;
 		FGeneratedExtensionDataMap GeneratedExtensionData;
 
-		void GenerateExtensionData(FExtensionDataGenerationResult& OutResult, const FGenericGenerationOptions&, const NodeExtensionDataPtrConst&);
+		void GenerateExtensionData(FExtensionDataGenerationResult& OutResult, const FGenericGenerationOptions&, const Ptr<const NodeExtensionData>&);
 		void GenerateExtensionData_Constant(FExtensionDataGenerationResult& OutResult, const FGenericGenerationOptions&, const class NodeExtensionDataConstant*);
 		void GenerateExtensionData_Switch(FExtensionDataGenerationResult& OutResult, const FGenericGenerationOptions&, const class NodeExtensionDataSwitch*);
 		void GenerateExtensionData_Variation(FExtensionDataGenerationResult& OutResult, const FGenericGenerationOptions&, const class NodeExtensionDataVariation*);
@@ -702,7 +757,7 @@ namespace mu
             Ptr<ASTOp> surfaceOp;
         };
 
-        void GenerateSurface( FSurfaceGenerationResult&, const FGenericGenerationOptions&,
+        void GenerateSurface( FSurfaceGenerationResult&, const FSurfaceGenerationOptions&,
                               Ptr<const NodeSurfaceNew>,
                               const TArray<FirstPassGenerator::FSurface::FEdit>& Edits );
 
@@ -741,18 +796,18 @@ namespace mu
 
 		if (NumRows == 0)
 		{
-			m_pErrorLog->GetPrivate()->Add("The table has no rows.", ELMT_ERROR, node.m_errorContext);
+			m_pErrorLog->GetPrivate()->Add("The table has no rows.", ELMT_ERROR, node.m_pNode->GetMessageContext());
 			return nullptr;
 		}
         else if (ColIndex < 0)
         {
-            m_pErrorLog->GetPrivate()->Add("Table column not found.", ELMT_ERROR, node.m_errorContext);
+            m_pErrorLog->GetPrivate()->Add("Table column not found.", ELMT_ERROR, node.m_pNode->GetMessageContext());
             return nullptr;
         }
 
         if (NodeTable->GetPrivate()->Columns[ ColIndex ].Type != TYPE )
         {
-            m_pErrorLog->GetPrivate()->Add("Table column type is not the right type.", ELMT_ERROR, node.m_errorContext);
+            m_pErrorLog->GetPrivate()->Add("Table column type is not the right type.", ELMT_ERROR, node.m_pNode->GetMessageContext());
             return nullptr;
         }
 
diff --git a/Engine/Plugins/Experimental/Mutable/Source/MutableTools/Private/MuT/CodeGenerator_Color.cpp b/Engine/Plugins/Experimental/Mutable/Source/MutableTools/Private/MuT/CodeGenerator_Color.cpp
--- a/Engine/Plugins/Experimental/Mutable/Source/MutableTools/Private/MuT/CodeGenerator_Color.cpp
+++ b/Engine/Plugins/Experimental/Mutable/Source/MutableTools/Private/MuT/CodeGenerator_Color.cpp
@@ -198,7 +198,7 @@ namespace mu
 		if (node.m_options.Num() == 0)
 		{
 			// No options in the switch!
-			Ptr<ASTOp> missingOp = GenerateMissingColourCode(TEXT("Switch option"), node.m_errorContext);
+			Ptr<ASTOp> missingOp = GenerateMissingColourCode(TEXT("Switch option"), Typed->GetMessageContext());
 			result.op = missingOp;
 			return;
 		}
@@ -209,26 +209,26 @@ namespace mu
 		// Variable value
 		if (node.m_pParameter)
 		{
-			op->variable = Generate(node.m_pParameter.get(), Options);
+			op->variable = Generate_Generic(node.m_pParameter.get(), Options);
 		}
 		else
 		{
 			// This argument is required
-			op->variable = GenerateMissingScalarCode(TEXT("Switch variable"), 0.0f, node.m_errorContext);
+			op->variable = GenerateMissingScalarCode(TEXT("Switch variable"), 0.0f, Typed->GetMessageContext());
 		}
 
 		// Options
-		for (std::size_t t = 0; t < node.m_options.Num(); ++t)
+		for (int32 t = 0; t < node.m_options.Num(); ++t)
 		{
 			Ptr<ASTOp> branch;
 			if (node.m_options[t])
 			{
-				branch = Generate(node.m_options[t].get(), Options);
+				branch = Generate_Generic(node.m_options[t].get(), Options);
 			}
 			else
 			{
 				// This argument is required
-				branch = GenerateMissingColourCode(TEXT("Switch option"), node.m_errorContext);
+				branch = GenerateMissingColourCode(TEXT("Switch option"), Typed->GetMessageContext());
 			}
 			op->cases.Emplace((int16)t, op, branch);
 		}
@@ -247,8 +247,9 @@ namespace mu
 		// Default case
 		if (node.m_defaultColour)
 		{
-			FMeshGenerationResult branchResults;
-			currentOp = Generate(node.m_defaultColour, Options);
+			FColorGenerationResult BranchResults;
+			GenerateColor(BranchResults, Options, node.m_defaultColour);
+			currentOp = BranchResults.op;
 		}
 
 		// Process variations in reverse order, since conditionals are built bottom-up.
@@ -267,19 +268,21 @@ namespace mu
 			if (tagIndex < 0)
 			{
 				FString Msg = FString::Printf(TEXT("Unknown tag found in color variation [%s]."), *tag);
-				m_pErrorLog->GetPrivate()->Add(Msg, ELMT_WARNING, node.m_errorContext);
+				m_pErrorLog->GetPrivate()->Add(Msg, ELMT_WARNING, Typed->GetMessageContext());
 				continue;
 			}
 
 			Ptr<ASTOp> variationOp;
 			if (node.m_variations[t].m_colour)
 			{
-				variationOp = Generate(node.m_variations[t].m_colour, Options);
+				FColorGenerationResult BranchResults;
+				GenerateColor(BranchResults,Options,node.m_variations[t].m_colour);
+				variationOp = BranchResults.op;
 			}
 			else
 			{
 				// This argument is required
-				variationOp = GenerateMissingColourCode(TEXT("Variation option"), node.m_errorContext);
+				variationOp = GenerateMissingColourCode(TEXT("Variation option"), Typed->GetMessageContext());
 			}
 
 
@@ -324,7 +327,7 @@ namespace mu
 		else
 		{
 			// This argument is required
-			base = GenerateMissingImageCode(TEXT("Sample image"), EImageFormat::IF_RGB_UBYTE, node.m_errorContext, ImageOptions);
+			base = GenerateMissingImageCode(TEXT("Sample image"), EImageFormat::IF_RGB_UBYTE, Typed->GetMessageContext(), ImageOptions);
 		}
 		base = GenerateImageFormat(base, EImageFormat::IF_RGB_UBYTE);
 		op->SetChild(op->op.args.ColourSampleImage.image, base);
@@ -470,7 +473,7 @@ namespace mu
 		else
 		{
 			op->SetChild(op->op.args.ColourArithmetic.a,
-				CodeGenerator::GenerateMissingColourCode(TEXT("ColourArithmetic A"), node.m_errorContext));
+				CodeGenerator::GenerateMissingColourCode(TEXT("ColourArithmetic A"), Typed->GetMessageContext()));
 		}
 
 		// B
@@ -482,7 +485,7 @@ namespace mu
 		else
 		{
 			op->SetChild(op->op.args.ColourArithmetic.b,
-				CodeGenerator::GenerateMissingColourCode(TEXT("ColourArithmetic B"),node.m_errorContext));
+				CodeGenerator::GenerateMissingColourCode(TEXT("ColourArithmetic B"), Typed->GetMessageContext()));
 		}
 
 		result.op = op;
@@ -497,10 +500,12 @@ namespace mu
 		result.op = GenerateTableSwitch<NodeColourTable::Private, ETableColumnType::Color, OP_TYPE::CO_SWITCH>(node,
 			[this, &Options](const NodeColourTable::Private& node, int colIndex, int row, ErrorLog* pErrorLog)
 			{
-				NodeColourConstantPtr CellData = new NodeColourConstant();
+				Ptr<NodeColourConstant> CellData = new NodeColourConstant();
 				FVector4f Colour = node.Table->GetPrivate()->Rows[row].Values[colIndex].Color;
 				CellData->SetValue(Colour);
-				return Generate(CellData, Options);
+				FColorGenerationResult BranchResults;
+				GenerateColor(BranchResults, Options, CellData );
+				return BranchResults.op;
 			});
 	}
 
diff --git a/Engine/Plugins/Experimental/Mutable/Source/MutableTools/Private/MuT/CodeGenerator_ExtensionData.cpp b/Engine/Plugins/Experimental/Mutable/Source/MutableTools/Private/MuT/CodeGenerator_ExtensionData.cpp
--- a/Engine/Plugins/Experimental/Mutable/Source/MutableTools/Private/MuT/CodeGenerator_ExtensionData.cpp
+++ b/Engine/Plugins/Experimental/Mutable/Source/MutableTools/Private/MuT/CodeGenerator_ExtensionData.cpp
@@ -21,11 +21,8 @@
 #include "MuT/NodeScalar.h"
 #include "MuT/NodeExtensionData.h"
 #include "MuT/NodeExtensionDataConstant.h"
-#include "MuT/NodeExtensionDataConstantPrivate.h"
 #include "MuT/NodeExtensionDataSwitch.h"
-#include "MuT/NodeExtensionDataSwitchPrivate.h"
 #include "MuT/NodeExtensionDataVariation.h"
-#include "MuT/NodeExtensionDataVariationPrivate.h"
 
 
 namespace mu
@@ -33,7 +30,7 @@ namespace mu
 class Node;
 
 	//---------------------------------------------------------------------------------------------
-	void CodeGenerator::GenerateExtensionData(FExtensionDataGenerationResult& OutResult, const FGenericGenerationOptions& Options, const NodeExtensionDataPtrConst& InUntypedNode)
+	void CodeGenerator::GenerateExtensionData(FExtensionDataGenerationResult& OutResult, const FGenericGenerationOptions& Options, const Ptr<const NodeExtensionData>& InUntypedNode)
 	{
 		if (!InUntypedNode)
 		{
@@ -53,12 +50,12 @@ class Node;
 		const NodeExtensionData* Node = InUntypedNode.get();
 
 		// Generate for each different type of node
-		switch (Node->GetExtensionDataNodeType())
+		switch (Node->GetType()->Type)
 		{
-			case NodeExtensionData::EType::Constant:  GenerateExtensionData_Constant(OutResult, Options, static_cast<const NodeExtensionDataConstant*>(Node)); break;
-			case NodeExtensionData::EType::Switch:    GenerateExtensionData_Switch(OutResult, Options, static_cast<const NodeExtensionDataSwitch*>(Node)); break;
-			case NodeExtensionData::EType::Variation: GenerateExtensionData_Variation(OutResult, Options, static_cast<const NodeExtensionDataVariation*>(Node)); break;
-			case NodeExtensionData::EType::None: check(false);
+			case Node::EType::ExtensionDataConstant:  GenerateExtensionData_Constant(OutResult, Options, static_cast<const NodeExtensionDataConstant*>(Node)); break;
+			case Node::EType::ExtensionDataSwitch:    GenerateExtensionData_Switch(OutResult, Options, static_cast<const NodeExtensionDataSwitch*>(Node)); break;
+			case Node::EType::ExtensionDataVariation: GenerateExtensionData_Variation(OutResult, Options, static_cast<const NodeExtensionDataVariation*>(Node)); break;
+			default: check(false);
 		}
 
 		// Cache the result
@@ -68,19 +65,17 @@ class Node;
 	//---------------------------------------------------------------------------------------------
 	void CodeGenerator::GenerateExtensionData_Constant(FExtensionDataGenerationResult& OutResult, const FGenericGenerationOptions& Options, const NodeExtensionDataConstant* Constant)
 	{
-		NodeExtensionDataConstant::Private& Node = *Constant->GetPrivate();
-
 		Ptr<ASTOpConstantExtensionData> Op = new ASTOpConstantExtensionData();
 		OutResult.Op = Op;
 
-		ExtensionDataPtrConst Data = Node.Value;
+		ExtensionDataPtrConst Data = Constant->Value;
 		if (!Data)
 		{
 			// Data can't be null, so make an empty one
 			Data = new ExtensionData();
 			
 			// Log an error message
-			m_pErrorLog->GetPrivate()->Add("Constant extension data not set", ELMT_WARNING, Node.m_errorContext);
+			m_pErrorLog->GetPrivate()->Add("Constant extension data not set", ELMT_WARNING, Constant->GetMessageContext());
 		}
 
 		Op->Value = Data;
@@ -89,13 +84,11 @@ class Node;
 	//---------------------------------------------------------------------------------------------
 	void CodeGenerator::GenerateExtensionData_Switch(FExtensionDataGenerationResult& OutResult, const FGenericGenerationOptions& Options, const class NodeExtensionDataSwitch* Switch)
 	{
-		NodeExtensionDataSwitch::Private& Node = *Switch->GetPrivate();
-
 		MUTABLE_CPUPROFILER_SCOPE(NodeExtensionDataSwitch);
 
-		if (Node.Options.IsEmpty())
+		if (Switch->Options.IsEmpty())
 		{
-			Ptr<ASTOp> MissingOp = GenerateMissingExtensionDataCode(TEXT("Switch option"), Node.m_errorContext);
+			Ptr<ASTOp> MissingOp = GenerateMissingExtensionDataCode(TEXT("Switch option"), Switch->GetMessageContext());
 			OutResult.Op = MissingOp;
 			return;
 		}
@@ -104,23 +97,23 @@ class Node;
 		Op->type = OP_TYPE::ED_SWITCH;
 
 		// Variable name
-		if (Node.Parameter)
+		if (Switch->Parameter)
 		{
-			Op->variable = Generate(Node.Parameter.get(), Options);
+			Op->variable = Generate_Generic(Switch->Parameter.get(), Options);
 		}
 		else
 		{
 			// This argument is required
-			Op->variable = GenerateMissingScalarCode(TEXT("Switch variable"), 0.0f, Node.m_errorContext);
+			Op->variable = GenerateMissingScalarCode(TEXT("Switch variable"), 0.0f, Switch->GetMessageContext());
 		}
 
 		// Options
-		for (int32 OptionIndex = 0; OptionIndex < Node.Options.Num(); ++OptionIndex)
+		for (int32 OptionIndex = 0; OptionIndex < Switch->Options.Num(); ++OptionIndex)
 		{
 			Ptr<ASTOp> Branch;
-			if (Node.Options[OptionIndex])
+			if (Switch->Options[OptionIndex])
 			{
-				NodeExtensionDataPtr SwitchOption = Node.Options[OptionIndex];
+				Ptr<NodeExtensionData> SwitchOption = Switch->Options[OptionIndex];
 
 				FExtensionDataGenerationResult OptionResult;
 				GenerateExtensionData(OptionResult, Options, SwitchOption);
@@ -130,7 +123,7 @@ class Node;
 			else
 			{
 				// This argument is required
-				Branch = GenerateMissingExtensionDataCode(TEXT("Switch option"), Node.m_errorContext);
+				Branch = GenerateMissingExtensionDataCode(TEXT("Switch option"), Switch->GetMessageContext());
 			}
 			Op->cases.Emplace(static_cast<int16_t>(OptionIndex), Op, Branch);
 		}
@@ -141,23 +134,21 @@ class Node;
 	//---------------------------------------------------------------------------------------------
 	void CodeGenerator::GenerateExtensionData_Variation(FExtensionDataGenerationResult& OutResult, const FGenericGenerationOptions& Options, const class NodeExtensionDataVariation* Variation)
 	{
-		const NodeExtensionDataVariation::Private& Node = *Variation->GetPrivate();
-
 		Ptr<ASTOp> CurrentOp;
 
 		// Default case
-		if (Node.DefaultValue)
+		if (Variation->DefaultValue)
 		{
 			FExtensionDataGenerationResult DefaultResult;
-			GenerateExtensionData(DefaultResult, Options, Node.DefaultValue);
+			GenerateExtensionData(DefaultResult, Options, Variation->DefaultValue);
 
 			CurrentOp = DefaultResult.Op;
 		}
 
 		// Process variations in reverse order, since conditionals are built bottom-up.
-		for (int32 VariationIndex = Node.Variations.Num() - 1; VariationIndex >= 0; --VariationIndex)
+		for (int32 VariationIndex = Variation->Variations.Num() - 1; VariationIndex >= 0; --VariationIndex)
 		{
-			const FString& Tag = Node.Variations[VariationIndex].Tag;
+			const FString& Tag = Variation->Variations[VariationIndex].Tag;
 			const int32 TagIndex = m_firstPass.m_tags.IndexOfByPredicate([Tag](const FirstPassGenerator::FTag& CandidateTag)
 			{
 				return CandidateTag.tag == Tag;
@@ -166,12 +157,12 @@ class Node;
 			if (TagIndex == INDEX_NONE)
 			{
 				const FString Msg = FString::Printf(TEXT("Unknown tag found in Extension Data variation [%s]"), *Tag);
-				m_pErrorLog->GetPrivate()->Add(Msg, ELMT_WARNING, Node.m_errorContext);
+				m_pErrorLog->GetPrivate()->Add(Msg, ELMT_WARNING, Variation->GetMessageContext());
 				continue;
 			}
 
 			Ptr<ASTOp> VariationOp;
-			if (NodeExtensionDataPtr VariationValue = Node.Variations[VariationIndex].Value)
+			if (Ptr<NodeExtensionData> VariationValue = Variation->Variations[VariationIndex].Value)
 			{
 				FExtensionDataGenerationResult VariationResult;
 				GenerateExtensionData(VariationResult, Options, VariationValue);
@@ -181,7 +172,7 @@ class Node;
 			else
 			{
 				// This argument is required
-				VariationOp = GenerateMissingExtensionDataCode(TEXT("Variation option"), Node.m_errorContext);
+				VariationOp = GenerateMissingExtensionDataCode(TEXT("Variation option"), Variation->GetMessageContext());
 			}
 
 			Ptr<ASTOpConditional> Conditional = new ASTOpConditional;
@@ -203,7 +194,7 @@ class Node;
 		m_pErrorLog->GetPrivate()->Add(Msg, ELMT_ERROR, ErrorContext);
 
 		// Create a constant extension data
-		NodeExtensionDataConstantPtrConst Node = new NodeExtensionDataConstant;
+		Ptr<const NodeExtensionDataConstant> Node = new NodeExtensionDataConstant;
 
 		FExtensionDataGenerationResult Result;
 		FGenericGenerationOptions Options;
