[mutable] Redesign the object node hierarchy to better match unreal component and LOD concepts.

diff --git a/Engine/Plugins/Experimental/Mutable/Source/MutableTools/Private/MuT/CodeGenerator_FirstPass.cpp b/Engine/Plugins/Experimental/Mutable/Source/MutableTools/Private/MuT/CodeGenerator_FirstPass.cpp
--- a/Engine/Plugins/Experimental/Mutable/Source/MutableTools/Private/MuT/CodeGenerator_FirstPass.cpp
+++ b/Engine/Plugins/Experimental/Mutable/Source/MutableTools/Private/MuT/CodeGenerator_FirstPass.cpp
@@ -17,16 +17,13 @@
 #include "MuT/ErrorLog.h"
 #include "MuT/ErrorLogPrivate.h"
 #include "MuT/NodeComponent.h"
-#include "MuT/NodeComponentEditPrivate.h"
-#include "MuT/NodeComponentNewPrivate.h"
-#include "MuT/NodeLODPrivate.h"
 #include "MuT/NodeModifierMeshClipDeformPrivate.h"
 #include "MuT/NodeModifierMeshClipMorphPlanePrivate.h"
 #include "MuT/NodeModifierMeshClipWithMeshPrivate.h"
 #include "MuT/NodeModifierMeshClipWithUVMaskPrivate.h"
 #include "MuT/NodeObject.h"
 #include "MuT/NodeObjectGroupPrivate.h"
-#include "MuT/NodeObjectNewPrivate.h"
+#include "MuT/NodeObjectNew.h"
 #include "MuT/NodePrivate.h"
 #include "MuT/NodeSurface.h"
 #include "MuT/NodeSurfaceEditPrivate.h"
@@ -69,23 +66,23 @@ namespace mu
 		// Step 2: Collect all tags and a list of the surfaces that activate them
 		for (int32 s=0; s<surfaces.Num(); ++s)
 		{
-            // \todo: edit surfaces should also be able to activate tags. T1245
-            for (int32 t=0; t<surfaces[s].node->GetPrivate()->m_tags.Num(); ++t)
+			// Collect the tags in new surfaces
+			for (int32 t=0; t<surfaces[s].node->GetPrivate()->m_tags.Num(); ++t)
 			{
-				int tag = -1;
+				int32 tag = -1;
                 const FString& tagStr = surfaces[s].node->GetPrivate()->m_tags[t];
-                for (std::size_t i = 0; i<m_tags.Num() && tag<0; ++i)
+                for (int32 i = 0; i<m_tags.Num() && tag<0; ++i)
 				{
                     if (m_tags[i].tag == tagStr)
 					{
-						tag = (int)i;
+						tag = i;
 					}
 				}
 
 				// New tag?
 				if (tag < 0)
 				{
-                    tag = (int)m_tags.Num();
+                    tag = m_tags.Num();
 					FTag newTag;
                     newTag.tag = tagStr;
                     m_tags.Add(newTag);
@@ -93,39 +90,39 @@ namespace mu
 
                 if (m_tags[tag].surfaces.Find(s)==INDEX_NONE)
 				{
-                    m_tags[tag].surfaces.Add((int)s);
+                    m_tags[tag].surfaces.Add(s);
 				}
 			}
 
             // Collect the tags in edit surfaces
-            for (std::size_t e=0; e<surfaces[s].edits.Num(); ++e)
+            for (int32 e=0; e<surfaces[s].edits.Num(); ++e)
             {
-                const auto& edit = surfaces[s].edits[e];
-                for (int32 t=0; t<edit.node->m_tags.Num(); ++t)
+                const FSurface::FEdit& edit = surfaces[s].edits[e];
+                for (int32 t=0; t<edit.node->GetPrivate()->m_tags.Num(); ++t)
                 {
-                    int tag = -1;
-                    auto tagStr = edit.node->m_tags[t];
+                    int32 tag = -1;
+					const FString& tagStr = edit.node->GetPrivate()->m_tags[t];
 
                     for (int32 i = 0; i<m_tags.Num() && tag<0; ++i)
                     {
                         if (m_tags[i].tag == tagStr)
                         {
-                            tag = (int)i;
+                            tag = i;
                         }
                     }
 
                     // New tag?
                     if (tag < 0)
                     {
-                        tag = (int)m_tags.Num();
+                        tag = m_tags.Num();
 						FTag newTag;
                         newTag.tag = tagStr;
                         m_tags.Add(newTag);
                     }
 
-                    if (m_tags[tag].edits.Find({int(s),int(e)}) == INDEX_NONE)
+                    if (m_tags[tag].edits.Find({s,e}) == INDEX_NONE)
                     {
-                        m_tags[tag].edits.Add({ int(s),int(e) });
+                        m_tags[tag].edits.Add({s,e});
                     }
                 }
             }
@@ -141,8 +138,8 @@ namespace mu
         if ( m_states.IsEmpty() )
         {
             FObjectState data;
-            data.m_name = "Default";
-            m_states.Emplace( data, Root->GetBasePrivate() );
+            data.Name = "Default";
+            m_states.Emplace( data, Root );
         }
 	}
 
@@ -198,12 +195,14 @@ namespace mu
 	//---------------------------------------------------------------------------------------------
 	void FirstPassGenerator::Generate_Modifier(const NodeModifier* InNode)
 	{
+		check(CurrentLOD>=0);
+
 		// Add the data about this modifier
 		FModifier thisData;
-		thisData.node = static_cast<const NodeModifier::Private*>(InNode->GetBasePrivate());
+		thisData.node = InNode;
 		thisData.objectCondition = m_currentCondition.Last().objectCondition;
 		thisData.stateCondition = m_currentStateCondition.Last();
-		thisData.lod = m_currentLOD;
+		thisData.LOD = CurrentLOD;
 		thisData.positiveTags = m_currentPositiveTags;
 		thisData.negativeTags = m_currentNegativeTags;
 		modifiers.Add(thisData);
@@ -216,7 +215,8 @@ namespace mu
 		// Add the data about this surface
 		FSurface thisData;
 		thisData.node = InNode;
-		thisData.component = m_currentComponent;
+		thisData.Component = m_currentComponent;
+		thisData.LOD = CurrentLOD;
 		thisData.objectCondition = m_currentCondition.Last().objectCondition;
 		thisData.stateCondition = m_currentStateCondition.Last();
 		thisData.positiveTags = m_currentPositiveTags;
@@ -240,7 +240,7 @@ namespace mu
             // Are we editing an edit node modifying this surface?
             for (const auto& e: s.edits)
             {
-                if (Private->m_pParent && e.node==Private->m_pParent->GetBasePrivate())
+                if (Private->m_pParent && e.node==Private->m_pParent)
                 {
                     return true;
                 }
@@ -256,13 +256,13 @@ namespace mu
 			FSurface::FEdit edit;
 			edit.PositiveTags = m_currentPositiveTags;
 			edit.NegativeTags = m_currentNegativeTags;
-            edit.node = Private;
+            edit.node = InNode;
             edit.condition = m_currentCondition.Last().objectCondition;
 			Surface->edits.Add(edit);
 		}
 		else
 		{
-			m_pErrorLog->GetPrivate()->Add("Missing parent object for edit node.", ELMT_WARNING, Private->m_errorContext);
+			m_pErrorLog->GetPrivate()->Add("Missing parent object for edit node.", ELMT_WARNING, InNode->GetMessageContext());
 		}
 	}
 
@@ -340,7 +340,7 @@ namespace mu
                     {
                         for( size_t s=0; s<stateCount; ++s )
                         {
-                            if (m_states[s].Key.m_name==v.m_tag)
+                            if (m_states[s].Key.Name==v.m_tag)
                             {
                                 // Remove this state from the default options, since it has its own variation
                                 defaultStates[s] = false;
@@ -372,7 +372,7 @@ namespace mu
 
                 for( size_t s=0; s<stateCount; ++s )
                 {
-                    if (m_states[s].Key.m_name==v.m_tag)
+                    if (m_states[s].Key.Name==v.m_tag)
                     {
                         variationStates[s] = true;
                     }
@@ -424,7 +424,7 @@ namespace mu
 		else
 		{
 			// This argument is required
-			ScalarResult.op = Generator->GenerateMissingScalarCode(TEXT("Switch variable"), 0.0f, Private->m_errorContext);
+			ScalarResult.op = Generator->GenerateMissingScalarCode(TEXT("Switch variable"), 0.0f, InNode->GetMessageContext());
 		}
 
 		// Parse the options
@@ -463,17 +463,18 @@ namespace mu
 	//---------------------------------------------------------------------------------------------
 	void FirstPassGenerator::Generate_ComponentNew(const NodeComponentNew* InNode)
 	{
-		const NodeComponentNew::Private* Private = InNode->GetPrivate();
+        m_currentComponent = InNode;
 
-        m_currentComponent = Private->GetParentComponentNew();
-
-		for (const Ptr<NodeSurface>& c : Private->m_surfaces)
+		CurrentLOD = 0;
+		for (const Ptr<NodeLOD>& c : InNode->LODs)
 		{
 			if (c)
 			{
-				Generate_Generic(c.get());
+				Generate_LOD(c.get());
 			}
+			++CurrentLOD;
 		}
+		CurrentLOD = -1;
 
 		m_currentComponent = nullptr;
 	}
@@ -482,17 +483,18 @@ namespace mu
 	//---------------------------------------------------------------------------------------------
 	void FirstPassGenerator::Generate_ComponentEdit(const NodeComponentEdit* InNode)
 	{
-		const NodeComponentEdit::Private* Private = InNode->GetPrivate();
+		m_currentComponent = InNode->GetParentComponentNew();
 
-		m_currentComponent = Private->GetParentComponentNew();
-
-		for (const auto& c : Private->m_surfaces)
+		CurrentLOD = 0;
+		for (const Ptr<NodeLOD>& c : InNode->LODs)
 		{
 			if (c)
 			{
-				Generate_Generic(c.get());
+				Generate_LOD(c.get());
 			}
+			++CurrentLOD;
 		}
+		CurrentLOD = -1;
 
 		m_currentComponent = nullptr;
 	}
@@ -501,16 +503,15 @@ namespace mu
 	//---------------------------------------------------------------------------------------------
 	void FirstPassGenerator::Generate_LOD(const NodeLOD* InNode)
 	{
-		const NodeLOD::Private* Private = InNode->GetPrivate();
-
-		for (const auto& c : Private->m_components)
+		for (const Ptr<NodeSurface>& c : InNode->Surfaces)
 		{
 			if (c)
 			{
 				Generate_Generic(c.get());
 			}
 		}
-		for (const auto& c : Private->m_modifiers)
+
+		for (const Ptr<NodeModifier>& c : InNode->Modifiers)
 		{
 			if (c)
 			{
@@ -523,47 +524,41 @@ namespace mu
 	//---------------------------------------------------------------------------------------------
 	void FirstPassGenerator::Generate_ObjectNew(const NodeObjectNew* InNode)
 	{
-		const NodeObjectNew::Private* Private = InNode->GetPrivate();
-
 		// Add the data about this object
 		FObject thisData;
-		thisData.node = Private;
+		thisData.node = InNode;
         thisData.condition = m_currentCondition.Last().objectCondition;
 		objects.Add(thisData);
 
         // Accumulate the model states
-        for ( const auto& s: Private->m_states )
+        for ( const FObjectState& s: InNode->States )
         {
-            m_states.Emplace( s, Private );
+            m_states.Emplace( s, InNode );
 
-            if ( s.m_runtimeParams.Num() > MUTABLE_MAX_RUNTIME_PARAMETERS_PER_STATE )
+            if ( s.RuntimeParams.Num() > MUTABLE_MAX_RUNTIME_PARAMETERS_PER_STATE )
             {
                 FString Msg = FString::Printf( TEXT("State [%s] has more than %d runtime parameters. Their update may fail."), 
-					*s.m_name,
+					*s.Name,
                     MUTABLE_MAX_RUNTIME_PARAMETERS_PER_STATE);
-                m_pErrorLog->GetPrivate()->Add(Msg, ELMT_ERROR, Private->m_errorContext );
+                m_pErrorLog->GetPrivate()->Add(Msg, ELMT_ERROR, InNode->GetMessageContext());
             }
         }
 
-		// Process the lods
-		int i = 0;
-		for (const auto& l : Private->m_lods)
+		// Process the components
+		for (const Ptr<NodeComponent>& Component : InNode->Components)
 		{
-			if (l)
+			if (Component)
 			{
-                m_currentLOD = i++;
-				Generate_Generic(l.get());
+				Generate_Generic(Component.get());
 			}
 		}
 
-		m_currentLOD = -1;
-
 		// Process the children
-		for (const auto& c : Private->m_children)
+		for (const Ptr<NodeObject>& Child : InNode->Children)
 		{
-			if (c)
+			if (Child)
 			{
-				Generate_Generic(c.get());
+				Generate_Generic(Child.get());
 			}
 		}
 	}
diff --git a/Engine/Plugins/Experimental/Mutable/Source/MutableTools/Private/MuT/CodeGenerator_FirstPass.h b/Engine/Plugins/Experimental/Mutable/Source/MutableTools/Private/MuT/CodeGenerator_FirstPass.h
--- a/Engine/Plugins/Experimental/Mutable/Source/MutableTools/Private/MuT/CodeGenerator_FirstPass.h
+++ b/Engine/Plugins/Experimental/Mutable/Source/MutableTools/Private/MuT/CodeGenerator_FirstPass.h
@@ -63,7 +63,7 @@ namespace mu
 		//! Store the conditions that will enable or disable every object
 		struct FObject
 		{
-			const NodeObjectNew::Private* node;
+			const NodeObjectNew* node;
             Ptr<ASTOp> condition;
 		};
 		TArray<FObject> objects;
@@ -82,11 +82,12 @@ namespace mu
         //! it is reached with.
 		struct FSurface
 		{
-            NodeSurfaceNewPtrConst node;
+            Ptr<const NodeSurfaceNew> node;
 
-			// Parent component where this surface will be added. It may be different from the 
-			// component that defined it (if it was an edit component).
-            const NodeComponentNew::Private* component = nullptr;
+			// Parent Component where this surface will be added. It may be different from the 
+			// Component that defined it (if it was an edit component).
+            const NodeComponentNew* Component = nullptr;
+			int32 LOD = 0;
 
             // List of tags that are required for the presence of this surface
 			TArray<FString> positiveTags;
@@ -117,8 +118,8 @@ namespace mu
                 //! Condition that enables the effects of this edit node on the surface
                 Ptr<ASTOp> condition;
 
-                //! Weak reference to the edit node, used during compilation.
-                const NodeSurfaceEdit::Private* node = nullptr;
+                //! Reference to the edit node, used during compilation.
+                const NodeSurfaceEdit* node = nullptr;
             };
 			TArray<FEdit> edits;
 
@@ -131,7 +132,7 @@ namespace mu
 		//! Store the conditions that enable every modifier.
 		struct FModifier
 		{
-            const NodeModifier::Private* node = nullptr;
+            const NodeModifier* node = nullptr;
 
             // List of tags that are required to apply this modifier
 			TArray<FString> positiveTags;
@@ -151,7 +152,7 @@ namespace mu
             StateCondition stateCondition;
 
             //
-            int32 lod = 0;
+            int32 LOD = 0;
         };
 		TArray<FModifier> modifiers;
 
@@ -175,9 +176,8 @@ namespace mu
         };
         TArray<FTag> m_tags;
 
-        //! Accumulate the model states found while generating code, with their generated root
-        //! nodes.
-        typedef TArray< TPair<FObjectState, const Node::Private*> > StateList;
+        //! Accumulate the model states found while generating code, with their generated root nodes.
+        typedef TArray< TPair<FObjectState, const Node*> > StateList;
         StateList m_states;
 
 		/** Parameters added for every node. */
@@ -195,14 +195,14 @@ namespace mu
 		TArray< StateCondition > m_currentStateCondition;
 
 		//! When processing surfaces, this is the parent component the surfaces may be added to
-        const NodeComponentNew::Private* m_currentComponent = nullptr;
+        const NodeComponentNew* m_currentComponent = nullptr;
 
         //! Current relevant tags so far. Used during traversal.
 		TArray<FString> m_currentPositiveTags;
 		TArray<FString> m_currentNegativeTags;
 
-		//! Index of the LOD we are processing
-        int m_currentLOD = -1;
+		//** Index of the LOD we are processing. */
+        int32 CurrentLOD = -1;
 
 		/** Non-owned reference to main code generator. */
 		CodeGenerator* Generator = nullptr;
diff --git a/Engine/Plugins/Experimental/Mutable/Source/MutableTools/Private/MuT/CodeGenerator_Image.cpp b/Engine/Plugins/Experimental/Mutable/Source/MutableTools/Private/MuT/CodeGenerator_Image.cpp
--- a/Engine/Plugins/Experimental/Mutable/Source/MutableTools/Private/MuT/CodeGenerator_Image.cpp
+++ b/Engine/Plugins/Experimental/Mutable/Source/MutableTools/Private/MuT/CodeGenerator_Image.cpp
@@ -133,34 +133,33 @@ namespace mu
 			const NodeImage* Node = Untyped.get();
 
 			// Generate for each different type of node
-			switch (Untyped->GetImageNodeType())
+			switch (Untyped->GetType()->Type)
 			{
-			case NodeImage::EType::Constant: GenerateImage_Constant(Options, Result, static_cast<const NodeImageConstant*>(Node)); break;
-			case NodeImage::EType::Difference_Deprecated: check(false); break;
-			case NodeImage::EType::Interpolate: GenerateImage_Interpolate(Options, Result, static_cast<const NodeImageInterpolate*>(Node)); break;
-			case NodeImage::EType::Saturate: GenerateImage_Saturate(Options, Result, static_cast<const NodeImageSaturate*>(Node)); break;
-			case NodeImage::EType::Table: GenerateImage_Table(Options, Result, static_cast<const NodeImageTable*>(Node)); break;
-			case NodeImage::EType::Swizzle: GenerateImage_Swizzle(Options, Result, static_cast<const NodeImageSwizzle*>(Node)); break;
-			case NodeImage::EType::ColourMap: GenerateImage_ColourMap(Options, Result, static_cast<const NodeImageColourMap*>(Node)); break;
-			case NodeImage::EType::Gradient: GenerateImage_Gradient(Options, Result, static_cast<const NodeImageGradient*>(Node)); break;
-			case NodeImage::EType::Binarise: GenerateImage_Binarise(Options, Result, static_cast<const NodeImageBinarise*>(Node)); break;
-			case NodeImage::EType::Luminance: GenerateImage_Luminance(Options, Result, static_cast<const NodeImageLuminance*>(Node)); break;
-			case NodeImage::EType::Layer: GenerateImage_Layer(Options, Result, static_cast<const NodeImageLayer*>(Node)); break;
-			case NodeImage::EType::LayerColour: GenerateImage_LayerColour(Options, Result, static_cast<const NodeImageLayerColour*>(Node)); break;
-			case NodeImage::EType::Resize: GenerateImage_Resize(Options, Result, static_cast<const NodeImageResize*>(Node)); break;
-			case NodeImage::EType::PlainColour: GenerateImage_PlainColour(Options, Result, static_cast<const NodeImagePlainColour*>(Node)); break;
-			case NodeImage::EType::Project: GenerateImage_Project(Options, Result, static_cast<const NodeImageProject*>(Node)); break;
-			case NodeImage::EType::Mipmap: GenerateImage_Mipmap(Options, Result, static_cast<const NodeImageMipmap*>(Node)); break;
-			case NodeImage::EType::Switch: GenerateImage_Switch(Options, Result, static_cast<const NodeImageSwitch*>(Node)); break;
-			case NodeImage::EType::Conditional: GenerateImage_Conditional(Options, Result, static_cast<const NodeImageConditional*>(Node)); break;
-			case NodeImage::EType::Format: GenerateImage_Format(Options, Result, static_cast<const NodeImageFormat*>(Node)); break;
-			case NodeImage::EType::Parameter: GenerateImage_Parameter(Options, Result, static_cast<const NodeImageParameter*>(Node)); break;
-			case NodeImage::EType::MultiLayer: GenerateImage_MultiLayer(Options, Result, static_cast<const NodeImageMultiLayer*>(Node)); break;
-			case NodeImage::EType::Invert: GenerateImage_Invert(Options, Result, static_cast<const NodeImageInvert*>(Node)); break;
-			case NodeImage::EType::Variation: GenerateImage_Variation(Options, Result, static_cast<const NodeImageVariation*>(Node)); break;
-			case NodeImage::EType::NormalComposite: GenerateImage_NormalComposite(Options, Result, static_cast<const NodeImageNormalComposite*>(Node)); break;
-			case NodeImage::EType::Transform: GenerateImage_Transform(Options, Result, static_cast<const NodeImageTransform*>(Node)); break;
-			case NodeImage::EType::None: check(false);
+			case Node::EType::ImageConstant: GenerateImage_Constant(Options, Result, static_cast<const NodeImageConstant*>(Node)); break;
+			case Node::EType::ImageInterpolate: GenerateImage_Interpolate(Options, Result, static_cast<const NodeImageInterpolate*>(Node)); break;
+			case Node::EType::ImageSaturate: GenerateImage_Saturate(Options, Result, static_cast<const NodeImageSaturate*>(Node)); break;
+			case Node::EType::ImageTable: GenerateImage_Table(Options, Result, static_cast<const NodeImageTable*>(Node)); break;
+			case Node::EType::ImageSwizzle: GenerateImage_Swizzle(Options, Result, static_cast<const NodeImageSwizzle*>(Node)); break;
+			case Node::EType::ImageColorMap: GenerateImage_ColourMap(Options, Result, static_cast<const NodeImageColourMap*>(Node)); break;
+			case Node::EType::ImageGradient: GenerateImage_Gradient(Options, Result, static_cast<const NodeImageGradient*>(Node)); break;
+			case Node::EType::ImageBinarise: GenerateImage_Binarise(Options, Result, static_cast<const NodeImageBinarise*>(Node)); break;
+			case Node::EType::ImageLuminance: GenerateImage_Luminance(Options, Result, static_cast<const NodeImageLuminance*>(Node)); break;
+			case Node::EType::ImageLayer: GenerateImage_Layer(Options, Result, static_cast<const NodeImageLayer*>(Node)); break;
+			case Node::EType::ImageLayerColour: GenerateImage_LayerColour(Options, Result, static_cast<const NodeImageLayerColour*>(Node)); break;
+			case Node::EType::ImageResize: GenerateImage_Resize(Options, Result, static_cast<const NodeImageResize*>(Node)); break;
+			case Node::EType::ImagePlainColour: GenerateImage_PlainColour(Options, Result, static_cast<const NodeImagePlainColour*>(Node)); break;
+			case Node::EType::ImageProject: GenerateImage_Project(Options, Result, static_cast<const NodeImageProject*>(Node)); break;
+			case Node::EType::ImageMipmap: GenerateImage_Mipmap(Options, Result, static_cast<const NodeImageMipmap*>(Node)); break;
+			case Node::EType::ImageSwitch: GenerateImage_Switch(Options, Result, static_cast<const NodeImageSwitch*>(Node)); break;
+			case Node::EType::ImageConditional: GenerateImage_Conditional(Options, Result, static_cast<const NodeImageConditional*>(Node)); break;
+			case Node::EType::ImageFormat: GenerateImage_Format(Options, Result, static_cast<const NodeImageFormat*>(Node)); break;
+			case Node::EType::ImageParameter: GenerateImage_Parameter(Options, Result, static_cast<const NodeImageParameter*>(Node)); break;
+			case Node::EType::ImageMultiLayer: GenerateImage_MultiLayer(Options, Result, static_cast<const NodeImageMultiLayer*>(Node)); break;
+			case Node::EType::ImageInvert: GenerateImage_Invert(Options, Result, static_cast<const NodeImageInvert*>(Node)); break;
+			case Node::EType::ImageVariation: GenerateImage_Variation(Options, Result, static_cast<const NodeImageVariation*>(Node)); break;
+			case Node::EType::ImageNormalComposite: GenerateImage_NormalComposite(Options, Result, static_cast<const NodeImageNormalComposite*>(Node)); break;
+			case Node::EType::ImageTransform: GenerateImage_Transform(Options, Result, static_cast<const NodeImageTransform*>(Node)); break;
+			default: check(false);
 			}
 
 			// Cache the Result
@@ -187,7 +186,7 @@ namespace mu
             pImage = GenerateMissingImage(EImageFormat::IF_RGB_UBYTE );
 
             // Log an error message
-            m_pErrorLog->GetPrivate()->Add( "Constant image not set.", ELMT_WARNING, node.m_errorContext );
+            m_pErrorLog->GetPrivate()->Add( "Constant image not set.", ELMT_WARNING, InNode->GetMessageContext());
         }
 
 
@@ -319,7 +318,7 @@ namespace mu
         else
         {
             // This argument is required
-            base = GenerateMissingImageCode(TEXT("Image Layer base"), EImageFormat::IF_RGB_UBYTE, node.m_errorContext, Options);
+            base = GenerateMissingImageCode(TEXT("Image Layer base"), EImageFormat::IF_RGB_UBYTE, InNode->GetMessageContext(), Options);
         }
 
 		UE::Math::TIntVector2<int32> TargetSize = Options.RectSize;
@@ -395,7 +394,7 @@ namespace mu
         else
         {
             // This argument is required
-            base = GenerateMissingImageCode(TEXT("Layer base image"), EImageFormat::IF_RGB_UBYTE, node.m_errorContext, Options );
+            base = GenerateMissingImageCode(TEXT("Layer base image"), EImageFormat::IF_RGB_UBYTE, InNode->GetMessageContext(), Options );
         }
 
 		UE::Math::TIntVector2<int32> TargetSize = Options.RectSize;
@@ -435,7 +434,7 @@ namespace mu
         else
         {
             // This argument is required
-            colour = GenerateMissingColourCode(TEXT("Layer colour"), node.m_errorContext );
+            colour = GenerateMissingColourCode(TEXT("Layer colour"), InNode->GetMessageContext() );
         }
         op->color = colour;
 
@@ -465,7 +464,7 @@ namespace mu
         else
         {
             // This argument is required
-            base = GenerateMissingImageCode(TEXT("Image MultiLayer base"), EImageFormat::IF_RGB_UBYTE, node.m_errorContext, Options);
+            base = GenerateMissingImageCode(TEXT("Image MultiLayer base"), EImageFormat::IF_RGB_UBYTE, InNode->GetMessageContext(), Options);
         }
 
 		FImageDesc BaseDesc = base->GetImageDesc(true);
@@ -550,7 +549,7 @@ namespace mu
         else
         {
             // This argument is required
-            base = GenerateMissingImageCode(TEXT("Image Composite Base"), EImageFormat::IF_RGB_UBYTE, node.m_errorContext, Options);
+            base = GenerateMissingImageCode(TEXT("Image Composite Base"), EImageFormat::IF_RGB_UBYTE, InNode->GetMessageContext(), Options);
         }
 
 		FImageDesc BaseDesc = base->GetImageDesc(true);
@@ -580,7 +579,7 @@ namespace mu
 		else
 		{
             // This argument is required
-            normal = GenerateMissingImageCode(TEXT("Image Composite Normal"), EImageFormat::IF_RGB_UBYTE, node.m_errorContext, Options);
+            normal = GenerateMissingImageCode(TEXT("Image Composite Normal"), EImageFormat::IF_RGB_UBYTE, InNode->GetMessageContext(), Options);
 		}
 
 		normal = GenerateImageSize(normal, TargetSize);
@@ -601,31 +600,41 @@ namespace mu
 		Ptr<ASTOp> OffsetX;
 		if (Node.m_pOffsetX)
 		{
-			OffsetX = Generate(Node.m_pOffsetX, Options);
+			FScalarGenerationResult ScalarResult;
+			GenerateScalar(ScalarResult, Options, Node.m_pOffsetX);
+			OffsetX = ScalarResult.op;
 		}
 
 		Ptr<ASTOp> OffsetY;
 		if (Node.m_pOffsetY)
 		{
-			OffsetY = Generate(Node.m_pOffsetY, Options);
+			FScalarGenerationResult ScalarResult;
+			GenerateScalar(ScalarResult, Options, Node.m_pOffsetY);
+			OffsetY = ScalarResult.op;
 		}
 	
 		Ptr<ASTOp> ScaleX;
 		if (Node.m_pScaleX)
 		{
-			ScaleX = Generate(Node.m_pScaleX, Options);
+			FScalarGenerationResult ScalarResult;
+			GenerateScalar(ScalarResult, Options, Node.m_pScaleX);
+			ScaleX = ScalarResult.op;
 		}
 	
 		Ptr<ASTOp> ScaleY;
 		if (Node.m_pScaleY)
 		{
-			ScaleY = Generate(Node.m_pScaleY, Options);
+			FScalarGenerationResult ScalarResult;
+			GenerateScalar(ScalarResult, Options, Node.m_pScaleY);
+			ScaleY = ScalarResult.op;
 		}
 
 		Ptr<ASTOp> Rotation;
 		if (Node.m_pRotation)
 		{
-			Rotation = Generate(Node.m_pRotation, Options);
+			FScalarGenerationResult ScalarResult;
+			GenerateScalar(ScalarResult, Options, Node.m_pRotation);
+			Rotation = ScalarResult.op;
 		}
 
 		// If one of the inputs (offset or scale) is missig assume unifrom translation/scaling 
@@ -656,7 +665,7 @@ namespace mu
         else
         {
             // This argument is required
-            Base = GenerateMissingImageCode(TEXT("Image Transform Base"), EImageFormat::IF_RGB_UBYTE, Node.m_errorContext, NewOptions);
+            Base = GenerateMissingImageCode(TEXT("Image Transform Base"), EImageFormat::IF_RGB_UBYTE, InNode->GetMessageContext(), NewOptions);
         }
 		
 		FImageDesc BaseDesc = Base->GetImageDesc();
@@ -681,13 +690,13 @@ namespace mu
         // Factor
         if ( Node* pFactor = node.m_pFactor.get() )
         {
-            op->SetChild( op->op.args.ImageInterpolate.factor, Generate( pFactor, Options));
+            op->SetChild( op->op.args.ImageInterpolate.factor, Generate_Generic( pFactor, Options));
         }
         else
         {
             // This argument is required
             op->SetChild( op->op.args.ImageInterpolate.factor,
-                          GenerateMissingScalarCode(TEXT("Interpolation factor"), 0.5f, node.m_errorContext ));
+                          GenerateMissingScalarCode(TEXT("Interpolation factor"), 0.5f, InNode->GetMessageContext() ));
         }
 
         // Target images
@@ -725,7 +734,7 @@ namespace mu
         // At least one target is required
         if (!op->op.args.ImageInterpolate.targets[0])
         {
-            Ptr<ASTOp> target = GenerateMissingImageCode(TEXT("First interpolation image"), EImageFormat::IF_RGB_UBYTE, node.m_errorContext, Options);
+            Ptr<ASTOp> target = GenerateMissingImageCode(TEXT("First interpolation image"), EImageFormat::IF_RGB_UBYTE, InNode->GetMessageContext(), Options);
             target = GenerateImageSize( target, Options.RectSize);
             op->SetChild( op->op.args.ImageInterpolate.targets[0], target);
         }
@@ -804,7 +813,7 @@ namespace mu
 				if (!source)
 				{
 					// TODO: Warn?
-					source = GenerateMissingImageCode(TEXT("Swizzle channel"), EImageFormat::IF_L_UBYTE, node.m_errorContext, Options);
+					source = GenerateMissingImageCode(TEXT("Swizzle channel"), EImageFormat::IF_L_UBYTE, InNode->GetMessageContext(), Options);
 				}
 
                 Ptr<ASTOp> sizedSource;
@@ -827,7 +836,7 @@ namespace mu
 		// At least one source is required
         if (!op->Sources[0])
 		{
-            Ptr<ASTOp> source = GenerateMissingImageCode(TEXT("First swizzle image"), EImageFormat::IF_RGBA_UBYTE, node.m_errorContext, Options);
+            Ptr<ASTOp> source = GenerateMissingImageCode(TEXT("First swizzle image"), EImageFormat::IF_RGBA_UBYTE, InNode->GetMessageContext(), Options);
             op->Sources[0] = source;
 		}
 
@@ -861,7 +870,7 @@ namespace mu
 		// Source is required
 		if (!node.m_source)
 		{
-            fop->Source = GenerateMissingImageCode(TEXT("Source image for format."), EImageFormat::IF_RGBA_UBYTE, node.m_errorContext, Options);
+            fop->Source = GenerateMissingImageCode(TEXT("Source image for format."), EImageFormat::IF_RGBA_UBYTE, InNode->GetMessageContext(), Options);
 		}
 		else
 		{
@@ -894,7 +903,7 @@ namespace mu
         else
         {
             // This argument is required
-            base = GenerateMissingImageCode(TEXT("Saturate image"), EImageFormat::IF_RGB_UBYTE, node.m_errorContext, Options);
+            base = GenerateMissingImageCode(TEXT("Saturate image"), EImageFormat::IF_RGB_UBYTE, InNode->GetMessageContext(), Options);
         }
 		
         base = GenerateImageFormat(base, GetRGBOrRGBAFormat(base->GetImageDesc().m_format));
@@ -905,12 +914,12 @@ namespace mu
         // Factor
         if ( Node* pFactor = node.m_pFactor.get() )
         {
-            op->SetChild( op->op.args.ImageSaturate.factor, Generate( pFactor, Options));
+            op->SetChild( op->op.args.ImageSaturate.factor, Generate_Generic( pFactor, Options));
         }
         else
         {
             // This argument is required
-            op->SetChild( op->op.args.ImageSaturate.factor, GenerateMissingScalarCode(TEXT("Saturation factor"), 0.5f, node.m_errorContext ) );
+            op->SetChild( op->op.args.ImageSaturate.factor, GenerateMissingScalarCode(TEXT("Saturation factor"), 0.5f, InNode->GetMessageContext() ) );
         }
 
         Result.op = op;
@@ -936,7 +945,7 @@ namespace mu
         else
         {
             // This argument is required
-            base = GenerateMissingImageCode(TEXT("Image luminance"), EImageFormat::IF_RGB_UBYTE, node.m_errorContext, Options);
+            base = GenerateMissingImageCode(TEXT("Image luminance"), EImageFormat::IF_RGB_UBYTE, InNode->GetMessageContext(), Options);
         }
         base = GenerateImageFormat( base, EImageFormat::IF_RGB_UBYTE );
         op->SetChild( op->op.args.ImageLuminance.base, base);
@@ -964,7 +973,7 @@ namespace mu
         else
         {
             // This argument is required
-            base = GenerateMissingImageCode(TEXT("Colourmap base image"), EImageFormat::IF_RGB_UBYTE, node.m_errorContext, Options);
+            base = GenerateMissingImageCode(TEXT("Colourmap base image"), EImageFormat::IF_RGB_UBYTE, InNode->GetMessageContext(), Options);
         }
         base = GenerateImageSize( base, Options.RectSize);
         op->SetChild( op->op.args.ImageColourMap.base, base);
@@ -1005,7 +1014,7 @@ namespace mu
         else
         {
             // This argument is required
-			MapImageOp = GenerateMissingImageCode(TEXT("Map image"), EImageFormat::IF_RGB_UBYTE, node.m_errorContext, Options);
+			MapImageOp = GenerateMissingImageCode(TEXT("Map image"), EImageFormat::IF_RGB_UBYTE, InNode->GetMessageContext(), Options);
         }
         op->SetChild( op->op.args.ImageColourMap.map, MapImageOp);
 
@@ -1025,12 +1034,12 @@ namespace mu
         Ptr<ASTOp> colour0 = 0;
         if ( Node* pColour0 = node.m_pColour0.get() )
         {
-            colour0 = Generate( pColour0, Options);
+            colour0 = Generate_Generic( pColour0, Options);
         }
         else
         {
             // This argument is required
-            colour0 = GenerateMissingColourCode(TEXT("Gradient colour 0"), node.m_errorContext );
+            colour0 = GenerateMissingColourCode(TEXT("Gradient colour 0"), InNode->GetMessageContext() );
         }
         op->SetChild( op->op.args.ImageGradient.colour0, colour0);
 
@@ -1038,12 +1047,12 @@ namespace mu
         Ptr<ASTOp> colour1 = 0;
         if ( Node* pColour1 = node.m_pColour1.get() )
         {
-            colour1 = Generate( pColour1, Options);
+            colour1 = Generate_Generic( pColour1, Options);
         }
         else
         {
             // This argument is required
-            colour1 = GenerateMissingColourCode(TEXT("Gradient colour 1"), node.m_errorContext );
+            colour1 = GenerateMissingColourCode(TEXT("Gradient colour 1"), InNode->GetMessageContext() );
         }
         op->SetChild( op->op.args.ImageGradient.colour1, colour1);
 
@@ -1073,7 +1082,7 @@ namespace mu
         else
         {
             // This argument is required
-            a = GenerateMissingImageCode(TEXT("Image Binarise Base"), EImageFormat::IF_RGB_UBYTE, node.m_errorContext, Options);
+            a = GenerateMissingImageCode(TEXT("Image Binarise Base"), EImageFormat::IF_RGB_UBYTE, InNode->GetMessageContext(), Options);
         }
         a = GenerateImageFormat( a, EImageFormat::IF_RGB_UBYTE );
         a = GenerateImageSize( a, Options.RectSize);
@@ -1083,12 +1092,12 @@ namespace mu
         Ptr<ASTOp> b = 0;
         if ( Node* pScalar = node.m_pThreshold.get() )
         {
-            b = Generate( pScalar, Options);
+            b = Generate_Generic( pScalar, Options);
         }
         else
         {
             // This argument is required
-            b = GenerateMissingScalarCode(TEXT("Image Binarise Threshold"), 0.5f, node.m_errorContext );
+            b = GenerateMissingScalarCode(TEXT("Image Binarise Threshold"), 0.5f, InNode->GetMessageContext() );
         }
         op->SetChild( op->op.args.ImageBinarise.threshold, b );
 
@@ -1124,7 +1133,7 @@ namespace mu
         else
         {
             // This argument is required
-            base = GenerateMissingImageCode(TEXT("Image resize base"), EImageFormat::IF_RGB_UBYTE, node.m_errorContext, Options);
+            base = GenerateMissingImageCode(TEXT("Image resize base"), EImageFormat::IF_RGB_UBYTE, InNode->GetMessageContext(), Options);
         }
 
         // Size
@@ -1185,12 +1194,12 @@ namespace mu
         Ptr<ASTOp> base = 0;
         if ( node.m_pColour )
         {
-            base = Generate( node.m_pColour.get(), Options);
+            base = Generate_Generic( node.m_pColour.get(), Options);
         }
         else
         {
             // This argument is required
-            base = GenerateMissingColourCode(TEXT("Image plain colour base"), node.m_errorContext );
+            base = GenerateMissingColourCode(TEXT("Image plain colour base"), InNode->GetMessageContext() );
         }
 
 		UE::Math::TIntVector2<int32> FinalImageSize = { 0, 0 };
@@ -1251,7 +1260,7 @@ namespace mu
         if (node.m_options.Num() == 0)
 		{
 			// No options in the switch!
-            Ptr<ASTOp> missingOp = GenerateMissingImageCode(TEXT("Switch option"), EImageFormat::IF_RGBA_UBYTE, node.m_errorContext, Options);
+            Ptr<ASTOp> missingOp = GenerateMissingImageCode(TEXT("Switch option"), EImageFormat::IF_RGBA_UBYTE, InNode->GetMessageContext(), Options);
 			Result.op = missingOp;
 			return;
 		}
@@ -1262,12 +1271,12 @@ namespace mu
 		// Variable value
 		if ( node.m_pParameter )
 		{
-            op->variable = Generate( node.m_pParameter.get(), Options);
+            op->variable = Generate_Generic( node.m_pParameter.get(), Options);
 		}
 		else
 		{
 			// This argument is required
-            op->variable = GenerateMissingScalarCode(TEXT("Switch variable"), 0.0f, node.m_errorContext );
+            op->variable = GenerateMissingScalarCode(TEXT("Switch variable"), 0.0f, InNode->GetMessageContext() );
 		}
 
 		// Options
@@ -1284,7 +1293,7 @@ namespace mu
             else
             {
                 // This argument is required
-                branch = GenerateMissingImageCode(TEXT("Switch option"), EImageFormat::IF_RGBA_UBYTE, node.m_errorContext, Options);
+                branch = GenerateMissingImageCode(TEXT("Switch option"), EImageFormat::IF_RGBA_UBYTE, InNode->GetMessageContext(), Options);
             }
 
             op->cases.Emplace((int16_t)t,op,branch);
@@ -1336,12 +1345,12 @@ namespace mu
         // Condition
         if ( node.m_parameter )
         {
-            op->condition = Generate( node.m_parameter.get(), Options);
+            op->condition = Generate_Generic( node.m_parameter.get(), Options);
         }
         else
         {
             // This argument is required
-            op->condition = GenerateMissingBoolCode(TEXT("Conditional condition"), true, node.m_errorContext );
+            op->condition = GenerateMissingBoolCode(TEXT("Conditional condition"), true, InNode->GetMessageContext() );
         }
 
         // Options
@@ -1381,7 +1390,7 @@ namespace mu
         else
         {
             // This argument is required
-            GenerateMissingProjectorCode( projectorResult, node.m_errorContext );
+            GenerateMissingProjectorCode( projectorResult, InNode->GetMessageContext() );
         }
 
 		ProjectOp->SetChild(ProjectOp->op.args.MeshProject.projector, projectorResult.op );
@@ -1420,7 +1429,7 @@ namespace mu
 			}
 			else
 			{
-				m_pErrorLog->GetPrivate()->Add("Layout or block index error.", ELMT_ERROR, node.m_errorContext);
+				m_pErrorLog->GetPrivate()->Add("Layout or block index error.", ELMT_ERROR, InNode->GetMessageContext());
 			}
 
 
@@ -1488,7 +1497,7 @@ namespace mu
             cop->Type = OP_TYPE::ME_CONSTANT;
             cop->SetValue(TempMesh, m_compilerOptions->OptimisationOptions.DiskCacheContext);
 			ProjectOp->SetChild(ProjectOp->op.args.MeshProject.mesh, cop );
-            m_pErrorLog->GetPrivate()->Add( "Projector mesh not set.", ELMT_ERROR, node.m_errorContext );
+            m_pErrorLog->GetPrivate()->Add( "Projector mesh not set.", ELMT_ERROR, InNode->GetMessageContext() );
         }
 
 
@@ -1561,7 +1570,7 @@ namespace mu
         else
         {
             // This argument is required
-			ImageRasterOp->image = GenerateMissingImageCode(TEXT("Projector image"), EImageFormat::IF_RGB_UBYTE, node.m_errorContext, Options);
+			ImageRasterOp->image = GenerateMissingImageCode(TEXT("Projector image"), EImageFormat::IF_RGB_UBYTE, InNode->GetMessageContext(), Options);
         }
 
 		if (RasterImageSize.X == 0)
@@ -1619,7 +1628,7 @@ namespace mu
 			if (node.m_pAngleFadeEnd) pPropsNode->SetY(node.m_pAngleFadeEnd);
 			else pPropsNode->SetY(pDefaultFade);
 
-			ImageRasterOp->angleFadeProperties = Generate(pPropsNode, Options);
+			ImageRasterOp->angleFadeProperties = Generate_Generic(pPropsNode, Options);
 		}
 
         // Seam correction operations
@@ -1694,7 +1703,7 @@ namespace mu
         else
         {
             // This argument is required
-            base = GenerateMissingImageCode(TEXT("Mipmap image"), EImageFormat::IF_RGB_UBYTE, node.m_errorContext, Options);
+            base = GenerateMissingImageCode(TEXT("Mipmap image"), EImageFormat::IF_RGB_UBYTE, InNode->GetMessageContext(), Options);
         }
 
         op->Source = base;
@@ -1751,7 +1760,7 @@ namespace mu
 		else
 		{
 			// This argument is required
-			a = GenerateMissingImageCode(TEXT("Image Invert Color"), EImageFormat::IF_RGB_UBYTE, node.m_errorContext, Options);
+			a = GenerateMissingImageCode(TEXT("Image Invert Color"), EImageFormat::IF_RGB_UBYTE, InNode->GetMessageContext(), Options);
 		}
 		a = GenerateImageFormat(a, EImageFormat::IF_RGB_UBYTE);
 		a = GenerateImageSize(a, Options.RectSize);
@@ -1793,7 +1802,7 @@ namespace mu
             {
 				FString Msg = FString::Printf(TEXT("Unknown tag found in image variation [%s]."), *tag );
 
-                m_pErrorLog->GetPrivate()->Add( Msg, ELMT_WARNING, node.m_errorContext );
+                m_pErrorLog->GetPrivate()->Add( Msg, ELMT_WARNING, InNode->GetMessageContext() );
                 continue;
             }
 
@@ -1807,7 +1816,7 @@ namespace mu
             else
             {
                 // This argument is required
-                variationOp = GenerateMissingImageCode(TEXT("Variation option"), EImageFormat::IF_RGBA_UBYTE, node.m_errorContext, Options);
+                variationOp = GenerateMissingImageCode(TEXT("Variation option"), EImageFormat::IF_RGBA_UBYTE, InNode->GetMessageContext(), Options);
             }
 
 
@@ -1856,7 +1865,7 @@ namespace mu
 		const NodeImageTable::Private& node = *InNode->GetPrivate();
 
 		Result.op = GenerateTableSwitch<NodeImageTable::Private, ETableColumnType::Image, OP_TYPE::IM_SWITCH>(node,
-			[this,Options](const NodeImageTable::Private& node, int colIndex, int row, ErrorLog* pErrorLog)
+			[this, InNode, Options](const NodeImageTable::Private& node, int colIndex, int row, ErrorLog* pErrorLog)
 			{
 				const FTableValue& CellData = node.Table->GetPrivate()->Rows[row].Values[colIndex];
 				ImagePtrConst pImage = nullptr;
@@ -1871,7 +1880,7 @@ namespace mu
 				if (!pImage)
 				{
 					FString Msg = FString::Printf(TEXT("Table has a missing image in column %d, row %d."), colIndex, row);
-					pErrorLog->GetPrivate()->Add(Msg, ELMT_ERROR, node.m_errorContext);
+					pErrorLog->GetPrivate()->Add(Msg, ELMT_ERROR, InNode->GetMessageContext());
 
 					return ImageOp;
 				}
diff --git a/Engine/Plugins/Experimental/Mutable/Source/MutableTools/Private/MuT/CodeGenerator_Mesh.cpp b/Engine/Plugins/Experimental/Mutable/Source/MutableTools/Private/MuT/CodeGenerator_Mesh.cpp
--- a/Engine/Plugins/Experimental/Mutable/Source/MutableTools/Private/MuT/CodeGenerator_Mesh.cpp
+++ b/Engine/Plugins/Experimental/Mutable/Source/MutableTools/Private/MuT/CodeGenerator_Mesh.cpp
@@ -732,25 +732,25 @@ namespace mu
 		const NodeMesh* Node = InUntypedNode.get();
 
         // Generate for each different type of node
-		switch (Node->GetMeshNodeType())
+		switch (Node->GetType()->Type)
 		{
-		case NodeMesh::EType::Constant: GenerateMesh_Constant(InOptions, OutResult, static_cast<const NodeMeshConstant*>(Node)); break;
-		case NodeMesh::EType::Format: GenerateMesh_Format(InOptions, OutResult, static_cast<const NodeMeshFormat*>(Node)); break;
-		case NodeMesh::EType::Morph: GenerateMesh_Morph(InOptions, OutResult, static_cast<const NodeMeshMorph*>(Node)); break;
-		case NodeMesh::EType::MakeMorph: GenerateMesh_MakeMorph(InOptions, OutResult, static_cast<const NodeMeshMakeMorph*>(Node)); break;
-		case NodeMesh::EType::Fragment: GenerateMesh_Fragment(InOptions, OutResult, static_cast<const NodeMeshFragment*>(Node)); break;
-		case NodeMesh::EType::Interpolate: GenerateMesh_Interpolate(InOptions, OutResult, static_cast<const NodeMeshInterpolate*>(Node)); break;
-		case NodeMesh::EType::Switch: GenerateMesh_Switch(InOptions, OutResult, static_cast<const NodeMeshSwitch*>(Node)); break;
-		case NodeMesh::EType::Transform: GenerateMesh_Transform(InOptions, OutResult, static_cast<const NodeMeshTransform*>(Node)); break;
-		case NodeMesh::EType::ClipMorphPlane: GenerateMesh_ClipMorphPlane(InOptions, OutResult, static_cast<const NodeMeshClipMorphPlane*>(Node)); break;
-		case NodeMesh::EType::ClipWithMesh: GenerateMesh_ClipWithMesh(InOptions, OutResult, static_cast<const NodeMeshClipWithMesh*>(Node)); break;
-		case NodeMesh::EType::ApplyPose: GenerateMesh_ApplyPose(InOptions, OutResult, static_cast<const NodeMeshApplyPose*>(Node)); break;
-		case NodeMesh::EType::Variation: GenerateMesh_Variation(InOptions, OutResult, static_cast<const NodeMeshVariation*>(Node)); break;
-		case NodeMesh::EType::Table: GenerateMesh_Table(InOptions, OutResult, static_cast<const NodeMeshTable*>(Node)); break;
-		case NodeMesh::EType::GeometryOperation: GenerateMesh_GeometryOperation(InOptions, OutResult, static_cast<const NodeMeshGeometryOperation*>(Node)); break;
-		case NodeMesh::EType::Reshape: GenerateMesh_Reshape(InOptions, OutResult, static_cast<const NodeMeshReshape*>(Node)); break;
-		case NodeMesh::EType::ClipDeform: GenerateMesh_ClipDeform(InOptions, OutResult, static_cast<const NodeMeshClipDeform*>(Node)); break;
-		case NodeMesh::EType::None: check(false);
+		case Node::EType::MeshConstant: GenerateMesh_Constant(InOptions, OutResult, static_cast<const NodeMeshConstant*>(Node)); break;
+		case Node::EType::MeshFormat: GenerateMesh_Format(InOptions, OutResult, static_cast<const NodeMeshFormat*>(Node)); break;
+		case Node::EType::MeshMorph: GenerateMesh_Morph(InOptions, OutResult, static_cast<const NodeMeshMorph*>(Node)); break;
+		case Node::EType::MeshMakeMorph: GenerateMesh_MakeMorph(InOptions, OutResult, static_cast<const NodeMeshMakeMorph*>(Node)); break;
+		case Node::EType::MeshFragment: GenerateMesh_Fragment(InOptions, OutResult, static_cast<const NodeMeshFragment*>(Node)); break;
+		case Node::EType::MeshInterpolate: GenerateMesh_Interpolate(InOptions, OutResult, static_cast<const NodeMeshInterpolate*>(Node)); break;
+		case Node::EType::MeshSwitch: GenerateMesh_Switch(InOptions, OutResult, static_cast<const NodeMeshSwitch*>(Node)); break;
+		case Node::EType::MeshTransform: GenerateMesh_Transform(InOptions, OutResult, static_cast<const NodeMeshTransform*>(Node)); break;
+		case Node::EType::MeshClipMorphPlane: GenerateMesh_ClipMorphPlane(InOptions, OutResult, static_cast<const NodeMeshClipMorphPlane*>(Node)); break;
+		case Node::EType::MeshClipWithMesh: GenerateMesh_ClipWithMesh(InOptions, OutResult, static_cast<const NodeMeshClipWithMesh*>(Node)); break;
+		case Node::EType::MeshApplyPose: GenerateMesh_ApplyPose(InOptions, OutResult, static_cast<const NodeMeshApplyPose*>(Node)); break;
+		case Node::EType::MeshVariation: GenerateMesh_Variation(InOptions, OutResult, static_cast<const NodeMeshVariation*>(Node)); break;
+		case Node::EType::MeshTable: GenerateMesh_Table(InOptions, OutResult, static_cast<const NodeMeshTable*>(Node)); break;
+		case Node::EType::MeshGeometryOperation: GenerateMesh_GeometryOperation(InOptions, OutResult, static_cast<const NodeMeshGeometryOperation*>(Node)); break;
+		case Node::EType::MeshReshape: GenerateMesh_Reshape(InOptions, OutResult, static_cast<const NodeMeshReshape*>(Node)); break;
+		case Node::EType::MeshClipDeform: GenerateMesh_ClipDeform(InOptions, OutResult, static_cast<const NodeMeshClipDeform*>(Node)); break;
+		default: check(false);
 		}
 
         // Cache the result
@@ -772,12 +772,12 @@ namespace mu
         // Factor
         if ( node.Factor )
         {
-            OpMorph->Factor = Generate( node.Factor.get(), InOptions );
+            OpMorph->Factor = Generate_Generic( node.Factor.get(), InOptions );
         }
         else
         {
             // This argument is required
-            OpMorph->Factor = GenerateMissingScalarCode(TEXT("Morph factor"), 0.5f, node.m_errorContext );
+            OpMorph->Factor = GenerateMissingScalarCode(TEXT("Morph factor"), 0.5f, InMorphNode->GetMessageContext());
         }
 
         // Base
@@ -791,7 +791,7 @@ namespace mu
         {
             // This argument is required
             m_pErrorLog->GetPrivate()->Add( "Mesh morph base node is not set.",
-                                            ELMT_ERROR, node.m_errorContext );
+                                            ELMT_ERROR, InMorphNode->GetMessageContext());
         }        
 
 		if (node.Morph)
@@ -890,7 +890,7 @@ namespace mu
         {
             // This argument is required
             m_pErrorLog->GetPrivate()->Add( "Mesh make morph base node is not set.",
-                                            ELMT_ERROR, node.m_errorContext );
+                                            ELMT_ERROR, InMakeMorphNode->GetMessageContext());
         }
 
         // Target
@@ -909,7 +909,7 @@ namespace mu
         {
             // This argument is required
             m_pErrorLog->GetPrivate()->Add( "Mesh make morph target node is not set.",
-                                            ELMT_ERROR, node.m_errorContext );
+                                            ELMT_ERROR, InMakeMorphNode->GetMessageContext());
         }
 
         OutResult.meshOp = op;
@@ -919,9 +919,9 @@ namespace mu
 
     //---------------------------------------------------------------------------------------------
     void CodeGenerator::GenerateMesh_Fragment(const FMeshGenerationOptions& InOptions, FMeshGenerationResult& OutResult,
-		const NodeMeshFragment* fragment )
+		const NodeMeshFragment* FragmentNode )
     {
-        NodeMeshFragment::Private& node = *fragment->GetPrivate();
+        NodeMeshFragment::Private& node = *FragmentNode->GetPrivate();
 
         FMeshGenerationResult BaseResult;
         if ( node.m_pMesh )
@@ -955,7 +955,7 @@ namespace mu
                         else
                         {
                             m_pErrorLog->GetPrivate()->Add( "Internal layout block index error.",
-                                                            ELMT_ERROR, node.m_errorContext );
+                                                            ELMT_ERROR, FragmentNode->GetMessageContext());
                         }
                     }
                 }
@@ -963,7 +963,7 @@ namespace mu
                 {
                     // This argument is required
                     m_pErrorLog->GetPrivate()->Add( "Missing layout in mesh fragment source.",
-                                                    ELMT_ERROR, node.m_errorContext );
+                                                    ELMT_ERROR, FragmentNode->GetMessageContext());
                 }
             }
 
@@ -977,7 +977,7 @@ namespace mu
         {
             // This argument is required
             m_pErrorLog->GetPrivate()->Add( "Mesh fragment source is not set.",
-                                            ELMT_ERROR, node.m_errorContext );
+                                            ELMT_ERROR, FragmentNode->GetMessageContext());
         }
 
         OutResult.baseMeshOp = BaseResult.baseMeshOp;
@@ -987,9 +987,9 @@ namespace mu
 
     //---------------------------------------------------------------------------------------------
     void CodeGenerator::GenerateMesh_Interpolate(const FMeshGenerationOptions& InOptions, FMeshGenerationResult& OutResult,
-		const NodeMeshInterpolate* interpolate )
+		const NodeMeshInterpolate* InterpolateNode )
     {
-        NodeMeshInterpolate::Private& node = *interpolate->GetPrivate();
+        NodeMeshInterpolate::Private& node = *InterpolateNode->GetPrivate();
 
         // Generate the code
         Ptr<ASTOpFixed> op = new ASTOpFixed();
@@ -999,13 +999,13 @@ namespace mu
         // Factor
         if ( Node* pFactor = node.m_pFactor.get() )
         {
-            op->SetChild( op->op.args.MeshInterpolate.factor, Generate( pFactor, InOptions ) );
+            op->SetChild( op->op.args.MeshInterpolate.factor, Generate_Generic( pFactor, InOptions ) );
         }
         else
         {
             // This argument is required
             op->SetChild( op->op.args.MeshInterpolate.factor,
-                          GenerateMissingScalarCode(TEXT("Interpolation factor"), 0.5f, node.m_errorContext ) );
+                          GenerateMissingScalarCode(TEXT("Interpolation factor"), 0.5f, InterpolateNode->GetMessageContext()) );
         }
 
         //
@@ -1066,15 +1066,15 @@ namespace mu
             //op.args.MeshInterpolate.target[0] = GenerateMissingImageCode( "First mesh", IF_RGB_UBYTE );
             m_pErrorLog->GetPrivate()->Add
                 ( "Mesh interpolation: at least the first mesh is required.",
-                  ELMT_ERROR, node.m_errorContext );
+                  ELMT_ERROR, InterpolateNode->GetMessageContext());
         }
     }
 
 
     //---------------------------------------------------------------------------------------------
-    void CodeGenerator::GenerateMesh_Switch(const FMeshGenerationOptions& InOptions, FMeshGenerationResult& OutResult, const NodeMeshSwitch* sw )
+    void CodeGenerator::GenerateMesh_Switch(const FMeshGenerationOptions& InOptions, FMeshGenerationResult& OutResult, const NodeMeshSwitch* SwitchNode )
     {
-        NodeMeshSwitch::Private& node = *sw->GetPrivate();
+        NodeMeshSwitch::Private& node = *SwitchNode->GetPrivate();
 
         if (node.m_options.Num() == 0)
         {
@@ -1090,12 +1090,12 @@ namespace mu
         // Factor
         if ( node.m_pParameter )
         {
-            op->variable = Generate( node.m_pParameter.get(), InOptions);
+            op->variable = Generate_Generic( node.m_pParameter.get(), InOptions);
         }
         else
         {
             // This argument is required
-            op->variable = GenerateMissingScalarCode(TEXT("Switch variable"), 0.0f, node.m_errorContext );
+            op->variable = GenerateMissingScalarCode(TEXT("Switch variable"), 0.0f, SwitchNode->GetMessageContext());
         }
 
         // Options
@@ -1187,9 +1187,9 @@ namespace mu
 
     //---------------------------------------------------------------------------------------------
     void CodeGenerator::GenerateMesh_Variation(const FMeshGenerationOptions& InOptions, FMeshGenerationResult& OutResult,
-		const NodeMeshVariation* va )
+		const NodeMeshVariation* VariationNode )
     {
-        NodeMeshVariation::Private& node = *va->GetPrivate();
+        NodeMeshVariation::Private& node = *VariationNode->GetPrivate();
 
         FMeshGenerationResult currentResult;
         Ptr<ASTOp> currentMeshOp;
@@ -1226,7 +1226,7 @@ namespace mu
                 m_pErrorLog->GetPrivate()->Add( 
 					FString::Printf(TEXT("Unknown tag found in mesh variation [%s]."), *tag),
 					ELMT_WARNING,
-					node.m_errorContext,
+					VariationNode->GetMessageContext(),
 					ELMSB_UNKNOWN_TAG
 				);
                 continue;
@@ -1295,7 +1295,7 @@ namespace mu
 			GeneratedConstantMeshes.Add(MeshEntry);
 
 			// Log an error message
-			m_pErrorLog->GetPrivate()->Add("Constant mesh not set.", ELMT_WARNING, Node.m_errorContext);
+			m_pErrorLog->GetPrivate()->Add("Constant mesh not set.", ELMT_WARNING, InNode->GetMessageContext());
 
 			return;
 		}
@@ -1442,7 +1442,7 @@ namespace mu
 
 						Ptr<const Layout> SourceLayout = TypedNode->GetPrivate()->m_pLayout;
 						Ptr<const Layout> GeneratedLayout = AddLayout(SourceLayout, MeshIDPrefix);
-						const void* Context = InOptions.OverrideContext.Get(Node.m_errorContext);
+						const void* Context = InOptions.OverrideContext.Get(InNode->GetMessageContext());
 
 						bool bUseAbsoluteBlockIds = false;
 						PrepareForLayout(GeneratedLayout, Cloned, LayoutIndex, Context, InOptions, bUseAbsoluteBlockIds);
@@ -1456,7 +1456,7 @@ namespace mu
 					for (int32 LayoutIndex = 0; LayoutIndex < InOptions.OverrideLayouts.Num(); ++LayoutIndex)
 					{
 						Ptr<const Layout> GeneratedLayout = InOptions.OverrideLayouts[LayoutIndex];
-						const void* Context = InOptions.OverrideContext.Get(Node.m_errorContext);
+						const void* Context = InOptions.OverrideContext.Get(InNode->GetMessageContext());
 
 						// In this case we need the layout block ids to use the ids in the parent layout, and not be prefixed with
 						// the current mesh id prefix. For this reason we need them to be absolute.
@@ -1483,7 +1483,7 @@ namespace mu
 
 		// Apply the modifier for the pre-normal operations stage.
 		bool bModifiersForBeforeOperations = true;
-		OutResult.meshOp = ApplyMeshModifiers(InOptions, LastMeshOp, bModifiersForBeforeOperations, Node.m_errorContext);
+		OutResult.meshOp = ApplyMeshModifiers(InOptions, LastMeshOp, bModifiersForBeforeOperations, InNode->GetMessageContext());
     }
 
 
@@ -1541,9 +1541,9 @@ namespace mu
 
     //---------------------------------------------------------------------------------------------
     void CodeGenerator::GenerateMesh_Transform(const FMeshGenerationOptions& InOptions, FMeshGenerationResult& OutResult,
-                                                const NodeMeshTransform* trans )
+                                                const NodeMeshTransform* TransformNode )
     {
-        const auto& node = *trans->GetPrivate();
+        const auto& node = *TransformNode->GetPrivate();
 
         Ptr<ASTOpMeshTransform> op = new ASTOpMeshTransform();
 
@@ -1556,8 +1556,7 @@ namespace mu
         else
         {
             // This argument is required
-            m_pErrorLog->GetPrivate()->Add("Mesh transform base node is not set.",
-                ELMT_ERROR, node.m_errorContext);
+            m_pErrorLog->GetPrivate()->Add("Mesh transform base node is not set.", ELMT_ERROR, TransformNode->GetMessageContext() );
         }
 
         op->matrix = node.Transform;
@@ -1568,9 +1567,9 @@ namespace mu
 
     //---------------------------------------------------------------------------------------------
     void CodeGenerator::GenerateMesh_ClipMorphPlane(const FMeshGenerationOptions& InOptions, FMeshGenerationResult& OutResult,
-                                                     const NodeMeshClipMorphPlane* clip )
+                                                     const NodeMeshClipMorphPlane* ClipNode )
     {
-        const auto& node = *clip->GetPrivate();
+        const auto& node = *ClipNode->GetPrivate();
 
         Ptr<ASTOpMeshClipMorphPlane> op = new ASTOpMeshClipMorphPlane();
 
@@ -1585,7 +1584,7 @@ namespace mu
         {
             // This argument is required
             m_pErrorLog->GetPrivate()->Add("Mesh clip-morph-plane source node is not set.",
-                ELMT_ERROR, node.m_errorContext);
+                ELMT_ERROR, ClipNode->GetMessageContext());
         }
 
         // Morph to an ellipse
@@ -1640,9 +1639,9 @@ namespace mu
 
     //---------------------------------------------------------------------------------------------
     void CodeGenerator::GenerateMesh_ClipWithMesh(const FMeshGenerationOptions& InOptions, FMeshGenerationResult& OutResult,
-                                                   const NodeMeshClipWithMesh* clip )
+                                                   const NodeMeshClipWithMesh* ClipNode)
     {
-        const auto& node = *clip->GetPrivate();
+        const auto& node = *ClipNode->GetPrivate();
 
         Ptr<ASTOpFixed> op = new ASTOpFixed();
         op->op.type = OP_TYPE::ME_CLIPWITHMESH;
@@ -1657,7 +1656,7 @@ namespace mu
         {
             // This argument is required
             m_pErrorLog->GetPrivate()->Add("Mesh clip-with-mesh source node is not set.",
-                ELMT_ERROR, node.m_errorContext);
+                ELMT_ERROR, ClipNode->GetMessageContext());
         }
 
         // Clipping mesh
@@ -1676,7 +1675,7 @@ namespace mu
         {
             // This argument is required
             m_pErrorLog->GetPrivate()->Add("Mesh clip-with-mesh clipping mesh node is not set.",
-                ELMT_ERROR, node.m_errorContext);
+                ELMT_ERROR, ClipNode->GetMessageContext());
         }
 
         OutResult.meshOp = op;
@@ -1699,8 +1698,7 @@ namespace mu
 		else
 		{
 			// This argument is required
-			m_pErrorLog->GetPrivate()->Add("Mesh Clip Deform base mesh node is not set.",
-				ELMT_ERROR, Node.m_errorContext);
+			m_pErrorLog->GetPrivate()->Add("Mesh Clip Deform base mesh node is not set.", ELMT_ERROR, ClipDeform->GetMessageContext());
 		}
 
 		// Base Shape
@@ -1724,9 +1722,9 @@ namespace mu
 
     //---------------------------------------------------------------------------------------------
     void CodeGenerator::GenerateMesh_ApplyPose(const FMeshGenerationOptions& InOptions, FMeshGenerationResult& OutResult,
-                                                const NodeMeshApplyPose* pose )
+                                                const NodeMeshApplyPose* PoseNode )
     {
-        const auto& node = *pose->GetPrivate();
+        const auto& node = *PoseNode->GetPrivate();
 
         Ptr<ASTOpMeshApplyPose> op = new ASTOpMeshApplyPose();
 
@@ -1740,7 +1738,7 @@ namespace mu
         {
             // This argument is required
             m_pErrorLog->GetPrivate()->Add("Mesh apply-pose base node is not set.",
-                ELMT_ERROR, node.m_errorContext);
+                ELMT_ERROR, PoseNode->GetMessageContext());
         }
 
         // Pose mesh
@@ -1759,7 +1757,7 @@ namespace mu
         {
             // This argument is required
             m_pErrorLog->GetPrivate()->Add("Mesh apply-pose pose node is not set.",
-                ELMT_ERROR, node.m_errorContext);
+                ELMT_ERROR, PoseNode->GetMessageContext());
         }
 
         OutResult.meshOp = op;
@@ -1767,9 +1765,9 @@ namespace mu
 
 
 	//---------------------------------------------------------------------------------------------
-	void CodeGenerator::GenerateMesh_GeometryOperation(const FMeshGenerationOptions& InOptions, FMeshGenerationResult& OutResult, const NodeMeshGeometryOperation* geom)
+	void CodeGenerator::GenerateMesh_GeometryOperation(const FMeshGenerationOptions& InOptions, FMeshGenerationResult& OutResult, const NodeMeshGeometryOperation* GeomNode)
 	{
-		const auto& node = *geom->GetPrivate();
+		const auto& node = *GeomNode->GetPrivate();
 
 		Ptr<ASTOpMeshGeometryOperation> op = new ASTOpMeshGeometryOperation();
 
@@ -1783,7 +1781,7 @@ namespace mu
 		{
 			// This argument is required
 			m_pErrorLog->GetPrivate()->Add("Mesh geometric op mesh-a node is not set.",
-				ELMT_ERROR, node.m_errorContext);
+				ELMT_ERROR, GeomNode->GetMessageContext());
 		}
 
 		// Mesh B
@@ -1799,8 +1797,8 @@ namespace mu
 			op->meshB = bResult.meshOp;
 		}
 
-		op->scalarA = Generate(node.m_pScalarA, InOptions);
-		op->scalarB = Generate(node.m_pScalarB, InOptions);
+		op->scalarA = Generate_Generic(node.m_pScalarA, InOptions);
+		op->scalarB = Generate_Generic(node.m_pScalarB, InOptions);
 
 		OutResult.meshOp = op;
 	}
@@ -1843,7 +1841,7 @@ namespace mu
 		else
 		{
 			// This argument is required
-			m_pErrorLog->GetPrivate()->Add("Mesh reshape base node is not set.", ELMT_ERROR, Node.m_errorContext);
+			m_pErrorLog->GetPrivate()->Add("Mesh reshape base node is not set.", ELMT_ERROR, Reshape->GetMessageContext());
 		}
 
 		// Base and target shapes shouldn't have layouts or modifiers.
diff --git a/Engine/Plugins/Experimental/Mutable/Source/MutableTools/Private/MuT/CodeGenerator_Scalar.cpp b/Engine/Plugins/Experimental/Mutable/Source/MutableTools/Private/MuT/CodeGenerator_Scalar.cpp
--- a/Engine/Plugins/Experimental/Mutable/Source/MutableTools/Private/MuT/CodeGenerator_Scalar.cpp
+++ b/Engine/Plugins/Experimental/Mutable/Source/MutableTools/Private/MuT/CodeGenerator_Scalar.cpp
@@ -229,7 +229,7 @@ namespace mu
 			// No options in the switch!
 			Ptr<ASTOp> missingOp = GenerateMissingScalarCode(TEXT("Switch option"),
 				1.0f,
-				node.m_errorContext);
+				Typed->GetMessageContext());
 			result.op = missingOp;
 			return;
 		}
@@ -247,7 +247,7 @@ namespace mu
 		else
 		{
 			// This argument is required
-			op->variable = GenerateMissingScalarCode(TEXT("Switch variable"), 0.0f, node.m_errorContext);
+			op->variable = GenerateMissingScalarCode(TEXT("Switch variable"), 0.0f, Typed->GetMessageContext());
 		}
 
 		// Options
@@ -263,7 +263,7 @@ namespace mu
 			else
 			{
 				// This argument is required
-				branch = GenerateMissingScalarCode(TEXT("Switch option"), 1.0f, node.m_errorContext);
+				branch = GenerateMissingScalarCode(TEXT("Switch option"), 1.0f, Typed->GetMessageContext());
 			}
 			op->cases.Emplace((int16_t)t, op, branch);
 		}
@@ -306,7 +306,7 @@ namespace mu
 			{
 				FString Msg = FString::Printf(TEXT("Unknown tag found in image variation [%s]."), *tag);
 
-				m_pErrorLog->GetPrivate()->Add(Msg, ELMT_WARNING, node.m_errorContext);
+				m_pErrorLog->GetPrivate()->Add(Msg, ELMT_WARNING, Typed->GetMessageContext());
 				continue;
 			}
 
@@ -320,7 +320,7 @@ namespace mu
 			else
 			{
 				// This argument is required
-				variationOp = GenerateMissingScalarCode(TEXT("Variation option"), 0.0f, node.m_errorContext);
+				variationOp = GenerateMissingScalarCode(TEXT("Variation option"), 0.0f, Typed->GetMessageContext());
 			}
 
 
@@ -345,11 +345,11 @@ namespace mu
 		// T
 		if (Node* pA = Typed->CurveSampleValue.get())
 		{
-			op->time = Generate(pA, Options);
+			op->time = Generate_Generic(pA, Options);
 		}
 		else
 		{
-			op->time = CodeGenerator::GenerateMissingScalarCode(TEXT("Curve T"), 0.5f, Typed->GetPrivate()->m_errorContext);
+			op->time = CodeGenerator::GenerateMissingScalarCode(TEXT("Curve T"), 0.5f, Typed->GetMessageContext());
 		}
 
 		op->Curve = Typed->Curve;
@@ -381,24 +381,24 @@ namespace mu
 		// A
 		if (Node* pA = node.m_pA.get())
 		{
-			op->SetChild(op->op.args.ScalarArithmetic.a, Generate(pA, Options));
+			op->SetChild(op->op.args.ScalarArithmetic.a, Generate_Generic(pA, Options));
 		}
 		else
 		{
 			op->SetChild(op->op.args.ScalarArithmetic.a,
-				CodeGenerator::GenerateMissingScalarCode( TEXT("ScalarArithmetic A"), 1.0f, node.m_errorContext )
+				CodeGenerator::GenerateMissingScalarCode( TEXT("ScalarArithmetic A"), 1.0f, Typed->GetMessageContext() )
 			);
 		}
 
 		// B
 		if (Node* pB = node.m_pB.get())
 		{
-			op->SetChild(op->op.args.ScalarArithmetic.b, Generate(pB, Options));
+			op->SetChild(op->op.args.ScalarArithmetic.b, Generate_Generic(pB, Options));
 		}
 		else
 		{
 			op->SetChild(op->op.args.ScalarArithmetic.b,
-				CodeGenerator::GenerateMissingScalarCode( TEXT("ScalarArithmetic B"), 1.0f, node.m_errorContext )
+				CodeGenerator::GenerateMissingScalarCode( TEXT("ScalarArithmetic B"), 1.0f, Typed->GetMessageContext() )
 			);
 		}
 
@@ -417,7 +417,7 @@ namespace mu
 				NodeScalarConstantPtr pCell = new NodeScalarConstant();
 				float scalar = node.Table->GetPrivate()->Rows[row].Values[colIndex].Scalar;
 				pCell->SetValue(scalar);
-				return Generate(pCell, Options);
+				return Generate_Generic(pCell, Options);
 			});
 
 		result.op = Op;
@@ -436,7 +436,7 @@ namespace mu
 		pNode->SetValue(value);
 
 		FGenericGenerationOptions Options;
-		Ptr<ASTOp> result = Generate(pNode, Options);
+		Ptr<ASTOp> result = Generate_Generic(pNode, Options);
 
 		return result;
 	}
diff --git a/Engine/Plugins/Experimental/Mutable/Source/MutableTools/Private/MuT/CodeGenerator_SecondPass.cpp b/Engine/Plugins/Experimental/Mutable/Source/MutableTools/Private/MuT/CodeGenerator_SecondPass.cpp
--- a/Engine/Plugins/Experimental/Mutable/Source/MutableTools/Private/MuT/CodeGenerator_SecondPass.cpp
+++ b/Engine/Plugins/Experimental/Mutable/Source/MutableTools/Private/MuT/CodeGenerator_SecondPass.cpp
@@ -406,9 +406,7 @@ namespace mu
 
 
 	//---------------------------------------------------------------------------------------------
-	bool SecondPassGenerator::Generate(
-		ErrorLogPtr pErrorLog,
-		const Node::Private* root)
+	bool SecondPassGenerator::Generate( ErrorLogPtr pErrorLog, const Node* root)
 	{
 		MUTABLE_CPUPROFILER_SCOPE(SecondPassGenerate);
 
diff --git a/Engine/Plugins/Experimental/Mutable/Source/MutableTools/Private/MuT/CodeGenerator_SecondPass.h b/Engine/Plugins/Experimental/Mutable/Source/MutableTools/Private/MuT/CodeGenerator_SecondPass.h
--- a/Engine/Plugins/Experimental/Mutable/Source/MutableTools/Private/MuT/CodeGenerator_SecondPass.h
+++ b/Engine/Plugins/Experimental/Mutable/Source/MutableTools/Private/MuT/CodeGenerator_SecondPass.h
@@ -27,7 +27,7 @@ namespace mu
 		SecondPassGenerator( FirstPassGenerator* firstPass, const CompilerOptions::Private* options  );
 
 		// Return true on success.
-        bool Generate( ErrorLogPtr pErrorLog, const Node::Private* root );
+        bool Generate( ErrorLogPtr pErrorLog, const Node* root );
 
 	private:
 
@@ -64,6 +64,7 @@ namespace mu
         };
 
         // List of surfaces that activate or deactivate every tag, or another surface that activates a tag in this set.
+		// \TODO: Change to UE containers
 		vector< set<size_t> > m_surfacesPerTag;
 		vector< set<size_t> > m_tagsPerTag;
 
diff --git a/Engine/Plugins/Experimental/Mutable/Source/MutableTools/Private/MuT/CodeOptimiser.cpp b/Engine/Plugins/Experimental/Mutable/Source/MutableTools/Private/MuT/CodeOptimiser.cpp
--- a/Engine/Plugins/Experimental/Mutable/Source/MutableTools/Private/MuT/CodeOptimiser.cpp
+++ b/Engine/Plugins/Experimental/Mutable/Source/MutableTools/Private/MuT/CodeOptimiser.cpp
@@ -681,6 +681,7 @@ namespace mu
 						break;
 					}
 
+					case OP_TYPE::ME_REFERENCE:
 					case OP_TYPE::IM_REFERENCE:
 						// If we are in a reference-resolution optimization phase, then the ops are not special.
 						if (Pass < 2)
diff --git a/Engine/Plugins/Experimental/Mutable/Source/MutableTools/Private/MuT/CodeOptimiser_Parameters.cpp b/Engine/Plugins/Experimental/Mutable/Source/MutableTools/Private/MuT/CodeOptimiser_Parameters.cpp
--- a/Engine/Plugins/Experimental/Mutable/Source/MutableTools/Private/MuT/CodeOptimiser_Parameters.cpp
+++ b/Engine/Plugins/Experimental/Mutable/Source/MutableTools/Private/MuT/CodeOptimiser_Parameters.cpp
@@ -51,7 +51,7 @@ namespace mu
 
     bool RuntimeParameterVisitorAST::HasAny( const Ptr<ASTOp>& root )
     {
-        if (!m_pState->nodeState.m_runtimeParams.Num())
+        if (!m_pState->nodeState.RuntimeParams.Num())
         {
             return false;
         }
@@ -142,7 +142,7 @@ namespace mu
                     case OP_TYPE::IM_PARAMETER:
                     {
 						const ASTOpParameter* typed = static_cast<const ASTOpParameter*>(at.get());
-                        const TArray<FString>& params = m_pState->nodeState.m_runtimeParams;
+                        const TArray<FString>& params = m_pState->nodeState.RuntimeParams;
                         if ( params.Find( typed->parameter.m_name)
                              !=
                              INDEX_NONE )
@@ -2000,19 +2000,20 @@ namespace mu
     //---------------------------------------------------------------------------------------------
     void CodeOptimiser::OptimiseStatesAST()
     {
-        MUTABLE_CPUPROFILER_SCOPE(OptimiseStatesAST);
+		MUTABLE_CPUPROFILER_SCOPE(OptimiseStatesAST);
 
-         for ( int32 s=0; s<m_states.Num(); ++s )
-         {
-            // Remove the unnecessary lods
-            if (m_states[s].nodeState.m_optimisation.bOnlyFirstLOD)
-            {
-				LODCountReducerAST(m_states[s].root, m_states[s].nodeState.m_optimisation.FirstLOD + 1 + m_states[s].nodeState.m_optimisation.NumExtraLODsToBuildAfterFirstLOD);
-            }
+		for ( int32 s=0; s<m_states.Num(); ++s )
+		{
+			// Remove the unnecessary lods
+			FStateOptimizationOptions StateOptimization = m_states[s].nodeState.Optimisation;
+			if (StateOptimization.bOnlyFirstLOD)
+			{
+				LODCountReducerAST(m_states[s].root, StateOptimization.FirstLOD + 1 + StateOptimization.NumExtraLODsToBuildAfterFirstLOD);
+			}
 
 			// Apply texture compression strategy
 			bool bModified = false;
-			switch (m_states[s].nodeState.m_optimisation.TextureCompressionStrategy)
+			switch (StateOptimization.TextureCompressionStrategy)
 			{
 			case ETextureCompressionStrategy::DontCompressRuntime:
 			{
@@ -2035,7 +2036,7 @@ namespace mu
 			}
 
             // If a state has no runtime parameters, skip its optimisation alltogether
-            if (bModified || m_states[s].nodeState.m_runtimeParams.Num())
+            if (bModified || m_states[s].nodeState.RuntimeParams.Num())
             {
                 // Promote the intructions that depend on runtime parameters, and sink new
                 // format instructions.
diff --git a/Engine/Plugins/Experimental/Mutable/Source/MutableTools/Private/MuT/Compiler.cpp b/Engine/Plugins/Experimental/Mutable/Source/MutableTools/Private/MuT/Compiler.cpp
--- a/Engine/Plugins/Experimental/Mutable/Source/MutableTools/Private/MuT/Compiler.cpp
+++ b/Engine/Plugins/Experimental/Mutable/Source/MutableTools/Private/MuT/Compiler.cpp
@@ -192,52 +192,6 @@ namespace mu
 		UE_LOG(LogMutableCore, Log, TEXT("   Cache MB Written    : %" PRIu64), m_pD->DiskCacheContext.BytesWritten.load() >> 20);
 		UE_LOG(LogMutableCore, Log, TEXT("   Cache MB Read       : %" PRIu64), m_pD->DiskCacheContext.BytesRead.load()>>20);
 	}
-
-
-	void FObjectState::Serialise(OutputArchive& arch) const
-    {
-    	const int32 ver = 6;
-    	arch << ver;
-
-    	arch << m_name;
-    	arch << m_optimisation;
-    	arch << m_runtimeParams;
-    }
-
-
-	void FObjectState::Unserialise(InputArchive& arch)
-    {
-    	int32 ver = 0;
-    	arch >> ver;
-    	check( ver>=5 && ver<=6 );
-
-    	if (ver <= 5)
-    	{
-    		std::string Temp;
-    		arch >> Temp;
-    		m_name = Temp.c_str();
-    	}
-    	else
-    	{
-    		arch >> m_name;
-    	}
-    	arch >> m_optimisation;
-
-    	if (ver <= 5)
-    	{
-    		TArray<std::string> Temp;
-    		arch >> Temp;
-    		m_runtimeParams.SetNum(Temp.Num());
-    		for ( int32 i=0; i<Temp.Num(); ++i)
-    		{
-    			m_runtimeParams[i] = Temp[i].c_str();
-    		}
-    	}
-    	else
-    	{
-    		arch >> m_runtimeParams;
-    	}
-    }
 	
 	
 	//---------------------------------------------------------------------------------------------
@@ -283,7 +237,7 @@ namespace mu
                 FStateCompilationData data;
                 data.nodeState = s.Key;
                 data.root = s.Value;
-                data.state.Name = s.Key.m_name;
+                data.state.Name = s.Key.Name;
                 states.Add( data );
             }
 
@@ -365,14 +319,14 @@ namespace mu
         // Set the runtime parameter indices.
         for(FStateCompilationData& s: states )
         {
-            for ( int32 p=0; p<s.nodeState.m_runtimeParams.Num(); ++p )
+            for ( int32 p=0; p<s.nodeState.RuntimeParams.Num(); ++p )
             {
                 int32 paramIndex = -1;
                 for ( int32 i=0; paramIndex<0 && i<program.m_parameters.Num(); ++i )
                 {
                     if ( program.m_parameters[i].m_name
                          ==
-                         s.nodeState.m_runtimeParams[p] )
+                         s.nodeState.RuntimeParams[p] )
                     {
                         paramIndex = (int)i;
                     }
@@ -387,9 +341,9 @@ namespace mu
 					FString Temp = FString::Printf(TEXT(
 						"The state [%s] refers to a parameter [%s]  that has not been found in the model. This error can be "
 						"safely dismissed in case of partial compilation."), 
-						*s.nodeState.m_name,
-						*s.nodeState.m_runtimeParams[p]);
-                    m_pD->m_pErrorLog->GetPrivate()->Add(Temp, ELMT_WARNING, pNode->GetBasePrivate()->m_errorContext );
+						*s.nodeState.Name,
+						*s.nodeState.RuntimeParams[p]);
+                    m_pD->m_pErrorLog->GetPrivate()->Add(Temp, ELMT_WARNING, pNode->GetMessageContext() );
                 }
             }
 
diff --git a/Engine/Plugins/Experimental/Mutable/Source/MutableTools/Private/MuT/Node.cpp b/Engine/Plugins/Experimental/Mutable/Source/MutableTools/Private/MuT/Node.cpp
--- a/Engine/Plugins/Experimental/Mutable/Source/MutableTools/Private/MuT/Node.cpp
+++ b/Engine/Plugins/Experimental/Mutable/Source/MutableTools/Private/MuT/Node.cpp
@@ -5,26 +5,39 @@
 #include "MuT/NodePrivate.h"
 #include "Misc/AssertionMacros.h"
 
+#include "MuT/NodeComponent.h"
+#include "MuT/NodeComponentNew.h"
+#include "MuT/NodeLOD.h"
+#include "MuT/NodeExtensionData.h"
+#include "MuT/NodeExtensionDataConstant.h"
 
 namespace mu
 {
 
+	FNodeType NodeComponent::StaticType = FNodeType(Node::EType::Component, Node::GetStaticType());
+	FNodeType NodeComponentNew::StaticType = FNodeType(Node::EType::ComponentNew, NodeComponent::GetStaticType());
+	FNodeType NodeLOD::StaticType = FNodeType(Node::EType::LOD, Node::GetStaticType());
+
+	FNodeType NodeExtensionData::StaticType = FNodeType(Node::EType::ExtensionData, Node::GetStaticType());
+	FNodeType NodeExtensionDataConstant::StaticType = FNodeType(Node::EType::ExtensionDataConstant, NodeExtensionData::GetStaticType());
+
 	// Static initialisation
-	static FNodeType s_nodeType = FNodeType( "Node", 0 );
+	static FNodeType s_nodeType = FNodeType(Node::EType::Node, nullptr );
 
 	FNodeType::FNodeType()
 	{
-		m_strName = "";
+		Type = Node::EType::None;
 		m_pParent = nullptr;
 	}
 
 
-	FNodeType::FNodeType( const char* strName, const FNodeType* pParent )
+	FNodeType::FNodeType(Node::EType InType, const FNodeType* pParent )
 	{
-        m_strName = strName;
+		Type = InType;
 		m_pParent = pParent;
 	}
 
+
 	const FNodeType* Node::GetType() const
 	{
 		return GetStaticType();
@@ -39,7 +52,12 @@ namespace mu
 
 	void Node::SetMessageContext( const void* context )
 	{
-		GetBasePrivate()->m_errorContext = context;
+		MessageContext = context;
+	}
+
+	const void* Node::GetMessageContext() const 
+	{ 
+		return MessageContext; 
 	}
 
 }
diff --git a/Engine/Plugins/Experimental/Mutable/Source/MutableTools/Private/MuT/NodeBool.cpp b/Engine/Plugins/Experimental/Mutable/Source/MutableTools/Private/MuT/NodeBool.cpp
--- a/Engine/Plugins/Experimental/Mutable/Source/MutableTools/Private/MuT/NodeBool.cpp
+++ b/Engine/Plugins/Experimental/Mutable/Source/MutableTools/Private/MuT/NodeBool.cpp
@@ -9,7 +9,7 @@
 namespace mu
 {
 
-	static FNodeType s_nodeBoolType = FNodeType( "NodeBool", Node::GetStaticType() );
+	static FNodeType s_nodeBoolType = FNodeType(Node::EType::Bool, Node::GetStaticType() );
 
 
 	//---------------------------------------------------------------------------------------------
diff --git a/Engine/Plugins/Experimental/Mutable/Source/MutableTools/Private/MuT/NodeBoolAnd.cpp b/Engine/Plugins/Experimental/Mutable/Source/MutableTools/Private/MuT/NodeBoolAnd.cpp
--- a/Engine/Plugins/Experimental/Mutable/Source/MutableTools/Private/MuT/NodeBoolAnd.cpp
+++ b/Engine/Plugins/Experimental/Mutable/Source/MutableTools/Private/MuT/NodeBoolAnd.cpp
@@ -14,15 +14,14 @@ namespace mu
 	//---------------------------------------------------------------------------------------------
 	// Static initialisation
 	//---------------------------------------------------------------------------------------------
-	FNodeType NodeBoolAnd::Private::s_type =
-			FNodeType( "BoolParameter", NodeBool::GetStaticType() );
+	FNodeType NodeBoolAnd::Private::s_type = FNodeType(Node::EType::BoolParameter, NodeBool::GetStaticType() );
 
 
 	//---------------------------------------------------------------------------------------------
 	//!
 	//---------------------------------------------------------------------------------------------
 
-	MUTABLE_IMPLEMENT_NODE( NodeBoolAnd, EType::And, Node, Node::EType::Bool);
+	MUTABLE_IMPLEMENT_NODE( NodeBoolAnd );
 
 
 	//---------------------------------------------------------------------------------------------
diff --git a/Engine/Plugins/Experimental/Mutable/Source/MutableTools/Private/MuT/NodeBoolConstant.cpp b/Engine/Plugins/Experimental/Mutable/Source/MutableTools/Private/MuT/NodeBoolConstant.cpp
--- a/Engine/Plugins/Experimental/Mutable/Source/MutableTools/Private/MuT/NodeBoolConstant.cpp
+++ b/Engine/Plugins/Experimental/Mutable/Source/MutableTools/Private/MuT/NodeBoolConstant.cpp
@@ -15,15 +15,14 @@ namespace mu
 	//---------------------------------------------------------------------------------------------
 	// Static initialisation
 	//---------------------------------------------------------------------------------------------
-	FNodeType NodeBoolConstant::Private::s_type =
-			FNodeType( "BoolConstant", NodeBool::GetStaticType() );
+	FNodeType NodeBoolConstant::Private::s_type = FNodeType(Node::EType::BoolConstant, NodeBool::GetStaticType() );
 
 
 	//---------------------------------------------------------------------------------------------
 	//!
 	//---------------------------------------------------------------------------------------------
 
-	MUTABLE_IMPLEMENT_NODE( NodeBoolConstant, EType::Constant, Node, Node::EType::Bool);
+	MUTABLE_IMPLEMENT_NODE( NodeBoolConstant );
 
 
 	//---------------------------------------------------------------------------------------------
diff --git a/Engine/Plugins/Experimental/Mutable/Source/MutableTools/Private/MuT/NodeBoolNot.cpp b/Engine/Plugins/Experimental/Mutable/Source/MutableTools/Private/MuT/NodeBoolNot.cpp
--- a/Engine/Plugins/Experimental/Mutable/Source/MutableTools/Private/MuT/NodeBoolNot.cpp
+++ b/Engine/Plugins/Experimental/Mutable/Source/MutableTools/Private/MuT/NodeBoolNot.cpp
@@ -14,15 +14,14 @@ namespace mu
 	//---------------------------------------------------------------------------------------------
 	// Static initialisation
 	//---------------------------------------------------------------------------------------------
-	FNodeType NodeBoolNot::Private::s_type =
-			FNodeType( "BoolParameter", NodeBool::GetStaticType() );
+	FNodeType NodeBoolNot::Private::s_type = FNodeType(Node::EType::BoolParameter, NodeBool::GetStaticType() );
 
 
 	//---------------------------------------------------------------------------------------------
 	//!
 	//---------------------------------------------------------------------------------------------
 
-	MUTABLE_IMPLEMENT_NODE( NodeBoolNot, EType::Not, Node, Node::EType::Bool);
+	MUTABLE_IMPLEMENT_NODE( NodeBoolNot );
 
 
 	//---------------------------------------------------------------------------------------------
diff --git a/Engine/Plugins/Experimental/Mutable/Source/MutableTools/Private/MuT/NodeBoolParameter.cpp b/Engine/Plugins/Experimental/Mutable/Source/MutableTools/Private/MuT/NodeBoolParameter.cpp
--- a/Engine/Plugins/Experimental/Mutable/Source/MutableTools/Private/MuT/NodeBoolParameter.cpp
+++ b/Engine/Plugins/Experimental/Mutable/Source/MutableTools/Private/MuT/NodeBoolParameter.cpp
@@ -17,15 +17,14 @@ namespace mu
 	//---------------------------------------------------------------------------------------------
 	// Static initialisation
 	//---------------------------------------------------------------------------------------------
-	FNodeType NodeBoolParameter::Private::s_type =
-			FNodeType( "BoolParameter", NodeBool::GetStaticType() );
+	FNodeType NodeBoolParameter::Private::s_type = FNodeType(Node::EType::BoolParameter, NodeBool::GetStaticType() );
 
 
 	//---------------------------------------------------------------------------------------------
 	//!
 	//---------------------------------------------------------------------------------------------
 
-	MUTABLE_IMPLEMENT_NODE( NodeBoolParameter, EType::Parameter, Node, Node::EType::Bool);
+	MUTABLE_IMPLEMENT_NODE( NodeBoolParameter );
 
 
 	//---------------------------------------------------------------------------------------------
diff --git a/Engine/Plugins/Experimental/Mutable/Source/MutableTools/Private/MuT/NodeColour.cpp b/Engine/Plugins/Experimental/Mutable/Source/MutableTools/Private/MuT/NodeColour.cpp
--- a/Engine/Plugins/Experimental/Mutable/Source/MutableTools/Private/MuT/NodeColour.cpp
+++ b/Engine/Plugins/Experimental/Mutable/Source/MutableTools/Private/MuT/NodeColour.cpp
@@ -9,7 +9,7 @@
 namespace mu
 {
 
-	static FNodeType s_nodeColourType = FNodeType( "NodeColour", Node::GetStaticType() );
+	static FNodeType s_nodeColourType = FNodeType( Node::EType::Color, Node::GetStaticType() );
 
 
 	//---------------------------------------------------------------------------------------------
diff --git a/Engine/Plugins/Experimental/Mutable/Source/MutableTools/Private/MuT/NodeColourArithmeticOperation.cpp b/Engine/Plugins/Experimental/Mutable/Source/MutableTools/Private/MuT/NodeColourArithmeticOperation.cpp
--- a/Engine/Plugins/Experimental/Mutable/Source/MutableTools/Private/MuT/NodeColourArithmeticOperation.cpp
+++ b/Engine/Plugins/Experimental/Mutable/Source/MutableTools/Private/MuT/NodeColourArithmeticOperation.cpp
@@ -15,14 +15,13 @@ namespace mu
 	//---------------------------------------------------------------------------------------------
 	// Static initialisation
 	//---------------------------------------------------------------------------------------------
-	FNodeType NodeColourArithmeticOperation::Private::s_type =
-			FNodeType( "ColourArithmenticOperation", NodeColour::GetStaticType() );
+	FNodeType NodeColourArithmeticOperation::Private::s_type = FNodeType(Node::EType::ColorArithmeticOperation, NodeColour::GetStaticType() );
 
 	//---------------------------------------------------------------------------------------------
 	//!
 	//---------------------------------------------------------------------------------------------
 
-	MUTABLE_IMPLEMENT_NODE( NodeColourArithmeticOperation, EType::ArithmeticOperation, Node, Node::EType::Colour)
+	MUTABLE_IMPLEMENT_NODE( NodeColourArithmeticOperation )
 
 
     //---------------------------------------------------------------------------------------------
diff --git a/Engine/Plugins/Experimental/Mutable/Source/MutableTools/Private/MuT/NodeColourConstant.cpp b/Engine/Plugins/Experimental/Mutable/Source/MutableTools/Private/MuT/NodeColourConstant.cpp
--- a/Engine/Plugins/Experimental/Mutable/Source/MutableTools/Private/MuT/NodeColourConstant.cpp
+++ b/Engine/Plugins/Experimental/Mutable/Source/MutableTools/Private/MuT/NodeColourConstant.cpp
@@ -12,14 +12,14 @@
 namespace mu
 {
 
-	FNodeType NodeColourConstant::Private::s_type = FNodeType( "ColourConstant", NodeColour::GetStaticType() );
+	FNodeType NodeColourConstant::Private::s_type = FNodeType(Node::EType::ColorConstant, NodeColour::GetStaticType() );
 
 
 	//---------------------------------------------------------------------------------------------
 	//!
 	//---------------------------------------------------------------------------------------------
 
-	MUTABLE_IMPLEMENT_NODE( NodeColourConstant, EType::Constant, Node, Node::EType::Colour)
+	MUTABLE_IMPLEMENT_NODE( NodeColourConstant)
 
 
 
diff --git a/Engine/Plugins/Experimental/Mutable/Source/MutableTools/Private/MuT/NodeColourFromScalars.cpp b/Engine/Plugins/Experimental/Mutable/Source/MutableTools/Private/MuT/NodeColourFromScalars.cpp
--- a/Engine/Plugins/Experimental/Mutable/Source/MutableTools/Private/MuT/NodeColourFromScalars.cpp
+++ b/Engine/Plugins/Experimental/Mutable/Source/MutableTools/Private/MuT/NodeColourFromScalars.cpp
@@ -15,15 +15,14 @@ namespace mu
 	//---------------------------------------------------------------------------------------------
 	// Static initialisation
 	//---------------------------------------------------------------------------------------------
-	FNodeType NodeColourFromScalars::Private::s_type =
-			FNodeType( "ColourFromScalars", NodeColour::GetStaticType() );
+	FNodeType NodeColourFromScalars::Private::s_type = FNodeType(Node::EType::ColorFromScalars, NodeColour::GetStaticType() );
 
 
 	//---------------------------------------------------------------------------------------------
 	//!
 	//---------------------------------------------------------------------------------------------
 
-	MUTABLE_IMPLEMENT_NODE( NodeColourFromScalars, EType::FromScalars, Node, Node::EType::Colour)
+	MUTABLE_IMPLEMENT_NODE( NodeColourFromScalars )
 
 
 	//---------------------------------------------------------------------------------------------
diff --git a/Engine/Plugins/Experimental/Mutable/Source/MutableTools/Private/MuT/NodeColourParameter.cpp b/Engine/Plugins/Experimental/Mutable/Source/MutableTools/Private/MuT/NodeColourParameter.cpp
--- a/Engine/Plugins/Experimental/Mutable/Source/MutableTools/Private/MuT/NodeColourParameter.cpp
+++ b/Engine/Plugins/Experimental/Mutable/Source/MutableTools/Private/MuT/NodeColourParameter.cpp
@@ -16,15 +16,14 @@ namespace mu
 	//---------------------------------------------------------------------------------------------
 	// Static initialisation
 	//---------------------------------------------------------------------------------------------
-	FNodeType NodeColourParameter::Private::s_type =
-			FNodeType( "ColourParameter", NodeColour::GetStaticType() );
+	FNodeType NodeColourParameter::Private::s_type = FNodeType(Node::EType::ColorParameter, NodeColour::GetStaticType() );
 
 
 	//---------------------------------------------------------------------------------------------
 	//!
 	//---------------------------------------------------------------------------------------------
 
-	MUTABLE_IMPLEMENT_NODE( NodeColourParameter, EType::Parameter, Node, Node::EType::Colour)
+	MUTABLE_IMPLEMENT_NODE( NodeColourParameter )
 
 
 	//---------------------------------------------------------------------------------------------
diff --git a/Engine/Plugins/Experimental/Mutable/Source/MutableTools/Private/MuT/NodeColourSampleImage.cpp b/Engine/Plugins/Experimental/Mutable/Source/MutableTools/Private/MuT/NodeColourSampleImage.cpp
--- a/Engine/Plugins/Experimental/Mutable/Source/MutableTools/Private/MuT/NodeColourSampleImage.cpp
+++ b/Engine/Plugins/Experimental/Mutable/Source/MutableTools/Private/MuT/NodeColourSampleImage.cpp
@@ -16,15 +16,14 @@ namespace mu
 	//---------------------------------------------------------------------------------------------
 	// Static initialisation
 	//---------------------------------------------------------------------------------------------
-	FNodeType NodeColourSampleImage::Private::s_type =
-			FNodeType( "ColourSampleImage", NodeColour::GetStaticType() );
+	FNodeType NodeColourSampleImage::Private::s_type = FNodeType (Node::EType::ColorSampleImage, NodeColour::GetStaticType() );
 
 
 	//---------------------------------------------------------------------------------------------
 	//!
 	//---------------------------------------------------------------------------------------------
 
-	MUTABLE_IMPLEMENT_NODE( NodeColourSampleImage, EType::SampleImage, Node, Node::EType::Colour)
+	MUTABLE_IMPLEMENT_NODE( NodeColourSampleImage )
 
 
 	//---------------------------------------------------------------------------------------------
diff --git a/Engine/Plugins/Experimental/Mutable/Source/MutableTools/Private/MuT/NodeColourSwitch.cpp b/Engine/Plugins/Experimental/Mutable/Source/MutableTools/Private/MuT/NodeColourSwitch.cpp
--- a/Engine/Plugins/Experimental/Mutable/Source/MutableTools/Private/MuT/NodeColourSwitch.cpp
+++ b/Engine/Plugins/Experimental/Mutable/Source/MutableTools/Private/MuT/NodeColourSwitch.cpp
@@ -14,15 +14,14 @@ namespace mu
 	//---------------------------------------------------------------------------------------------
 	// Static initialisation
 	//---------------------------------------------------------------------------------------------
-    FNodeType NodeColourSwitch::Private::s_type =
-            FNodeType( "ColourSwitch", NodeColour::GetStaticType() );
+    FNodeType NodeColourSwitch::Private::s_type = FNodeType(Node::EType::ColorSwitch, NodeColour::GetStaticType() );
 
 
 	//---------------------------------------------------------------------------------------------
 	//!
 	//---------------------------------------------------------------------------------------------
 
-    MUTABLE_IMPLEMENT_NODE( NodeColourSwitch, EType::Switch, Node, Node::EType::Colour)
+    MUTABLE_IMPLEMENT_NODE( NodeColourSwitch )
 
 
 	//---------------------------------------------------------------------------------------------
diff --git a/Engine/Plugins/Experimental/Mutable/Source/MutableTools/Private/MuT/NodeColourTable.cpp b/Engine/Plugins/Experimental/Mutable/Source/MutableTools/Private/MuT/NodeColourTable.cpp
--- a/Engine/Plugins/Experimental/Mutable/Source/MutableTools/Private/MuT/NodeColourTable.cpp
+++ b/Engine/Plugins/Experimental/Mutable/Source/MutableTools/Private/MuT/NodeColourTable.cpp
@@ -17,15 +17,14 @@ namespace mu
 	//---------------------------------------------------------------------------------------------
 	// Static initialisation
 	//---------------------------------------------------------------------------------------------
-	FNodeType NodeColourTable::Private::s_type =
-			FNodeType( "TableColour", NodeColour::GetStaticType() );
+	FNodeType NodeColourTable::Private::s_type = FNodeType(Node::EType::ColorTable, NodeColour::GetStaticType() );
 
 
 	//---------------------------------------------------------------------------------------------
 	//!
 	//---------------------------------------------------------------------------------------------
 
-	MUTABLE_IMPLEMENT_NODE( NodeColourTable, EType::Table, Node, Node::EType::Colour)
+	MUTABLE_IMPLEMENT_NODE( NodeColourTable )
 
 
 	//---------------------------------------------------------------------------------------------
diff --git a/Engine/Plugins/Experimental/Mutable/Source/MutableTools/Private/MuT/NodeColourVariation.cpp b/Engine/Plugins/Experimental/Mutable/Source/MutableTools/Private/MuT/NodeColourVariation.cpp
--- a/Engine/Plugins/Experimental/Mutable/Source/MutableTools/Private/MuT/NodeColourVariation.cpp
+++ b/Engine/Plugins/Experimental/Mutable/Source/MutableTools/Private/MuT/NodeColourVariation.cpp
@@ -14,11 +14,10 @@ namespace mu
     //---------------------------------------------------------------------------------------------
     // Static initialisation
     //---------------------------------------------------------------------------------------------
-    FNodeType NodeColourVariation::Private::s_type =
-        FNodeType( "ColourVariation", NodeColour::GetStaticType() );
+    FNodeType NodeColourVariation::Private::s_type = FNodeType(Node::EType::ColorVariation, NodeColour::GetStaticType() );
 
 
-    MUTABLE_IMPLEMENT_NODE( NodeColourVariation, EType::Variation, Node, Node::EType::Colour)
+    MUTABLE_IMPLEMENT_NODE( NodeColourVariation )
 
 
     //---------------------------------------------------------------------------------------------
diff --git a/Engine/Plugins/Experimental/Mutable/Source/MutableTools/Private/MuT/NodeComponentEdit.cpp b/Engine/Plugins/Experimental/Mutable/Source/MutableTools/Private/MuT/NodeComponentEdit.cpp
--- a/Engine/Plugins/Experimental/Mutable/Source/MutableTools/Private/MuT/NodeComponentEdit.cpp
+++ b/Engine/Plugins/Experimental/Mutable/Source/MutableTools/Private/MuT/NodeComponentEdit.cpp
@@ -4,7 +4,6 @@
 #include "MuT/NodeComponentEdit.h"
 
 #include "Misc/AssertionMacros.h"
-#include "MuT/NodeComponentEditPrivate.h"
 #include "MuT/NodeComponentNew.h"
 #include "MuT/NodePrivate.h"
 #include "MuT/NodeSurface.h"
@@ -13,19 +12,7 @@
 namespace mu
 {
 
-
-	//---------------------------------------------------------------------------------------------
-	// Static initialisation
-	//---------------------------------------------------------------------------------------------
-	FNodeType NodeComponentEdit::Private::s_type =
-			FNodeType( "EditComponent", NodeComponent::GetStaticType() );
-
-
-	//---------------------------------------------------------------------------------------------
-	//---------------------------------------------------------------------------------------------
-	//---------------------------------------------------------------------------------------------
-
-    MUTABLE_IMPLEMENT_NODE( NodeComponentEdit, EType::Edit, Node, Node::EType::Component)
+	FNodeType NodeComponentEdit::StaticType = FNodeType(Node::EType::ComponentEdit, NodeComponent::GetStaticType() );
 
 
 	//---------------------------------------------------------------------------------------------
@@ -33,62 +20,28 @@ namespace mu
 	//---------------------------------------------------------------------------------------------
     void NodeComponentEdit::SetParent( NodeComponent* p )
 	{
-		m_pD->m_pParent = p;
+		m_pParent = p;
 	}
 
 
 	//---------------------------------------------------------------------------------------------
     NodeComponent* NodeComponentEdit::GetParent() const
 	{
-        return m_pD->m_pParent.get();
+        return m_pParent.get();
 	}
 
 
-    //---------------------------------------------------------------------------------------------
-    int NodeComponentEdit::GetSurfaceCount() const
-    {
-        return m_pD->m_surfaces.Num();
-    }
-
-
-    //---------------------------------------------------------------------------------------------
-    void NodeComponentEdit::SetSurfaceCount( int32 num )
-    {
-        check( num >=0 );
-        m_pD->m_surfaces.SetNum( num );
-    }
-
-
-    //---------------------------------------------------------------------------------------------
-    NodeSurface* NodeComponentEdit::GetSurface( int32 index ) const
-    {
-        check( index >=0 && index < m_pD->m_surfaces.Num() );
-
-        return m_pD->m_surfaces[ index ].get();
-    }
-
-
-    //---------------------------------------------------------------------------------------------
-    void NodeComponentEdit::SetSurface( int32 index, NodeSurface* pNode )
-    {
-        check( index >=0 && index < m_pD->m_surfaces.Num() );
-
-        m_pD->m_surfaces[ index ] = pNode;
-    }
-
-
 	//---------------------------------------------------------------------------------------------
-	const NodeComponentNew::Private* NodeComponentEdit::Private::GetParentComponentNew() const
+	const NodeComponentNew* NodeComponentEdit::GetParentComponentNew() const
 	{
-		const NodeComponentNew::Private* parent = nullptr;
+		const NodeComponentNew* Parent = nullptr;
 		if (m_pParent)
 		{
-			NodeComponent::Private* ParentPrivate = static_cast<NodeComponent::Private*>(m_pParent->GetBasePrivate());
-			parent = ParentPrivate->GetParentComponentNew();
+			Parent = m_pParent->GetParentComponentNew();
 		}
 
-		check(parent);
-		return parent;
+		check(Parent);
+		return Parent;
 	}
 
 }
