[mutable] Redesign the object node hierarchy to better match unreal component and LOD concepts.

diff --git a/Engine/Plugins/Experimental/Mutable/Source/MutableTools/Private/MuT/NodeExtensionDataSwitch.cpp b/Engine/Plugins/Experimental/Mutable/Source/MutableTools/Private/MuT/NodeExtensionDataSwitch.cpp
--- a/Engine/Plugins/Experimental/Mutable/Source/MutableTools/Private/MuT/NodeExtensionDataSwitch.cpp
+++ b/Engine/Plugins/Experimental/Mutable/Source/MutableTools/Private/MuT/NodeExtensionDataSwitch.cpp
@@ -3,49 +3,38 @@
 #include "MuT/NodeExtensionDataSwitch.h"
 
 #include "Misc/AssertionMacros.h"
-#include "MuT/NodeExtensionDataSwitchPrivate.h"
-// #include "MuT/NodePrivate.h"
 #include "MuT/NodeScalar.h"
 
 namespace mu
 {
-	//---------------------------------------------------------------------------------------------
-	// Static initialisation
-	//---------------------------------------------------------------------------------------------
-	FNodeType NodeExtensionDataSwitch::Private::s_type =
-		FNodeType("ExtensionDataSwitch", NodeExtensionData::GetStaticType());
+	FNodeType NodeExtensionDataSwitch::StaticType = FNodeType(Node::EType::ExtensionDataSwitch, NodeExtensionData::GetStaticType());
 
-	//---------------------------------------------------------------------------------------------
-	//!
-	//---------------------------------------------------------------------------------------------
 
-	MUTABLE_IMPLEMENT_NODE(NodeExtensionDataSwitch, EType::Switch, Node, Node::EType::ExtensionData);
-
-	NodeScalarPtr NodeExtensionDataSwitch::GetParameter() const
+	Ptr<NodeScalar> NodeExtensionDataSwitch::GetParameter() const
 	{
-		return m_pD->Parameter.get();
+		return Parameter.get();
 	}
 
-	void NodeExtensionDataSwitch::SetParameter(NodeScalarPtr pNode)
+	void NodeExtensionDataSwitch::SetParameter(Ptr<NodeScalar> pNode)
 	{
-		m_pD->Parameter = pNode;
+		Parameter = pNode;
 	}
 
 	void NodeExtensionDataSwitch::SetOptionCount(int InNumOptions)
 	{
 		check(InNumOptions >= 0);
-		m_pD->Options.SetNum(InNumOptions);
+		Options.SetNum(InNumOptions);
 	}
 
-	NodeExtensionDataPtr NodeExtensionDataSwitch::GetOption(int t) const
+	Ptr<NodeExtensionData> NodeExtensionDataSwitch::GetOption(int32 t) const
 	{
-		check(m_pD->Options.IsValidIndex(t));
-		return m_pD->Options[t].get();
+		check(Options.IsValidIndex(t));
+		return Options[t].get();
 	}
 
-	void NodeExtensionDataSwitch::SetOption(int t, NodeExtensionDataPtr pNode)
+	void NodeExtensionDataSwitch::SetOption(int32 t, Ptr<NodeExtensionData> pNode)
 	{
-		check(m_pD->Options.IsValidIndex(t));
-		m_pD->Options[t] = pNode;
+		check(Options.IsValidIndex(t));
+		Options[t] = pNode;
 	}
 }
\ No newline at end of file
diff --git a/Engine/Plugins/Experimental/Mutable/Source/MutableTools/Private/MuT/NodeExtensionDataVariation.cpp b/Engine/Plugins/Experimental/Mutable/Source/MutableTools/Private/MuT/NodeExtensionDataVariation.cpp
--- a/Engine/Plugins/Experimental/Mutable/Source/MutableTools/Private/MuT/NodeExtensionDataVariation.cpp
+++ b/Engine/Plugins/Experimental/Mutable/Source/MutableTools/Private/MuT/NodeExtensionDataVariation.cpp
@@ -3,53 +3,40 @@
 #include "MuT/NodeExtensionDataVariation.h"
 
 #include "Misc/AssertionMacros.h"
-#include "MuT/NodeExtensionDataVariationPrivate.h"
 
 namespace mu
 {
-	//---------------------------------------------------------------------------------------------
-	// Static initialisation
-	//---------------------------------------------------------------------------------------------
-	FNodeType NodeExtensionDataVariation::Private::s_type =
-		FNodeType("ExtensionDataVariation", NodeExtensionData::GetStaticType());
+	FNodeType NodeExtensionDataVariation::StaticType = FNodeType(Node::EType::ExtensionDataVariation, NodeExtensionData::GetStaticType());
 
-	//---------------------------------------------------------------------------------------------
-	//!
-	//---------------------------------------------------------------------------------------------
-	MUTABLE_IMPLEMENT_NODE(NodeExtensionDataVariation, EType::Variation, Node, Node::EType::ExtensionData);
 
-
-	//---------------------------------------------------------------------------------------------
-	// Own Interface
-	//---------------------------------------------------------------------------------------------
-	void NodeExtensionDataVariation::SetDefaultValue(NodeExtensionDataPtr InValue)
+	void NodeExtensionDataVariation::SetDefaultValue(Ptr<NodeExtensionData> InValue)
 	{
-		m_pD->DefaultValue = InValue;
+		DefaultValue = InValue;
 	}
 
-	void NodeExtensionDataVariation::SetVariationCount(int InCount)
+	void NodeExtensionDataVariation::SetVariationCount(int32 InCount)
 	{
 		check(InCount >= 0);
-		m_pD->Variations.SetNum(InCount);
+		Variations.SetNum(InCount);
 	}
 
 	int NodeExtensionDataVariation::GetVariationCount() const
 	{
-		return m_pD->Variations.Num();
+		return Variations.Num();
 	}
 
-	void NodeExtensionDataVariation::SetVariationTag(int InIndex, const FString& Tag)
+	void NodeExtensionDataVariation::SetVariationTag(int32 InIndex, const FString& Tag)
 	{
-		check(m_pD->Variations.IsValidIndex(InIndex));
+		check(Variations.IsValidIndex(InIndex));
 
-		m_pD->Variations[InIndex].Tag = Tag;
+		Variations[InIndex].Tag = Tag;
 	}
 	//---------------------------------------------------------------------------------------------
-	void NodeExtensionDataVariation::SetVariationValue(int InIndex, NodeExtensionDataPtr InValue)
+	void NodeExtensionDataVariation::SetVariationValue(int32 InIndex, Ptr<NodeExtensionData> InValue)
 	{
-		check(m_pD->Variations.IsValidIndex(InIndex));
+		check(Variations.IsValidIndex(InIndex));
 
-		m_pD->Variations[InIndex].Value = InValue;
+		Variations[InIndex].Value = InValue;
 	}
 
 }
diff --git a/Engine/Plugins/Experimental/Mutable/Source/MutableTools/Private/MuT/NodeImage.cpp b/Engine/Plugins/Experimental/Mutable/Source/MutableTools/Private/MuT/NodeImage.cpp
--- a/Engine/Plugins/Experimental/Mutable/Source/MutableTools/Private/MuT/NodeImage.cpp
+++ b/Engine/Plugins/Experimental/Mutable/Source/MutableTools/Private/MuT/NodeImage.cpp
@@ -10,7 +10,7 @@ namespace mu
 	//---------------------------------------------------------------------------------------------
 	// Static initialisation
 	//---------------------------------------------------------------------------------------------
-	static FNodeType s_nodeImageType = 	FNodeType( "NodeImage", Node::GetStaticType() );
+	static FNodeType s_nodeImageType = 	FNodeType( Node::EType::Image, Node::GetStaticType() );
 
 
 	//---------------------------------------------------------------------------------------------
diff --git a/Engine/Plugins/Experimental/Mutable/Source/MutableTools/Private/MuT/NodeImageBinarise.cpp b/Engine/Plugins/Experimental/Mutable/Source/MutableTools/Private/MuT/NodeImageBinarise.cpp
--- a/Engine/Plugins/Experimental/Mutable/Source/MutableTools/Private/MuT/NodeImageBinarise.cpp
+++ b/Engine/Plugins/Experimental/Mutable/Source/MutableTools/Private/MuT/NodeImageBinarise.cpp
@@ -15,12 +15,11 @@ namespace mu
 	//---------------------------------------------------------------------------------------------
 	// Static initialisation
 	//---------------------------------------------------------------------------------------------
-	FNodeType NodeImageBinarise::Private::s_type =
-			FNodeType( "ImageMultiply", NodeImage::GetStaticType() );
+	FNodeType NodeImageBinarise::Private::s_type = FNodeType(Node::EType::ImageBinarise, NodeImage::GetStaticType() );
 
 	//---------------------------------------------------------------------------------------------
 
-	MUTABLE_IMPLEMENT_NODE( NodeImageBinarise, EType::Binarise, Node, Node::EType::Image)
+	MUTABLE_IMPLEMENT_NODE( NodeImageBinarise )
 
 
 	//---------------------------------------------------------------------------------------------
diff --git a/Engine/Plugins/Experimental/Mutable/Source/MutableTools/Private/MuT/NodeImageColourMap.cpp b/Engine/Plugins/Experimental/Mutable/Source/MutableTools/Private/MuT/NodeImageColourMap.cpp
--- a/Engine/Plugins/Experimental/Mutable/Source/MutableTools/Private/MuT/NodeImageColourMap.cpp
+++ b/Engine/Plugins/Experimental/Mutable/Source/MutableTools/Private/MuT/NodeImageColourMap.cpp
@@ -14,15 +14,14 @@ namespace mu
 	//---------------------------------------------------------------------------------------------
 	// Static initialisation
 	//---------------------------------------------------------------------------------------------
-	FNodeType NodeImageColourMap::Private::s_type =
-			FNodeType( "ImageColourMap", NodeImage::GetStaticType() );
+	FNodeType NodeImageColourMap::Private::s_type = FNodeType(Node::EType::ImageColorMap, NodeImage::GetStaticType() );
 
 
 	//---------------------------------------------------------------------------------------------
 	//!
 	//---------------------------------------------------------------------------------------------
 
-	MUTABLE_IMPLEMENT_NODE( NodeImageColourMap, EType::ColourMap, Node, Node::EType::Image)
+	MUTABLE_IMPLEMENT_NODE( NodeImageColourMap )
 
 
 	//---------------------------------------------------------------------------------------------
diff --git a/Engine/Plugins/Experimental/Mutable/Source/MutableTools/Private/MuT/NodeImageConditional.cpp b/Engine/Plugins/Experimental/Mutable/Source/MutableTools/Private/MuT/NodeImageConditional.cpp
--- a/Engine/Plugins/Experimental/Mutable/Source/MutableTools/Private/MuT/NodeImageConditional.cpp
+++ b/Engine/Plugins/Experimental/Mutable/Source/MutableTools/Private/MuT/NodeImageConditional.cpp
@@ -15,15 +15,14 @@ namespace mu
 	//---------------------------------------------------------------------------------------------
 	// Static initialisation
 	//---------------------------------------------------------------------------------------------
-    FNodeType NodeImageConditional::Private::s_type =
-            FNodeType( "ImageConditional", NodeImage::GetStaticType() );
+    FNodeType NodeImageConditional::Private::s_type = FNodeType(Node::EType::ImageConditional, NodeImage::GetStaticType() );
 
 
 	//---------------------------------------------------------------------------------------------
 	//!
 	//---------------------------------------------------------------------------------------------
 
-    MUTABLE_IMPLEMENT_NODE( NodeImageConditional, EType::Conditional, Node, Node::EType::Image)
+    MUTABLE_IMPLEMENT_NODE( NodeImageConditional )
 
 
 	//---------------------------------------------------------------------------------------------
diff --git a/Engine/Plugins/Experimental/Mutable/Source/MutableTools/Private/MuT/NodeImageConstant.cpp b/Engine/Plugins/Experimental/Mutable/Source/MutableTools/Private/MuT/NodeImageConstant.cpp
--- a/Engine/Plugins/Experimental/Mutable/Source/MutableTools/Private/MuT/NodeImageConstant.cpp
+++ b/Engine/Plugins/Experimental/Mutable/Source/MutableTools/Private/MuT/NodeImageConstant.cpp
@@ -16,12 +16,11 @@ namespace mu
 	//---------------------------------------------------------------------------------------------
 	// Static initialisation
 	//---------------------------------------------------------------------------------------------
-	FNodeType NodeImageConstant::Private::s_type =
-		FNodeType("ImageConstant", NodeImage::GetStaticType());
+	FNodeType NodeImageConstant::Private::s_type = FNodeType(Node::EType::ImageConstant, NodeImage::GetStaticType());
 
 
 	//---------------------------------------------------------------------------------------------
-	MUTABLE_IMPLEMENT_NODE(NodeImageConstant, EType::Constant, Node, Node::EType::Image);
+	MUTABLE_IMPLEMENT_NODE(NodeImageConstant);
 
 
 	//---------------------------------------------------------------------------------------------
diff --git a/Engine/Plugins/Experimental/Mutable/Source/MutableTools/Private/MuT/NodeImageFormat.cpp b/Engine/Plugins/Experimental/Mutable/Source/MutableTools/Private/MuT/NodeImageFormat.cpp
--- a/Engine/Plugins/Experimental/Mutable/Source/MutableTools/Private/MuT/NodeImageFormat.cpp
+++ b/Engine/Plugins/Experimental/Mutable/Source/MutableTools/Private/MuT/NodeImageFormat.cpp
@@ -14,15 +14,14 @@ namespace mu
 	//---------------------------------------------------------------------------------------------
 	// Static initialisation
 	//---------------------------------------------------------------------------------------------
-	FNodeType NodeImageFormat::Private::s_type =
-			FNodeType( "ImageFormat", NodeImage::GetStaticType() );
+	FNodeType NodeImageFormat::Private::s_type = FNodeType(Node::EType::ImageFormat, NodeImage::GetStaticType() );
 
 
 	//---------------------------------------------------------------------------------------------
 	//!
 	//---------------------------------------------------------------------------------------------
 
-	MUTABLE_IMPLEMENT_NODE( NodeImageFormat, EType::Format, Node, Node::EType::Image)
+	MUTABLE_IMPLEMENT_NODE( NodeImageFormat )
 
 
 	//---------------------------------------------------------------------------------------------
diff --git a/Engine/Plugins/Experimental/Mutable/Source/MutableTools/Private/MuT/NodeImageGradient.cpp b/Engine/Plugins/Experimental/Mutable/Source/MutableTools/Private/MuT/NodeImageGradient.cpp
--- a/Engine/Plugins/Experimental/Mutable/Source/MutableTools/Private/MuT/NodeImageGradient.cpp
+++ b/Engine/Plugins/Experimental/Mutable/Source/MutableTools/Private/MuT/NodeImageGradient.cpp
@@ -16,15 +16,14 @@ namespace mu
 	//---------------------------------------------------------------------------------------------
 	// Static initialisation
 	//---------------------------------------------------------------------------------------------
-	FNodeType NodeImageGradient::Private::s_type =
-			FNodeType( "ImageGradient", NodeImage::GetStaticType() );
+	FNodeType NodeImageGradient::Private::s_type = FNodeType(Node::EType::ImageGradient, NodeImage::GetStaticType() );
 
 
 	//---------------------------------------------------------------------------------------------
 	//!
 	//---------------------------------------------------------------------------------------------
 
-	MUTABLE_IMPLEMENT_NODE( NodeImageGradient, EType::Gradient, Node, Node::EType::Image)
+	MUTABLE_IMPLEMENT_NODE( NodeImageGradient )
 
 
 	//---------------------------------------------------------------------------------------------
diff --git a/Engine/Plugins/Experimental/Mutable/Source/MutableTools/Private/MuT/NodeImageInterpolate.cpp b/Engine/Plugins/Experimental/Mutable/Source/MutableTools/Private/MuT/NodeImageInterpolate.cpp
--- a/Engine/Plugins/Experimental/Mutable/Source/MutableTools/Private/MuT/NodeImageInterpolate.cpp
+++ b/Engine/Plugins/Experimental/Mutable/Source/MutableTools/Private/MuT/NodeImageInterpolate.cpp
@@ -15,15 +15,14 @@ namespace mu
 	//---------------------------------------------------------------------------------------------
 	// Static initialisation
 	//---------------------------------------------------------------------------------------------
-	FNodeType NodeImageInterpolate::Private::s_type =
-			FNodeType( "ImageInterpolate", NodeImage::GetStaticType() );
+	FNodeType NodeImageInterpolate::Private::s_type = FNodeType(Node::EType::ImageInterpolate, NodeImage::GetStaticType() );
 
 
 	//---------------------------------------------------------------------------------------------
 	//!
 	//---------------------------------------------------------------------------------------------
 
-	MUTABLE_IMPLEMENT_NODE( NodeImageInterpolate, EType::Interpolate, Node, Node::EType::Image)
+	MUTABLE_IMPLEMENT_NODE( NodeImageInterpolate )
 
 
 	//---------------------------------------------------------------------------------------------
diff --git a/Engine/Plugins/Experimental/Mutable/Source/MutableTools/Private/MuT/NodeImageInvert.cpp b/Engine/Plugins/Experimental/Mutable/Source/MutableTools/Private/MuT/NodeImageInvert.cpp
--- a/Engine/Plugins/Experimental/Mutable/Source/MutableTools/Private/MuT/NodeImageInvert.cpp
+++ b/Engine/Plugins/Experimental/Mutable/Source/MutableTools/Private/MuT/NodeImageInvert.cpp
@@ -14,12 +14,11 @@ namespace mu
 	//---------------------------------------------------------------------------------------------
 	// Static initialisation
 	//---------------------------------------------------------------------------------------------
-	FNodeType NodeImageInvert::Private::s_type =
-		FNodeType("ImageInvert", NodeImage::GetStaticType());
+	FNodeType NodeImageInvert::Private::s_type = FNodeType(Node::EType::ImageInvert, NodeImage::GetStaticType());
 
 	//---------------------------------------------------------------------------------------------
 
-	MUTABLE_IMPLEMENT_NODE( NodeImageInvert, EType::Invert, Node, Node::EType::Image)
+	MUTABLE_IMPLEMENT_NODE( NodeImageInvert )
 
 	//---------------------------------------------------------------------------------------------
 	// Own Interface
diff --git a/Engine/Plugins/Experimental/Mutable/Source/MutableTools/Private/MuT/NodeImageLayer.cpp b/Engine/Plugins/Experimental/Mutable/Source/MutableTools/Private/MuT/NodeImageLayer.cpp
--- a/Engine/Plugins/Experimental/Mutable/Source/MutableTools/Private/MuT/NodeImageLayer.cpp
+++ b/Engine/Plugins/Experimental/Mutable/Source/MutableTools/Private/MuT/NodeImageLayer.cpp
@@ -14,15 +14,14 @@ namespace mu
 	//---------------------------------------------------------------------------------------------
 	// Static initialisation
 	//---------------------------------------------------------------------------------------------
-	FNodeType NodeImageLayer::Private::s_type =
-			FNodeType( "ImageLayer", NodeImage::GetStaticType() );
+	FNodeType NodeImageLayer::Private::s_type = FNodeType(Node::EType::ImageLayer, NodeImage::GetStaticType() );
 
 
 	//---------------------------------------------------------------------------------------------
 	//!
 	//---------------------------------------------------------------------------------------------
 
-	MUTABLE_IMPLEMENT_NODE( NodeImageLayer, EType::Layer, Node, Node::EType::Image)
+	MUTABLE_IMPLEMENT_NODE( NodeImageLayer )
 
 
 	//---------------------------------------------------------------------------------------------
diff --git a/Engine/Plugins/Experimental/Mutable/Source/MutableTools/Private/MuT/NodeImageLayerColour.cpp b/Engine/Plugins/Experimental/Mutable/Source/MutableTools/Private/MuT/NodeImageLayerColour.cpp
--- a/Engine/Plugins/Experimental/Mutable/Source/MutableTools/Private/MuT/NodeImageLayerColour.cpp
+++ b/Engine/Plugins/Experimental/Mutable/Source/MutableTools/Private/MuT/NodeImageLayerColour.cpp
@@ -15,15 +15,14 @@ namespace mu
 	//---------------------------------------------------------------------------------------------
 	// Static initialisation
 	//---------------------------------------------------------------------------------------------
-	FNodeType NodeImageLayerColour::Private::s_type =
-			FNodeType( "ImageLayerColour", NodeImage::GetStaticType() );
+	FNodeType NodeImageLayerColour::Private::s_type = FNodeType(Node::EType::ImageLayerColour, NodeImage::GetStaticType() );
 
 
 	//---------------------------------------------------------------------------------------------
 	//!
 	//---------------------------------------------------------------------------------------------
 
-	MUTABLE_IMPLEMENT_NODE( NodeImageLayerColour, EType::LayerColour, Node, Node::EType::Image);
+	MUTABLE_IMPLEMENT_NODE( NodeImageLayerColour );
 
 
 	//---------------------------------------------------------------------------------------------
diff --git a/Engine/Plugins/Experimental/Mutable/Source/MutableTools/Private/MuT/NodeImageLuminance.cpp b/Engine/Plugins/Experimental/Mutable/Source/MutableTools/Private/MuT/NodeImageLuminance.cpp
--- a/Engine/Plugins/Experimental/Mutable/Source/MutableTools/Private/MuT/NodeImageLuminance.cpp
+++ b/Engine/Plugins/Experimental/Mutable/Source/MutableTools/Private/MuT/NodeImageLuminance.cpp
@@ -14,15 +14,14 @@ namespace mu
 	//---------------------------------------------------------------------------------------------
 	// Static initialisation
 	//---------------------------------------------------------------------------------------------
-	FNodeType NodeImageLuminance::Private::s_type =
-			FNodeType( "ImageLuminance", NodeImage::GetStaticType() );
+	FNodeType NodeImageLuminance::Private::s_type = FNodeType(Node::EType::ImageLuminance, NodeImage::GetStaticType() );
 
 
 	//---------------------------------------------------------------------------------------------
 	//!
 	//---------------------------------------------------------------------------------------------
 
-	MUTABLE_IMPLEMENT_NODE( NodeImageLuminance, EType::Luminance, Node, Node::EType::Image)
+	MUTABLE_IMPLEMENT_NODE( NodeImageLuminance )
 
 
 	//---------------------------------------------------------------------------------------------
diff --git a/Engine/Plugins/Experimental/Mutable/Source/MutableTools/Private/MuT/NodeImageMipmap.cpp b/Engine/Plugins/Experimental/Mutable/Source/MutableTools/Private/MuT/NodeImageMipmap.cpp
--- a/Engine/Plugins/Experimental/Mutable/Source/MutableTools/Private/MuT/NodeImageMipmap.cpp
+++ b/Engine/Plugins/Experimental/Mutable/Source/MutableTools/Private/MuT/NodeImageMipmap.cpp
@@ -15,15 +15,14 @@ namespace mu
 	//---------------------------------------------------------------------------------------------
 	// Static initialisation
 	//---------------------------------------------------------------------------------------------
-    FNodeType NodeImageMipmap::Private::s_type =
-            FNodeType( "ImageMipmap", NodeImage::GetStaticType() );
+    FNodeType NodeImageMipmap::Private::s_type = FNodeType(Node::EType::ImageMipmap, NodeImage::GetStaticType() );
 
 
 	//---------------------------------------------------------------------------------------------
 	//!
 	//---------------------------------------------------------------------------------------------
 
-    MUTABLE_IMPLEMENT_NODE( NodeImageMipmap, EType::Mipmap, Node, Node::EType::Image)
+    MUTABLE_IMPLEMENT_NODE( NodeImageMipmap )
 
 
 	//---------------------------------------------------------------------------------------------
diff --git a/Engine/Plugins/Experimental/Mutable/Source/MutableTools/Private/MuT/NodeImageMultiLayer.cpp b/Engine/Plugins/Experimental/Mutable/Source/MutableTools/Private/MuT/NodeImageMultiLayer.cpp
--- a/Engine/Plugins/Experimental/Mutable/Source/MutableTools/Private/MuT/NodeImageMultiLayer.cpp
+++ b/Engine/Plugins/Experimental/Mutable/Source/MutableTools/Private/MuT/NodeImageMultiLayer.cpp
@@ -14,15 +14,14 @@ namespace mu
 	//---------------------------------------------------------------------------------------------
 	// Static initialisation
 	//---------------------------------------------------------------------------------------------
-    FNodeType NodeImageMultiLayer::Private::s_type =
-            FNodeType( "ImageMultiLayer", NodeImage::GetStaticType() );
+    FNodeType NodeImageMultiLayer::Private::s_type = FNodeType(Node::EType::ImageMultiLayer, NodeImage::GetStaticType() );
 
 
 	//---------------------------------------------------------------------------------------------
 	//!
 	//---------------------------------------------------------------------------------------------
 
-    MUTABLE_IMPLEMENT_NODE( NodeImageMultiLayer, EType::MultiLayer, Node, Node::EType::Image)
+    MUTABLE_IMPLEMENT_NODE( NodeImageMultiLayer )
 
 
 	//---------------------------------------------------------------------------------------------
diff --git a/Engine/Plugins/Experimental/Mutable/Source/MutableTools/Private/MuT/NodeImageNormalComposite.cpp b/Engine/Plugins/Experimental/Mutable/Source/MutableTools/Private/MuT/NodeImageNormalComposite.cpp
--- a/Engine/Plugins/Experimental/Mutable/Source/MutableTools/Private/MuT/NodeImageNormalComposite.cpp
+++ b/Engine/Plugins/Experimental/Mutable/Source/MutableTools/Private/MuT/NodeImageNormalComposite.cpp
@@ -15,15 +15,14 @@ namespace mu
 	//---------------------------------------------------------------------------------------------
 	// Static initialisation
 	//---------------------------------------------------------------------------------------------
-	FNodeType NodeImageNormalComposite::Private::s_type =
-			FNodeType( "ImageNormalComposite", NodeImage::GetStaticType() );
+	FNodeType NodeImageNormalComposite::Private::s_type = FNodeType(Node::EType::ImageNormalComposite, NodeImage::GetStaticType() );
 
 
 	//---------------------------------------------------------------------------------------------
 	//!
 	//---------------------------------------------------------------------------------------------
 
-	MUTABLE_IMPLEMENT_NODE( NodeImageNormalComposite, EType::NormalComposite, Node, Node::EType::Image )
+	MUTABLE_IMPLEMENT_NODE( NodeImageNormalComposite )
 
 
 	//---------------------------------------------------------------------------------------------
diff --git a/Engine/Plugins/Experimental/Mutable/Source/MutableTools/Private/MuT/NodeImageParameter.cpp b/Engine/Plugins/Experimental/Mutable/Source/MutableTools/Private/MuT/NodeImageParameter.cpp
--- a/Engine/Plugins/Experimental/Mutable/Source/MutableTools/Private/MuT/NodeImageParameter.cpp
+++ b/Engine/Plugins/Experimental/Mutable/Source/MutableTools/Private/MuT/NodeImageParameter.cpp
@@ -14,15 +14,14 @@ namespace mu
 	//---------------------------------------------------------------------------------------------
 	// Static initialisation
 	//---------------------------------------------------------------------------------------------
-    FNodeType NodeImageParameter::Private::s_type =
-            FNodeType( "ImageParameter", NodeImage::GetStaticType() );
+    FNodeType NodeImageParameter::Private::s_type = FNodeType(Node::EType::ImageParameter, NodeImage::GetStaticType() );
 
 
 	//---------------------------------------------------------------------------------------------
 	//!
 	//---------------------------------------------------------------------------------------------
 
-    MUTABLE_IMPLEMENT_NODE( NodeImageParameter, EType::Parameter, Node, Node::EType::Image)
+    MUTABLE_IMPLEMENT_NODE( NodeImageParameter )
 
 
 	//---------------------------------------------------------------------------------------------
diff --git a/Engine/Plugins/Experimental/Mutable/Source/MutableTools/Private/MuT/NodeImagePlainColour.cpp b/Engine/Plugins/Experimental/Mutable/Source/MutableTools/Private/MuT/NodeImagePlainColour.cpp
--- a/Engine/Plugins/Experimental/Mutable/Source/MutableTools/Private/MuT/NodeImagePlainColour.cpp
+++ b/Engine/Plugins/Experimental/Mutable/Source/MutableTools/Private/MuT/NodeImagePlainColour.cpp
@@ -15,13 +15,12 @@ namespace mu
 	//---------------------------------------------------------------------------------------------
 	// Static initialisation
 	//---------------------------------------------------------------------------------------------
-	FNodeType NodeImagePlainColour::Private::s_type =
-			FNodeType( "ImagePlainColour", NodeImage::GetStaticType() );
+	FNodeType NodeImagePlainColour::Private::s_type = FNodeType(Node::EType::ImagePlainColour, NodeImage::GetStaticType() );
 
 
 	//---------------------------------------------------------------------------------------------
 
-	MUTABLE_IMPLEMENT_NODE( NodeImagePlainColour, EType::PlainColour, Node, Node::EType::Image)
+	MUTABLE_IMPLEMENT_NODE( NodeImagePlainColour )
 
 
 	//---------------------------------------------------------------------------------------------
diff --git a/Engine/Plugins/Experimental/Mutable/Source/MutableTools/Private/MuT/NodeImageProject.cpp b/Engine/Plugins/Experimental/Mutable/Source/MutableTools/Private/MuT/NodeImageProject.cpp
--- a/Engine/Plugins/Experimental/Mutable/Source/MutableTools/Private/MuT/NodeImageProject.cpp
+++ b/Engine/Plugins/Experimental/Mutable/Source/MutableTools/Private/MuT/NodeImageProject.cpp
@@ -19,15 +19,14 @@ namespace mu
 	//---------------------------------------------------------------------------------------------
 	// Static initialisation
 	//---------------------------------------------------------------------------------------------
-	FNodeType NodeImageProject::Private::s_type =
-			FNodeType( "ImageProject", NodeImage::GetStaticType() );
+	FNodeType NodeImageProject::Private::s_type = FNodeType(Node::EType::ImageProject, NodeImage::GetStaticType() );
 
 
 	//---------------------------------------------------------------------------------------------
 	//!
 	//---------------------------------------------------------------------------------------------
 
-	MUTABLE_IMPLEMENT_NODE( NodeImageProject, EType::Project, Node, Node::EType::Image)
+	MUTABLE_IMPLEMENT_NODE( NodeImageProject )
 
 
 	//---------------------------------------------------------------------------------------------
diff --git a/Engine/Plugins/Experimental/Mutable/Source/MutableTools/Private/MuT/NodeImageResize.cpp b/Engine/Plugins/Experimental/Mutable/Source/MutableTools/Private/MuT/NodeImageResize.cpp
--- a/Engine/Plugins/Experimental/Mutable/Source/MutableTools/Private/MuT/NodeImageResize.cpp
+++ b/Engine/Plugins/Experimental/Mutable/Source/MutableTools/Private/MuT/NodeImageResize.cpp
@@ -14,12 +14,11 @@ namespace mu
 	//---------------------------------------------------------------------------------------------
 	// Static initialisation
 	//---------------------------------------------------------------------------------------------
-	FNodeType NodeImageResize::Private::s_type =
-			FNodeType( "ImageResize", NodeImage::GetStaticType() );
+	FNodeType NodeImageResize::Private::s_type = FNodeType(Node::EType::ImageResize, NodeImage::GetStaticType() );
 
 	//---------------------------------------------------------------------------------------------
 
-	MUTABLE_IMPLEMENT_NODE( NodeImageResize, EType::Resize, Node, Node::EType::Image)
+	MUTABLE_IMPLEMENT_NODE( NodeImageResize )
 
 
 	//---------------------------------------------------------------------------------------------
diff --git a/Engine/Plugins/Experimental/Mutable/Source/MutableTools/Private/MuT/NodeImageSaturate.cpp b/Engine/Plugins/Experimental/Mutable/Source/MutableTools/Private/MuT/NodeImageSaturate.cpp
--- a/Engine/Plugins/Experimental/Mutable/Source/MutableTools/Private/MuT/NodeImageSaturate.cpp
+++ b/Engine/Plugins/Experimental/Mutable/Source/MutableTools/Private/MuT/NodeImageSaturate.cpp
@@ -15,15 +15,14 @@ namespace mu
 	//---------------------------------------------------------------------------------------------
 	// Static initialisation
 	//---------------------------------------------------------------------------------------------
-	FNodeType NodeImageSaturate::Private::s_type =
-			FNodeType( "ImageSaturate", NodeImage::GetStaticType() );
+	FNodeType NodeImageSaturate::Private::s_type = FNodeType(Node::EType::ImageSaturate, NodeImage::GetStaticType() );
 
 
 	//---------------------------------------------------------------------------------------------
 	//!
 	//---------------------------------------------------------------------------------------------
 
-	MUTABLE_IMPLEMENT_NODE( NodeImageSaturate, EType::Saturate, Node, Node::EType::Image)
+	MUTABLE_IMPLEMENT_NODE( NodeImageSaturate )
 
 
 	//---------------------------------------------------------------------------------------------
diff --git a/Engine/Plugins/Experimental/Mutable/Source/MutableTools/Private/MuT/NodeImageSwitch.cpp b/Engine/Plugins/Experimental/Mutable/Source/MutableTools/Private/MuT/NodeImageSwitch.cpp
--- a/Engine/Plugins/Experimental/Mutable/Source/MutableTools/Private/MuT/NodeImageSwitch.cpp
+++ b/Engine/Plugins/Experimental/Mutable/Source/MutableTools/Private/MuT/NodeImageSwitch.cpp
@@ -15,15 +15,14 @@ namespace mu
 	//---------------------------------------------------------------------------------------------
 	// Static initialisation
 	//---------------------------------------------------------------------------------------------
-	FNodeType NodeImageSwitch::Private::s_type =
-			FNodeType( "ImageSwitch", NodeImage::GetStaticType() );
+	FNodeType NodeImageSwitch::Private::s_type = FNodeType(Node::EType::ImageSwitch, NodeImage::GetStaticType() );
 
 
 	//---------------------------------------------------------------------------------------------
 	//!
 	//---------------------------------------------------------------------------------------------
 
-	MUTABLE_IMPLEMENT_NODE( NodeImageSwitch, EType::Switch, Node, Node::EType::Image)
+	MUTABLE_IMPLEMENT_NODE( NodeImageSwitch )
 
 
 	//---------------------------------------------------------------------------------------------
diff --git a/Engine/Plugins/Experimental/Mutable/Source/MutableTools/Private/MuT/NodeImageSwizzle.cpp b/Engine/Plugins/Experimental/Mutable/Source/MutableTools/Private/MuT/NodeImageSwizzle.cpp
--- a/Engine/Plugins/Experimental/Mutable/Source/MutableTools/Private/MuT/NodeImageSwizzle.cpp
+++ b/Engine/Plugins/Experimental/Mutable/Source/MutableTools/Private/MuT/NodeImageSwizzle.cpp
@@ -15,15 +15,14 @@ namespace mu
 	//---------------------------------------------------------------------------------------------
 	// Static initialisation
 	//---------------------------------------------------------------------------------------------
-	FNodeType NodeImageSwizzle::Private::s_type =
-			FNodeType( "ImageSwizzle", NodeImage::GetStaticType() );
+	FNodeType NodeImageSwizzle::Private::s_type = FNodeType(Node::EType::ImageSwizzle, NodeImage::GetStaticType() );
 
 
 	//---------------------------------------------------------------------------------------------
 	//!
 	//---------------------------------------------------------------------------------------------
 
-	MUTABLE_IMPLEMENT_NODE( NodeImageSwizzle, EType::Swizzle, Node, Node::EType::Image)
+	MUTABLE_IMPLEMENT_NODE( NodeImageSwizzle )
 
 
 	//---------------------------------------------------------------------------------------------
diff --git a/Engine/Plugins/Experimental/Mutable/Source/MutableTools/Private/MuT/NodeImageTable.cpp b/Engine/Plugins/Experimental/Mutable/Source/MutableTools/Private/MuT/NodeImageTable.cpp
--- a/Engine/Plugins/Experimental/Mutable/Source/MutableTools/Private/MuT/NodeImageTable.cpp
+++ b/Engine/Plugins/Experimental/Mutable/Source/MutableTools/Private/MuT/NodeImageTable.cpp
@@ -16,15 +16,14 @@ namespace mu
 	//---------------------------------------------------------------------------------------------
 	// Static initialisation
 	//---------------------------------------------------------------------------------------------
-	FNodeType NodeImageTable::Private::s_type =
-			FNodeType( "TableImage", NodeImage::GetStaticType() );
+	FNodeType NodeImageTable::Private::s_type = FNodeType(Node::EType::ImageTable, NodeImage::GetStaticType() );
 
 
 	//---------------------------------------------------------------------------------------------
 	//!
 	//---------------------------------------------------------------------------------------------
 
-	MUTABLE_IMPLEMENT_NODE( NodeImageTable, EType::Table, Node, Node::EType::Image);
+	MUTABLE_IMPLEMENT_NODE( NodeImageTable );
 
 
 	//---------------------------------------------------------------------------------------------
diff --git a/Engine/Plugins/Experimental/Mutable/Source/MutableTools/Private/MuT/NodeImageTransform.cpp b/Engine/Plugins/Experimental/Mutable/Source/MutableTools/Private/MuT/NodeImageTransform.cpp
--- a/Engine/Plugins/Experimental/Mutable/Source/MutableTools/Private/MuT/NodeImageTransform.cpp
+++ b/Engine/Plugins/Experimental/Mutable/Source/MutableTools/Private/MuT/NodeImageTransform.cpp
@@ -15,12 +15,11 @@ namespace mu
 	//---------------------------------------------------------------------------------------------
 	// Static initialisation
 	//---------------------------------------------------------------------------------------------
-	FNodeType NodeImageTransform::Private::s_type =
-			FNodeType( "ImageTransform", NodeImage::GetStaticType() );
+	FNodeType NodeImageTransform::Private::s_type = FNodeType(Node::EType::ImageTransform, NodeImage::GetStaticType() );
 
 	//---------------------------------------------------------------------------------------------
 
-	MUTABLE_IMPLEMENT_NODE( NodeImageTransform, EType::Transform, Node, Node::EType::Image)
+	MUTABLE_IMPLEMENT_NODE( NodeImageTransform )
 
 
 	//---------------------------------------------------------------------------------------------
diff --git a/Engine/Plugins/Experimental/Mutable/Source/MutableTools/Private/MuT/NodeImageVariation.cpp b/Engine/Plugins/Experimental/Mutable/Source/MutableTools/Private/MuT/NodeImageVariation.cpp
--- a/Engine/Plugins/Experimental/Mutable/Source/MutableTools/Private/MuT/NodeImageVariation.cpp
+++ b/Engine/Plugins/Experimental/Mutable/Source/MutableTools/Private/MuT/NodeImageVariation.cpp
@@ -14,15 +14,14 @@ namespace mu
     //---------------------------------------------------------------------------------------------
     // Static initialisation
     //---------------------------------------------------------------------------------------------
-    FNodeType NodeImageVariation::Private::s_type =
-        FNodeType( "ImageVariation", NodeImage::GetStaticType() );
+    FNodeType NodeImageVariation::Private::s_type = FNodeType(Node::EType::ImageVariation, NodeImage::GetStaticType() );
 
 
     //---------------------------------------------------------------------------------------------
     //---------------------------------------------------------------------------------------------
     //---------------------------------------------------------------------------------------------
 
-    MUTABLE_IMPLEMENT_NODE( NodeImageVariation, EType::Variation, Node, Node::EType::Image)
+    MUTABLE_IMPLEMENT_NODE( NodeImageVariation )
 
 
     //---------------------------------------------------------------------------------------------
diff --git a/Engine/Plugins/Experimental/Mutable/Source/MutableTools/Private/MuT/NodeLayout.cpp b/Engine/Plugins/Experimental/Mutable/Source/MutableTools/Private/MuT/NodeLayout.cpp
--- a/Engine/Plugins/Experimental/Mutable/Source/MutableTools/Private/MuT/NodeLayout.cpp
+++ b/Engine/Plugins/Experimental/Mutable/Source/MutableTools/Private/MuT/NodeLayout.cpp
@@ -15,7 +15,7 @@
 namespace mu
 {
 
-	static FNodeType s_nodeLayoutType = FNodeType( "NodeLayout", Node::GetStaticType() );
+	static FNodeType s_nodeLayoutType = FNodeType(Node::EType::Layout, Node::GetStaticType() );
 
 
 	//---------------------------------------------------------------------------------------------
@@ -33,14 +33,14 @@ namespace mu
 
 
 	//---------------------------------------------------------------------------------------------
-	FNodeType NodeLayoutBlocks::Private::s_type = FNodeType( "LayoutBlocks", NodeLayout::GetStaticType() );
+	FNodeType NodeLayoutBlocks::Private::s_type = FNodeType(Node::EType::LayoutBlocks, NodeLayout::GetStaticType() );
 
 
 	//---------------------------------------------------------------------------------------------
 	//!
 	//---------------------------------------------------------------------------------------------
 
-	MUTABLE_IMPLEMENT_NODE( NodeLayoutBlocks, EType::Blocks, Node, Node::EType::Layout)
+	MUTABLE_IMPLEMENT_NODE( NodeLayoutBlocks )
 
 
 	//---------------------------------------------------------------------------------------------
diff --git a/Engine/Plugins/Experimental/Mutable/Source/MutableTools/Private/MuT/NodeMesh.cpp b/Engine/Plugins/Experimental/Mutable/Source/MutableTools/Private/MuT/NodeMesh.cpp
--- a/Engine/Plugins/Experimental/Mutable/Source/MutableTools/Private/MuT/NodeMesh.cpp
+++ b/Engine/Plugins/Experimental/Mutable/Source/MutableTools/Private/MuT/NodeMesh.cpp
@@ -11,7 +11,7 @@ namespace mu
 	//---------------------------------------------------------------------------------------------
 	// Static initialisation
 	//---------------------------------------------------------------------------------------------
-	static FNodeType s_nodeMeshType = FNodeType( "NodeMesh", Node::GetStaticType() );
+	static FNodeType s_nodeMeshType = FNodeType(Node::EType::Mesh , Node::GetStaticType());
 
 
 	//---------------------------------------------------------------------------------------------
diff --git a/Engine/Plugins/Experimental/Mutable/Source/MutableTools/Private/MuT/NodeMeshApplyPose.cpp b/Engine/Plugins/Experimental/Mutable/Source/MutableTools/Private/MuT/NodeMeshApplyPose.cpp
--- a/Engine/Plugins/Experimental/Mutable/Source/MutableTools/Private/MuT/NodeMeshApplyPose.cpp
+++ b/Engine/Plugins/Experimental/Mutable/Source/MutableTools/Private/MuT/NodeMeshApplyPose.cpp
@@ -15,15 +15,14 @@ namespace mu
 	//---------------------------------------------------------------------------------------------
 	// Static initialisation
 	//---------------------------------------------------------------------------------------------
-    FNodeType NodeMeshApplyPose::Private::s_type =
-            FNodeType( "MeshApplyPose", NodeMesh::GetStaticType() );
+    FNodeType NodeMeshApplyPose::Private::s_type = FNodeType(Node::EType::MeshApplyPose, NodeMesh::GetStaticType() );
 
 
 	//---------------------------------------------------------------------------------------------
 	//!
 	//---------------------------------------------------------------------------------------------
 
-    MUTABLE_IMPLEMENT_NODE( NodeMeshApplyPose, EType::ApplyPose, Node, Node::EType::Mesh)
+    MUTABLE_IMPLEMENT_NODE( NodeMeshApplyPose )
 
 
 	//---------------------------------------------------------------------------------------------
@@ -56,22 +55,6 @@ namespace mu
     }
 
 
-	//---------------------------------------------------------------------------------------------
-    NodeLayoutPtr NodeMeshApplyPose::Private::GetLayout( int index ) const
-	{
-		NodeLayoutPtr pResult;
-
-		if ( m_pBase )
-		{
-			NodeMesh::Private* pPrivate = static_cast<NodeMesh::Private*>( m_pBase->GetBasePrivate() );
-
-			pResult = pPrivate->GetLayout( index );
-		}
-
-		return pResult;
-	}
-
-
 }
 
 
diff --git a/Engine/Plugins/Experimental/Mutable/Source/MutableTools/Private/MuT/NodeMeshApplyPosePrivate.h b/Engine/Plugins/Experimental/Mutable/Source/MutableTools/Private/MuT/NodeMeshApplyPosePrivate.h
--- a/Engine/Plugins/Experimental/Mutable/Source/MutableTools/Private/MuT/NodeMeshApplyPosePrivate.h
+++ b/Engine/Plugins/Experimental/Mutable/Source/MutableTools/Private/MuT/NodeMeshApplyPosePrivate.h
@@ -19,9 +19,6 @@ namespace mu
         NodeMeshPtr m_pBase;
         NodeMeshPtr m_pPose;
 
-        // NodeMesh::Private interface
-        NodeLayoutPtr GetLayout( int index ) const override;
-
 	};
 
 
diff --git a/Engine/Plugins/Experimental/Mutable/Source/MutableTools/Private/MuT/NodeMeshClipDeform.cpp b/Engine/Plugins/Experimental/Mutable/Source/MutableTools/Private/MuT/NodeMeshClipDeform.cpp
--- a/Engine/Plugins/Experimental/Mutable/Source/MutableTools/Private/MuT/NodeMeshClipDeform.cpp
+++ b/Engine/Plugins/Experimental/Mutable/Source/MutableTools/Private/MuT/NodeMeshClipDeform.cpp
@@ -16,15 +16,14 @@ namespace mu
 	//---------------------------------------------------------------------------------------------
 	// Static initialisation
 	//---------------------------------------------------------------------------------------------
-	FNodeType NodeMeshClipDeform::Private::s_type =
-			FNodeType( "MeshClipDeform", NodeMesh::GetStaticType() );
+	FNodeType NodeMeshClipDeform::Private::s_type = FNodeType(Node::EType::MeshClipDeform, NodeMesh::GetStaticType() );
 
 
 	//---------------------------------------------------------------------------------------------
 	//!
 	//---------------------------------------------------------------------------------------------
 
-	MUTABLE_IMPLEMENT_NODE( NodeMeshClipDeform, EType::ClipDeform, Node, Node::EType::Mesh)
+	MUTABLE_IMPLEMENT_NODE( NodeMeshClipDeform )
 
 
 	//---------------------------------------------------------------------------------------------
@@ -57,21 +56,6 @@ namespace mu
 	}
 
 
-	//---------------------------------------------------------------------------------------------
-	Ptr<NodeLayout> NodeMeshClipDeform::Private::GetLayout( int index ) const
-	{
-		NodeLayoutPtr pResult;
-
-		if (m_pBaseMesh)
-		{
-			NodeMesh::Private* pPrivate = static_cast<NodeMesh::Private*>(m_pBaseMesh->GetBasePrivate() );
-			pResult = pPrivate->GetLayout( index );
-		}
-
-		return pResult;
-	}
-
-
 }
 
 
diff --git a/Engine/Plugins/Experimental/Mutable/Source/MutableTools/Private/MuT/NodeMeshClipDeformPrivate.h b/Engine/Plugins/Experimental/Mutable/Source/MutableTools/Private/MuT/NodeMeshClipDeformPrivate.h
--- a/Engine/Plugins/Experimental/Mutable/Source/MutableTools/Private/MuT/NodeMeshClipDeformPrivate.h
+++ b/Engine/Plugins/Experimental/Mutable/Source/MutableTools/Private/MuT/NodeMeshClipDeformPrivate.h
@@ -20,9 +20,6 @@ namespace mu
 		Ptr<NodeMesh> m_pClipShape;
 		Ptr<NodeImage> m_pShapeWeights;
 
-		// NodeMesh::Private interface
-        Ptr<NodeLayout> GetLayout( int index ) const override;
-
 	};
 
 }
diff --git a/Engine/Plugins/Experimental/Mutable/Source/MutableTools/Private/MuT/NodeMeshClipMorphPlane.cpp b/Engine/Plugins/Experimental/Mutable/Source/MutableTools/Private/MuT/NodeMeshClipMorphPlane.cpp
--- a/Engine/Plugins/Experimental/Mutable/Source/MutableTools/Private/MuT/NodeMeshClipMorphPlane.cpp
+++ b/Engine/Plugins/Experimental/Mutable/Source/MutableTools/Private/MuT/NodeMeshClipMorphPlane.cpp
@@ -16,15 +16,14 @@ namespace mu
 	//---------------------------------------------------------------------------------------------
 	// Static initialisation
 	//---------------------------------------------------------------------------------------------
-    FNodeType NodeMeshClipMorphPlane::Private::s_type =
-            FNodeType( "MeshClipMorphPlane", NodeMesh::GetStaticType() );
+    FNodeType NodeMeshClipMorphPlane::Private::s_type = FNodeType(Node::EType::MeshClipMorphPlane, NodeMesh::GetStaticType() );
 
 
 	//---------------------------------------------------------------------------------------------
 	//!
 	//---------------------------------------------------------------------------------------------
 
-    MUTABLE_IMPLEMENT_NODE( NodeMeshClipMorphPlane, EType::ClipMorphPlane, Node, Node::EType::Mesh)
+    MUTABLE_IMPLEMENT_NODE( NodeMeshClipMorphPlane )
 
 
 	//---------------------------------------------------------------------------------------------
@@ -82,20 +81,6 @@ namespace mu
 		m_pD->m_maxEffectRadius = maxEffectRadius;
 	}
 
-	//---------------------------------------------------------------------------------------------
-    NodeLayoutPtr NodeMeshClipMorphPlane::Private::GetLayout( int index ) const
-	{
-		NodeLayoutPtr pResult;
-
-		if ( m_pSource )
-		{
-			NodeMesh::Private* pPrivate = static_cast<NodeMesh::Private*>( m_pSource->GetBasePrivate() );
-			pResult = pPrivate->GetLayout( index );
-		}
-
-		return pResult;
-	}
-
 	//---------------------------------------------------------------------------------------------
 	void NodeMeshClipMorphPlane::AddTag(const char* tagName)
 	{
diff --git a/Engine/Plugins/Experimental/Mutable/Source/MutableTools/Private/MuT/NodeMeshClipMorphPlanePrivate.h b/Engine/Plugins/Experimental/Mutable/Source/MutableTools/Private/MuT/NodeMeshClipMorphPlanePrivate.h
--- a/Engine/Plugins/Experimental/Mutable/Source/MutableTools/Private/MuT/NodeMeshClipMorphPlanePrivate.h
+++ b/Engine/Plugins/Experimental/Mutable/Source/MutableTools/Private/MuT/NodeMeshClipMorphPlanePrivate.h
@@ -67,9 +67,6 @@ namespace mu
 		// means no limit.
 		float m_maxEffectRadius;
 
-        // NodeMesh::Private interface
-        NodeLayoutPtr GetLayout( int index ) const override;
-
 	};
 
 
diff --git a/Engine/Plugins/Experimental/Mutable/Source/MutableTools/Private/MuT/NodeMeshClipWithMesh.cpp b/Engine/Plugins/Experimental/Mutable/Source/MutableTools/Private/MuT/NodeMeshClipWithMesh.cpp
--- a/Engine/Plugins/Experimental/Mutable/Source/MutableTools/Private/MuT/NodeMeshClipWithMesh.cpp
+++ b/Engine/Plugins/Experimental/Mutable/Source/MutableTools/Private/MuT/NodeMeshClipWithMesh.cpp
@@ -16,15 +16,14 @@ namespace mu
 	//---------------------------------------------------------------------------------------------
 	// Static initialisation
 	//---------------------------------------------------------------------------------------------
-    FNodeType NodeMeshClipWithMesh::Private::s_type =
-            FNodeType( "MeshClipMorphPlane", NodeMesh::GetStaticType() );
+    FNodeType NodeMeshClipWithMesh::Private::s_type = FNodeType(Node::EType::MeshClipMorphPlane, NodeMesh::GetStaticType() );
 
 
 	//---------------------------------------------------------------------------------------------
 	//!
 	//---------------------------------------------------------------------------------------------
 
-    MUTABLE_IMPLEMENT_NODE( NodeMeshClipWithMesh, EType::ClipWithMesh, Node, Node::EType::Mesh)
+    MUTABLE_IMPLEMENT_NODE( NodeMeshClipWithMesh )
 
 
 	//---------------------------------------------------------------------------------------------
@@ -50,20 +49,6 @@ namespace mu
 	}
 
 
-	//---------------------------------------------------------------------------------------------
-    NodeLayoutPtr NodeMeshClipWithMesh::Private::GetLayout( int index ) const
-	{
-		NodeLayoutPtr pResult;
-
-		if ( m_pSource )
-		{
-			NodeMesh::Private* pPrivate = static_cast<NodeMesh::Private*>( m_pSource->GetBasePrivate() );
-			pResult = pPrivate->GetLayout( index );
-		}
-
-		return pResult;
-	}
-
 	//---------------------------------------------------------------------------------------------
 	void NodeMeshClipWithMesh::AddTag(const char* tagName)
 	{
diff --git a/Engine/Plugins/Experimental/Mutable/Source/MutableTools/Private/MuT/NodeMeshConstant.cpp b/Engine/Plugins/Experimental/Mutable/Source/MutableTools/Private/MuT/NodeMeshConstant.cpp
--- a/Engine/Plugins/Experimental/Mutable/Source/MutableTools/Private/MuT/NodeMeshConstant.cpp
+++ b/Engine/Plugins/Experimental/Mutable/Source/MutableTools/Private/MuT/NodeMeshConstant.cpp
@@ -11,10 +11,10 @@
 
 namespace mu
 {
-	FNodeType NodeMeshConstant::Private::s_type = FNodeType( "MeshConstant", NodeMesh::GetStaticType() );
+	FNodeType NodeMeshConstant::Private::s_type = FNodeType(Node::EType::MeshConstant, NodeMesh::GetStaticType() );
 
 
-	MUTABLE_IMPLEMENT_NODE( NodeMeshConstant, EType::Constant, Node, Node::EType::Mesh);
+	MUTABLE_IMPLEMENT_NODE( NodeMeshConstant );
 
 
 	Ptr<Mesh> NodeMeshConstant::GetValue() const
@@ -52,19 +52,12 @@ namespace mu
 	{
 		check( index >=0 && index < m_pD->Layouts.Num() );
 
-		return m_pD->GetLayout( index );
-	}
-
-
-	Ptr<NodeLayout> NodeMeshConstant::Private::GetLayout( int32 index ) const
-	{
 		NodeLayoutPtr pResult;
 
-		if ( index >=0 && index < Layouts.Num() )
+		if (index >= 0 && index < m_pD->Layouts.Num())
 		{
-			pResult = Layouts[ index ];
+			pResult = m_pD->Layouts[index];
 		}
-
 		return pResult;
 	}
 
diff --git a/Engine/Plugins/Experimental/Mutable/Source/MutableTools/Private/MuT/NodeMeshFormat.cpp b/Engine/Plugins/Experimental/Mutable/Source/MutableTools/Private/MuT/NodeMeshFormat.cpp
--- a/Engine/Plugins/Experimental/Mutable/Source/MutableTools/Private/MuT/NodeMeshFormat.cpp
+++ b/Engine/Plugins/Experimental/Mutable/Source/MutableTools/Private/MuT/NodeMeshFormat.cpp
@@ -17,15 +17,14 @@ class FMeshBufferSet;
 	//---------------------------------------------------------------------------------------------
 	// Static initialisation
 	//---------------------------------------------------------------------------------------------
-	FNodeType NodeMeshFormat::Private::s_type =
-			FNodeType( "MeshFormat", NodeMesh::GetStaticType() );
+	FNodeType NodeMeshFormat::Private::s_type = FNodeType(Node::EType::MeshFormat, NodeMesh::GetStaticType() );
 
 
 	//---------------------------------------------------------------------------------------------
 	//!
 	//---------------------------------------------------------------------------------------------
 
-	MUTABLE_IMPLEMENT_NODE( NodeMeshFormat, EType::Format, Node, Node::EType::Mesh)
+	MUTABLE_IMPLEMENT_NODE( NodeMeshFormat )
 
 
 	//---------------------------------------------------------------------------------------------
@@ -60,20 +59,6 @@ class FMeshBufferSet;
 		m_pD->bOptimizeBuffers = bEnable;
 	}
 
-
-	Ptr<NodeLayout> NodeMeshFormat::Private::GetLayout( int32 index ) const
-	{
-		Ptr<NodeLayout> Result;
-
-		if ( Source )
-		{
-			NodeMesh::Private* Private = static_cast<NodeMesh::Private*>( Source->GetBasePrivate() );
-			Result = Private->GetLayout( index );
-		}
-
-		return Result;
-	}
-
 }
 
 
diff --git a/Engine/Plugins/Experimental/Mutable/Source/MutableTools/Private/MuT/NodeMeshFragment.cpp b/Engine/Plugins/Experimental/Mutable/Source/MutableTools/Private/MuT/NodeMeshFragment.cpp
--- a/Engine/Plugins/Experimental/Mutable/Source/MutableTools/Private/MuT/NodeMeshFragment.cpp
+++ b/Engine/Plugins/Experimental/Mutable/Source/MutableTools/Private/MuT/NodeMeshFragment.cpp
@@ -16,15 +16,14 @@ namespace mu
 	//---------------------------------------------------------------------------------------------
 	// Static initialisation
 	//---------------------------------------------------------------------------------------------
-	FNodeType NodeMeshFragment::Private::s_type =
-			FNodeType( "MeshFragment", NodeMesh::GetStaticType() );
+	FNodeType NodeMeshFragment::Private::s_type = FNodeType(Node::EType::MeshFragment, NodeMesh::GetStaticType() );
 
 
 	//---------------------------------------------------------------------------------------------
 	//!
 	//---------------------------------------------------------------------------------------------
 
-	MUTABLE_IMPLEMENT_NODE( NodeMeshFragment, EType::Fragment, Node, Node::EType::Mesh)
+	MUTABLE_IMPLEMENT_NODE( NodeMeshFragment )
 
 
 	//---------------------------------------------------------------------------------------------
@@ -93,24 +92,6 @@ namespace mu
 		m_pD->Blocks[t] = b;
 	}
 
-
-	//---------------------------------------------------------------------------------------------
-	Ptr<NodeLayout> NodeMeshFragment::Private::GetLayout( int32 index ) const
-	{
-		Ptr<NodeLayout> pResult;
-
-		if ( m_pMesh )
-		{
-			// TODO: Cut a fragment out of the layout.
-			NodeMesh::Private* pPrivate = static_cast<NodeMesh::Private*>( m_pMesh->GetBasePrivate() );
-			pResult = pPrivate->GetLayout( index );
-		}
-
-		return pResult;
-	}
-
-
-
 }
 
 
diff --git a/Engine/Plugins/Experimental/Mutable/Source/MutableTools/Private/MuT/NodeMeshGeometryOperation.cpp b/Engine/Plugins/Experimental/Mutable/Source/MutableTools/Private/MuT/NodeMeshGeometryOperation.cpp
--- a/Engine/Plugins/Experimental/Mutable/Source/MutableTools/Private/MuT/NodeMeshGeometryOperation.cpp
+++ b/Engine/Plugins/Experimental/Mutable/Source/MutableTools/Private/MuT/NodeMeshGeometryOperation.cpp
@@ -18,15 +18,14 @@ namespace mu
 	//---------------------------------------------------------------------------------------------
 	// Static initialisation
 	//---------------------------------------------------------------------------------------------
-	FNodeType NodeMeshGeometryOperation::Private::s_type =
-			FNodeType( "MeshGeometryOperation", NodeMesh::GetStaticType() );
+	FNodeType NodeMeshGeometryOperation::Private::s_type = FNodeType(Node::EType::MeshGeometryOperation, NodeMesh::GetStaticType() );
 
 
 	//---------------------------------------------------------------------------------------------
 	//!
 	//---------------------------------------------------------------------------------------------
 
-	MUTABLE_IMPLEMENT_NODE( NodeMeshGeometryOperation, EType::GeometryOperation, Node, Node::EType::Mesh)
+	MUTABLE_IMPLEMENT_NODE( NodeMeshGeometryOperation )
 
 
 	//---------------------------------------------------------------------------------------------
@@ -86,22 +85,6 @@ namespace mu
 		m_pD->m_pMeshB = pNode;
 	}
 
-
-	//---------------------------------------------------------------------------------------------
-	NodeLayoutPtr NodeMeshGeometryOperation::Private::GetLayout( int index ) const
-	{
-		NodeLayoutPtr pResult;
-
-		if (m_pMeshA)
-		{
-			NodeMesh::Private* pPrivate = static_cast<NodeMesh::Private*>(m_pMeshA->GetBasePrivate() );
-			pResult = pPrivate->GetLayout( index );
-		}
-
-		return pResult;
-	}
-
-
 }
 
 
diff --git a/Engine/Plugins/Experimental/Mutable/Source/MutableTools/Private/MuT/NodeMeshInterpolate.cpp b/Engine/Plugins/Experimental/Mutable/Source/MutableTools/Private/MuT/NodeMeshInterpolate.cpp
--- a/Engine/Plugins/Experimental/Mutable/Source/MutableTools/Private/MuT/NodeMeshInterpolate.cpp
+++ b/Engine/Plugins/Experimental/Mutable/Source/MutableTools/Private/MuT/NodeMeshInterpolate.cpp
@@ -17,15 +17,14 @@ namespace mu
 	//---------------------------------------------------------------------------------------------
 	// Static initialisation
 	//---------------------------------------------------------------------------------------------
-	FNodeType NodeMeshInterpolate::Private::s_type =
-			FNodeType( "MeshInterpolate", NodeMesh::GetStaticType() );
+	FNodeType NodeMeshInterpolate::Private::s_type = FNodeType(Node::EType::MeshInterpolate, NodeMesh::GetStaticType() );
 
 
 	//---------------------------------------------------------------------------------------------
 	//!
 	//---------------------------------------------------------------------------------------------
 
-	MUTABLE_IMPLEMENT_NODE( NodeMeshInterpolate, EType::Interpolate, Node, Node::EType::Mesh)
+	MUTABLE_IMPLEMENT_NODE( NodeMeshInterpolate )
 
 
 	//---------------------------------------------------------------------------------------------
@@ -95,22 +94,4 @@ namespace mu
 	}
 
 
-	//---------------------------------------------------------------------------------------------
-	//---------------------------------------------------------------------------------------------
-	//---------------------------------------------------------------------------------------------
-	NodeLayoutPtr NodeMeshInterpolate::Private::GetLayout( int index ) const
-	{
-		NodeLayoutPtr pResult;
-
-		// TODO: Substract layouts too? Usually they are ignored.
-		if ( m_targets.Num()>0 && m_targets[0] )
-		{
-			NodeMesh::Private* pPrivate = static_cast<NodeMesh::Private*>( m_targets[0]->GetBasePrivate() );
-			pResult = pPrivate->GetLayout( index );
-		}
-
-		return pResult;
-	}
-
-
 }
diff --git a/Engine/Plugins/Experimental/Mutable/Source/MutableTools/Private/MuT/NodeMeshMakeMorph.cpp b/Engine/Plugins/Experimental/Mutable/Source/MutableTools/Private/MuT/NodeMeshMakeMorph.cpp
--- a/Engine/Plugins/Experimental/Mutable/Source/MutableTools/Private/MuT/NodeMeshMakeMorph.cpp
+++ b/Engine/Plugins/Experimental/Mutable/Source/MutableTools/Private/MuT/NodeMeshMakeMorph.cpp
@@ -16,15 +16,14 @@ namespace mu
 	//---------------------------------------------------------------------------------------------
 	// Static initialisation
 	//---------------------------------------------------------------------------------------------
-    FNodeType NodeMeshMakeMorph::Private::s_type =
-			FNodeType( "MeshMakeMorph", NodeMesh::GetStaticType() );
+    FNodeType NodeMeshMakeMorph::Private::s_type = FNodeType(Node::EType::MeshMakeMorph, NodeMesh::GetStaticType() );
 
 
 	//---------------------------------------------------------------------------------------------
 	//!
 	//---------------------------------------------------------------------------------------------
 
-    MUTABLE_IMPLEMENT_NODE( NodeMeshMakeMorph, EType::MakeMorph, Node, Node::EType::Mesh)
+    MUTABLE_IMPLEMENT_NODE( NodeMeshMakeMorph )
 
 
 	//---------------------------------------------------------------------------------------------
@@ -69,20 +68,4 @@ namespace mu
 		return m_pD->bOnlyPositionAndNormal;
 	}
 
-	//---------------------------------------------------------------------------------------------
-    NodeLayoutPtr NodeMeshMakeMorph::Private::GetLayout( int index ) const
-	{
-		NodeLayoutPtr pResult;
-
-		// TODO: Substract layouts too? Usually they are ignored.
-		if ( m_pBase )
-		{
-			NodeMesh::Private* pPrivate = static_cast<NodeMesh::Private*>( m_pBase->GetBasePrivate() );
-			pResult = pPrivate->GetLayout( index );
-		}
-
-		return pResult;
-	}
-
-
 }
diff --git a/Engine/Plugins/Experimental/Mutable/Source/MutableTools/Private/MuT/NodeMeshMorph.cpp b/Engine/Plugins/Experimental/Mutable/Source/MutableTools/Private/MuT/NodeMeshMorph.cpp
--- a/Engine/Plugins/Experimental/Mutable/Source/MutableTools/Private/MuT/NodeMeshMorph.cpp
+++ b/Engine/Plugins/Experimental/Mutable/Source/MutableTools/Private/MuT/NodeMeshMorph.cpp
@@ -18,14 +18,14 @@ namespace mu
 	//---------------------------------------------------------------------------------------------
 	// Static initialisation
 	//---------------------------------------------------------------------------------------------
-	FNodeType NodeMeshMorph::Private::s_type = FNodeType( "MeshMorph", NodeMesh::GetStaticType() );
+	FNodeType NodeMeshMorph::Private::s_type = FNodeType(Node::EType::MeshMorph, NodeMesh::GetStaticType() );
 
 
 	//---------------------------------------------------------------------------------------------
 	//!
 	//---------------------------------------------------------------------------------------------
 
-	MUTABLE_IMPLEMENT_NODE( NodeMeshMorph, EType::Morph, Node, Node::EType::Mesh)
+	MUTABLE_IMPLEMENT_NODE( NodeMeshMorph )
 
 
 	//---------------------------------------------------------------------------------------------
@@ -94,21 +94,6 @@ namespace mu
 		m_pD->bReshapePhysicsVolumes = bEnable;
 	}
 
-	//---------------------------------------------------------------------------------------------
-	NodeLayoutPtr NodeMeshMorph::Private::GetLayout( int index ) const
-	{
-		NodeLayoutPtr pResult;
-
-		if ( Base )
-		{
-			NodeMesh::Private* pPrivate = static_cast<NodeMesh::Private*>( Base->GetBasePrivate() );
-			pResult = pPrivate->GetLayout( index );
-		}
-
-		return pResult;
-	}
-
-
 }
 
 
diff --git a/Engine/Plugins/Experimental/Mutable/Source/MutableTools/Private/MuT/NodeMeshReshape.cpp b/Engine/Plugins/Experimental/Mutable/Source/MutableTools/Private/MuT/NodeMeshReshape.cpp
--- a/Engine/Plugins/Experimental/Mutable/Source/MutableTools/Private/MuT/NodeMeshReshape.cpp
+++ b/Engine/Plugins/Experimental/Mutable/Source/MutableTools/Private/MuT/NodeMeshReshape.cpp
@@ -17,15 +17,14 @@ namespace mu
 	//---------------------------------------------------------------------------------------------
 	// Static initialisation
 	//---------------------------------------------------------------------------------------------
-	FNodeType NodeMeshReshape::Private::s_type =
-			FNodeType("MeshReshape", NodeMesh::GetStaticType());
+	FNodeType NodeMeshReshape::Private::s_type = FNodeType(Node::EType::MeshReshape, NodeMesh::GetStaticType());
 
 
 	//---------------------------------------------------------------------------------------------
 	//!
 	//---------------------------------------------------------------------------------------------
 
-	MUTABLE_IMPLEMENT_NODE(NodeMeshReshape, EType::Reshape, Node, Node::EType::Mesh)
+	MUTABLE_IMPLEMENT_NODE(NodeMeshReshape )
 
 
 	//---------------------------------------------------------------------------------------------
@@ -123,20 +122,6 @@ namespace mu
 	{
 		m_pD->bReshapePhysicsVolumes = bEnable;
 	}
-	
-	//---------------------------------------------------------------------------------------------
-	Ptr<NodeLayout> NodeMeshReshape::Private::GetLayout(int32 Index) const
-	{
-		NodeLayoutPtr pResult;
-
-		if (BaseMesh)
-		{
-			NodeMesh::Private* pPrivate = static_cast<NodeMesh::Private*>(BaseMesh->GetBasePrivate() );
-			pResult = pPrivate->GetLayout(Index);
-		}
-
-		return pResult;
-	}
 
 }
 
diff --git a/Engine/Plugins/Experimental/Mutable/Source/MutableTools/Private/MuT/NodeMeshReshapePrivate.h b/Engine/Plugins/Experimental/Mutable/Source/MutableTools/Private/MuT/NodeMeshReshapePrivate.h
--- a/Engine/Plugins/Experimental/Mutable/Source/MutableTools/Private/MuT/NodeMeshReshapePrivate.h
+++ b/Engine/Plugins/Experimental/Mutable/Source/MutableTools/Private/MuT/NodeMeshReshapePrivate.h
@@ -34,9 +34,6 @@ namespace mu
 		TArray<FBoneName> BonesToDeform;
 		TArray<FBoneName> PhysicsToDeform;
 
-		// NodeMesh::Private interface
-        Ptr<NodeLayout> GetLayout( int index ) const override;
-
 	};
 
 }
diff --git a/Engine/Plugins/Experimental/Mutable/Source/MutableTools/Private/MuT/NodeMeshSwitch.cpp b/Engine/Plugins/Experimental/Mutable/Source/MutableTools/Private/MuT/NodeMeshSwitch.cpp
--- a/Engine/Plugins/Experimental/Mutable/Source/MutableTools/Private/MuT/NodeMeshSwitch.cpp
+++ b/Engine/Plugins/Experimental/Mutable/Source/MutableTools/Private/MuT/NodeMeshSwitch.cpp
@@ -16,15 +16,14 @@ namespace mu
 	//---------------------------------------------------------------------------------------------
 	// Static initialisation
 	//---------------------------------------------------------------------------------------------
-	FNodeType NodeMeshSwitch::Private::s_type =
-			FNodeType( "MeshSwitch", NodeMesh::GetStaticType() );
+	FNodeType NodeMeshSwitch::Private::s_type = FNodeType(Node::EType::MeshSwitch, NodeMesh::GetStaticType() );
 
 
 	//---------------------------------------------------------------------------------------------
 	//!
 	//---------------------------------------------------------------------------------------------
 
-	MUTABLE_IMPLEMENT_NODE( NodeMeshSwitch, EType::Switch, Node, Node::EType::Mesh)
+	MUTABLE_IMPLEMENT_NODE( NodeMeshSwitch )
 
 
 	//---------------------------------------------------------------------------------------------
@@ -65,25 +64,6 @@ namespace mu
 		m_pD->m_options[t] = pNode;
 	}
 
-
-	//---------------------------------------------------------------------------------------------
-	NodeLayoutPtr NodeMeshSwitch::Private::GetLayout( int index ) const
-	{
-		for (int32 i=0; i<m_options.Num(); ++i)
-		{
-			if (m_options[i])
-			{
-				NodeMesh::Private* pPrivate = static_cast<NodeMesh::Private*>(m_options[i]->GetBasePrivate());
-				NodeLayoutPtr pResult = pPrivate->GetLayout(index);
-				return pResult;
-			}
-		}
-
-		return nullptr;
-	}
-
-
-
 }
 
 
diff --git a/Engine/Plugins/Experimental/Mutable/Source/MutableTools/Private/MuT/NodeMeshSwitchPrivate.h b/Engine/Plugins/Experimental/Mutable/Source/MutableTools/Private/MuT/NodeMeshSwitchPrivate.h
--- a/Engine/Plugins/Experimental/Mutable/Source/MutableTools/Private/MuT/NodeMeshSwitchPrivate.h
+++ b/Engine/Plugins/Experimental/Mutable/Source/MutableTools/Private/MuT/NodeMeshSwitchPrivate.h
@@ -20,8 +20,6 @@ namespace mu
 		NodeScalarPtr m_pParameter;
 		TArray<NodeMeshPtr> m_options;
 
-		// NodeMesh::Private interface
-        NodeLayoutPtr GetLayout( int index ) const override;
 	};
 
 
diff --git a/Engine/Plugins/Experimental/Mutable/Source/MutableTools/Private/MuT/NodeMeshTable.cpp b/Engine/Plugins/Experimental/Mutable/Source/MutableTools/Private/MuT/NodeMeshTable.cpp
--- a/Engine/Plugins/Experimental/Mutable/Source/MutableTools/Private/MuT/NodeMeshTable.cpp
+++ b/Engine/Plugins/Experimental/Mutable/Source/MutableTools/Private/MuT/NodeMeshTable.cpp
@@ -17,15 +17,14 @@ namespace mu
 	//---------------------------------------------------------------------------------------------
 	// Static initialisation
 	//---------------------------------------------------------------------------------------------
-	FNodeType NodeMeshTable::Private::s_type =
-			FNodeType( "TableMesh", NodeMesh::GetStaticType() );
+	FNodeType NodeMeshTable::Private::s_type = FNodeType(Node::EType::MeshTable, NodeMesh::GetStaticType() );
 
 
 	//---------------------------------------------------------------------------------------------
 	//!
 	//---------------------------------------------------------------------------------------------
 
-	MUTABLE_IMPLEMENT_NODE( NodeMeshTable, EType::Table, Node, Node::EType::Mesh);
+	MUTABLE_IMPLEMENT_NODE( NodeMeshTable );
 
 
 	//---------------------------------------------------------------------------------------------
@@ -75,8 +74,14 @@ namespace mu
 	//---------------------------------------------------------------------------------------------
 	NodeLayoutPtr NodeMeshTable::GetLayout( int i ) const
 	{
-		check( i>=0 && i<GetLayoutCount() );
-		return m_pD->GetLayout( i );
+		NodeLayoutPtr pResult;
+
+		if (i >= 0 && i < m_pD->Layouts.Num())
+		{
+			pResult = m_pD->Layouts[i];
+		}
+
+		return pResult;
 	}
 
 
@@ -88,20 +93,6 @@ namespace mu
 	}
 
 
-	//---------------------------------------------------------------------------------------------
-	NodeLayoutPtr NodeMeshTable::Private::GetLayout( int i ) const
-	{
-		NodeLayoutPtr pResult;
-
-		if ( i>=0 && i< Layouts.Num() )
-		{
-			pResult = Layouts[i];
-		}
-
-		return pResult;
-	}
-
-
 	//---------------------------------------------------------------------------------------------
 	void NodeMeshTable::SetNoneOption(bool bAddNoneOption)
 	{
diff --git a/Engine/Plugins/Experimental/Mutable/Source/MutableTools/Private/MuT/NodeMeshTangents.cpp b/Engine/Plugins/Experimental/Mutable/Source/MutableTools/Private/MuT/NodeMeshTangents.cpp
--- a/Engine/Plugins/Experimental/Mutable/Source/MutableTools/Private/MuT/NodeMeshTangents.cpp
+++ b/Engine/Plugins/Experimental/Mutable/Source/MutableTools/Private/MuT/NodeMeshTangents.cpp
@@ -16,15 +16,14 @@ namespace mu
 	//---------------------------------------------------------------------------------------------
 	// Static initialisation
 	//---------------------------------------------------------------------------------------------
-	FNodeType NodeMeshTangents::Private::s_type =
-			FNodeType( "MeshTangents", NodeMesh::GetStaticType() );
+	FNodeType NodeMeshTangents::Private::s_type = FNodeType(Node::EType::MeshTangents, NodeMesh::GetStaticType() );
 
 
 	//---------------------------------------------------------------------------------------------
 	//!
 	//---------------------------------------------------------------------------------------------
 
-	MUTABLE_IMPLEMENT_NODE( NodeMeshTangents, EType::Tangents, Node, Node::EType::Mesh)
+	MUTABLE_IMPLEMENT_NODE( NodeMeshTangents )
 
 
 	//---------------------------------------------------------------------------------------------
@@ -42,20 +41,4 @@ namespace mu
 		m_pD->m_pSource = p;
 	}
 
-
-	//---------------------------------------------------------------------------------------------
-	NodeLayoutPtr NodeMeshTangents::Private::GetLayout( int index ) const
-	{
-		NodeLayoutPtr pResult;
-
-		if ( m_pSource )
-		{
-			NodeMesh::Private* pPrivate = static_cast<NodeMesh::Private*>( m_pSource->GetBasePrivate() );
-			pResult = pPrivate->GetLayout( index );
-		}
-
-		return pResult;
-	}
-
-
 }
diff --git a/Engine/Plugins/Experimental/Mutable/Source/MutableTools/Private/MuT/NodeMeshTransform.cpp b/Engine/Plugins/Experimental/Mutable/Source/MutableTools/Private/MuT/NodeMeshTransform.cpp
--- a/Engine/Plugins/Experimental/Mutable/Source/MutableTools/Private/MuT/NodeMeshTransform.cpp
+++ b/Engine/Plugins/Experimental/Mutable/Source/MutableTools/Private/MuT/NodeMeshTransform.cpp
@@ -18,15 +18,14 @@ namespace mu
 	//---------------------------------------------------------------------------------------------
 	// Static initialisation
 	//---------------------------------------------------------------------------------------------
-    FNodeType NodeMeshTransform::Private::s_type =
-            FNodeType( "MeshTransform", NodeMesh::GetStaticType() );
+    FNodeType NodeMeshTransform::Private::s_type = FNodeType(Node::EType::MeshTransform, NodeMesh::GetStaticType() );
 
 
 	//---------------------------------------------------------------------------------------------
 	//!
 	//---------------------------------------------------------------------------------------------
 
-    MUTABLE_IMPLEMENT_NODE( NodeMeshTransform, EType::Transform, Node, Node::EType::Mesh)
+    MUTABLE_IMPLEMENT_NODE( NodeMeshTransform )
 
 
 	//---------------------------------------------------------------------------------------------
@@ -58,20 +57,4 @@ namespace mu
         return m_pD->Transform;
     }
 
-
-	//---------------------------------------------------------------------------------------------
-    NodeLayoutPtr NodeMeshTransform::Private::GetLayout( int index ) const
-	{
-		NodeLayoutPtr pResult;
-
-		if ( Source )
-		{
-			NodeMesh::Private* pPrivate = static_cast<NodeMesh::Private*>( Source->GetBasePrivate() );
-			pResult = pPrivate->GetLayout( index );
-		}
-
-		return pResult;
-	}
-
-
 }
diff --git a/Engine/Plugins/Experimental/Mutable/Source/MutableTools/Private/MuT/NodeMeshVariation.cpp b/Engine/Plugins/Experimental/Mutable/Source/MutableTools/Private/MuT/NodeMeshVariation.cpp
--- a/Engine/Plugins/Experimental/Mutable/Source/MutableTools/Private/MuT/NodeMeshVariation.cpp
+++ b/Engine/Plugins/Experimental/Mutable/Source/MutableTools/Private/MuT/NodeMeshVariation.cpp
@@ -13,15 +13,14 @@ namespace mu
     //---------------------------------------------------------------------------------------------
     // Static initialisation
     //---------------------------------------------------------------------------------------------
-    FNodeType NodeMeshVariation::Private::s_type =
-        FNodeType( "MeshVariation", NodeMesh::GetStaticType() );
+    FNodeType NodeMeshVariation::Private::s_type = FNodeType(Node::EType::MeshVariation, NodeMesh::GetStaticType() );
 
 
     //---------------------------------------------------------------------------------------------
     //---------------------------------------------------------------------------------------------
     //---------------------------------------------------------------------------------------------
 
-    MUTABLE_IMPLEMENT_NODE( NodeMeshVariation, EType::Variation, Node, Node::EType::Mesh)
+    MUTABLE_IMPLEMENT_NODE( NodeMeshVariation )
 
 
     //---------------------------------------------------------------------------------------------
diff --git a/Engine/Plugins/Experimental/Mutable/Source/MutableTools/Private/MuT/NodeMeshVariationPrivate.h b/Engine/Plugins/Experimental/Mutable/Source/MutableTools/Private/MuT/NodeMeshVariationPrivate.h
--- a/Engine/Plugins/Experimental/Mutable/Source/MutableTools/Private/MuT/NodeMeshVariationPrivate.h
+++ b/Engine/Plugins/Experimental/Mutable/Source/MutableTools/Private/MuT/NodeMeshVariationPrivate.h
@@ -18,11 +18,11 @@ namespace mu
 
         static FNodeType s_type;
 
-        NodeMeshPtr m_defaultMesh;
+        Ptr<NodeMesh> m_defaultMesh;
 
         struct FVariation
         {
-            NodeMeshPtr m_mesh;
+			Ptr<NodeMesh> m_mesh;
             FString m_tag;
         };
 
diff --git a/Engine/Plugins/Experimental/Mutable/Source/MutableTools/Private/MuT/NodeModifier.cpp b/Engine/Plugins/Experimental/Mutable/Source/MutableTools/Private/MuT/NodeModifier.cpp
--- a/Engine/Plugins/Experimental/Mutable/Source/MutableTools/Private/MuT/NodeModifier.cpp
+++ b/Engine/Plugins/Experimental/Mutable/Source/MutableTools/Private/MuT/NodeModifier.cpp
@@ -10,7 +10,7 @@
 namespace mu
 {
 	// Static initialisation
-	static FNodeType s_nodeModifierType = FNodeType( "NodeModifier", Node::GetStaticType() );
+	static FNodeType s_nodeModifierType = FNodeType(Node::EType::Modifier, Node::GetStaticType() );
 
 	MUTABLE_IMPLEMENT_ENUM_SERIALISABLE(EMutableMultipleTagPolicy);
 
@@ -27,22 +27,19 @@ namespace mu
 
 	void NodeModifier::AddTag(const FString& Value)
 	{
-		NodeModifier::Private* pD = static_cast<NodeModifier::Private*>(GetBasePrivate());
-		pD->RequiredTags.Add(Value);
+		RequiredTags.Add(Value);
 	}
 
 
 	void NodeModifier::SetMultipleTagPolicy(EMutableMultipleTagPolicy Value)
 	{
-		NodeModifier::Private* pD = static_cast<NodeModifier::Private*>(GetBasePrivate());
-		pD->MultipleTagsPolicy = Value;
+		MultipleTagsPolicy = Value;
 	}
 
 
 	void NodeModifier::SetStage(bool bBeforeNormalOperation)
 	{
-		NodeModifier::Private* pD = static_cast<NodeModifier::Private*>(GetBasePrivate());
-		pD->bApplyBeforeNormalOperations = bBeforeNormalOperation; 
+		bApplyBeforeNormalOperations = bBeforeNormalOperation; 
 	}
 
 }
diff --git a/Engine/Plugins/Experimental/Mutable/Source/MutableTools/Private/MuT/NodeModifierMeshClipDeform.cpp b/Engine/Plugins/Experimental/Mutable/Source/MutableTools/Private/MuT/NodeModifierMeshClipDeform.cpp
--- a/Engine/Plugins/Experimental/Mutable/Source/MutableTools/Private/MuT/NodeModifierMeshClipDeform.cpp
+++ b/Engine/Plugins/Experimental/Mutable/Source/MutableTools/Private/MuT/NodeModifierMeshClipDeform.cpp
@@ -16,15 +16,14 @@ namespace mu
 	//---------------------------------------------------------------------------------------------
 	// Static initialisation
 	//---------------------------------------------------------------------------------------------
-	FNodeType NodeModifierMeshClipDeform::Private::s_type =
-			FNodeType( "NodeModifierMeshClipDeform", NodeModifier::GetStaticType() );
+	FNodeType NodeModifierMeshClipDeform::Private::s_type = FNodeType(Node::EType::ModifierMeshClipDeform, NodeModifier::GetStaticType() );
 
 
 	//---------------------------------------------------------------------------------------------
 	//!
 	//---------------------------------------------------------------------------------------------
 
-	MUTABLE_IMPLEMENT_NODE( NodeModifierMeshClipDeform, EType::MeshClipDeform, Node, Node::EType::Modifier)
+	MUTABLE_IMPLEMENT_NODE( NodeModifierMeshClipDeform )
 
 
 	//---------------------------------------------------------------------------------------------
diff --git a/Engine/Plugins/Experimental/Mutable/Source/MutableTools/Private/MuT/NodeModifierMeshClipMorphPlane.cpp b/Engine/Plugins/Experimental/Mutable/Source/MutableTools/Private/MuT/NodeModifierMeshClipMorphPlane.cpp
--- a/Engine/Plugins/Experimental/Mutable/Source/MutableTools/Private/MuT/NodeModifierMeshClipMorphPlane.cpp
+++ b/Engine/Plugins/Experimental/Mutable/Source/MutableTools/Private/MuT/NodeModifierMeshClipMorphPlane.cpp
@@ -15,15 +15,14 @@ namespace mu
 	//---------------------------------------------------------------------------------------------
 	// Static initialisation
 	//---------------------------------------------------------------------------------------------
-    FNodeType NodeModifierMeshClipMorphPlane::Private::s_type =
-            FNodeType( "ModifierMeshClipMorphPlane", NodeModifier::GetStaticType() );
+    FNodeType NodeModifierMeshClipMorphPlane::Private::s_type = FNodeType(Node::EType::ModifierMeshClipMorphPlane, NodeModifier::GetStaticType() );
 
 
 	//---------------------------------------------------------------------------------------------
 	//!
 	//---------------------------------------------------------------------------------------------
 
-    MUTABLE_IMPLEMENT_NODE( NodeModifierMeshClipMorphPlane, EType::MeshClipMorphPlane, Node, Node::EType::Modifier)
+    MUTABLE_IMPLEMENT_NODE( NodeModifierMeshClipMorphPlane )
 
 
 	//---------------------------------------------------------------------------------------------
diff --git a/Engine/Plugins/Experimental/Mutable/Source/MutableTools/Private/MuT/NodeModifierMeshClipWithMesh.cpp b/Engine/Plugins/Experimental/Mutable/Source/MutableTools/Private/MuT/NodeModifierMeshClipWithMesh.cpp
--- a/Engine/Plugins/Experimental/Mutable/Source/MutableTools/Private/MuT/NodeModifierMeshClipWithMesh.cpp
+++ b/Engine/Plugins/Experimental/Mutable/Source/MutableTools/Private/MuT/NodeModifierMeshClipWithMesh.cpp
@@ -15,15 +15,14 @@ namespace mu
 	//---------------------------------------------------------------------------------------------
 	// Static initialisation
 	//---------------------------------------------------------------------------------------------
-    FNodeType NodeModifierMeshClipWithMesh::Private::s_type =
-            FNodeType( "NodeModifierMeshClipWithMesh", NodeModifier::GetStaticType() );
+    FNodeType NodeModifierMeshClipWithMesh::Private::s_type = FNodeType(Node::EType::ModifierMeshClipWithMesh, NodeModifier::GetStaticType() );
 
 
 	//---------------------------------------------------------------------------------------------
 	//!
 	//---------------------------------------------------------------------------------------------
 
-    MUTABLE_IMPLEMENT_NODE( NodeModifierMeshClipWithMesh, EType::MeshClipWithMesh, Node, Node::EType::Modifier)
+    MUTABLE_IMPLEMENT_NODE( NodeModifierMeshClipWithMesh )
 
 
 	//---------------------------------------------------------------------------------------------
diff --git a/Engine/Plugins/Experimental/Mutable/Source/MutableTools/Private/MuT/NodeModifierMeshClipWithUVMask.cpp b/Engine/Plugins/Experimental/Mutable/Source/MutableTools/Private/MuT/NodeModifierMeshClipWithUVMask.cpp
--- a/Engine/Plugins/Experimental/Mutable/Source/MutableTools/Private/MuT/NodeModifierMeshClipWithUVMask.cpp
+++ b/Engine/Plugins/Experimental/Mutable/Source/MutableTools/Private/MuT/NodeModifierMeshClipWithUVMask.cpp
@@ -15,15 +15,14 @@ namespace mu
 	//---------------------------------------------------------------------------------------------
 	// Static initialisation
 	//---------------------------------------------------------------------------------------------
-    FNodeType NodeModifierMeshClipWithUVMask::Private::s_type =
-            FNodeType( "NodeModifierMeshClipWithUVMask", NodeModifier::GetStaticType() );
+    FNodeType NodeModifierMeshClipWithUVMask::Private::s_type = FNodeType(Node::EType::ModifierMeshClipWithUVMask, NodeModifier::GetStaticType() );
 
 
 	//---------------------------------------------------------------------------------------------
 	//!
 	//---------------------------------------------------------------------------------------------
 
-    MUTABLE_IMPLEMENT_NODE( NodeModifierMeshClipWithUVMask, EType::MeshClipWithMesh, Node, Node::EType::Modifier)
+    MUTABLE_IMPLEMENT_NODE( NodeModifierMeshClipWithUVMask )
 
 
 	//---------------------------------------------------------------------------------------------
diff --git a/Engine/Plugins/Experimental/Mutable/Source/MutableTools/Private/MuT/NodeObject.cpp b/Engine/Plugins/Experimental/Mutable/Source/MutableTools/Private/MuT/NodeObject.cpp
--- a/Engine/Plugins/Experimental/Mutable/Source/MutableTools/Private/MuT/NodeObject.cpp
+++ b/Engine/Plugins/Experimental/Mutable/Source/MutableTools/Private/MuT/NodeObject.cpp
@@ -8,7 +8,7 @@
 namespace mu
 {
 
-	FNodeType s_nodeObjectType = FNodeType( "NodeObject", Node::GetStaticType() );
+	FNodeType s_nodeObjectType = FNodeType(Node::EType::Object, Node::GetStaticType() );
 
 
 	//---------------------------------------------------------------------------------------------
diff --git a/Engine/Plugins/Experimental/Mutable/Source/MutableTools/Private/MuT/NodeObjectGroup.cpp b/Engine/Plugins/Experimental/Mutable/Source/MutableTools/Private/MuT/NodeObjectGroup.cpp
--- a/Engine/Plugins/Experimental/Mutable/Source/MutableTools/Private/MuT/NodeObjectGroup.cpp
+++ b/Engine/Plugins/Experimental/Mutable/Source/MutableTools/Private/MuT/NodeObjectGroup.cpp
@@ -18,14 +18,14 @@ namespace mu
 	//---------------------------------------------------------------------------------------------
 	// Static initialisation
 	//---------------------------------------------------------------------------------------------
-	FNodeType NodeObjectGroup::Private::s_type = FNodeType( "ObjectGroup", NodeObject::GetStaticType() );
+	FNodeType NodeObjectGroup::Private::s_type = FNodeType(Node::EType::ObjectGroup, NodeObject::GetStaticType() );
 
 
 	//---------------------------------------------------------------------------------------------
 	//!
 	//---------------------------------------------------------------------------------------------
 
-	MUTABLE_IMPLEMENT_NODE( NodeObjectGroup, EType::Group, Node, Node::EType::Object)
+	MUTABLE_IMPLEMENT_NODE( NodeObjectGroup )
 
 
 	//---------------------------------------------------------------------------------------------
@@ -104,24 +104,6 @@ namespace mu
 	}
 
 
-	//---------------------------------------------------------------------------------------------
-    NodeLayoutPtr NodeObjectGroup::Private::GetLayout(int lod, int component, int surface, int texture ) const
-	{
-		NodeLayoutPtr pLayout;
-
-		for ( int32 i=0; !pLayout && i<m_children.Num(); ++i )
-		{
-			if (m_children[i])
-			{
-				NodeObject::Private* pPrivate = static_cast<NodeObject::Private*>( m_children[i]->GetBasePrivate() );
-                pLayout = pPrivate->GetLayout( lod, component, surface, texture );
-			}
-		}
-
-		// TODO: layout index
-		return pLayout;
-	}
-
 	void NodeObjectGroup::SetDefaultValue(int32 Value)
 	{
 		m_pD->DefaultValue = Value;
diff --git a/Engine/Plugins/Experimental/Mutable/Source/MutableTools/Private/MuT/NodeObjectNew.cpp b/Engine/Plugins/Experimental/Mutable/Source/MutableTools/Private/MuT/NodeObjectNew.cpp
--- a/Engine/Plugins/Experimental/Mutable/Source/MutableTools/Private/MuT/NodeObjectNew.cpp
+++ b/Engine/Plugins/Experimental/Mutable/Source/MutableTools/Private/MuT/NodeObjectNew.cpp
@@ -8,11 +8,9 @@
 #include "MuT/NodeComponent.h"
 #include "MuT/NodeComponentNew.h"
 #include "MuT/NodeLOD.h"
-#include "MuT/NodeLODPrivate.h"
 #include "MuT/NodeMesh.h"
 #include "MuT/NodeMeshConstant.h"
 #include "MuT/NodeMeshPrivate.h"
-#include "MuT/NodeObjectNewPrivate.h"
 #include "MuT/NodePrivate.h"
 #include "MuT/NodeSurface.h"
 #include "MuT/NodeSurfaceNew.h"
@@ -20,225 +18,83 @@
 
 namespace mu
 {
-
-
-	//---------------------------------------------------------------------------------------------
-	// Static initialisation
 	//---------------------------------------------------------------------------------------------
-	FNodeType NodeObjectNew::Private::s_type = FNodeType( "Object", NodeObject::GetStaticType() );
+	FNodeType NodeObjectNew::StaticType = FNodeType(Node::EType::ObjectNew, NodeObject::GetStaticType() );
 
 
-	//---------------------------------------------------------------------------------------------
-	//!
-	//---------------------------------------------------------------------------------------------
-
-	MUTABLE_IMPLEMENT_NODE( NodeObjectNew, EType::New, Node, Node::EType::Object )
-
-
-	//---------------------------------------------------------------------------------------------
-	// Own Interface
-	//---------------------------------------------------------------------------------------------
 	const FString& NodeObjectNew::GetName() const
 	{
-		return m_pD->m_name;
+		return Name;
 	}
 
 
-	//---------------------------------------------------------------------------------------------
-	void NodeObjectNew::SetName( const FString& Name )
+	void NodeObjectNew::SetName( const FString& InName )
 	{
-		m_pD->m_name = Name;
+		Name = InName;
 	}
 
 
 	const FString& NodeObjectNew::GetUid() const
 	{
-		return m_pD->m_uid;
-	}
-
-
-	//---------------------------------------------------------------------------------------------
-	void NodeObjectNew::SetUid( const FString& Uid )
-	{
-		m_pD->m_uid = Uid;
-	}
-
-
-	//---------------------------------------------------------------------------------------------
-	int NodeObjectNew::GetLODCount() const
-	{
-		return m_pD->m_lods.Num();
-	}
-
-
-	//---------------------------------------------------------------------------------------------
-	void NodeObjectNew::SetLODCount( int num )
-	{
-		check( num >=0 );
-		m_pD->m_lods.SetNum( num );
-	}
-
-
-	//---------------------------------------------------------------------------------------------
-	NodeLODPtr NodeObjectNew::GetLOD( int index ) const
-	{
-		check( index >=0 && index < m_pD->m_lods.Num() );
-
-		return m_pD->m_lods[ index ].get();
-	}
-
-
-	//---------------------------------------------------------------------------------------------
-	void NodeObjectNew::SetLOD( int index, NodeLODPtr pLOD )
-	{
-		check( index >=0 && index < m_pD->m_lods.Num() );
-
-		m_pD->m_lods[ index ] = pLOD;
-	}
-
-
-	//---------------------------------------------------------------------------------------------
-	int NodeObjectNew::GetChildCount() const
-	{
-		return m_pD->m_children.Num();
-	}
-
-
-	//---------------------------------------------------------------------------------------------
-	void NodeObjectNew::SetChildCount( int num )
-	{
-		check( num >=0 );
-		m_pD->m_children.SetNum( num );
-	}
-
-
-	//---------------------------------------------------------------------------------------------
-	NodeObjectPtr NodeObjectNew::GetChild( int index ) const
-	{
-		check( index >=0 && index < m_pD->m_children.Num() );
-
-		return m_pD->m_children[ index ].get();
-	}
-
-
-	//---------------------------------------------------------------------------------------------
-	void NodeObjectNew::SetChild( int index, NodeObjectPtr pObject )
-	{
-		check( index >=0 && index < m_pD->m_children.Num() );
-
-		m_pD->m_children[ index ] = pObject;
+		return Uid;
 	}
 
 
-	//---------------------------------------------------------------------------------------------
-    NodeLayoutPtr NodeObjectNew::Private::GetLayout( int lod,
-                                                     int component,
-                                                     int surface,
-                                                     int //texture
-                                                     ) const
+	void NodeObjectNew::SetUid( const FString& InUid )
 	{
-		check( lod>=0 && lod<m_lods.Num() );
-
-		NodeLayoutPtr pLayout;
-
-		const NodeComponent* pCompGeneric = m_lods[lod]->GetComponent( component ).get();
-
-        if (pCompGeneric->GetType()==NodeComponentNew::GetStaticType())
-        {
-			const NodeComponentNew* pComp = static_cast<const NodeComponentNew*> (pCompGeneric);
-			const NodeSurface* pSurfaceGeneric = pComp->GetSurface( surface );
-            if (pSurfaceGeneric->GetType()==NodeSurfaceNew::GetStaticType())
-            {
-				const NodeSurfaceNew* pSurface = static_cast<const NodeSurfaceNew*> (pSurfaceGeneric);
-				
-				// TODO: Look for the layout index of the given texture
-                // TODO: Multiple meshes
-                if ( pSurface->GetMeshCount()>0 )
-                {
-                    const NodeMesh* pMesh = pSurface->GetMesh( 0 ).get();
-                    if ( pMesh )
-                    {
-                        NodeMesh::Private* pPrivate = static_cast<NodeMesh::Private*>( pMesh->GetBasePrivate() );
-                        pLayout = pPrivate->GetLayout( 0 );
-                    }
-                }
-            }
-        }
-
-		// TODO: layout index
-		return pLayout;
+		Uid = InUid;
 	}
 
 
-    //---------------------------------------------------------------------------------------------
-    bool NodeObjectNew::Private::HasComponent( const NodeComponent* pComponent ) const
-    {
-        bool found = false;
-        for ( int32 l=0; !found && l<m_lods.Num(); ++l )
-        {
-			found = m_lods[l]->GetPrivate()->m_components.FindByPredicate([pComponent](const NodeComponentPtr& e) { return e.get() == pComponent; }) != nullptr;
-        }
-
-        return found;
-    }
-
-
-	//---------------------------------------------------------------------------------------------
 	int32 NodeObjectNew::GetStateCount() const
 	{
-		return m_pD->m_states.Num();
+		return States.Num();
 	}
 
 
-	//---------------------------------------------------------------------------------------------
 	void NodeObjectNew::SetStateCount( int32 c )
 	{
-		m_pD->m_states.SetNum( c );
+		States.SetNum( c );
 	}
 
 
-	//---------------------------------------------------------------------------------------------
 	void NodeObjectNew::SetStateName( int32 s, const FString& n )
 	{
 		check( s>=0 && s<GetStateCount() );
-		m_pD->m_states[s].m_name = n;
+		States[s].Name = n;
 	}
 
 
-	//---------------------------------------------------------------------------------------------
 	bool NodeObjectNew::HasStateParam( int32 s, const FString& param ) const
 	{
 		check( s>=0 && s<GetStateCount() );
-		return m_pD->m_states[s].m_runtimeParams.Contains( param );
+		return States[s].RuntimeParams.Contains( param );
 	}
 
 
-	//---------------------------------------------------------------------------------------------
 	void NodeObjectNew::AddStateParam( int32 s, const FString& param )
 	{
 		check( s>=0 && s<GetStateCount() );
 
 		if (!HasStateParam(s,param))
 		{
-			m_pD->m_states[s].m_runtimeParams.Add( param );
+			States[s].RuntimeParams.Add( param );
 		}
 	}
 
 
-	//---------------------------------------------------------------------------------------------
 	void NodeObjectNew::RemoveStateParam( int32 s, const FString& param )
 	{
 		check( s>=0 && s<GetStateCount() );
 
-		int32 it = m_pD->m_states[s].m_runtimeParams.Find( param );
+		int32 it = States[s].RuntimeParams.Find( param );
 		if ( it != INDEX_NONE )
 		{
-			m_pD->m_states[s].m_runtimeParams.RemoveAt( it );
+			States[s].RuntimeParams.RemoveAt( it );
 		}
 	}
 
 
-    //---------------------------------------------------------------------------------------------
     void NodeObjectNew::SetStateProperties( int32 StateIndex, 
 		ETextureCompressionStrategy TextureCompressionStrategy, 
 		bool bOnlyFirstLOD, 
@@ -247,19 +103,18 @@ namespace mu
     {
         check(StateIndex >=0 && StateIndex<GetStateCount() );
 
-		FStateOptimizationOptions& Data = m_pD->m_states[StateIndex].m_optimisation;
+		FStateOptimizationOptions& Data = States[StateIndex].Optimisation;
         Data.TextureCompressionStrategy = TextureCompressionStrategy;
 		Data.bOnlyFirstLOD = bOnlyFirstLOD;
 		Data.FirstLOD = FirstLOD;
 		Data.NumExtraLODsToBuildAfterFirstLOD = NumExtraLODsToBuildAfterFirstLOD;
     }
 
-    //---------------------------------------------------------------------------------------------
-	void NodeObjectNew::AddExtensionDataNode(NodeExtensionDataPtr Node, const FString& Name)
+	void NodeObjectNew::AddExtensionDataNode(Ptr<NodeExtensionData> Node, const FString& InName)
 	{
-		NodeObjectNew::Private::NamedExtensionDataNode& Entry = m_pD->m_extensionDataNodes.AddDefaulted_GetRef();
+		NodeObjectNew::FNamedExtensionDataNode& Entry = ExtensionDataNodes.AddDefaulted_GetRef();
 		Entry.Node = Node;
-		Entry.Name = Name;
+		Entry.Name = InName;
 	}
 }
 
diff --git a/Engine/Plugins/Experimental/Mutable/Source/MutableTools/Private/MuT/NodePatchImage.cpp b/Engine/Plugins/Experimental/Mutable/Source/MutableTools/Private/MuT/NodePatchImage.cpp
--- a/Engine/Plugins/Experimental/Mutable/Source/MutableTools/Private/MuT/NodePatchImage.cpp
+++ b/Engine/Plugins/Experimental/Mutable/Source/MutableTools/Private/MuT/NodePatchImage.cpp
@@ -15,14 +15,14 @@ namespace mu
     //---------------------------------------------------------------------------------------------
     // Static initialisation
     //---------------------------------------------------------------------------------------------
-    FNodeType NodePatchImage::Private::s_type = FNodeType( "PatchTexture", Node::GetStaticType() );
+    FNodeType NodePatchImage::Private::s_type = FNodeType(Node::EType::PatchImage, Node::GetStaticType() );
 
 
     //---------------------------------------------------------------------------------------------
     //!
     //---------------------------------------------------------------------------------------------
 
-    MUTABLE_IMPLEMENT_NODE( NodePatchImage, EType::PatchImage, Node, Node::EType::None)
+    MUTABLE_IMPLEMENT_NODE( NodePatchImage )
 
 
     //---------------------------------------------------------------------------------------------
diff --git a/Engine/Plugins/Experimental/Mutable/Source/MutableTools/Private/MuT/NodePatchMesh.cpp b/Engine/Plugins/Experimental/Mutable/Source/MutableTools/Private/MuT/NodePatchMesh.cpp
--- a/Engine/Plugins/Experimental/Mutable/Source/MutableTools/Private/MuT/NodePatchMesh.cpp
+++ b/Engine/Plugins/Experimental/Mutable/Source/MutableTools/Private/MuT/NodePatchMesh.cpp
@@ -15,15 +15,14 @@ namespace mu
 	//---------------------------------------------------------------------------------------------
 	// Static initialisation
 	//---------------------------------------------------------------------------------------------
-	FNodeType NodePatchMesh::Private::s_type =
-			FNodeType( "MeshPatch", Node::GetStaticType() );
+	FNodeType NodePatchMesh::Private::s_type = FNodeType(Node::EType::PatchMesh, Node::GetStaticType() );
 
 
 	//---------------------------------------------------------------------------------------------
 	//!
 	//---------------------------------------------------------------------------------------------
 
-	MUTABLE_IMPLEMENT_NODE( NodePatchMesh, EType::PatchMesh, Node, Node::EType::None);
+	MUTABLE_IMPLEMENT_NODE( NodePatchMesh );
 
 
 	//---------------------------------------------------------------------------------------------
diff --git a/Engine/Plugins/Experimental/Mutable/Source/MutableTools/Private/MuT/NodeProjector.cpp b/Engine/Plugins/Experimental/Mutable/Source/MutableTools/Private/MuT/NodeProjector.cpp
--- a/Engine/Plugins/Experimental/Mutable/Source/MutableTools/Private/MuT/NodeProjector.cpp
+++ b/Engine/Plugins/Experimental/Mutable/Source/MutableTools/Private/MuT/NodeProjector.cpp
@@ -8,7 +8,7 @@
 namespace mu
 {
 
-	static FNodeType s_nodeProjectorType = FNodeType( "NodeProjector", Node::GetStaticType() );
+	static FNodeType s_nodeProjectorType = FNodeType(Node::EType::Projector, Node::GetStaticType() );
 
 
 	//---------------------------------------------------------------------------------------------
diff --git a/Engine/Plugins/Experimental/Mutable/Source/MutableTools/Private/MuT/NodeProjectorConstant.cpp b/Engine/Plugins/Experimental/Mutable/Source/MutableTools/Private/MuT/NodeProjectorConstant.cpp
--- a/Engine/Plugins/Experimental/Mutable/Source/MutableTools/Private/MuT/NodeProjectorConstant.cpp
+++ b/Engine/Plugins/Experimental/Mutable/Source/MutableTools/Private/MuT/NodeProjectorConstant.cpp
@@ -16,15 +16,14 @@ namespace mu
 	//---------------------------------------------------------------------------------------------
 	// Static initialisation
 	//---------------------------------------------------------------------------------------------
-	FNodeType NodeProjectorConstant::Private::s_type =
-			FNodeType( "ProjectorConstant", NodeProjector::GetStaticType() );
+	FNodeType NodeProjectorConstant::Private::s_type = FNodeType(Node::EType::ProjectorConstant, NodeProjector::GetStaticType() );
 
 
 	//---------------------------------------------------------------------------------------------
 	//!
 	//---------------------------------------------------------------------------------------------
 
-	MUTABLE_IMPLEMENT_NODE( NodeProjectorConstant, EType::Constant, Node, Node::EType::Projector)
+	MUTABLE_IMPLEMENT_NODE( NodeProjectorConstant )
 
 
 	//---------------------------------------------------------------------------------------------
