[mutable] Redesign the object node hierarchy to better match unreal component and LOD concepts.

diff --git a/Engine/Plugins/Experimental/Mutable/Source/MutableTools/Private/MuT/NodeProjectorParameter.cpp b/Engine/Plugins/Experimental/Mutable/Source/MutableTools/Private/MuT/NodeProjectorParameter.cpp
--- a/Engine/Plugins/Experimental/Mutable/Source/MutableTools/Private/MuT/NodeProjectorParameter.cpp
+++ b/Engine/Plugins/Experimental/Mutable/Source/MutableTools/Private/MuT/NodeProjectorParameter.cpp
@@ -19,15 +19,14 @@ namespace mu
 	//---------------------------------------------------------------------------------------------
 	// Static initialisation
 	//---------------------------------------------------------------------------------------------
-	FNodeType NodeProjectorParameter::Private::s_type =
-			FNodeType( "ProjectorParameter", NodeProjector::GetStaticType() );
+	FNodeType NodeProjectorParameter::Private::s_type = FNodeType(Node::EType::ProjectorParameter, NodeProjector::GetStaticType() );
 
 
 	//---------------------------------------------------------------------------------------------
 	//!
 	//---------------------------------------------------------------------------------------------
 
-	MUTABLE_IMPLEMENT_NODE( NodeProjectorParameter, EType::Parameter, Node, Node::EType::Projector)
+	MUTABLE_IMPLEMENT_NODE( NodeProjectorParameter )
 
 
 	//---------------------------------------------------------------------------------------------
diff --git a/Engine/Plugins/Experimental/Mutable/Source/MutableTools/Private/MuT/NodeRange.cpp b/Engine/Plugins/Experimental/Mutable/Source/MutableTools/Private/MuT/NodeRange.cpp
--- a/Engine/Plugins/Experimental/Mutable/Source/MutableTools/Private/MuT/NodeRange.cpp
+++ b/Engine/Plugins/Experimental/Mutable/Source/MutableTools/Private/MuT/NodeRange.cpp
@@ -11,7 +11,7 @@ namespace mu
 	//---------------------------------------------------------------------------------------------
 	// Static initialisation
 	//---------------------------------------------------------------------------------------------
-    static FNodeType s_nodeRangeType = 	FNodeType( "NodeRange", Node::GetStaticType() );
+    static FNodeType s_nodeRangeType = 	FNodeType(Node::EType::Range, Node::GetStaticType() );
 
 
 	//---------------------------------------------------------------------------------------------
diff --git a/Engine/Plugins/Experimental/Mutable/Source/MutableTools/Private/MuT/NodeRangeFromScalar.cpp b/Engine/Plugins/Experimental/Mutable/Source/MutableTools/Private/MuT/NodeRangeFromScalar.cpp
--- a/Engine/Plugins/Experimental/Mutable/Source/MutableTools/Private/MuT/NodeRangeFromScalar.cpp
+++ b/Engine/Plugins/Experimental/Mutable/Source/MutableTools/Private/MuT/NodeRangeFromScalar.cpp
@@ -16,15 +16,14 @@ namespace mu
 	//---------------------------------------------------------------------------------------------
 	// Static initialisation
 	//---------------------------------------------------------------------------------------------
-    FNodeType NodeRangeFromScalar::Private::s_type =
-            FNodeType( "RangeFromScalar", NodeRange::GetStaticType() );
+    FNodeType NodeRangeFromScalar::Private::s_type = FNodeType(Node::EType::RangeFromScalar, NodeRange::GetStaticType() );
 
 
 	//---------------------------------------------------------------------------------------------
 	//!
 	//---------------------------------------------------------------------------------------------
 
-    MUTABLE_IMPLEMENT_NODE( NodeRangeFromScalar, EType::FromScalar, Node, Node::EType::Range)
+    MUTABLE_IMPLEMENT_NODE( NodeRangeFromScalar )
 
 
 	//---------------------------------------------------------------------------------------------
diff --git a/Engine/Plugins/Experimental/Mutable/Source/MutableTools/Private/MuT/NodeScalar.cpp b/Engine/Plugins/Experimental/Mutable/Source/MutableTools/Private/MuT/NodeScalar.cpp
--- a/Engine/Plugins/Experimental/Mutable/Source/MutableTools/Private/MuT/NodeScalar.cpp
+++ b/Engine/Plugins/Experimental/Mutable/Source/MutableTools/Private/MuT/NodeScalar.cpp
@@ -9,7 +9,7 @@
 namespace mu
 {
 
-	static FNodeType s_nodeScalarType = FNodeType( "NodeScalar", Node::GetStaticType() );
+	static FNodeType s_nodeScalarType = FNodeType(Node::EType::Scalar, Node::GetStaticType() );
 
 
 	//---------------------------------------------------------------------------------------------
diff --git a/Engine/Plugins/Experimental/Mutable/Source/MutableTools/Private/MuT/NodeScalarArithmeticOperation.cpp b/Engine/Plugins/Experimental/Mutable/Source/MutableTools/Private/MuT/NodeScalarArithmeticOperation.cpp
--- a/Engine/Plugins/Experimental/Mutable/Source/MutableTools/Private/MuT/NodeScalarArithmeticOperation.cpp
+++ b/Engine/Plugins/Experimental/Mutable/Source/MutableTools/Private/MuT/NodeScalarArithmeticOperation.cpp
@@ -18,15 +18,14 @@ namespace mu
 	//---------------------------------------------------------------------------------------------
 	// Static initialisation
 	//---------------------------------------------------------------------------------------------
-    FNodeType NodeScalarArithmeticOperation::Private::s_type =
-            FNodeType( "ScalarArithmenticOperation", NodeScalar::GetStaticType() );
+    FNodeType NodeScalarArithmeticOperation::Private::s_type = FNodeType(Node::EType::ScalarArithmeticOperation, NodeScalar::GetStaticType() );
 
 
 	//---------------------------------------------------------------------------------------------
 	//!
 	//---------------------------------------------------------------------------------------------
 
-    MUTABLE_IMPLEMENT_NODE( NodeScalarArithmeticOperation, EType::ArithmeticOperation, Node, Node::EType::Scalar)
+    MUTABLE_IMPLEMENT_NODE( NodeScalarArithmeticOperation )
 
 
     //---------------------------------------------------------------------------------------------
diff --git a/Engine/Plugins/Experimental/Mutable/Source/MutableTools/Private/MuT/NodeScalarConstant.cpp b/Engine/Plugins/Experimental/Mutable/Source/MutableTools/Private/MuT/NodeScalarConstant.cpp
--- a/Engine/Plugins/Experimental/Mutable/Source/MutableTools/Private/MuT/NodeScalarConstant.cpp
+++ b/Engine/Plugins/Experimental/Mutable/Source/MutableTools/Private/MuT/NodeScalarConstant.cpp
@@ -14,15 +14,14 @@ namespace mu
 	//---------------------------------------------------------------------------------------------
 	// Static initialisation
 	//---------------------------------------------------------------------------------------------
-	FNodeType NodeScalarConstant::Private::s_type =
-			FNodeType( "ScalarConstant", NodeScalar::GetStaticType() );
+	FNodeType NodeScalarConstant::Private::s_type = FNodeType(Node::EType::ScalarConstant, NodeScalar::GetStaticType() );
 
 
 	//---------------------------------------------------------------------------------------------
 	//!
 	//---------------------------------------------------------------------------------------------
 
-	MUTABLE_IMPLEMENT_NODE( NodeScalarConstant, EType::Constant, Node, Node::EType::Scalar)
+	MUTABLE_IMPLEMENT_NODE( NodeScalarConstant )
 
 
 	//---------------------------------------------------------------------------------------------
diff --git a/Engine/Plugins/Experimental/Mutable/Source/MutableTools/Private/MuT/NodeScalarCurve.cpp b/Engine/Plugins/Experimental/Mutable/Source/MutableTools/Private/MuT/NodeScalarCurve.cpp
--- a/Engine/Plugins/Experimental/Mutable/Source/MutableTools/Private/MuT/NodeScalarCurve.cpp
+++ b/Engine/Plugins/Experimental/Mutable/Source/MutableTools/Private/MuT/NodeScalarCurve.cpp
@@ -13,9 +13,9 @@
 namespace mu
 {
 
-    FNodeType NodeScalarCurve::Private::s_type = FNodeType( "Curve", NodeScalar::GetStaticType() );
+    FNodeType NodeScalarCurve::Private::s_type = FNodeType(Node::EType::ScalarCurve, NodeScalar::GetStaticType() );
 
-    MUTABLE_IMPLEMENT_NODE( NodeScalarCurve, EType::Curve, Node, Node::EType::Scalar)
+    MUTABLE_IMPLEMENT_NODE( NodeScalarCurve )
 
 }
 
diff --git a/Engine/Plugins/Experimental/Mutable/Source/MutableTools/Private/MuT/NodeScalarEnumParameter.cpp b/Engine/Plugins/Experimental/Mutable/Source/MutableTools/Private/MuT/NodeScalarEnumParameter.cpp
--- a/Engine/Plugins/Experimental/Mutable/Source/MutableTools/Private/MuT/NodeScalarEnumParameter.cpp
+++ b/Engine/Plugins/Experimental/Mutable/Source/MutableTools/Private/MuT/NodeScalarEnumParameter.cpp
@@ -14,15 +14,14 @@ namespace mu
 	//---------------------------------------------------------------------------------------------
 	// Static initialisation
 	//---------------------------------------------------------------------------------------------
-	FNodeType NodeScalarEnumParameter::Private::s_type =
-			FNodeType( "ScalarEnumParameter", NodeScalar::GetStaticType() );
+	FNodeType NodeScalarEnumParameter::Private::s_type = FNodeType(Node::EType::ScalarEnumParameter, NodeScalar::GetStaticType() );
 
 
 	//---------------------------------------------------------------------------------------------
 	//!
 	//---------------------------------------------------------------------------------------------
 
-	MUTABLE_IMPLEMENT_NODE( NodeScalarEnumParameter, EType::EnumParameter, Node, Node::EType::Scalar)
+	MUTABLE_IMPLEMENT_NODE( NodeScalarEnumParameter )
 
 
 	//---------------------------------------------------------------------------------------------
diff --git a/Engine/Plugins/Experimental/Mutable/Source/MutableTools/Private/MuT/NodeScalarParameter.cpp b/Engine/Plugins/Experimental/Mutable/Source/MutableTools/Private/MuT/NodeScalarParameter.cpp
--- a/Engine/Plugins/Experimental/Mutable/Source/MutableTools/Private/MuT/NodeScalarParameter.cpp
+++ b/Engine/Plugins/Experimental/Mutable/Source/MutableTools/Private/MuT/NodeScalarParameter.cpp
@@ -16,15 +16,14 @@ namespace mu
 	//---------------------------------------------------------------------------------------------
 	// Static initialisation
 	//---------------------------------------------------------------------------------------------
-	FNodeType NodeScalarParameter::Private::s_type =
-			FNodeType( "ScalarParameter", NodeScalar::GetStaticType() );
+	FNodeType NodeScalarParameter::Private::s_type = FNodeType(Node::EType::ScalarParameter, NodeScalar::GetStaticType() );
 
 
 	//---------------------------------------------------------------------------------------------
 	//!
 	//---------------------------------------------------------------------------------------------
 
-	MUTABLE_IMPLEMENT_NODE( NodeScalarParameter, EType::Parameter, Node, Node::EType::Scalar)
+	MUTABLE_IMPLEMENT_NODE( NodeScalarParameter )
 
 
 	//---------------------------------------------------------------------------------------------
diff --git a/Engine/Plugins/Experimental/Mutable/Source/MutableTools/Private/MuT/NodeScalarSwitch.cpp b/Engine/Plugins/Experimental/Mutable/Source/MutableTools/Private/MuT/NodeScalarSwitch.cpp
--- a/Engine/Plugins/Experimental/Mutable/Source/MutableTools/Private/MuT/NodeScalarSwitch.cpp
+++ b/Engine/Plugins/Experimental/Mutable/Source/MutableTools/Private/MuT/NodeScalarSwitch.cpp
@@ -15,15 +15,14 @@ namespace mu
 	//---------------------------------------------------------------------------------------------
 	// Static initialisation
 	//---------------------------------------------------------------------------------------------
-    FNodeType NodeScalarSwitch::Private::s_type =
-            FNodeType( "ScalarSwitch", NodeScalar::GetStaticType() );
+    FNodeType NodeScalarSwitch::Private::s_type = FNodeType(Node::EType::ScalarSwitch, NodeScalar::GetStaticType() );
 
 
 	//---------------------------------------------------------------------------------------------
 	//!
 	//---------------------------------------------------------------------------------------------
 
-    MUTABLE_IMPLEMENT_NODE( NodeScalarSwitch, EType::Switch, Node, Node::EType::Scalar)
+    MUTABLE_IMPLEMENT_NODE( NodeScalarSwitch )
 
 
 	//---------------------------------------------------------------------------------------------
diff --git a/Engine/Plugins/Experimental/Mutable/Source/MutableTools/Private/MuT/NodeScalarTable.cpp b/Engine/Plugins/Experimental/Mutable/Source/MutableTools/Private/MuT/NodeScalarTable.cpp
--- a/Engine/Plugins/Experimental/Mutable/Source/MutableTools/Private/MuT/NodeScalarTable.cpp
+++ b/Engine/Plugins/Experimental/Mutable/Source/MutableTools/Private/MuT/NodeScalarTable.cpp
@@ -16,15 +16,14 @@ namespace mu
 	//---------------------------------------------------------------------------------------------
 	// Static initialisation
 	//---------------------------------------------------------------------------------------------
-	FNodeType NodeScalarTable::Private::s_type =
-			FNodeType( "TableScalar", NodeScalar::GetStaticType() );
+	FNodeType NodeScalarTable::Private::s_type = FNodeType(Node::EType::ScalarTable, NodeScalar::GetStaticType() );
 
 
 	//---------------------------------------------------------------------------------------------
 	//!
 	//---------------------------------------------------------------------------------------------
 
-	MUTABLE_IMPLEMENT_NODE(NodeScalarTable, EType::Table, Node, Node::EType::Scalar)
+	MUTABLE_IMPLEMENT_NODE(NodeScalarTable )
 
 
 	//---------------------------------------------------------------------------------------------
diff --git a/Engine/Plugins/Experimental/Mutable/Source/MutableTools/Private/MuT/NodeScalarVariation.cpp b/Engine/Plugins/Experimental/Mutable/Source/MutableTools/Private/MuT/NodeScalarVariation.cpp
--- a/Engine/Plugins/Experimental/Mutable/Source/MutableTools/Private/MuT/NodeScalarVariation.cpp
+++ b/Engine/Plugins/Experimental/Mutable/Source/MutableTools/Private/MuT/NodeScalarVariation.cpp
@@ -14,15 +14,14 @@ namespace mu
     //---------------------------------------------------------------------------------------------
     // Static initialisation
     //---------------------------------------------------------------------------------------------
-    FNodeType NodeScalarVariation::Private::s_type =
-        FNodeType( "ScalarVariation", NodeScalar::GetStaticType() );
+    FNodeType NodeScalarVariation::Private::s_type = FNodeType(Node::EType::ScalarVariation, NodeScalar::GetStaticType() );
 
 
     //---------------------------------------------------------------------------------------------
     //---------------------------------------------------------------------------------------------
     //---------------------------------------------------------------------------------------------
 
-    MUTABLE_IMPLEMENT_NODE( NodeScalarVariation, EType::Variation, Node, Node::EType::Scalar)
+    MUTABLE_IMPLEMENT_NODE( NodeScalarVariation )
 
 
     //---------------------------------------------------------------------------------------------
diff --git a/Engine/Plugins/Experimental/Mutable/Source/MutableTools/Private/MuT/NodeString.cpp b/Engine/Plugins/Experimental/Mutable/Source/MutableTools/Private/MuT/NodeString.cpp
--- a/Engine/Plugins/Experimental/Mutable/Source/MutableTools/Private/MuT/NodeString.cpp
+++ b/Engine/Plugins/Experimental/Mutable/Source/MutableTools/Private/MuT/NodeString.cpp
@@ -12,7 +12,7 @@ namespace mu
 	//---------------------------------------------------------------------------------------------
 	// Static initialisation
 	//---------------------------------------------------------------------------------------------
-	static FNodeType s_nodeStringType = FNodeType( "NodeString", Node::GetStaticType() );
+	static FNodeType s_nodeStringType = FNodeType(Node::EType::String, Node::GetStaticType() );
 
 
 	//---------------------------------------------------------------------------------------------
diff --git a/Engine/Plugins/Experimental/Mutable/Source/MutableTools/Private/MuT/NodeStringConstant.cpp b/Engine/Plugins/Experimental/Mutable/Source/MutableTools/Private/MuT/NodeStringConstant.cpp
--- a/Engine/Plugins/Experimental/Mutable/Source/MutableTools/Private/MuT/NodeStringConstant.cpp
+++ b/Engine/Plugins/Experimental/Mutable/Source/MutableTools/Private/MuT/NodeStringConstant.cpp
@@ -14,15 +14,14 @@ namespace mu
 	//---------------------------------------------------------------------------------------------
 	// Static initialisation
 	//---------------------------------------------------------------------------------------------
-	FNodeType NodeStringConstant::Private::s_type =
-			FNodeType( "StringConstant", NodeString::GetStaticType() );
+	FNodeType NodeStringConstant::Private::s_type = FNodeType(Node::EType::StringConstant, NodeString::GetStaticType() );
 
 
 	//---------------------------------------------------------------------------------------------
 	//!
 	//---------------------------------------------------------------------------------------------
 
-	MUTABLE_IMPLEMENT_NODE( NodeStringConstant, EType::Constant, Node, Node::EType::String)
+	MUTABLE_IMPLEMENT_NODE( NodeStringConstant )
 
 
 	//---------------------------------------------------------------------------------------------
diff --git a/Engine/Plugins/Experimental/Mutable/Source/MutableTools/Private/MuT/NodeStringParameter.cpp b/Engine/Plugins/Experimental/Mutable/Source/MutableTools/Private/MuT/NodeStringParameter.cpp
--- a/Engine/Plugins/Experimental/Mutable/Source/MutableTools/Private/MuT/NodeStringParameter.cpp
+++ b/Engine/Plugins/Experimental/Mutable/Source/MutableTools/Private/MuT/NodeStringParameter.cpp
@@ -16,15 +16,14 @@ namespace mu
 	//---------------------------------------------------------------------------------------------
 	// Static initialisation
 	//---------------------------------------------------------------------------------------------
-	FNodeType NodeStringParameter::Private::s_type =
-			FNodeType( "StringParameter", NodeString::GetStaticType() );
+	FNodeType NodeStringParameter::Private::s_type = FNodeType(Node::EType::StringParameter, NodeString::GetStaticType() );
 
 
 	//---------------------------------------------------------------------------------------------
 	//!
 	//---------------------------------------------------------------------------------------------
 
-	MUTABLE_IMPLEMENT_NODE( NodeStringParameter, EType::Parameter, Node, Node::EType::String)
+	MUTABLE_IMPLEMENT_NODE( NodeStringParameter )
 
 
 	//---------------------------------------------------------------------------------------------
diff --git a/Engine/Plugins/Experimental/Mutable/Source/MutableTools/Private/MuT/NodeSurface.cpp b/Engine/Plugins/Experimental/Mutable/Source/MutableTools/Private/MuT/NodeSurface.cpp
--- a/Engine/Plugins/Experimental/Mutable/Source/MutableTools/Private/MuT/NodeSurface.cpp
+++ b/Engine/Plugins/Experimental/Mutable/Source/MutableTools/Private/MuT/NodeSurface.cpp
@@ -11,7 +11,7 @@ namespace mu
 	//---------------------------------------------------------------------------------------------
 	// Static initialisation
 	//---------------------------------------------------------------------------------------------
-    static FNodeType s_nodeSurfaceType = FNodeType( "NodeSurface", Node::GetStaticType() );
+    static FNodeType s_nodeSurfaceType = FNodeType(Node::EType::Surface, Node::GetStaticType() );
 
 
 	//---------------------------------------------------------------------------------------------
diff --git a/Engine/Plugins/Experimental/Mutable/Source/MutableTools/Private/MuT/NodeSurfaceEdit.cpp b/Engine/Plugins/Experimental/Mutable/Source/MutableTools/Private/MuT/NodeSurfaceEdit.cpp
--- a/Engine/Plugins/Experimental/Mutable/Source/MutableTools/Private/MuT/NodeSurfaceEdit.cpp
+++ b/Engine/Plugins/Experimental/Mutable/Source/MutableTools/Private/MuT/NodeSurfaceEdit.cpp
@@ -20,15 +20,14 @@ namespace mu
 	//---------------------------------------------------------------------------------------------
 	// Static initialisation
 	//---------------------------------------------------------------------------------------------
-    FNodeType NodeSurfaceEdit::Private::s_type =
-            FNodeType( "EditSurface", NodeSurface::GetStaticType() );
+    FNodeType NodeSurfaceEdit::Private::s_type = FNodeType(Node::EType::SurfaceEdit, NodeSurface::GetStaticType() );
 
 
 	//---------------------------------------------------------------------------------------------
 	//---------------------------------------------------------------------------------------------
 	//---------------------------------------------------------------------------------------------
 
-    MUTABLE_IMPLEMENT_NODE( NodeSurfaceEdit, EType::Edit, Node, Node::EType::Surface)
+    MUTABLE_IMPLEMENT_NODE( NodeSurfaceEdit )
 
 
 	//---------------------------------------------------------------------------------------------
diff --git a/Engine/Plugins/Experimental/Mutable/Source/MutableTools/Private/MuT/NodeSurfaceNew.cpp b/Engine/Plugins/Experimental/Mutable/Source/MutableTools/Private/MuT/NodeSurfaceNew.cpp
--- a/Engine/Plugins/Experimental/Mutable/Source/MutableTools/Private/MuT/NodeSurfaceNew.cpp
+++ b/Engine/Plugins/Experimental/Mutable/Source/MutableTools/Private/MuT/NodeSurfaceNew.cpp
@@ -22,15 +22,14 @@ namespace mu
 	//---------------------------------------------------------------------------------------------
 	// Static initialisation
 	//---------------------------------------------------------------------------------------------
-    FNodeType NodeSurfaceNew::Private::s_type =
-            FNodeType( "NewSurface", NodeSurface::GetStaticType() );
+    FNodeType NodeSurfaceNew::Private::s_type = FNodeType(Node::EType::SurfaceNew, NodeSurface::GetStaticType() );
 
 
 	//---------------------------------------------------------------------------------------------
 	//---------------------------------------------------------------------------------------------
 	//---------------------------------------------------------------------------------------------
 
-    MUTABLE_IMPLEMENT_NODE( NodeSurfaceNew, EType::New, Node, Node::EType::Surface)
+    MUTABLE_IMPLEMENT_NODE( NodeSurfaceNew )
 
 
 	//---------------------------------------------------------------------------------------------
@@ -57,53 +56,16 @@ namespace mu
 
 
 	//---------------------------------------------------------------------------------------------
-    int NodeSurfaceNew::GetMeshCount() const
+	Ptr<NodeMesh> NodeSurfaceNew::GetMesh() const
 	{
-		return (int)m_pD->m_meshes.Num();
+		return m_pD->Mesh;
 	}
 
 
 	//---------------------------------------------------------------------------------------------
-    void NodeSurfaceNew::SetMeshCount( int num )
+    void NodeSurfaceNew::SetMesh( Ptr<NodeMesh> InNode)
 	{
-		check( num >=0 );
-		m_pD->m_meshes.SetNum( num );
-	}
-
-
-	//---------------------------------------------------------------------------------------------
-    NodeMeshPtr NodeSurfaceNew::GetMesh( int index ) const
-	{
-		check( index >=0 && index < (int)m_pD->m_meshes.Num() );
-
-		return m_pD->m_meshes[ index ].m_pMesh.get();
-	}
-
-
-	//---------------------------------------------------------------------------------------------
-    void NodeSurfaceNew::SetMesh( int index, NodeMeshPtr pNode )
-	{
-		check( index >=0 && index < (int)m_pD->m_meshes.Num() );
-
-		m_pD->m_meshes[ index ].m_pMesh = pNode;
-	}
-
-
-	//---------------------------------------------------------------------------------------------
-    const FString& NodeSurfaceNew::GetMeshName( int index ) const
-	{
-		check( index >=0 && index < (int)m_pD->m_meshes.Num() );
-
-		return m_pD->m_meshes[ index ].m_name;
-	}
-
-
-	//---------------------------------------------------------------------------------------------
-    void NodeSurfaceNew::SetMeshName( int index, const FString& strName )
-	{
-		check( index >=0 && index < (int)m_pD->m_meshes.Num() );
-
-		m_pD->m_meshes[ index ].m_name = strName;
+		m_pD->Mesh = InNode;
 	}
 
 
@@ -359,21 +321,6 @@ namespace mu
 	}
 
 
-	//---------------------------------------------------------------------------------------------
-	int NodeSurfaceNew::Private::FindMesh(const FString& strName) const
-	{
-		for ( int32 i = 0; i < m_meshes.Num(); ++i)
-		{
-			if (m_meshes[i].m_name == strName)
-			{
-				return i;
-			}
-		}
-
-		return -1;
-	}
-
-
     //---------------------------------------------------------------------------------------------
     int NodeSurfaceNew::Private::FindVector(const FString& strName) const
     {
diff --git a/Engine/Plugins/Experimental/Mutable/Source/MutableTools/Private/MuT/NodeSurfaceSwitch.cpp b/Engine/Plugins/Experimental/Mutable/Source/MutableTools/Private/MuT/NodeSurfaceSwitch.cpp
--- a/Engine/Plugins/Experimental/Mutable/Source/MutableTools/Private/MuT/NodeSurfaceSwitch.cpp
+++ b/Engine/Plugins/Experimental/Mutable/Source/MutableTools/Private/MuT/NodeSurfaceSwitch.cpp
@@ -15,15 +15,14 @@ namespace mu
 	//---------------------------------------------------------------------------------------------
 	// Static initialisation
 	//---------------------------------------------------------------------------------------------
-	FNodeType NodeSurfaceSwitch::Private::s_type =
-			FNodeType( "SurfaceSwitch", NodeSurface::GetStaticType() );
+	FNodeType NodeSurfaceSwitch::Private::s_type = FNodeType(Node::EType::SurfaceSwitch, NodeSurface::GetStaticType() );
 
 
 	//---------------------------------------------------------------------------------------------
 	//!
 	//---------------------------------------------------------------------------------------------
 
-	MUTABLE_IMPLEMENT_NODE( NodeSurfaceSwitch, EType::Switch, Node, Node::EType::Surface)
+	MUTABLE_IMPLEMENT_NODE( NodeSurfaceSwitch )
 
 
 	//---------------------------------------------------------------------------------------------
diff --git a/Engine/Plugins/Experimental/Mutable/Source/MutableTools/Private/MuT/NodeSurfaceVariation.cpp b/Engine/Plugins/Experimental/Mutable/Source/MutableTools/Private/MuT/NodeSurfaceVariation.cpp
--- a/Engine/Plugins/Experimental/Mutable/Source/MutableTools/Private/MuT/NodeSurfaceVariation.cpp
+++ b/Engine/Plugins/Experimental/Mutable/Source/MutableTools/Private/MuT/NodeSurfaceVariation.cpp
@@ -17,15 +17,14 @@ namespace mu
 	//---------------------------------------------------------------------------------------------
 	// Static initialisation
 	//---------------------------------------------------------------------------------------------
-    FNodeType NodeSurfaceVariation::Private::s_type =
-            FNodeType( "SurfaceVariation", NodeSurface::GetStaticType() );
+    FNodeType NodeSurfaceVariation::Private::s_type = FNodeType(Node::EType::SurfaceVariation, NodeSurface::GetStaticType() );
 
 
 	//---------------------------------------------------------------------------------------------
 	//---------------------------------------------------------------------------------------------
 	//---------------------------------------------------------------------------------------------
 
-    MUTABLE_IMPLEMENT_NODE( NodeSurfaceVariation, EType::Variation, Node, Node::EType::Surface)
+    MUTABLE_IMPLEMENT_NODE( NodeSurfaceVariation )
 
 
 	//---------------------------------------------------------------------------------------------
diff --git a/Engine/Plugins/Experimental/Mutable/Source/MutableTools/Public/MuT/Compiler.h b/Engine/Plugins/Experimental/Mutable/Source/MutableTools/Public/MuT/Compiler.h
--- a/Engine/Plugins/Experimental/Mutable/Source/MutableTools/Public/MuT/Compiler.h
+++ b/Engine/Plugins/Experimental/Mutable/Source/MutableTools/Public/MuT/Compiler.h
@@ -198,18 +198,13 @@ namespace mu
 	struct FObjectState
 	{
 		//! Name used to identify the state from the code and user interface.
-		FString m_name;
+		FString Name;
 
 		//! GPU Optimisation options
-		FStateOptimizationOptions m_optimisation;
+		FStateOptimizationOptions Optimisation;
 
 		//! List of names of the runtime parameters in this state
-		TArray<FString> m_runtimeParams;
-
-		void Serialise( OutputArchive& arch ) const;
-
-
-		void Unserialise( InputArchive& arch );
+		TArray<FString> RuntimeParams;
 	};
 
 	
diff --git a/Engine/Plugins/Experimental/Mutable/Source/MutableTools/Public/MuT/Node.h b/Engine/Plugins/Experimental/Mutable/Source/MutableTools/Public/MuT/Node.h
--- a/Engine/Plugins/Experimental/Mutable/Source/MutableTools/Public/MuT/Node.h
+++ b/Engine/Plugins/Experimental/Mutable/Source/MutableTools/Public/MuT/Node.h
@@ -5,21 +5,8 @@
 #include "HAL/Platform.h"
 #include "MuR/RefCounted.h"
 
-
 namespace mu
 {
-	//! \defgroup tools Tools library classes
-	//! Tools library classes.
-
-	//! \defgroup model Model nodes
-	//! \ingroup tools
-	//! This group contains the nodes that can be used to compose models.
-
-	//! \defgroup transform Transform nodes
-	//! \ingroup tools
-	//! This group contains the nodes that can be used to compose model transformations.
-
-
 	// Forward declarations
 	class Node;
 	typedef Ptr<Node> NodePtr;
@@ -30,60 +17,136 @@ namespace mu
 	typedef Ptr<const NodeMap> NodeMapPtrConst;
 
 
-	/** Information about the type of a node, to provide some means to the tools to deal generically with nodes. */
-	struct FNodeType
-	{
-		FNodeType();
-		FNodeType( const char* strName, const FNodeType* pParent );
-
-		const char* m_strName;
-		const FNodeType* m_pParent;
-
-		inline bool IsA(const FNodeType* CandidateType) const
-		{
-			if (CandidateType == this)
-			{
-				return true;
-			}
-
-			if (m_pParent)
-			{
-				return m_pParent->IsA(CandidateType);
-			}
-
-			return false;
-		}
-	};
-
 
     //! %Base class for all graphs used in the source data to define models and transforms.
 	class MUTABLETOOLS_API Node : public RefCounted
 	{
 	public:
 
-		// Possible subclasses
+		/** Non-stable enumeration of all node types. */
 		enum class EType : uint8
 		{
-			Colour = 0,
-			Component = 1,
-			Image = 2,
-			Layout = 3,
-			LOD = 4,
-			Mesh = 5,
-			Object = 6,
-			PatchImage = 8,
-			Scalar = 9,
-			PatchMesh = 13,
-			Volume_Deprecated = 14,
-			Projector = 15,
-			Surface = 16,
-			Modifier = 18,
-			Range = 19,
-			String = 20,
-			Bool = 21,
-			ExtensionData = 22,
-
-			None
+			None,
+
+			Node,
+
+			Mesh,
+			MeshConstant,
+			MeshInterpolate,
+			MeshTable,
+			MeshFormat,
+			MeshTangents,
+			MeshMorph,
+			MeshMakeMorph,
+			MeshSwitch,
+			MeshFragment,
+			MeshTransform,
+			MeshClipMorphPlane,
+			MeshClipWithMesh,
+			MeshApplyPose,
+			MeshVariation,
+			MeshGeometryOperation,
+			MeshReshape,
+			MeshClipDeform,
+
+			Image,
+			ImageConstant,
+			ImageInterpolate,
+			ImageSaturate,
+			ImageTable,
+			ImageSwizzle,
+			ImageColorMap,
+			ImageGradient,
+			ImageBinarise,
+			ImageLuminance,
+			ImageLayer,
+			ImageLayerColour,
+			ImageResize,
+			ImagePlainColour,
+			ImageProject,
+			ImageMipmap,
+			ImageSwitch,
+			ImageConditional,
+			ImageFormat,
+			ImageParameter, 
+			ImageMultiLayer,
+			ImageInvert,
+			ImageVariation,
+			ImageNormalComposite,
+			ImageTransform,
+
+			Bool,
+			BoolConstant,
+			BoolParameter,
+			BoolNot,
+			BoolAnd,
+
+			Color,
+			ColorConstant,
+			ColorParameter,
+			ColorSampleImage,
+			ColorTable,
+			ColorImageSize,
+			ColorFromScalars,
+			ColorArithmeticOperation,
+			ColorSwitch,
+			ColorVariation,
+
+			Scalar,
+			ScalarConstant,
+			ScalarParameter,
+			ScalarEnumParameter,
+			ScalarCurve,
+			ScalarSwitch,
+			ScalarArithmeticOperation,
+			ScalarVariation,
+			ScalarTable,
+
+			String,
+			StringConstant,
+			StringParameter,
+
+			Projector,
+			ProjectorConstant,
+			ProjectorParameter,
+
+			Range,
+			RangeFromScalar,
+
+			Layout,
+			LayoutBlocks,
+
+			PatchImage,
+			PatchMesh,
+
+			Surface,
+			SurfaceNew,
+			SurfaceEdit,
+			SurfaceSwitch,
+			SurfaceVariation,
+
+			LOD,
+
+			Component,
+			ComponentNew,
+			ComponentEdit,
+
+			Object,
+			ObjectNew,
+			ObjectGroup,
+
+			Modifier,
+			ModifierMeshClipMorphPlane,
+			ModifierMeshClipWithMesh,
+			ModifierMeshClipDeform,
+			ModifierMeshClipWithUVMask,
+
+			ExtensionData,
+			ExtensionDataConstant,
+			ExtensionDataSwitch,
+			ExtensionDataVariation,
+
+			Count
 		};
 
 		//-----------------------------------------------------------------------------------------
@@ -91,17 +154,18 @@ namespace mu
 		//-----------------------------------------------------------------------------------------
 
 		/** Node type hierarchy data. */
-        virtual const FNodeType* GetType() const;
-		static const FNodeType* GetStaticType();
+        virtual const struct FNodeType* GetType() const;
+		static const struct FNodeType* GetStaticType();
 
 		/** Set the opaque context returned in messages in the compiler log. */
-		void SetMessageContext( const void* context );
+		void SetMessageContext(const void* context);
+		const void* GetMessageContext() const;
 
 		//-----------------------------------------------------------------------------------------
         // Interface pattern
 		//-----------------------------------------------------------------------------------------
 		class Private;
-		virtual Private* GetBasePrivate() const = 0;
+		//virtual Private* GetBasePrivate() const = 0;
 
 	protected:
 
@@ -110,7 +174,34 @@ namespace mu
 		//!
 		EType Type = EType::None;
 
+		/** This is an opaque context used to attach to reported error messages. */
+		const void* MessageContext = nullptr;
+
 	};
 
+	/** Information about the type of a node, to provide some means to the tools to deal generically with nodes. */
+	struct FNodeType
+	{
+		FNodeType();
+		FNodeType(Node::EType, const FNodeType* pParent);
+
+		Node::EType Type;
+		const FNodeType* m_pParent;
+
+		inline bool IsA(const FNodeType* CandidateType) const
+		{
+			if (CandidateType == this)
+			{
+				return true;
+			}
+
+			if (m_pParent)
+			{
+				return m_pParent->IsA(CandidateType);
+			}
+
+			return false;
+		}
+	};
 
 }
diff --git a/Engine/Plugins/Experimental/Mutable/Source/MutableTools/Public/MuT/NodeBool.h b/Engine/Plugins/Experimental/Mutable/Source/MutableTools/Public/MuT/NodeBool.h
--- a/Engine/Plugins/Experimental/Mutable/Source/MutableTools/Public/MuT/NodeBool.h
+++ b/Engine/Plugins/Experimental/Mutable/Source/MutableTools/Public/MuT/NodeBool.h
@@ -20,19 +20,6 @@ namespace mu
 	{
 	public:
 
-		// Possible subclasses
-		enum class EType : uint8
-		{
-			Constant = 0,
-			Parameter = 1,
-			DEPRECATED_IsNull = 2,
-			Not = 3,
-			And = 4,
-
-			None
-		};
-
-
 		//-----------------------------------------------------------------------------------------
 		// Node Interface
 		//-----------------------------------------------------------------------------------------
@@ -50,10 +37,6 @@ namespace mu
 
 		//! Forbidden. Manage with the Ptr<> template.
 		inline ~NodeBool() {}
-
-		//!
-		EType Type = EType::None;
-
 	};
 
 
@@ -90,7 +73,6 @@ namespace mu
 		//-----------------------------------------------------------------------------------------
 		class Private;
 		Private* GetPrivate() const;
-        Node::Private* GetBasePrivate() const override;
 
 	protected:
 
@@ -141,7 +123,6 @@ namespace mu
 		//-----------------------------------------------------------------------------------------
 		class Private;
 		Private* GetPrivate() const;
-        Node::Private* GetBasePrivate() const override;
 
 	protected:
 
@@ -185,7 +166,6 @@ namespace mu
 		//-----------------------------------------------------------------------------------------
 		class Private;
 		Private* GetPrivate() const;
-        Node::Private* GetBasePrivate() const override;
 
 	protected:
 
@@ -232,7 +212,6 @@ namespace mu
 		//-----------------------------------------------------------------------------------------
 		class Private;
 		Private* GetPrivate() const;
-        Node::Private* GetBasePrivate() const override;
 
 	protected:
 
diff --git a/Engine/Plugins/Experimental/Mutable/Source/MutableTools/Public/MuT/NodeColour.h b/Engine/Plugins/Experimental/Mutable/Source/MutableTools/Public/MuT/NodeColour.h
--- a/Engine/Plugins/Experimental/Mutable/Source/MutableTools/Public/MuT/NodeColour.h
+++ b/Engine/Plugins/Experimental/Mutable/Source/MutableTools/Public/MuT/NodeColour.h
@@ -22,23 +22,6 @@ namespace mu
     {
     public:
 
-		// Possible subclasses
-		enum class EType : uint8
-		{
-			Constant = 0,
-			Parameter = 1,
-			SampleImage = 2,
-			Table = 3,
-			ImageSize = 4,
-			FromScalars = 5,
-			ArithmeticOperation = 6,
-			Switch = 7,
-			Variation = 8,
-
-			None
-		};
-
-
         //-----------------------------------------------------------------------------------------
         // Node Interface
         //-----------------------------------------------------------------------------------------
@@ -56,10 +39,6 @@ namespace mu
 
         //! Forbidden. Manage with the Ptr<> template.
         inline ~NodeColour() {}
-
-		//!
-		EType Type = EType::None;
-
     };
 
 }
diff --git a/Engine/Plugins/Experimental/Mutable/Source/MutableTools/Public/MuT/NodeColourArithmeticOperation.h b/Engine/Plugins/Experimental/Mutable/Source/MutableTools/Public/MuT/NodeColourArithmeticOperation.h
--- a/Engine/Plugins/Experimental/Mutable/Source/MutableTools/Public/MuT/NodeColourArithmeticOperation.h
+++ b/Engine/Plugins/Experimental/Mutable/Source/MutableTools/Public/MuT/NodeColourArithmeticOperation.h
@@ -68,7 +68,6 @@ namespace mu
 		//-----------------------------------------------------------------------------------------
 		class Private;
 		Private* GetPrivate() const;
-        Node::Private* GetBasePrivate() const override;
 
 	protected:
 
diff --git a/Engine/Plugins/Experimental/Mutable/Source/MutableTools/Public/MuT/NodeComponent.h b/Engine/Plugins/Experimental/Mutable/Source/MutableTools/Public/MuT/NodeComponent.h
--- a/Engine/Plugins/Experimental/Mutable/Source/MutableTools/Public/MuT/NodeComponent.h
+++ b/Engine/Plugins/Experimental/Mutable/Source/MutableTools/Public/MuT/NodeComponent.h
@@ -6,76 +6,41 @@
 #include "MuR/Ptr.h"
 #include "MuR/RefCounted.h"
 #include "MuT/Node.h"
+#include "MuT/NodeLOD.h"
 
 
 namespace mu
 {
 
-	// Forward definitions
-	class NodeComponent;
-	typedef Ptr<NodeComponent> NodeComponentPtr;
-	typedef Ptr<const NodeComponent> NodeComponentPtrConst;
-
-    class NodeSurface;
-
-
-	//! This class is the parent of all nodes that output a component.
-	//! \ingroup model
+	/** */
 	class MUTABLETOOLS_API NodeComponent : public Node
 	{
 	public:
 
-		// Possible subclasses
-		enum class EType : uint8
-		{
-			New = 0,
-			Edit = 1,
+		TArray<Ptr<NodeLOD>> LODs;
 
-			None
-		};
+	public:
 
 		//-----------------------------------------------------------------------------------------
         // Node interface
 		//-----------------------------------------------------------------------------------------
-
-        const FNodeType* GetType() const override;
-		static const FNodeType* GetStaticType();
+		virtual const FNodeType* GetType() const override { return GetStaticType(); }
+		static const FNodeType* GetStaticType() { return &StaticType; }
 
         //-----------------------------------------------------------------------------------------
         // Own interface
         //-----------------------------------------------------------------------------------------
 
-        //! \name Surfaces
-        //! \{
-
-        //! Get the number of meshes in the component.
-        virtual int GetSurfaceCount() const = 0;
-
-        //! Set the number of meshes in the component.
-        virtual void SetSurfaceCount( int ) = 0;
-
-        //! Get the node generating one of the meshes in the component.
-        //! \param index index of the mesh, from 0 to GetSurfaceCount()-1
-        virtual NodeSurface* GetSurface( int index ) const = 0;
-
-        //! Set the node generating one of the meshes in the component.
-        //! \param index index of the mesh, from 0 to GetSurfaceCount()-1
-        virtual void SetSurface( int index, NodeSurface* ) = 0;
-
-        //! \}
-
-		//-----------------------------------------------------------------------------------------
-		// Interface pattern
-		//-----------------------------------------------------------------------------------------
-		class Private;
+		virtual const class NodeComponentNew* GetParentComponentNew() const = 0;
 
 	protected:
 
 		//! Forbidden. Manage with the Ptr<> template.
 		inline ~NodeComponent() {}
 
-		//!
-		EType Type = EType::None;
+	private:
+
+		static FNodeType StaticType;
 
 	};
 
diff --git a/Engine/Plugins/Experimental/Mutable/Source/MutableTools/Public/MuT/NodeComponentEdit.h b/Engine/Plugins/Experimental/Mutable/Source/MutableTools/Public/MuT/NodeComponentEdit.h
--- a/Engine/Plugins/Experimental/Mutable/Source/MutableTools/Public/MuT/NodeComponentEdit.h
+++ b/Engine/Plugins/Experimental/Mutable/Source/MutableTools/Public/MuT/NodeComponentEdit.h
@@ -11,66 +11,43 @@
 namespace mu
 {
 
-	// Forward definitions
-	class NodeComponentEdit;
-	typedef Ptr<NodeComponentEdit> NodeComponentEditPtr;
-	typedef Ptr<const NodeComponentEdit> NodeComponentEditPtrConst;
-
-    class NodeSurface;
-    typedef Ptr<NodeSurface> NodeSurfacePtr;
-    typedef Ptr<const NodeSurface> NodeSurfacePtrConst;
-
-
 	//! This node modifies a node of the parent object of the object that this node belongs to.
 	//! It allows to extend, cut and morph the parent component's meshes.
 	//! It also allows to patch the parent component's textures.
-	//! \ingroup model
 	class MUTABLETOOLS_API NodeComponentEdit : public NodeComponent
 	{
 	public:
 
-		NodeComponentEdit();
-
 		//-----------------------------------------------------------------------------------------
         // Node interface
 		//-----------------------------------------------------------------------------------------
 
-		const FNodeType* GetType() const override;
-		static const FNodeType* GetStaticType();
+		virtual const FNodeType* GetType() const override { return GetStaticType(); }
+		static const FNodeType* GetStaticType() { return &StaticType; }
 
         //-----------------------------------------------------------------------------------------
         // NodeComponent interface
         //-----------------------------------------------------------------------------------------
-        int GetSurfaceCount() const override;
-        void SetSurfaceCount( int ) override;
-        NodeSurface* GetSurface( int index ) const override;
-        void SetSurface( int index, NodeSurface* ) override;
+		virtual const class NodeComponentNew* GetParentComponentNew() const override;
 
 		//-----------------------------------------------------------------------------------------
         // Own interface
 		//-----------------------------------------------------------------------------------------
 
-		//! Set the parent component to modify. It should be a node from the same LOD of one of
-		//! the parent objects of this node's object.
+		//! Set the parent component to modify. It should be a node from the parent objects of this node's object.
         void SetParent( NodeComponent* );
         NodeComponent* GetParent() const;
 
-
-		//-----------------------------------------------------------------------------------------
-		// Interface pattern
-		//-----------------------------------------------------------------------------------------
-		class Private;
-		Private* GetPrivate() const;
-        Node::Private* GetBasePrivate() const override;
-
 	protected:
 
 		//! Forbidden. Manage with the Ptr<> template.
-		~NodeComponentEdit();
+		~NodeComponentEdit() {}
 
 	private:
 
-		Private* m_pD;
+		static FNodeType StaticType;
+
+		Ptr<NodeComponent> m_pParent;
 
 	};
 
diff --git a/Engine/Plugins/Experimental/Mutable/Source/MutableTools/Public/MuT/NodeComponentNew.h b/Engine/Plugins/Experimental/Mutable/Source/MutableTools/Public/MuT/NodeComponentNew.h
--- a/Engine/Plugins/Experimental/Mutable/Source/MutableTools/Public/MuT/NodeComponentNew.h
+++ b/Engine/Plugins/Experimental/Mutable/Source/MutableTools/Public/MuT/NodeComponentNew.h
@@ -12,76 +12,40 @@
 namespace mu
 {
 
-	// Forward definitions
-	class NodeComponent;
-	typedef Ptr<NodeComponent> NodeComponentPtr;
-	typedef Ptr<const NodeComponent> NodeComponentPtrConst;
-
-	class NodeComponentNew;
-	typedef Ptr<NodeComponentNew> NodeComponentNewPtr;
-	typedef Ptr<const NodeComponentNew> NodeComponentNewPtrConst;
-
-    class NodeSurface;
-    typedef Ptr<NodeSurface> NodeSurfacePtr;
-    typedef Ptr<const NodeSurface> NodeSurfacePtrConst;
-
-
-	//! This node makes a new component from several meshes and images.
-	//! \ingroup model
 	class MUTABLETOOLS_API NodeComponentNew : public NodeComponent
 	{
 	public:
 
-		NodeComponentNew();
+		/** Externally managed id assign to this component. */
+		uint16 Id = 0;
+
+	public:
 
 		//-----------------------------------------------------------------------------------------
         // Node interface
 		//-----------------------------------------------------------------------------------------		
 
-        const FNodeType* GetType() const override;
-		static const FNodeType* GetStaticType();
+		virtual const FNodeType* GetType() const override { return GetStaticType(); }
+		static const FNodeType* GetStaticType() { return &StaticType; }
 
 		//-----------------------------------------------------------------------------------------
-        // NodeComponent interface
-        //-----------------------------------------------------------------------------------------
-        int GetSurfaceCount() const override;
-        void SetSurfaceCount( int ) override;
-        NodeSurface* GetSurface( int index ) const override;
-        void SetSurface( int index, NodeSurface* ) override;
+		// NodeComponent interface
+		//-----------------------------------------------------------------------------------------		
+		virtual const class NodeComponentNew* GetParentComponentNew() const override { return this;  }
 
 		//-----------------------------------------------------------------------------------------
         // Own interface
 		//-----------------------------------------------------------------------------------------
 
-		//! Get the name of the component.
-		const FString& GetName() const;
-
-		//! Set the name of the component.
-		void SetName( const FString& );
-
-		//! Get the id of the component.
-		uint16 GetId() const;
-
-		//! Set the id of the component.
-		void SetId( uint16 );
-
-		//-----------------------------------------------------------------------------------------
-		// Interface pattern
-		//-----------------------------------------------------------------------------------------
-		class Private;
-		Private* GetPrivate() const;
-        Node::Private* GetBasePrivate() const override;
-
 	protected:
 
 		//! Forbidden. Manage with the Ptr<> template.
-		~NodeComponentNew();
+		~NodeComponentNew() {}
 
 	private:
 
-		Private* m_pD;
-
+		static FNodeType StaticType;
+		
 	};
 
-
 }
diff --git a/Engine/Plugins/Experimental/Mutable/Source/MutableTools/Public/MuT/NodeExtensionData.h b/Engine/Plugins/Experimental/Mutable/Source/MutableTools/Public/MuT/NodeExtensionData.h
--- a/Engine/Plugins/Experimental/Mutable/Source/MutableTools/Public/MuT/NodeExtensionData.h
+++ b/Engine/Plugins/Experimental/Mutable/Source/MutableTools/Public/MuT/NodeExtensionData.h
@@ -11,45 +11,27 @@
 namespace mu
 {
 
-	// Forward definitions
-	class NodeExtensionData;
-	typedef Ptr<NodeExtensionData> NodeExtensionDataPtr;
-	typedef Ptr<const NodeExtensionData> NodeExtensionDataPtrConst;
-
 
 	//! Node that evaluates to an ExtensionData
 	class MUTABLETOOLS_API NodeExtensionData : public Node
 	{
 	public:
-		// Possible subclasses
-		enum class EType : uint8
-		{
-			Constant = 0,
-			Switch,
-			Variation,
-			None
-		};
 
 		//-----------------------------------------------------------------------------------------
 		// Node Interface
 		//-----------------------------------------------------------------------------------------
-		const FNodeType* GetType() const override;
-		static const FNodeType* GetStaticType();
-
-		inline EType GetExtensionDataNodeType() const { return Type; }
-
-
-		//-----------------------------------------------------------------------------------------
-		// Interface pattern
-		//-----------------------------------------------------------------------------------------
-		class Private;
+		virtual const FNodeType* GetType() const override { return GetStaticType(); }
+		static const FNodeType* GetStaticType() { return &StaticType; }
 
 	protected:
 
 		//! Forbidden. Manage with the Ptr<> template.
 		inline ~NodeExtensionData() {}
 
-		EType Type = EType::None;
+	private:
+
+		static FNodeType StaticType;
+
 	};
 
 
diff --git a/Engine/Plugins/Experimental/Mutable/Source/MutableTools/Public/MuT/NodeExtensionDataConstant.h b/Engine/Plugins/Experimental/Mutable/Source/MutableTools/Public/MuT/NodeExtensionDataConstant.h
--- a/Engine/Plugins/Experimental/Mutable/Source/MutableTools/Public/MuT/NodeExtensionDataConstant.h
+++ b/Engine/Plugins/Experimental/Mutable/Source/MutableTools/Public/MuT/NodeExtensionDataConstant.h
@@ -11,48 +11,37 @@
 namespace mu
 {
 
-	// Forward definitions
-	class NodeExtensionDataConstant;
-	typedef Ptr<NodeExtensionDataConstant> NodeExtensionDataConstantPtr;
-	typedef Ptr<const NodeExtensionDataConstant> NodeExtensionDataConstantPtrConst;
-
 	//! Node that outputs a constant ExtensionData
 	//! \ingroup model
 	class MUTABLETOOLS_API NodeExtensionDataConstant : public NodeExtensionData
 	{
 	public:
 
-		NodeExtensionDataConstant();
+		Ptr<const ExtensionData> Value;
+
+	public:
 
 		//-----------------------------------------------------------------------------------------
 		// Node Interface
 		//-----------------------------------------------------------------------------------------
-		const FNodeType* GetType() const override;
-		static const FNodeType* GetStaticType();
+		virtual const FNodeType* GetType() const override { return GetStaticType(); }
+		static const FNodeType* GetStaticType() { return &StaticType; }
 
 		//-----------------------------------------------------------------------------------------
 		// Own Interface
 		//-----------------------------------------------------------------------------------------
 
-		//! Get the constant ExtensionData that will be returned.
-		ExtensionDataPtrConst GetValue() const;
-
-		//! Set the constant ExtensionData that will be returned.
-		void SetValue(ExtensionDataPtrConst Value);
-
-		//-----------------------------------------------------------------------------------------
-		// Interface pattern
-		//-----------------------------------------------------------------------------------------
-		class Private;
-		Private* GetPrivate() const;
-		Node::Private* GetBasePrivate() const override;
+		/** Deprecated. Access Value attribute directly. */
+		void SetValue(const Ptr<const ExtensionData>& In) { Value = In; }
 
 	protected:
+
 		//! Forbidden. Manage with the Ptr<> template.
-		~NodeExtensionDataConstant();
+		~NodeExtensionDataConstant() {}
 
 	private:
-		Private* m_pD;
+
+		static FNodeType StaticType;
 
 	};
 
diff --git a/Engine/Plugins/Experimental/Mutable/Source/MutableTools/Public/MuT/NodeExtensionDataSwitch.h b/Engine/Plugins/Experimental/Mutable/Source/MutableTools/Public/MuT/NodeExtensionDataSwitch.h
--- a/Engine/Plugins/Experimental/Mutable/Source/MutableTools/Public/MuT/NodeExtensionDataSwitch.h
+++ b/Engine/Plugins/Experimental/Mutable/Source/MutableTools/Public/MuT/NodeExtensionDataSwitch.h
@@ -6,55 +6,46 @@
 #include "MuR/Ptr.h"
 #include "MuR/RefCounted.h"
 #include "MuT/NodeExtensionData.h"
+#include "MuT/NodeScalar.h"
 
 namespace mu
 {
-	// Forward declarations
-	class NodeScalar;
-	using NodeScalarPtr = Ptr<NodeScalar>;
-	using NodeScalarPtrConst = Ptr<const NodeScalar>;
-
-	class NodeExtensionDataSwitch;
-	using NodeExtensionDataSwitchPtr = Ptr<NodeExtensionDataSwitch>;
-	using NodeExtensionDataSwitchPtrConst = Ptr<const NodeExtensionDataSwitch>;
 
 	class MUTABLETOOLS_API NodeExtensionDataSwitch : public NodeExtensionData
 	{
 	public:
 
-		NodeExtensionDataSwitch();
+		Ptr<NodeScalar> Parameter;
+		TArray<Ptr<NodeExtensionData>> Options;
+
+	public:
 
 		//-----------------------------------------------------------------------------------------
 		// Node Interface
 		//-----------------------------------------------------------------------------------------
-		const FNodeType* GetType() const override;
-		static const FNodeType* GetStaticType();
+		virtual const FNodeType* GetType() const override { return GetStaticType(); }
+		static const FNodeType* GetStaticType() { return &StaticType; }
 
 		//-----------------------------------------------------------------------------------------
 		// Own Interface
 		//-----------------------------------------------------------------------------------------
-		NodeScalarPtr GetParameter() const;
-		void SetParameter(NodeScalarPtr InParameter);
+		Ptr<NodeScalar> GetParameter() const;
+		void SetParameter(Ptr<NodeScalar> InParameter);
 
 		void SetOptionCount(int);
 
-		NodeExtensionDataPtr GetOption(int t) const;
-		void SetOption(int t, NodeExtensionDataPtr);
+		Ptr<NodeExtensionData> GetOption(int32 t) const;
+		void SetOption(int32 t, Ptr<NodeExtensionData>);
 
-		//-----------------------------------------------------------------------------------------
-		// Interface pattern
-		//-----------------------------------------------------------------------------------------
-		class Private;
-		Private* GetPrivate() const;
-		Node::Private* GetBasePrivate() const override;
 
 	protected:
 
 		//! Forbidden. Manage with the Ptr<> template.
-		~NodeExtensionDataSwitch();
+		~NodeExtensionDataSwitch() {}
 
 	private:
 
-		Private* m_pD;
+		static FNodeType StaticType;
+
 	};
 }
\ No newline at end of file
diff --git a/Engine/Plugins/Experimental/Mutable/Source/MutableTools/Public/MuT/NodeExtensionDataVariation.h b/Engine/Plugins/Experimental/Mutable/Source/MutableTools/Public/MuT/NodeExtensionDataVariation.h
--- a/Engine/Plugins/Experimental/Mutable/Source/MutableTools/Public/MuT/NodeExtensionDataVariation.h
+++ b/Engine/Plugins/Experimental/Mutable/Source/MutableTools/Public/MuT/NodeExtensionDataVariation.h
@@ -9,49 +9,50 @@
 
 namespace mu
 {
-	// Forward declarations
-	class NodeExtensionDataVariation;
-	using NodeExtensionDataVariationPtr = Ptr<NodeExtensionDataVariation>;
-	using NodeExtensionDataVariationPtrConst = Ptr<const NodeExtensionDataVariation>;
 
 	class MUTABLETOOLS_API NodeExtensionDataVariation : public NodeExtensionData
 	{
 	public:
 
-		NodeExtensionDataVariation();
+		Ptr<NodeExtensionData> DefaultValue;
+
+		struct FVariation
+		{
+			Ptr<NodeExtensionData> Value;
+			FString Tag;
+		};
+
+		TArray<FVariation> Variations;
+
+	public:
 
 		//-----------------------------------------------------------------------------------------
 		// Node Interface
 		//-----------------------------------------------------------------------------------------
-		const FNodeType* GetType() const override;
-		static const FNodeType* GetStaticType();
+		virtual const FNodeType* GetType() const override { return GetStaticType(); }
+		static const FNodeType* GetStaticType() { return &StaticType; }
 
 		//-----------------------------------------------------------------------------------------
 		// Own Interface
 		//-----------------------------------------------------------------------------------------
-		void SetDefaultValue(NodeExtensionDataPtr InValue);
+		void SetDefaultValue(Ptr<NodeExtensionData> InValue);
 
-		void SetVariationCount(int InCount);
+		void SetVariationCount(int32 InCount);
 		int GetVariationCount() const;
 
-		void SetVariationTag(int InIndex, const FString& InTag);
+		void SetVariationTag(int32 InIndex, const FString& InTag);
 
-		void SetVariationValue(int InIndex, NodeExtensionDataPtr InValue);
+		void SetVariationValue(int32 InIndex, Ptr<NodeExtensionData> InValue);
 
-		//-----------------------------------------------------------------------------------------
-		// Interface pattern
-		//-----------------------------------------------------------------------------------------
-		class Private;
-		Private* GetPrivate() const;
-		Node::Private* GetBasePrivate() const override;
 
 	protected:
 
 		//! Forbidden. Manage with the Ptr<> template.
-		~NodeExtensionDataVariation();
+		~NodeExtensionDataVariation() {}
 
 	private:
 
-		Private* m_pD;
+		static FNodeType StaticType;
+
 	};
 }
\ No newline at end of file
diff --git a/Engine/Plugins/Experimental/Mutable/Source/MutableTools/Public/MuT/NodeImage.h b/Engine/Plugins/Experimental/Mutable/Source/MutableTools/Public/MuT/NodeImage.h
--- a/Engine/Plugins/Experimental/Mutable/Source/MutableTools/Public/MuT/NodeImage.h
+++ b/Engine/Plugins/Experimental/Mutable/Source/MutableTools/Public/MuT/NodeImage.h
@@ -22,42 +22,6 @@ namespace mu
 	{
 	public:
 
-		// Possible subclasses
-		enum class EType : uint8
-		{
-			Constant = 1,
-			Difference_Deprecated = 2,
-			Interpolate = 4,
-			Saturate = 6,
-			Table = 9,
-			Swizzle = 10,
-			SelectColour_Deprecated = 11,
-			ColourMap = 12,
-			Gradient = 13,
-			VolumeLayer_Deprecated = 14,
-			Binarise = 17,
-			Luminance = 18,
-			Layer = 19,
-			LayerColour = 20,
-			Resize = 21,
-			PlainColour = 22,
-			Interpolate3 = 23,
-			Project = 24,
-			Mipmap = 25,
-			Switch = 26,
-			Conditional = 27,
-			Format = 28,
-			Parameter = 29,
-			MultiLayer = 30,
-			Invert = 31,
-			Variation = 32,
-			NormalComposite = 33,
-			Transform = 34,
-			Reference = 35,
-			None
-		};
-
-
 		//-----------------------------------------------------------------------------------------
 		// Node Interface
 		//-----------------------------------------------------------------------------------------
@@ -65,9 +29,6 @@ namespace mu
 		const FNodeType* GetType() const override;
 		static const FNodeType* GetStaticType();
 
-		inline EType GetImageNodeType() const { return Type; }
-
-
 		//-----------------------------------------------------------------------------------------
 		// Interface pattern
 		//-----------------------------------------------------------------------------------------
@@ -78,9 +39,6 @@ namespace mu
 		//! Forbidden. Manage with the Ptr<> template.
 		inline ~NodeImage() {}
 
-		//!
-		EType Type = EType::None;
-
 	};
 
 
diff --git a/Engine/Plugins/Experimental/Mutable/Source/MutableTools/Public/MuT/NodeImageColourMap.h b/Engine/Plugins/Experimental/Mutable/Source/MutableTools/Public/MuT/NodeImageColourMap.h
--- a/Engine/Plugins/Experimental/Mutable/Source/MutableTools/Public/MuT/NodeImageColourMap.h
+++ b/Engine/Plugins/Experimental/Mutable/Source/MutableTools/Public/MuT/NodeImageColourMap.h
@@ -55,9 +55,7 @@ namespace mu
 		// Interface pattern
 		//-----------------------------------------------------------------------------------------
 		class Private;
-
 		Private* GetPrivate() const;
-        Node::Private* GetBasePrivate() const override;
 
 	protected:
 
diff --git a/Engine/Plugins/Experimental/Mutable/Source/MutableTools/Public/MuT/NodeImageInvert.h b/Engine/Plugins/Experimental/Mutable/Source/MutableTools/Public/MuT/NodeImageInvert.h
--- a/Engine/Plugins/Experimental/Mutable/Source/MutableTools/Public/MuT/NodeImageInvert.h
+++ b/Engine/Plugins/Experimental/Mutable/Source/MutableTools/Public/MuT/NodeImageInvert.h
@@ -44,7 +44,6 @@ namespace mu
 		//-----------------------------------------------------------------------------------------
 		class Private;
 		Private* GetPrivate() const;
-		Node::Private*GetBasePrivate()const override;
 
 	protected:
 
diff --git a/Engine/Plugins/Experimental/Mutable/Source/MutableTools/Public/MuT/NodeImageNormalComposite.h b/Engine/Plugins/Experimental/Mutable/Source/MutableTools/Public/MuT/NodeImageNormalComposite.h
--- a/Engine/Plugins/Experimental/Mutable/Source/MutableTools/Public/MuT/NodeImageNormalComposite.h
+++ b/Engine/Plugins/Experimental/Mutable/Source/MutableTools/Public/MuT/NodeImageNormalComposite.h
@@ -51,7 +51,6 @@ namespace mu
 		//-----------------------------------------------------------------------------------------
 		class Private;
 		Private* GetPrivate() const;
-		Node::Private*GetBasePrivate()const override;
 
 	protected:
 
diff --git a/Engine/Plugins/Experimental/Mutable/Source/MutableTools/Public/MuT/NodeImageProject.h b/Engine/Plugins/Experimental/Mutable/Source/MutableTools/Public/MuT/NodeImageProject.h
--- a/Engine/Plugins/Experimental/Mutable/Source/MutableTools/Public/MuT/NodeImageProject.h
+++ b/Engine/Plugins/Experimental/Mutable/Source/MutableTools/Public/MuT/NodeImageProject.h
@@ -81,9 +81,7 @@ namespace mu
 		// Interface pattern
 		//-----------------------------------------------------------------------------------------
 		class Private;
-
 		Private* GetPrivate() const;
-        Node::Private* GetBasePrivate() const override;
 
 	protected:
 
diff --git a/Engine/Plugins/Experimental/Mutable/Source/MutableTools/Public/MuT/NodeImageSwizzle.h b/Engine/Plugins/Experimental/Mutable/Source/MutableTools/Public/MuT/NodeImageSwizzle.h
--- a/Engine/Plugins/Experimental/Mutable/Source/MutableTools/Public/MuT/NodeImageSwizzle.h
+++ b/Engine/Plugins/Experimental/Mutable/Source/MutableTools/Public/MuT/NodeImageSwizzle.h
@@ -48,14 +48,11 @@ namespace mu
 		void SetSource( int, NodeImagePtr );
 		void SetSourceChannel( int, int );
 
-
 		//-----------------------------------------------------------------------------------------
 		// Interface pattern
 		//-----------------------------------------------------------------------------------------
 		class Private;
-
 		Private* GetPrivate() const;
-        Node::Private* GetBasePrivate() const override;
 
 	protected:
 
diff --git a/Engine/Plugins/Experimental/Mutable/Source/MutableTools/Public/MuT/NodeLOD.h b/Engine/Plugins/Experimental/Mutable/Source/MutableTools/Public/MuT/NodeLOD.h
--- a/Engine/Plugins/Experimental/Mutable/Source/MutableTools/Public/MuT/NodeLOD.h
+++ b/Engine/Plugins/Experimental/Mutable/Source/MutableTools/Public/MuT/NodeLOD.h
@@ -5,82 +5,39 @@
 #include "MuR/Ptr.h"
 #include "MuR/RefCounted.h"
 #include "MuT/Node.h"
-
+#include "MuT/NodeSurface.h"
+#include "MuT/NodeModifier.h"
 
 namespace mu
 {
 
-	// Forward definitions
-	class NodeComponent;
-	typedef Ptr<NodeComponent> NodeComponentPtr;
-	typedef Ptr<const NodeComponent> NodeComponentPtrConst;
-
-	class NodeModifier;
-	typedef Ptr<NodeModifier> NodeModifierPtr;
-	typedef Ptr<const NodeModifier> NodeModifierPtrConst;
-
-	class NodeLOD;
-	typedef Ptr<NodeLOD> NodeLODPtr;
-	typedef Ptr<const NodeLOD> NodeLODPtrConst;
-
 
-	//! Node that creates a new level of detail by assembling several object components.
 	class MUTABLETOOLS_API NodeLOD : public Node
 	{
 	public:
 
-		NodeLOD();
-
 		//-----------------------------------------------------------------------------------------
 		// Node Interface
 		//-----------------------------------------------------------------------------------------
 
-		const FNodeType* GetType() const override;
-		static const FNodeType* GetStaticType();
+		virtual const FNodeType* GetType() const override { return GetStaticType(); }
+		static const FNodeType* GetStaticType() { return &StaticType; }
 
 		//-----------------------------------------------------------------------------------------
 		// Own Interface
 		//-----------------------------------------------------------------------------------------
 
-		//! Get the number of components in the LOD to create
-		int GetComponentCount() const;
-
-		//! Set the number of components in the LOD to create
-		void SetComponentCount( int );
-
-		//! Get the node generating one of the components of the LOD.
-		//! \param index index of the component from 0 to GetComponentCount()-1
-		NodeComponentPtr GetComponent( int index ) const;
-
-		//! Set the node generating one of the components of the LOD.
-		//! \param index index of the component from 0 to GetComponentCount()-1
-		void SetComponent( int index, NodeComponentPtr );
-
-		//! Set the number of modifiers in the LOD
-		void SetModifierCount(int);
-
-		//! Get the number of modifiers in the LOD
-		int GetModifierCount() const;
-
-		//! Set the node generating one of the modifiers of the LOD.
-		//! \param index index of the component from 0 to GetModifierCount()-1
-		void SetModifier(int index, NodeModifierPtr);
-
-		//-----------------------------------------------------------------------------------------
-		// Interface pattern
-		//-----------------------------------------------------------------------------------------
-		class Private;
-		Private* GetPrivate() const;
-        Node::Private* GetBasePrivate() const override;
+		TArray<Ptr<NodeSurface>> Surfaces;
+		TArray<Ptr<NodeModifier>> Modifiers;
 
 	protected:
 
 		//! Forbidden. Manage with the Ptr<> template.
-		~NodeLOD();
+		~NodeLOD() {}
 
 	private:
 
-		Private* m_pD;
+		static FNodeType StaticType;
 
 	};
 
diff --git a/Engine/Plugins/Experimental/Mutable/Source/MutableTools/Public/MuT/NodeLayout.h b/Engine/Plugins/Experimental/Mutable/Source/MutableTools/Public/MuT/NodeLayout.h
--- a/Engine/Plugins/Experimental/Mutable/Source/MutableTools/Public/MuT/NodeLayout.h
+++ b/Engine/Plugins/Experimental/Mutable/Source/MutableTools/Public/MuT/NodeLayout.h
@@ -30,14 +30,6 @@ namespace mu
 	{
 	public:
 
-		// Possible subclasses
-		enum class EType : uint8
-		{
-			Blocks = 0,
-
-			None
-		};
-
 		//-----------------------------------------------------------------------------------------
 		// Node Interface
 		//-----------------------------------------------------------------------------------------
@@ -56,9 +48,6 @@ namespace mu
 		//! Forbidden. Manage with the Ptr<> template.
 		~NodeLayout() {}
 
-		//!
-		EType Type = EType::None;
-
 	};
 
 
@@ -135,7 +124,6 @@ namespace mu
 		//-----------------------------------------------------------------------------------------
 		class Private;
 		Private* GetPrivate() const;
-        Node::Private* GetBasePrivate() const override;
 
 	protected:
 
diff --git a/Engine/Plugins/Experimental/Mutable/Source/MutableTools/Public/MuT/NodeMesh.h b/Engine/Plugins/Experimental/Mutable/Source/MutableTools/Public/MuT/NodeMesh.h
--- a/Engine/Plugins/Experimental/Mutable/Source/MutableTools/Public/MuT/NodeMesh.h
+++ b/Engine/Plugins/Experimental/Mutable/Source/MutableTools/Public/MuT/NodeMesh.h
@@ -21,41 +21,12 @@ namespace mu
 	{
 	public:
 
-		// Possible subclasses
-		enum class EType : uint8
-		{
-			Constant = 0,
-			Interpolate = 2,
-			Table = 4,
-			Subtract_Deprecated = 6,
-			Format = 7,
-			Tangents = 8,
-			Morph = 9,
-			MakeMorph = 10,
-			Switch = 11,
-			Fragment = 12,
-			Transform = 13,
-			ClipMorphPlane = 14,
-			ClipWithMesh = 15,
-			ApplyPose = 16,
-			Variation = 17,
-			GeometryOperation = 18,
-			Reshape = 19,
-			ClipDeform = 20,
-
-			None
-		};
-
 		//-----------------------------------------------------------------------------------------
 		// Node Interface
 		//-----------------------------------------------------------------------------------------
-
         const FNodeType* GetType() const override;
 		static const FNodeType* GetStaticType();
 
-		inline EType GetMeshNodeType() const { return Type; }
-
-
 		//-----------------------------------------------------------------------------------------
 		// Interface pattern
 		//-----------------------------------------------------------------------------------------
@@ -65,9 +36,6 @@ namespace mu
 
 		//! Forbidden. Manage with the Ptr<> template.
 		inline ~NodeMesh() {}
-
-		//!
-		EType Type = EType::None;
 	};
 
 
diff --git a/Engine/Plugins/Experimental/Mutable/Source/MutableTools/Public/MuT/NodeModifier.h b/Engine/Plugins/Experimental/Mutable/Source/MutableTools/Public/MuT/NodeModifier.h
--- a/Engine/Plugins/Experimental/Mutable/Source/MutableTools/Public/MuT/NodeModifier.h
+++ b/Engine/Plugins/Experimental/Mutable/Source/MutableTools/Public/MuT/NodeModifier.h
@@ -36,16 +36,6 @@ namespace mu
 	{
 	public:
 
-		// Possible subclasses
-		enum class EType : uint8
-		{
-			MeshClipMorphPlane = 0,
-			MeshClipWithMesh = 1,
-			MeshClipDeform = 2,
-			
-			None
-		};
-
 		//-----------------------------------------------------------------------------------------
         // Node interface
 		//-----------------------------------------------------------------------------------------
@@ -66,6 +56,15 @@ namespace mu
 		/** Set the stage to apply this modifier in.Default is before normal operations. */
 		void SetStage( bool bBeforeNormalOperation );
 
+		/** Tags that target surface need to have enabled to receive this modifier. */
+		TArray<FString> RequiredTags;
+
+		/** In case of multiple tags in RequiredTags: are they all required, or one is enough? */
+		EMutableMultipleTagPolicy MultipleTagsPolicy = EMutableMultipleTagPolicy::OnlyOneRequired;
+
+		// Wether the modifier has to be applied after the normal node operations or before
+		bool bApplyBeforeNormalOperations = true;
+
 		//-----------------------------------------------------------------------------------------
 		// Interface pattern
 		//-----------------------------------------------------------------------------------------
@@ -76,9 +75,6 @@ namespace mu
 		//! Forbidden. Manage with the Ptr<> template.
 		inline ~NodeModifier() {}
 
-		//!
-		EType Type = EType::None;
-
 	};
 
 
diff --git a/Engine/Plugins/Experimental/Mutable/Source/MutableTools/Public/MuT/NodeObject.h b/Engine/Plugins/Experimental/Mutable/Source/MutableTools/Public/MuT/NodeObject.h
--- a/Engine/Plugins/Experimental/Mutable/Source/MutableTools/Public/MuT/NodeObject.h
+++ b/Engine/Plugins/Experimental/Mutable/Source/MutableTools/Public/MuT/NodeObject.h
@@ -23,16 +23,6 @@ namespace mu
 	{
 	public:
 
-		// Possible subclasses
-		enum class EType : uint8
-		{
-			New = 0,
-			Group = 1,
-
-			None
-		};
-
-
 		//-----------------------------------------------------------------------------------------
 		// Node Interface
 		//-----------------------------------------------------------------------------------------
@@ -66,9 +56,6 @@ namespace mu
 		//! Forbidden. Manage with the Ptr<> template.
 		inline ~NodeObject() {}
 
-		//!
-		EType Type = EType::None;
-
 	};
 
 
diff --git a/Engine/Plugins/Experimental/Mutable/Source/MutableTools/Public/MuT/NodeObjectNew.h b/Engine/Plugins/Experimental/Mutable/Source/MutableTools/Public/MuT/NodeObjectNew.h
--- a/Engine/Plugins/Experimental/Mutable/Source/MutableTools/Public/MuT/NodeObjectNew.h
+++ b/Engine/Plugins/Experimental/Mutable/Source/MutableTools/Public/MuT/NodeObjectNew.h
@@ -1,31 +1,56 @@
 // Copyright Epic Games, Inc. All Rights Reserved.
 
 #pragma once
+
 #include "MuR/Ptr.h"
 #include "MuR/RefCounted.h"
 #include "MuR/System.h"
 #include "MuT/Node.h"
 #include "MuT/NodeObject.h"
+#include "MuT/NodeComponent.h"
 #include "MuT/NodeExtensionData.h"
+#include "MuT/Compiler.h"
 
 
 namespace mu
 {
-	class NodeLOD;
 
 	//! Node that creates a new object by setting its levels-of-detail and its children.
 	class MUTABLETOOLS_API NodeObjectNew : public NodeObject
 	{
 	public:
 
-		NodeObjectNew();
+		/** Name of the object. */
+		FString Name;
+
+		/** Externally provided id for the object. */
+		FString Uid;
+
+		/** Components defined in the object. */
+		TArray<Ptr<NodeComponent>> Components;
+
+		/** Children objects. */
+		TArray<Ptr<NodeObject>> Children;
 
+		/** Extension data attached to this object. */
+		struct FNamedExtensionDataNode
+		{
+			Ptr<NodeExtensionData> Node;
+			FString Name;
+		};
+		TArray<FNamedExtensionDataNode> ExtensionDataNodes;
+
+		/** States defined in this object. */
+		TArray<FObjectState> States;
+
+
+	public:
 
 		//-----------------------------------------------------------------------------------------
 		// Node Interface
 		//-----------------------------------------------------------------------------------------
-        const FNodeType* GetType() const override;
-		static const FNodeType* GetStaticType();
+		virtual const FNodeType* GetType() const override { return GetStaticType(); }
+		static const FNodeType* GetStaticType() { return &StaticType; }
 
 		//-----------------------------------------------------------------------------------------
 		// NodeObject Interface
@@ -39,22 +64,6 @@ namespace mu
 		// Own Interface
 		//-----------------------------------------------------------------------------------------
 
-		//! Get the number of levels of detail in the object.
-		int GetLODCount() const;
-		void SetLODCount( int );
-
-		//! Get a level of detail node from the object.
-		Ptr<NodeLOD> GetLOD( int index ) const;
-		void SetLOD( int index, Ptr<NodeLOD>);
-
-		//! Get the number of child objects
-		int GetChildCount() const;
-		void SetChildCount( int );
-
-		//! Get a child object node
-		NodeObjectPtr GetChild( int index ) const;
-		void SetChild( int index, NodeObjectPtr );
-
 		//! Set the number of states that the model can be in.
 		int32 GetStateCount() const;
 		void SetStateCount( int32 c );
@@ -79,23 +88,17 @@ namespace mu
 			uint8 NumExtraLODsToBuildAfterFirstLOD);
 
 		//! Connect a node that produces ExtensionData to be added to the final Instance, and provide a name to associate with the data
-		void AddExtensionDataNode(NodeExtensionDataPtr Node, const FString& Name);
-
-		//-----------------------------------------------------------------------------------------
-		// Interface pattern
-		//-----------------------------------------------------------------------------------------
-		class Private;
-		Private* GetPrivate() const;
-        Node::Private* GetBasePrivate() const override;
+		void AddExtensionDataNode(Ptr<NodeExtensionData> Node, const FString& Name);
 
 	protected:
 
 		//! Forbidden. Manage with the Ptr<> template.
-		~NodeObjectNew();
+		~NodeObjectNew() {}
 
 	private:
 
-		Private* m_pD;
+		static FNodeType StaticType;
+
 
 	};
 
diff --git a/Engine/Plugins/Experimental/Mutable/Source/MutableTools/Public/MuT/NodeProjector.h b/Engine/Plugins/Experimental/Mutable/Source/MutableTools/Public/MuT/NodeProjector.h
--- a/Engine/Plugins/Experimental/Mutable/Source/MutableTools/Public/MuT/NodeProjector.h
+++ b/Engine/Plugins/Experimental/Mutable/Source/MutableTools/Public/MuT/NodeProjector.h
@@ -34,15 +34,6 @@ namespace mu
 	{
 	public:
 
-		// Possible subclasses
-		enum class EType : uint8
-		{
-			Constant = 0,
-			Parameter = 1,
-
-			None
-		};
-
 		//-----------------------------------------------------------------------------------------
 		// Node Interface
 		//-----------------------------------------------------------------------------------------
@@ -60,9 +51,6 @@ namespace mu
 
 		//! Forbidden. Manage with the Ptr<> template.
 		inline ~NodeProjector() {}
-		
-		//!
-		EType Type = EType::None;
 
 	};
 
@@ -101,7 +89,6 @@ namespace mu
 		//-----------------------------------------------------------------------------------------
 		class Private;
 		Private* GetPrivate() const;
-        Node::Private* GetBasePrivate() const override;
 
 	protected:
 
@@ -159,7 +146,6 @@ namespace mu
 		//-----------------------------------------------------------------------------------------
 		class Private;
 		Private* GetPrivate() const;
-        Node::Private* GetBasePrivate() const override;
 
 	protected:
 
diff --git a/Engine/Plugins/Experimental/Mutable/Source/MutableTools/Public/MuT/NodeRange.h b/Engine/Plugins/Experimental/Mutable/Source/MutableTools/Public/MuT/NodeRange.h
--- a/Engine/Plugins/Experimental/Mutable/Source/MutableTools/Public/MuT/NodeRange.h
+++ b/Engine/Plugins/Experimental/Mutable/Source/MutableTools/Public/MuT/NodeRange.h
@@ -22,15 +22,6 @@ namespace mu
 	{
 	public:
 
-		// Possible subclasses
-		enum class EType : uint8
-		{
-			FromScalar = 0,
-
-			None
-		};
-
-
 		//-----------------------------------------------------------------------------------------
 		// Node Interface
 		//-----------------------------------------------------------------------------------------
@@ -49,9 +40,6 @@ namespace mu
 		//! Forbidden. Manage with the Ptr<> template.
         inline ~NodeRange() {}
 
-		//!
-		EType Type = EType::None;
-
 	};
 
 
diff --git a/Engine/Plugins/Experimental/Mutable/Source/MutableTools/Public/MuT/NodeScalar.h b/Engine/Plugins/Experimental/Mutable/Source/MutableTools/Public/MuT/NodeScalar.h
--- a/Engine/Plugins/Experimental/Mutable/Source/MutableTools/Public/MuT/NodeScalar.h
+++ b/Engine/Plugins/Experimental/Mutable/Source/MutableTools/Public/MuT/NodeScalar.h
@@ -22,22 +22,6 @@ namespace mu
 	{
 	public:
 
-		// Possible subclasses
-		enum class EType : uint8
-		{
-			Constant = 0,
-			Parameter = 1,
-			EnumParameter = 2,
-			Curve = 3,
-			Switch = 4,
-			ArithmeticOperation = 5,
-			Variation = 6,
-			Table = 7,
-
-			None
-		};
-
-
 		//-----------------------------------------------------------------------------------------
 		// Node Interface
 		//-----------------------------------------------------------------------------------------
@@ -56,9 +40,6 @@ namespace mu
 		//! Forbidden. Manage with the Ptr<> template.
 		inline ~NodeScalar() {}
 
-		//!
-		EType Type = EType::None;
-
 	};
 
 
diff --git a/Engine/Plugins/Experimental/Mutable/Source/MutableTools/Public/MuT/NodeString.h b/Engine/Plugins/Experimental/Mutable/Source/MutableTools/Public/MuT/NodeString.h
--- a/Engine/Plugins/Experimental/Mutable/Source/MutableTools/Public/MuT/NodeString.h
+++ b/Engine/Plugins/Experimental/Mutable/Source/MutableTools/Public/MuT/NodeString.h
@@ -22,15 +22,6 @@ namespace mu
 	{
 	public:
 
-		// Possible subclasses
-		enum class EType : uint8
-		{
-			Constant = 0,
-			Parameter = 1,
-
-			None
-		};
-
 		//-----------------------------------------------------------------------------------------
 		// Node Interface
 		//-----------------------------------------------------------------------------------------
@@ -49,10 +40,6 @@ namespace mu
 		//! Forbidden. Manage with the Ptr<> template.
 		inline ~NodeString() {}
 
-		//!
-		EType Type = EType::None;
-
-
 	};
 
 
diff --git a/Engine/Plugins/Experimental/Mutable/Source/MutableTools/Public/MuT/NodeSurface.h b/Engine/Plugins/Experimental/Mutable/Source/MutableTools/Public/MuT/NodeSurface.h
--- a/Engine/Plugins/Experimental/Mutable/Source/MutableTools/Public/MuT/NodeSurface.h
+++ b/Engine/Plugins/Experimental/Mutable/Source/MutableTools/Public/MuT/NodeSurface.h
@@ -22,17 +22,6 @@ namespace mu
 	{
 	public:
 
-		// Possible subclasses
-		enum class EType : uint8
-		{
-			New = 0,
-			Edit = 1,
-			Variation = 3,
-			Switch = 4,
-
-			None
-		};
-
 		//-----------------------------------------------------------------------------------------
 		// Node Interface
 		//-----------------------------------------------------------------------------------------
@@ -46,9 +35,6 @@ namespace mu
 		//! Forbidden. Manage with the Ptr<> template.
         inline ~NodeSurface() {}
 
-		//!
-		EType Type = EType::None;
-
 	};
 
 
diff --git a/Engine/Plugins/Experimental/Mutable/Source/MutableTools/Public/MuT/NodeSurfaceNew.h b/Engine/Plugins/Experimental/Mutable/Source/MutableTools/Public/MuT/NodeSurfaceNew.h
--- a/Engine/Plugins/Experimental/Mutable/Source/MutableTools/Public/MuT/NodeSurfaceNew.h
+++ b/Engine/Plugins/Experimental/Mutable/Source/MutableTools/Public/MuT/NodeSurfaceNew.h
@@ -84,28 +84,13 @@ namespace mu
 		//! \name Meshes
         //! \{
 
-        //! Get the number of meshes in the Surface.
-		int GetMeshCount() const;
-
-        //! Set the number of meshes in the Surface.
-		void SetMeshCount( int );
-
-        //! Get the node generating one of the meshes in the Surface.
+		//! Get the node generating one of the meshes in the Surface.
 		//! \param index index of the mesh, from 0 to GetMeshCount()-1
-		NodeMeshPtr GetMesh( int index ) const;
+		Ptr<NodeMesh> GetMesh() const;
 
         //! Set the node generating one of the meshes in the Surface.
 		//! \param index index of the mesh, from 0 to GetMeshCount()-1
-		void SetMesh( int index, NodeMeshPtr );
-
-		//! Get the name of a mesh in the Surface.
-		//! \param index index of the image, from 0 to GetMeshCount()-1
-		const FString& GetMeshName(int index) const;
-
-        //! Set the name of a mesh in the Surface.
-		//! \param index index of the mesh, from 0 to GetMeshCount()-1
-		//! \param strName name of the mesh
-		void SetMeshName( int index, const FString& strName );
+		void SetMesh( Ptr<NodeMesh> );
 
 		//! \}
 
@@ -243,7 +228,6 @@ namespace mu
 		//-----------------------------------------------------------------------------------------
 		class Private;
 		Private* GetPrivate() const;
-        Node::Private* GetBasePrivate() const override;
 
 	protected:
 
