[mutable] Treat layout blocks as PODs.
#rnx

diff --git a/Engine/Plugins/Experimental/Mutable/Source/CustomizableObject/Internal/MuCO/CustomizableObjectPrivate.h b/Engine/Plugins/Experimental/Mutable/Source/CustomizableObject/Internal/MuCO/CustomizableObjectPrivate.h
--- a/Engine/Plugins/Experimental/Mutable/Source/CustomizableObject/Internal/MuCO/CustomizableObjectPrivate.h
+++ b/Engine/Plugins/Experimental/Mutable/Source/CustomizableObject/Internal/MuCO/CustomizableObjectPrivate.h
@@ -1081,6 +1081,8 @@ public:
 
 		ComponentsLODsRedesign2,
 
+		LayoutToPOD,
+
 		// -----<new versions can be added above this line>--------
 		LastCustomizableObjectVersion
 	};
diff --git a/Engine/Plugins/Experimental/Mutable/Source/CustomizableObjectEditor/Private/MuCOE/CustomizableObjectLayout.cpp b/Engine/Plugins/Experimental/Mutable/Source/CustomizableObjectEditor/Private/MuCOE/CustomizableObjectLayout.cpp
--- a/Engine/Plugins/Experimental/Mutable/Source/CustomizableObjectEditor/Private/MuCOE/CustomizableObjectLayout.cpp
+++ b/Engine/Plugins/Experimental/Mutable/Source/CustomizableObjectEditor/Private/MuCOE/CustomizableObjectLayout.cpp
@@ -134,11 +134,10 @@ void UCustomizableObjectLayout::GenerateBlocksFromUVs()
 			// Generating the layout blocks with the mutable layout
 			for (int i = 0; i < Layout->GetBlockCount(); ++i)
 			{
-				uint16 minX, minY, sizeX, sizeY;
+				mu::FImageSize Min = Layout->GetLayout()->Blocks[i].Min;
+				mu::FImageSize Size = Layout->GetLayout()->Blocks[i].Size;
 		
-				Layout->GetLayout()->GetBlock(i, &minX, &minY, &sizeX, &sizeY);
-		
-				FCustomizableObjectLayoutBlock block(FIntPoint(minX, minY), FIntPoint(minX + sizeX, minY + sizeY));
+				FCustomizableObjectLayoutBlock block(FIntPoint(Min.X, Min.Y), FIntPoint(Min.X + Size.X, Min.Y + Size.Y));
 				Blocks.Add(block);
 			}
 		
diff --git a/Engine/Plugins/Experimental/Mutable/Source/CustomizableObjectEditor/Private/MuCOE/SMutableLayoutViewer.cpp b/Engine/Plugins/Experimental/Mutable/Source/CustomizableObjectEditor/Private/MuCOE/SMutableLayoutViewer.cpp
--- a/Engine/Plugins/Experimental/Mutable/Source/CustomizableObjectEditor/Private/MuCOE/SMutableLayoutViewer.cpp
+++ b/Engine/Plugins/Experimental/Mutable/Source/CustomizableObjectEditor/Private/MuCOE/SMutableLayoutViewer.cpp
@@ -104,9 +104,9 @@ void SMutableLayoutViewer::Tick(const FGeometry& AllottedGeometry, const double
 			Blocks.SetNum( BlockCount );
 			for (int32 BlockIndex = 0; BlockIndex < BlockCount; ++BlockIndex)
 			{
-				UE::Math::TIntVector2<uint16> Min(ForceInit);
-				UE::Math::TIntVector2<uint16> Size(ForceInit);
-				MutableLayout->GetBlock(BlockIndex, &Min[0], &Min[1], &Size[0], &Size[1]);
+				UE::Math::TIntVector2<uint16> Min = MutableLayout->Blocks[BlockIndex].Min;
+				UE::Math::TIntVector2<uint16> Size = MutableLayout->Blocks[BlockIndex].Size;
+
 				Blocks[BlockIndex].Min[0] = Min[0];
 				Blocks[BlockIndex].Min[1] = Min[1];
 				Blocks[BlockIndex].Max[0] = Min[0] + Size[0];
diff --git a/Engine/Plugins/Experimental/Mutable/Source/MutableRuntime/Private/MuR/CodeRunnerTasked.cpp b/Engine/Plugins/Experimental/Mutable/Source/MutableRuntime/Private/MuR/CodeRunnerTasked.cpp
--- a/Engine/Plugins/Experimental/Mutable/Source/MutableRuntime/Private/MuR/CodeRunnerTasked.cpp
+++ b/Engine/Plugins/Experimental/Mutable/Source/MutableRuntime/Private/MuR/CodeRunnerTasked.cpp
@@ -1918,12 +1918,8 @@ namespace mu
 		}
 
 		box< UE::Math::TIntVector2<uint16> > RectInblocks;
-		Layout->GetBlock
-		(
-			RelBlockIndex,
-			&RectInblocks.min[0], &RectInblocks.min[1],
-			&RectInblocks.size[0], &RectInblocks.size[1]
-		);
+		RectInblocks.min = Layout->Blocks[RelBlockIndex].Min;
+		RectInblocks.size = Layout->Blocks[RelBlockIndex].Size;
 
 		// Convert the rect from blocks to pixels
 		FIntPoint Grid = Layout->GetGridSize();
diff --git a/Engine/Plugins/Experimental/Mutable/Source/MutableRuntime/Private/MuR/Layout.cpp b/Engine/Plugins/Experimental/Mutable/Source/MutableRuntime/Private/MuR/Layout.cpp
--- a/Engine/Plugins/Experimental/Mutable/Source/MutableRuntime/Private/MuR/Layout.cpp
+++ b/Engine/Plugins/Experimental/Mutable/Source/MutableRuntime/Private/MuR/Layout.cpp
@@ -11,6 +11,9 @@
 
 namespace mu {
 
+	MUTABLE_IMPLEMENT_POD_SERIALISABLE(Layout::FBlock);
+	MUTABLE_IMPLEMENT_POD_VECTOR_SERIALISABLE(Layout::FBlock);
+
 	//---------------------------------------------------------------------------------------------
 	Layout::Layout()
 	{
@@ -25,20 +28,20 @@ namespace mu {
 
 
 	//---------------------------------------------------------------------------------------------
-	LayoutPtr Layout::StaticUnserialise( InputArchive& arch )
+	Ptr<Layout> Layout::StaticUnserialise( InputArchive& arch )
 	{
 		LLM_SCOPE_BYNAME(TEXT("MutableRuntime"));
-		LayoutPtr pResult = new Layout();
+		Ptr<Layout> pResult = new Layout();
 		arch >> *pResult;
 		return pResult;
 	}
 
 
 	//---------------------------------------------------------------------------------------------
-	LayoutPtr Layout::Clone() const
+	Ptr<Layout> Layout::Clone() const
 	{
 		LLM_SCOPE_BYNAME(TEXT("MutableRuntime"));
-		LayoutPtr pResult = new Layout();		
+		Ptr<Layout> pResult = new Layout();
 		pResult->Size = Size;
 		pResult->MaxSize = MaxSize;
 		pResult->Blocks = Blocks;
@@ -124,58 +127,9 @@ namespace mu {
 
 
 	//---------------------------------------------------------------------------------------------
-	void Layout::GetBlock( int32 index, uint16* pMinX, uint16* pMinY, uint16* pSizeX, uint16* pSizeY ) const
+	void Layout::SetLayoutPackingStrategy(EPackStrategy InStrategy)
 	{
-		check( index >=0 && index < Blocks.Num() );
-		check( pMinX && pMinY && pSizeX && pSizeY );
-
-		if (pMinX && pMinY && pSizeX && pSizeY)
-		{
-			*pMinX = Blocks[index].Min[0];
-			*pMinY = Blocks[index].Min[1];
-			*pSizeX = Blocks[index].Size[0];
-			*pSizeY = Blocks[index].Size[1];
-		}
-	}
-
-
-	//---------------------------------------------------------------------------------------------
-	void Layout::GetBlockOptions(int index, int& pPriority, bool& bReduceBothAxes, bool& bReduceByTwo) const
-	{
-		check(index >= 0 && index < Blocks.Num());
-
-		pPriority = Blocks[index].Priority;
-		bReduceBothAxes = Blocks[index].bReduceBothAxes;
-		bReduceByTwo = Blocks[index].bReduceByTwo;
-	}
-
-
-	//---------------------------------------------------------------------------------------------
-    void Layout::SetBlock( int index, int minx, int miny, int sizex, int sizey )
-	{
-		check( index >=0 && index < Blocks.Num() );
-
-		// Keeps the id
-		Blocks[index].Min = UE::Math::TIntVector2<uint16>((uint16)minx, (uint16)miny);
-		Blocks[index].Size = UE::Math::TIntVector2<uint16>((uint16)sizex, (uint16)sizey);
-	}
-
-
-	//---------------------------------------------------------------------------------------------
-	void Layout::SetBlockOptions(int index, int priority, bool bReduceBothAxes, bool bReduceByTwo)
-	{
-		check(index >= 0 && index < Blocks.Num());
-
-		Blocks[index].Priority = priority;
-		Blocks[index].bReduceBothAxes = bReduceBothAxes;
-		Blocks[index].bReduceByTwo = bReduceByTwo;
-	}
-
-
-	//---------------------------------------------------------------------------------------------
-	void Layout::SetLayoutPackingStrategy(EPackStrategy _strategy)
-	{
-		Strategy = _strategy;
+		Strategy = InStrategy;
 	}
 
 
@@ -189,9 +143,6 @@ namespace mu {
 	//---------------------------------------------------------------------------------------------
 	void Layout::Serialise(OutputArchive& arch) const
 	{
-		uint32 ver = 7;
-		arch << ver;
-
 		arch << Size;
 		arch << Blocks;
 
@@ -205,22 +156,18 @@ namespace mu {
 	//---------------------------------------------------------------------------------------------
 	void Layout::Unserialise(InputArchive& arch)
 	{
-		uint32 ver;
-		arch >> ver;
-		check(ver <= 7);
-
 		arch >> Size;
 		arch >> Blocks;
 		arch >> MaxSize;
 
-		uint32 temp;
-		arch >> temp;
-		Strategy = EPackStrategy(temp);
+		uint32 Temp;
+		arch >> Temp;
+		Strategy = EPackStrategy(Temp);
 
 		arch >> FirstLODToIgnoreWarnings;
 
-		arch >> temp;
-		ReductionMethod = EReductionMethod(temp);
+		arch >> Temp;
+		ReductionMethod = EReductionMethod(Temp);
 	}
 
 
@@ -296,29 +243,5 @@ namespace mu {
 		return ReductionMethod;
 	}
 
-	
-	//---------------------------------------------------------------------------------------------
-	void Layout::FBlock::Serialise(OutputArchive& arch) const
-	{
-		arch << Min;
-		arch << Size;
-		arch << Id;
-		arch << Priority;
-		arch << bReduceBothAxes;
-		arch << bReduceByTwo;
-	}
-
-
-	//---------------------------------------------------------------------------------------------
-	void Layout::FBlock::Unserialise(InputArchive& arch)
-	{
-		arch >> Min;
-		arch >> Size;
-		arch >> Id;
-		arch >> Priority;
-		arch >> bReduceBothAxes;
-		arch >> bReduceByTwo;
-	}
-
 }
 
diff --git a/Engine/Plugins/Experimental/Mutable/Source/MutableRuntime/Private/MuR/OpLayoutPack.h b/Engine/Plugins/Experimental/Mutable/Source/MutableRuntime/Private/MuR/OpLayoutPack.h
--- a/Engine/Plugins/Experimental/Mutable/Source/MutableRuntime/Private/MuR/OpLayoutPack.h
+++ b/Engine/Plugins/Experimental/Mutable/Source/MutableRuntime/Private/MuR/OpLayoutPack.h
@@ -420,11 +420,12 @@ namespace mu
         for ( int index=0; index<blockCount; ++index )
         {
             box< UE::Math::TIntVector2<uint16> > b;
-            pSourceLayout->GetBlock( index, &b.min[0], &b.min[1], &b.size[0], &b.size[1] );
+			b.min = pSourceLayout->Blocks[index].Min;
+			b.size = pSourceLayout->Blocks[index].Size;
 
-			int p;
-			bool bReduceBothAxes, bReduceByTwo;
-			pSourceLayout->GetBlockOptions(index, p, bReduceBothAxes, bReduceByTwo);
+			int32 p = pSourceLayout->Blocks[index].Priority;
+			bool bReduceBothAxes = pSourceLayout->Blocks[index].bReduceBothAxes;
+			bool bReduceByTwo = pSourceLayout->Blocks[index].bReduceByTwo;
 
 			FIntVector2 reductions;
 			reductions[0] = 0;
@@ -637,16 +638,13 @@ namespace mu
 		pResult->SetLayoutPackingStrategy(LayoutStrategy);
 		pResult->SetBlockReductionMethod(ReductionMethod);
 
-        for ( int index=0; index<blockCount; ++index )
+        for ( int32 index=0; index<blockCount; ++index )
         {
-            pResult->SetBlock
-                (
-                    index,
-                    scratch->positions[index][0], scratch->positions[index][1],
-                    scratch->blocks[index][0], scratch->blocks[index][1]
-                );
-
-			pResult->SetBlockOptions(index, scratch->priorities[index], (bool)scratch->ReduceBothAxes[index], (bool)scratch->ReduceByTwo[index]);
+			pResult->Blocks[index].Min = UE::Math::TIntVector2<uint16>(scratch->positions[index]);
+			pResult->Blocks[index].Size = scratch->blocks[index];
+			pResult->Blocks[index].Priority = scratch->priorities[index];
+			pResult->Blocks[index].bReduceBothAxes = scratch->ReduceBothAxes[index];
+			pResult->Blocks[index].bReduceByTwo = scratch->ReduceByTwo[index];
         }
     }
 }
diff --git a/Engine/Plugins/Experimental/Mutable/Source/MutableRuntime/Private/MuR/OpMeshApplyLayout.h b/Engine/Plugins/Experimental/Mutable/Source/MutableRuntime/Private/MuR/OpMeshApplyLayout.h
--- a/Engine/Plugins/Experimental/Mutable/Source/MutableRuntime/Private/MuR/OpMeshApplyLayout.h
+++ b/Engine/Plugins/Experimental/Mutable/Source/MutableRuntime/Private/MuR/OpMeshApplyLayout.h
@@ -55,7 +55,8 @@ namespace mu
 			FIntPoint grid = pLayout->GetGridSize();
 
 			box< UE::Math::TIntVector2<uint16> > block;
-			pLayout->GetBlock( b, &block.min[0], &block.min[1], &block.size[0], &block.size[1] );
+			block.min = pLayout->Blocks[b].Min;
+			block.size = pLayout->Blocks[b].Size;
 
 			Box rect;
 			rect.min[0] = ( (float)block.min[0] ) / (float) grid[0];
diff --git a/Engine/Plugins/Experimental/Mutable/Source/MutableRuntime/Public/MuR/Layout.h b/Engine/Plugins/Experimental/Mutable/Source/MutableRuntime/Public/MuR/Layout.h
--- a/Engine/Plugins/Experimental/Mutable/Source/MutableRuntime/Public/MuR/Layout.h
+++ b/Engine/Plugins/Experimental/Mutable/Source/MutableRuntime/Public/MuR/Layout.h
@@ -99,42 +99,6 @@ namespace mu
 		//! The existing blocks will be kept as much as possible. The new blocks will be undefined.
 		void SetBlockCount( int32 );
 
-		//! Return a block of the layout.
-		//! "Position" here means the lower-left corner of the block.
-		//! \param index Block to get.
-		//! \param[out] pMinX will be set to the x position of the block
-		//! (from 0 to the X size of the grid minus one )
-		//! \param[out] pMinY will be set to the y position of the block
-		//! (from 0 to the Y size of the grid minus one )
-		//! \param[out] pSizeX will be set to the x size of the block
-		//! \param[out] pSizeY will be set to the y size of the block
-		void GetBlock( int32 Index, uint16* MinX, uint16* MinY, uint16* SizeX, uint16* SizeY ) const;
-
-		//! Returns the reduction priority of a block.
-		//! \param index Block to get priority.
-		//! \param[out] priority reduction priority of the block
-		//! \param[out] bReduceBothAxes reduction method of the block
-		//! \param[out] bReduceByTwo reduction method of the block
-		void GetBlockOptions(int32 Index, int32& Priority, bool& bReduceBothAxes, bool& bReduceByTwo) const;
-
-		//! Set a block of the layout.
-		//! "Position" here means the lower-left corner of the block.
-		//! \param index Block to set.
-		//! \param minX will be set to the x position of the block
-		//! (from 0 to the X size of the grid minus one )
-		//! \param minY will be set to the y position of the block
-		//! (from 0 to the Y size of the grid minus one )
-		//! \param sizeX will be set to the x size of the block
-		//! \param sizeY will be set to the y size of the block
-        void SetBlock( int32 Index, int32 MinX, int32 MinY, int32 SizeX, int32 SizeY );
-
-		//! Set the reduction options of a block
-		//! \param index Block to set the options
-		//! \param priority will be set to the reduction priority of the block. The blocks with the highest values will be the last to be reduced
-		//! \param bReduceBothAxes will be set to reduce the block in both axis at the same time.
-		//! \param bReduceByTwo will reduce by two blocks on a unitary reduction.
-		void SetBlockOptions(int32 Index, int32 Priority, bool bReduceBothAxes, bool bReduceByTwo);
-
 		//! Set the texture layout packing strategy
 		//! By default the texture layout packing strategy is set to resizable layout
 		void SetLayoutPackingStrategy(EPackStrategy);
@@ -161,8 +125,6 @@ namespace mu
 
 	public:
 
-		// This used to be the data in the private implementation of the image interface
-		//-----------------------------------------------------------------------------------------
 		struct FBlock
 		{
 			FBlock(UE::Math::TIntVector2<uint16> InMin = UE::Math::TIntVector2<uint16>(), UE::Math::TIntVector2<uint16> InSize = UE::Math::TIntVector2<uint16>())
@@ -185,18 +147,12 @@ namespace mu
 			int32 Priority;
 
 			//! Value to control the method to reduce the block
-			bool bReduceBothAxes;
+			uint32 bReduceBothAxes : 1;
 
 			//! Value to control if a block has to be reduced by two in an unitary reduction strategy
-			bool bReduceByTwo;
+			uint32 bReduceByTwo : 1;
 
 
-			//!
-			void Serialise(OutputArchive& arch) const;
-
-			//!
-			void Unserialise(InputArchive& arch);
-
 			//!
 			inline bool operator==(const FBlock& o) const
 			{
@@ -252,5 +208,8 @@ namespace mu
 		bool IsSingleBlockAndFull() const;
 	};
 
+	MUTABLE_DEFINE_POD_SERIALISABLE(Layout::FBlock);
+	MUTABLE_DEFINE_POD_VECTOR_SERIALISABLE(Layout::FBlock);
+
 }
 
diff --git a/Engine/Plugins/Experimental/Mutable/Source/MutableTools/Private/MuT/CodeGenerator.cpp b/Engine/Plugins/Experimental/Mutable/Source/MutableTools/Private/MuT/CodeGenerator.cpp
--- a/Engine/Plugins/Experimental/Mutable/Source/MutableTools/Private/MuT/CodeGenerator.cpp
+++ b/Engine/Plugins/Experimental/Mutable/Source/MutableTools/Private/MuT/CodeGenerator.cpp
@@ -1313,12 +1313,8 @@ namespace mu
 
 								// Block in layout grid units (cells)
 								box< UE::Math::TIntVector2<uint16> > rectInCells;
-								pLayout->GetBlock
-								(
-									BlockIndex,
-									&rectInCells.min[0], &rectInCells.min[1],
-									&rectInCells.size[0], &rectInCells.size[1]
-								);
+								rectInCells.min = pLayout->Blocks[BlockIndex].Min;
+								rectInCells.size = pLayout->Blocks[BlockIndex].Size;
 
 								// If we don't know the size of a layout block in pixels, calculate it
 								UpdateBlockSize(BlockDesc, rectInCells.size);
@@ -1411,12 +1407,8 @@ namespace mu
 
 												// Block in layout grid units
 												box< UE::Math::TIntVector2<uint16> > rectInCells;
-												pExtendLayout->GetBlock
-												(
-													b,
-													&rectInCells.min[0], &rectInCells.min[1],
-													&rectInCells.size[0], &rectInCells.size[1]
-												);
+												rectInCells.min = pLayout->Blocks[b].Min;
+												rectInCells.size = pLayout->Blocks[b].Size;
 
 												FImageDesc ExtendDesc = fragmentAd->GetImageDesc();
 
diff --git a/Engine/Plugins/Experimental/Mutable/Source/MutableTools/Private/MuT/CodeGenerator_Image.cpp b/Engine/Plugins/Experimental/Mutable/Source/MutableTools/Private/MuT/CodeGenerator_Image.cpp
--- a/Engine/Plugins/Experimental/Mutable/Source/MutableTools/Private/MuT/CodeGenerator_Image.cpp
+++ b/Engine/Plugins/Experimental/Mutable/Source/MutableTools/Private/MuT/CodeGenerator_Image.cpp
@@ -225,12 +225,8 @@ namespace mu
 
 			// Block in layout grid units
 			box< UE::Math::TIntVector2<uint16> > RectInCells;
-			Options.LayoutToApply->GetBlock
-			(
-				BlockIndex,
-				&RectInCells.min[0], &RectInCells.min[1],
-				&RectInCells.size[0], &RectInCells.size[1]
-			);
+			RectInCells.min = Options.LayoutToApply->Blocks[BlockIndex].Min;
+			RectInCells.size = Options.LayoutToApply->Blocks[BlockIndex].Size;
 
 			FIntPoint grid = Options.LayoutToApply->GetGridSize();
 			grid[0] = FMath::Max(1, grid[0]);
@@ -1157,12 +1153,8 @@ namespace mu
 
 				// Block in layout grid units
 				box< UE::Math::TIntVector2<uint16> > RectInCells;
-				Options.LayoutToApply->GetBlock
-				(
-					BlockIndex,
-					&RectInCells.min[0], &RectInCells.min[1],
-					&RectInCells.size[0], &RectInCells.size[1]
-				);
+				RectInCells.min = Options.LayoutToApply->Blocks[BlockIndex].Min;
+				RectInCells.size = Options.LayoutToApply->Blocks[BlockIndex].Size;
 
 				FIntPoint grid = Options.LayoutToApply->GetGridSize();
 				grid[0] = FMath::Max(1, grid[0]);
@@ -1220,12 +1212,8 @@ namespace mu
 
 				// Block in layout grid units
 				box< UE::Math::TIntVector2<uint16> > RectInCells;
-				Options.LayoutToApply->GetBlock
-				(
-					BlockIndex,
-					&RectInCells.min[0], &RectInCells.min[1],
-					&RectInCells.size[0], &RectInCells.size[1]
-				);
+				RectInCells.min = Options.LayoutToApply->Blocks[BlockIndex].Min;
+				RectInCells.size = Options.LayoutToApply->Blocks[BlockIndex].Size;
 
 				FIntPoint grid = Options.LayoutToApply->GetGridSize();
 				grid[0] = FMath::Max(1, grid[0]);
@@ -1587,12 +1575,8 @@ namespace mu
 
 			// Block in layout grid units
 			box< UE::Math::TIntVector2<uint16> > RectInCells;
-			Options.LayoutToApply->GetBlock
-			(
-				BlockIndex,
-				&RectInCells.min[0], &RectInCells.min[1],
-				&RectInCells.size[0], &RectInCells.size[1]
-			);
+			RectInCells.min = Options.LayoutToApply->Blocks[BlockIndex].Min;
+			RectInCells.size = Options.LayoutToApply->Blocks[BlockIndex].Size;
 
 			FIntPoint grid = Options.LayoutToApply->GetGridSize();
 			grid[0] = FMath::Max(1, grid[0]);
diff --git a/Engine/Plugins/Experimental/Mutable/Source/MutableTools/Private/MuT/CodeGenerator_Mesh.cpp b/Engine/Plugins/Experimental/Mutable/Source/MutableTools/Private/MuT/CodeGenerator_Mesh.cpp
--- a/Engine/Plugins/Experimental/Mutable/Source/MutableTools/Private/MuT/CodeGenerator_Mesh.cpp
+++ b/Engine/Plugins/Experimental/Mutable/Source/MutableTools/Private/MuT/CodeGenerator_Mesh.cpp
@@ -320,21 +320,21 @@ namespace mu
 			BlockIds[BlockIndex] = Layout->Blocks[BlockIndex].Id;
 
 			// Get the block rect
-			uint16 MinX, MinY, SizeX, SizeY;
-			Layout->GetBlock(BlockIndex, &MinX, &MinY, &SizeX, &SizeY);
+			FImageSize Min = Layout->Blocks[BlockIndex].Min;
+			FImageSize Size = Layout->Blocks[BlockIndex].Size;
 
 			box<FVector2f>& BlockRect = BlockRects[BlockIndex];
-			BlockRect.min[0] = ((float)MinX) / (float)Grid.X;
-			BlockRect.min[1] = ((float)MinY) / (float)Grid.Y;
-			BlockRect.size[0] = ((float)SizeX) / (float)Grid.X;
-			BlockRect.size[1] = ((float)SizeY) / (float)Grid.Y;
+			BlockRect.min[0] = ((float)Min.X) / (float)Grid.X;
+			BlockRect.min[1] = ((float)Min.Y) / (float)Grid.Y;
+			BlockRect.size[0] = ((float)Size.X) / (float)Grid.X;
+			BlockRect.size[1] = ((float)Size.Y) / (float)Grid.Y;
 
 			// Create block index per cell array
-			for (uint16 Y = MinY; Y < MinY + SizeY; ++Y)
+			for (uint16 Y = Min.Y; Y < Min.Y + Size.Y; ++Y)
 			{
 				const uint16 PositionY = Y * Grid.X;
 
-				for (uint16 X = MinX; X < MinX + SizeX; ++X)
+				for (uint16 X = Min.X; X < Min.X + Size.X; ++X)
 				{
 					if (GridBlockBlockId[PositionY + X] == MAX_uint16)
 					{
diff --git a/Engine/Plugins/Experimental/Mutable/Source/MutableTools/Private/MuT/NodeLayout.cpp b/Engine/Plugins/Experimental/Mutable/Source/MutableTools/Private/MuT/NodeLayout.cpp
--- a/Engine/Plugins/Experimental/Mutable/Source/MutableTools/Private/MuT/NodeLayout.cpp
+++ b/Engine/Plugins/Experimental/Mutable/Source/MutableTools/Private/MuT/NodeLayout.cpp
@@ -46,21 +46,21 @@ namespace mu
 	//---------------------------------------------------------------------------------------------
 	// Own Interface
 	//---------------------------------------------------------------------------------------------
-	void NodeLayoutBlocks::SetGridSize( int x, int y )
+	void NodeLayoutBlocks::SetGridSize( int32 x, int32 y )
 	{
 		m_pD->m_pLayout->SetGridSize( x, y );
 	}
 
 
 	//---------------------------------------------------------------------------------------------
-	void NodeLayoutBlocks::SetMaxGridSize(int x, int y)
+	void NodeLayoutBlocks::SetMaxGridSize(int32 x, int32 y)
 	{
 		m_pD->m_pLayout->SetMaxGridSize(x, y);
 	}
 
 
 	//---------------------------------------------------------------------------------------------
-	void NodeLayoutBlocks::GetGridSize( int* pX, int* pY ) const
+	void NodeLayoutBlocks::GetGridSize( int32* pX, int32* pY ) const
 	{
 		FIntPoint grid = m_pD->m_pLayout->GetGridSize();
 		if (pX) *pX = grid[0];
@@ -69,21 +69,21 @@ namespace mu
 
 
 	//---------------------------------------------------------------------------------------------
-	void NodeLayoutBlocks::GetMaxGridSize(int* pX, int* pY) const
+	void NodeLayoutBlocks::GetMaxGridSize(int32* pX, int32* pY) const
 	{
 		m_pD->m_pLayout->GetMaxGridSize(pX, pY);
 	}
 
 
 	//---------------------------------------------------------------------------------------------
-	void NodeLayoutBlocks::SetBlockCount( int n )
+	void NodeLayoutBlocks::SetBlockCount( int32 n )
 	{
 		m_pD->m_pLayout->SetBlockCount( n );
 	}
 
 
 	//---------------------------------------------------------------------------------------------
-	int NodeLayoutBlocks::GetBlockCount()
+	int32 NodeLayoutBlocks::GetBlockCount()
 	{
 		return m_pD->m_pLayout->GetBlockCount();
 	}
@@ -97,16 +97,19 @@ namespace mu
 
 
 	//---------------------------------------------------------------------------------------------
-    void NodeLayoutBlocks::SetBlock( int index, int minx, int miny, int sizex, int sizey )
+    void NodeLayoutBlocks::SetBlock( int32 index, int32 minx, int32 miny, int32 sizex, int32 sizey )
 	{
-        m_pD->m_pLayout->SetBlock( index, minx, miny, sizex, sizey );
+		m_pD->m_pLayout->Blocks[index].Min = FImageSize( minx, miny );
+		m_pD->m_pLayout->Blocks[index].Size = FImageSize( sizex, sizey );
 	}
 
 
 	//---------------------------------------------------------------------------------------------
-	void NodeLayoutBlocks::SetBlockOptions(int index, int priority, bool bReduceBothAxes, bool bReduceByTwo)
+	void NodeLayoutBlocks::SetBlockOptions(int32 index, int32 priority, bool bReduceBothAxes, bool bReduceByTwo)
 	{
-		m_pD->m_pLayout->SetBlockOptions(index, priority, bReduceBothAxes, bReduceByTwo);
+		m_pD->m_pLayout->Blocks[index].Priority = priority;
+		m_pD->m_pLayout->Blocks[index].bReduceBothAxes = bReduceBothAxes;
+		m_pD->m_pLayout->Blocks[index].bReduceByTwo = bReduceByTwo;
 	}
 
 
@@ -118,13 +121,13 @@ namespace mu
 
 
 	//---------------------------------------------------------------------------------------------
-	NodeLayoutBlocksPtr NodeLayoutBlocks::GenerateLayoutBlocks(const MeshPtr pMesh, int layoutIndex, int gridSizeX, int gridSizeY)
+	NodeLayoutBlocksPtr NodeLayoutBlocks::GenerateLayoutBlocks(const MeshPtr pMesh, int32 layoutIndex, int32 gridSizeX, int32 gridSizeY)
 	{
 		NodeLayoutBlocksPtr layout = nullptr;
 
 		if (pMesh && layoutIndex >=0 && gridSizeX+gridSizeY>0)
 		{
-			int indexCount = pMesh->GetIndexCount();
+			int32 indexCount = pMesh->GetIndexCount();
 			TArray< FVector2f > UVs;
 			UVs.SetNumUninitialized(indexCount * 2);
 
@@ -132,7 +135,7 @@ namespace mu
 			UntypedMeshBufferIteratorConst texIt(pMesh->GetVertexBuffers(), MBS_TEXCOORDS, layoutIndex);
 			
 			//Getting UVs face by face
-			for (int v = 0; v < indexCount/3; ++v)
+			for (int32 v = 0; v < indexCount/3; ++v)
 			{
 				uint32_t i_1 = indexIt.GetAsUINT32(); 
 				indexIt++;
@@ -178,15 +181,15 @@ namespace mu
 			TArray<box<FIntVector2>> blocks;
 			
 			//Generating blocks
-			for (int i = 0; i < indexCount; ++i)
+			for (int32 i = 0; i < indexCount; ++i)
 			{
 				FIntVector2 a, b;
 			
-				a[0] = (int)floor(UVs[i * 2][0] * gridSizeX);
-				a[1] = (int)floor(UVs[i * 2][1] * gridSizeY);
+				a[0] = (int32)floor(UVs[i * 2][0] * gridSizeX);
+				a[1] = (int32)floor(UVs[i * 2][1] * gridSizeY);
 									  
-				b[0] = (int)floor(UVs[i * 2 +1][0] * gridSizeX);
-				b[1] = (int)floor(UVs[i * 2 +1][1] * gridSizeY);
+				b[0] = (int32)floor(UVs[i * 2 +1][0] * gridSizeX);
+				b[1] = (int32)floor(UVs[i * 2 +1][1] * gridSizeY);
 
 				//floor of UV = 1*gridSize is gridSize which is not a valid range
 				if (a[0] == gridSizeX){ a[0] = gridSizeX-1; }
@@ -219,19 +222,19 @@ namespace mu
 				}
 				else //they are in different blocks
 				{
-					int idxA = -1;
-					int idxB = -1;
+					int32 idxA = -1;
+					int32 idxB = -1;
 					
 					//Getting the blocks that contain them
 					for (int32 it = 0; it < blocks.Num(); ++it)
 					{
 						if (blocks[it].Contains(a))
 						{
-							idxA = (int)it;
+							idxA = (int32)it;
 						}
 						if (blocks[it].Contains(b))
 						{
-							idxB = (int)it;
+							idxB = (int32)it;
 						}
 					}
 					
@@ -307,13 +310,13 @@ namespace mu
 			{
 				layout->SetBlockCount(numBlocks);
 			
-				for (int i = 0; i < numBlocks; ++i)
+				for (int32 i = 0; i < numBlocks; ++i)
 				{
-					int blockIndex = i;
-					int minX = blocks[i].min[0];
-					int minY = blocks[i].min[1];
-					int sizeX = blocks[i].size[0];
-					int sizeY = blocks[i].size[1];
+					int32 blockIndex = i;
+					int32 minX = blocks[i].min[0];
+					int32 minY = blocks[i].min[1];
+					int32 sizeX = blocks[i].size[0];
+					int32 sizeY = blocks[i].size[1];
 			
 					layout->SetBlock(blockIndex, minX, minY, sizeX, sizeY);
 				}
diff --git a/Engine/Plugins/Experimental/Mutable/Source/MutableTools/Public/MuT/NodeLayout.h b/Engine/Plugins/Experimental/Mutable/Source/MutableTools/Public/MuT/NodeLayout.h
--- a/Engine/Plugins/Experimental/Mutable/Source/MutableTools/Public/MuT/NodeLayout.h
+++ b/Engine/Plugins/Experimental/Mutable/Source/MutableTools/Public/MuT/NodeLayout.h
@@ -72,39 +72,39 @@ namespace mu
 		//-----------------------------------------------------------------------------------------
 
 		//! Set the size of the grid where the blocks will be defined
-		void SetGridSize( int x, int y );
+		void SetGridSize( int32 x, int32 y );
 
 		//! Set the maximum size of the grid where the blocks can be defined
-		void SetMaxGridSize(int x, int y);
+		void SetMaxGridSize(int32 x, int32 y);
 
 		//! Get the size of the grid where the blocks will be defined
-		void GetGridSize( int* pX, int* pY ) const;
+		void GetGridSize( int32* pX, int32* pY ) const;
 
 		//! Get the maximum size of the grid where the blocks can be defined
-		void GetMaxGridSize(int* pX, int* pY) const;
+		void GetMaxGridSize(int32* pX, int32* pY) const;
 
 		//! Set the number of blocks in the layout.
 		//! It keeps the current data as much as possible, and the new data is undefined.
-		void SetBlockCount( int );
+		void SetBlockCount( int32 );
 
 		//! Get the number of blocks in the layout
-		int GetBlockCount();
+		int32 GetBlockCount();
 
 		/** */
 		Ptr<const Layout> GetLayout() const;
 
 		//! Set a block of the layout.
 		//! minx and miny refer to the lowest left corner of the block.
-        void SetBlock( int index, int minx, int miny, int sizex, int sizey );
+        void SetBlock( int32 index, int32 minx, int32 miny, int32 sizex, int32 sizey );
 
 		//! Set reduction block options like priority or if the block has to be reduced symmetrically.
-		void SetBlockOptions(int index, int priority, bool bReduceBothAxes, bool bReduceByTwo);
+		void SetBlockOptions(int32 index, int32 priority, bool bReduceBothAxes, bool bReduceByTwo);
 
 		//! Set the texture layout packing strategy 
 		void SetLayoutPackingStrategy(EPackStrategy strategy);
 
 		//! Generate the blocks of a layout using the UV of the meshes
-		static NodeLayoutBlocksPtr GenerateLayoutBlocks(const MeshPtr pMesh, int layoutIndex, int gridSizeX, int gridSizeY);
+		static NodeLayoutBlocksPtr GenerateLayoutBlocks(const MeshPtr pMesh, int32 layoutIndex, int32 gridSizeX, int32 gridSizeY);
 
 		//! Set at which LOD the unassigned vertices warnings will star to be ignored
 		void SetIgnoreWarningsLOD(int32 LOD);
