[Mutable] Avoid calling UpdateDelegates when AttachedParent has not been modified.

Only affects the optimized path. Now some EUpdateResult::Success updates will become EUpdateResult::ErrorOptimized.

#rnx

diff --git a/Engine/Plugins/Experimental/Mutable/Source/CustomizableObject/Internal/MuCO/CustomizableObjectSystemPrivate.h b/Engine/Plugins/Experimental/Mutable/Source/CustomizableObject/Internal/MuCO/CustomizableObjectSystemPrivate.h
--- a/Engine/Plugins/Experimental/Mutable/Source/CustomizableObject/Internal/MuCO/CustomizableObjectSystemPrivate.h
+++ b/Engine/Plugins/Experimental/Mutable/Source/CustomizableObject/Internal/MuCO/CustomizableObjectSystemPrivate.h
@@ -569,6 +569,9 @@ public:
 	bool UpdateStarted = false;
 	bool bLevelBegunPlay = false;
 
+	/** true if the update has been optimized (skips all Tasks and calls FinishUpdateGlobal directly on the Enqueue). */
+	bool bOptimizedUpdate = false;
+	
 	// Update stats
 	double StartQueueTime = 0.0;
 	double QueueTime = 0.0;
@@ -590,6 +593,9 @@ public:
 	/** Used for profiling in the editor. */
 	uint32 MutableRuntimeCycles = 0;
 
+	/** If a InstanceUsage is in this set it means that its AttachParent has been modified (USkeletalMesh changed, UMaterial changed...). */
+	TSet<TWeakObjectPtr<UCustomizableObjectInstanceUsage>> AttachedParentUpdated;
+	
 	/** Hard references to objects. Avoids GC to collect them. */
 	TArray<TStrongObjectPtr<const UObject>> Objects;
 };
diff --git a/Engine/Plugins/Experimental/Mutable/Source/CustomizableObject/Private/MuCO/CustomizableObjectInstanceUsage.cpp b/Engine/Plugins/Experimental/Mutable/Source/CustomizableObject/Private/MuCO/CustomizableObjectInstanceUsage.cpp
--- a/Engine/Plugins/Experimental/Mutable/Source/CustomizableObject/Private/MuCO/CustomizableObjectInstanceUsage.cpp
+++ b/Engine/Plugins/Experimental/Mutable/Source/CustomizableObject/Private/MuCO/CustomizableObjectInstanceUsage.cpp
@@ -232,14 +232,24 @@ bool RequiresReinitPose(USkeletalMesh* CurrentSkeletalMesh, USkeletalMesh* Skele
 }
 
 
-void UCustomizableObjectInstanceUsage::SetSkeletalMesh(USkeletalMesh* SkeletalMesh)
+void UCustomizableObjectInstanceUsage::SetSkeletalMesh(USkeletalMesh* SkeletalMesh, bool* bOutSkeletalMeshUpdated, bool* bOutMaterialsUpdated)
 {
 	USkeletalMeshComponent* Parent = Cast<USkeletalMeshComponent>(GetAttachParent());
 
 	if (Parent)
 	{
-		Parent->SetSkeletalMesh(SkeletalMesh, RequiresReinitPose(Parent->GetSkeletalMeshAsset(), SkeletalMesh));
+		if (SkeletalMesh != Parent->GetSkeletalMeshAsset())
+		{
+			Parent->SetSkeletalMesh(SkeletalMesh, RequiresReinitPose(Parent->GetSkeletalMeshAsset(), SkeletalMesh));
 
+			if (bOutSkeletalMeshUpdated)
+			{
+				*bOutSkeletalMeshUpdated = true;
+			}
+		}
+
+		TArray<TObjectPtr<UMaterialInterface>> OldOverridenMaterials = Parent->OverrideMaterials;
+		
 		if (Parent->HasOverrideMaterials())
 		{
 			// For some reason the reference skeletal mesh materials are added as override materials, clear them if necessary
@@ -267,17 +277,28 @@ void UCustomizableObjectInstanceUsage::SetSkeletalMesh(USkeletalMesh* SkeletalMe
 					Parent->SetMaterial(Index, ComponentData->OverrideMaterials[Index]);
 				}
 			}
-		}	
+		}
+
+		if (bOutMaterialsUpdated)
+		{
+			*bOutMaterialsUpdated = OldOverridenMaterials != Parent->OverrideMaterials;
+		}
 	}
 }
 
-void UCustomizableObjectInstanceUsage::SetPhysicsAsset(UPhysicsAsset* PhysicsAsset)
+void UCustomizableObjectInstanceUsage::SetPhysicsAsset(UPhysicsAsset* PhysicsAsset, bool* bOutPhysicsAssetUpdated)
 {
 	USkeletalMeshComponent* Parent = Cast<USkeletalMeshComponent>(GetAttachParent());
 
-	if (Parent && Parent->GetWorld())
+	if (Parent && Parent->GetWorld() &&
+		PhysicsAsset != Parent->GetPhysicsAsset())
 	{
 		Parent->SetPhysicsAsset(PhysicsAsset, true);
+		
+		if (bOutPhysicsAssetUpdated)
+		{
+			*bOutPhysicsAssetUpdated = true;
+		}
 	}
 }
 
diff --git a/Engine/Plugins/Experimental/Mutable/Source/CustomizableObject/Private/MuCO/CustomizableObjectSystem.cpp b/Engine/Plugins/Experimental/Mutable/Source/CustomizableObject/Private/MuCO/CustomizableObjectSystem.cpp
--- a/Engine/Plugins/Experimental/Mutable/Source/CustomizableObject/Private/MuCO/CustomizableObjectSystem.cpp
+++ b/Engine/Plugins/Experimental/Mutable/Source/CustomizableObject/Private/MuCO/CustomizableObjectSystem.cpp
@@ -855,18 +855,23 @@ void FinishUpdateGlobal(const TSharedRef<FUpdateContextPrivate>& Context)
 		// Call CustomizableObjectInstanceUsages updated callbacks.
 		for (TObjectIterator<UCustomizableObjectInstanceUsage> It; It; ++It) // Since iterating objects is expensive, for now CustomizableObjectInstanceUsage does not have a FinishUpdate function.
 		{
+			const UCustomizableObjectInstanceUsage* InstanceUsage = *It;
+			if (!IsValid(InstanceUsage))
+			{
+				continue;
+			}
+
 #if WITH_EDITOR
-			if (IsValid(*It) && It->IsNetMode(NM_DedicatedServer))
+			if (It->IsNetMode(NM_DedicatedServer))
 			{
 				continue;
 			}
 #endif
 
-			if (const UCustomizableObjectInstanceUsage* CustomizableObjectInstanceUsage = *It;
-				IsValid(CustomizableObjectInstanceUsage) &&
-				CustomizableObjectInstanceUsage->GetCustomizableObjectInstance() == Instance)
+			if (InstanceUsage->GetCustomizableObjectInstance() == Instance &&
+				(!Context->bOptimizedUpdate || Context->AttachedParentUpdated.Find(InstanceUsage)))
 			{
-				CustomizableObjectInstanceUsage->Callbacks();
+				InstanceUsage->Callbacks();
 			}
 		}
 	}
@@ -945,6 +950,10 @@ void UpdateSkeletalMesh(const TSharedRef<FUpdateContextPrivate>& Context)
 		}
 #endif
 
+		bool bSkeletalMeshUpdated = false;
+		bool bMaterialsUpdated = false;
+		bool bPhysicsAssetUpdated = false;	
+		
 		if (IsValid(CustomizableObjectInstanceUsage) &&
 			(CustomizableObjectInstanceUsage->GetCustomizableObjectInstance() == CustomizableObjectInstance) &&
 			CustomizableObjectInstancePrivateData->SkeletalMeshes.IsValidIndex(CustomizableObjectInstanceUsage->GetComponentIndex())
@@ -952,15 +961,20 @@ void UpdateSkeletalMesh(const TSharedRef<FUpdateContextPrivate>& Context)
 		{
 			MUTABLE_CPUPROFILER_SCOPE(UpdateSkeletalMesh_SetSkeletalMesh);
 
-			CustomizableObjectInstanceUsage->SetSkeletalMesh(CustomizableObjectInstancePrivateData->SkeletalMeshes[CustomizableObjectInstanceUsage->GetComponentIndex()]);
+			CustomizableObjectInstanceUsage->SetSkeletalMesh(CustomizableObjectInstancePrivateData->SkeletalMeshes[CustomizableObjectInstanceUsage->GetComponentIndex()], &bSkeletalMeshUpdated, &bMaterialsUpdated);
 
 			if (CustomizableObjectInstancePrivateData->HasCOInstanceFlags(ReplacePhysicsAssets))
 			{
 				CustomizableObjectInstanceUsage->SetPhysicsAsset(
 					CustomizableObjectInstancePrivateData->SkeletalMeshes[CustomizableObjectInstanceUsage->GetComponentIndex()] ?
-					CustomizableObjectInstancePrivateData->SkeletalMeshes[CustomizableObjectInstanceUsage->GetComponentIndex()]->GetPhysicsAsset() : nullptr);
+					CustomizableObjectInstancePrivateData->SkeletalMeshes[CustomizableObjectInstanceUsage->GetComponentIndex()]->GetPhysicsAsset() : nullptr, &bPhysicsAssetUpdated);
 			}
 		}
+
+		if (bSkeletalMeshUpdated || bMaterialsUpdated || bPhysicsAssetUpdated)
+		{
+			Context->AttachedParentUpdated.Add(CustomizableObjectInstanceUsage);
+		}
 	}
 }
 
@@ -1213,6 +1227,7 @@ void UCustomizableObjectSystemPrivate::EnqueueUpdateSkeletalMesh(const TSharedRe
 		{
 			if (Context->GetCapturedDescriptorHash().IsSubset(QueueElem->Context->GetCapturedDescriptorHash()))
 			{
+				Context->bOptimizedUpdate = true;
 				Context->UpdateResult = EUpdateResult::ErrorOptimized;
 				FinishUpdateGlobal(Context);			
 				return; // The the requested update is equal to the last enqueued update.
@@ -1223,6 +1238,7 @@ void UCustomizableObjectSystemPrivate::EnqueueUpdateSkeletalMesh(const TSharedRe
 			Instance == CurrentMutableOperation->Instance &&
 			Context->GetCapturedDescriptorHash().IsSubset(CurrentMutableOperation->GetCapturedDescriptorHash()))
 		{
+			Context->bOptimizedUpdate = true;
 			Context->UpdateResult = EUpdateResult::ErrorOptimized;
 			FinishUpdateGlobal(Context);
 			return; // The requested update is equal to the running update.
@@ -1232,8 +1248,13 @@ void UCustomizableObjectSystemPrivate::EnqueueUpdateSkeletalMesh(const TSharedRe
 			!(CurrentMutableOperation &&
 			Instance == CurrentMutableOperation->Instance)) // This condition is necessary because even if the descriptor is a subset, it will be replaced by the CurrentMutableOperation
 		{
-			Context->UpdateResult = EUpdateResult::Success;
-			UpdateSkeletalMesh(Context);
+			Context->bOptimizedUpdate = true;
+
+			// The user may have changed the AttachParent and we need to recustomize it.
+			// In case nothing need to be recustomized, the update will be considered ErrorOptimized.
+			UpdateSkeletalMesh(Context); 
+			Context->UpdateResult = Context->AttachedParentUpdated.IsEmpty() ? EUpdateResult::ErrorOptimized : EUpdateResult::Success;
+			
 			FinishUpdateGlobal(Context);
 		}
 		else
diff --git a/Engine/Plugins/Experimental/Mutable/Source/CustomizableObject/Public/MuCO/CustomizableObjectInstanceUsage.h b/Engine/Plugins/Experimental/Mutable/Source/CustomizableObject/Public/MuCO/CustomizableObjectInstanceUsage.h
--- a/Engine/Plugins/Experimental/Mutable/Source/CustomizableObject/Public/MuCO/CustomizableObjectInstanceUsage.h
+++ b/Engine/Plugins/Experimental/Mutable/Source/CustomizableObject/Public/MuCO/CustomizableObjectInstanceUsage.h
@@ -13,6 +13,7 @@ class FObjectPreSaveContext;
 class UObject;
 class USkeletalMesh;
 class UCustomizableSkeletalComponent;
+class UPhysicsAsset;
 struct FFrame;
 enum class EUpdateResult : uint8;
 
@@ -56,8 +57,8 @@ public:
 	UFUNCTION(BlueprintCallable, Category = CustomizableObjectInstance)
 	void UpdateSkeletalMeshAsyncResult(FInstanceUpdateDelegate Callback, bool bIgnoreCloseDist = false, bool bForceHighPriority = false);
 
-	void SetSkeletalMesh(USkeletalMesh* SkeletalMesh);
-	void SetPhysicsAsset(class UPhysicsAsset* PhysicsAsset);
+	void SetSkeletalMesh(USkeletalMesh* SkeletalMesh, bool* bOutSkeletalMeshUpdated = nullptr, bool* bOutMaterialsUpdated = nullptr);
+	void SetPhysicsAsset(UPhysicsAsset* PhysicsAsset, bool* bOutPhysicsAssetUpdated = nullptr);
 	void UpdateDistFromComponentToPlayer(const AActor* const Pawn, bool bForceEvenIfNotBegunPlay = false);
 
 	void SetVisibilityOfSkeletalMeshSectionWithMaterialName(bool bVisible, const FString& MaterialName, int32 LOD);
