[Mutable] Added new Gather References compilation mode.

Useful for encryption.

- Added Gather References button.
- Added Clear References button.
- Exposed Gather References in CompileCustomizableObjectSynchronously.
- Fixed return in IsCompilationOutOfDate.

#rnx

diff --git a/Engine/Plugins/Experimental/Mutable/Source/CustomizableObject/Internal/MuCO/CustomizableObjectCompilerTypes.h b/Engine/Plugins/Experimental/Mutable/Source/CustomizableObject/Internal/MuCO/CustomizableObjectCompilerTypes.h
--- a/Engine/Plugins/Experimental/Mutable/Source/CustomizableObject/Internal/MuCO/CustomizableObjectCompilerTypes.h
+++ b/Engine/Plugins/Experimental/Mutable/Source/CustomizableObject/Internal/MuCO/CustomizableObjectCompilerTypes.h
@@ -99,6 +99,9 @@ struct FCompilationOptions
 
 	// Control image tiled generation
 	int32 ImageTiling = 0;
+
+	/** If true, gather all game asset references and save them in the Customizable Object. */
+	bool bGatherReferences = false;
 };
 
 
diff --git a/Engine/Plugins/Experimental/Mutable/Source/CustomizableObject/Internal/MuCO/CustomizableObjectPrivate.h b/Engine/Plugins/Experimental/Mutable/Source/CustomizableObject/Internal/MuCO/CustomizableObjectPrivate.h
--- a/Engine/Plugins/Experimental/Mutable/Source/CustomizableObject/Internal/MuCO/CustomizableObjectPrivate.h
+++ b/Engine/Plugins/Experimental/Mutable/Source/CustomizableObject/Internal/MuCO/CustomizableObjectPrivate.h
@@ -1017,6 +1017,10 @@ public:
 	UPROPERTY(Transient)
 	TArray<FMutableModelParameterProperties> ParameterProperties;
 
+	/** Reference to all UObject used in game. Only updated during the compilation if the user explicitly wants to save all references. */
+	UPROPERTY()
+	FModelResources References;
+	
 	// Map of name to index of ParameterProperties.
 	// use this to lookup fast by Name
 	TMap<FString, FMutableParameterIndex> ParameterPropertiesLookupTable;
diff --git a/Engine/Plugins/Experimental/Mutable/Source/CustomizableObjectEditor/Private/MuCOE/CustomizableObjectCompiler.cpp b/Engine/Plugins/Experimental/Mutable/Source/CustomizableObjectEditor/Private/MuCOE/CustomizableObjectCompiler.cpp
--- a/Engine/Plugins/Experimental/Mutable/Source/CustomizableObjectEditor/Private/MuCOE/CustomizableObjectCompiler.cpp
+++ b/Engine/Plugins/Experimental/Mutable/Source/CustomizableObjectEditor/Private/MuCOE/CustomizableObjectCompiler.cpp
@@ -895,9 +895,7 @@ void FCustomizableObjectCompiler::CompileInternal(bool bAsync)
 	}
 	else
 	{
-		// Always work with the ModelResources (Editor) when compiling. They'll be copied to the cooked version during PreSave.
-		FModelResources& ModelResources = CurrentObject->GetPrivate()->GetModelResources(false);
-		ModelResources = FModelResources();
+		FModelResources ModelResources;
 		
 		ModelResources.ReferenceSkeletalMeshesData = MoveTemp(GenerationContext.ReferenceSkeletalMeshesData);
 
@@ -1108,6 +1106,16 @@ void FCustomizableObjectCompiler::CompileInternal(bool bAsync)
 		ModelResources.NumLODsToStream = GenerationContext.bEnableLODStreaming ? GenerationContext.NumMaxLODsToStream : 0;
 		ModelResources.FirstLODAvailable = GenerationContext.FirstLODAvailable;
 
+		if (CurrentOptions.bGatherReferences)
+		{
+			CurrentObject->GetPrivate()->References = ModelResources;
+			CurrentObject->GetPrivate()->References.RuntimeReferencedTextures.Empty(); // Empty in case the of none optimization. In maximum optimization, they are Mutable textures. 
+			CurrentObject->Modify();
+		}
+
+		// Always work with the ModelResources (Editor) when compiling. They'll be copied to the cooked version during PreSave.
+		CurrentObject->GetPrivate()->GetModelResources(false) = MoveTemp(ModelResources);
+
 		CurrentObject->GetPrivate()->GetStreamedResourceData() = MoveTemp(GenerationContext.StreamedResourceData);
 
 		// Pass-through textures
diff --git a/Engine/Plugins/Experimental/Mutable/Source/CustomizableObjectEditor/Private/MuCOE/CustomizableObjectEditor.cpp b/Engine/Plugins/Experimental/Mutable/Source/CustomizableObjectEditor/Private/MuCOE/CustomizableObjectEditor.cpp
--- a/Engine/Plugins/Experimental/Mutable/Source/CustomizableObjectEditor/Private/MuCOE/CustomizableObjectEditor.cpp
+++ b/Engine/Plugins/Experimental/Mutable/Source/CustomizableObjectEditor/Private/MuCOE/CustomizableObjectEditor.cpp
@@ -602,13 +602,13 @@ void FCustomizableObjectEditor::BindCommands()
 	// Compile and options
 	ToolkitCommands->MapAction(
 		Commands.Compile,
-		FExecuteAction::CreateSP(this, &FCustomizableObjectEditor::CompileObject, false),
+		FExecuteAction::CreateSP(this, &FCustomizableObjectEditor::CompileObject, false, false),
 		FCanExecuteAction::CreateStatic(&UCustomizableObjectSystem::IsActive),
 		FIsActionChecked());
 
 	ToolkitCommands->MapAction(
 		Commands.CompileOnlySelected,
-		FExecuteAction::CreateSP(this, &FCustomizableObjectEditor::CompileObject, true),
+		FExecuteAction::CreateSP(this, &FCustomizableObjectEditor::CompileObject, true, false),
 		FCanExecuteAction::CreateStatic(&UCustomizableObjectSystem::IsActive),
 		FIsActionChecked());
 
@@ -632,6 +632,19 @@ void FCustomizableObjectEditor::BindCommands()
 		FCanExecuteAction(),
 		FIsActionChecked());
 
+	// References
+	ToolkitCommands->MapAction(
+		Commands.CompileGatherReferences,
+		FExecuteAction::CreateSP(this, &FCustomizableObjectEditor::CompileObject, false, true),
+		FCanExecuteAction(),
+		FIsActionChecked());
+
+	ToolkitCommands->MapAction(
+		Commands.ClearGatheredReferences,
+		FExecuteAction::CreateSP(this, &FCustomizableObjectEditor::ClearGatheredReferences),
+		FCanExecuteAction(),
+		FIsActionChecked());
+	
 	// Texture Analyzer
 	ToolkitCommands->MapAction(
 		Commands.TextureAnalyzer,
@@ -1181,7 +1194,6 @@ void FCustomizableObjectEditor::ExtendToolbar()
 				LOCTEXT("Compile_Options_Tooltip", "Change Compile Options"),
 				TAttribute<FSlateIcon>(),
 				true);
-
 			ToolbarBuilder.EndSection();
 			
 			ToolbarBuilder.BeginSection("Information");
@@ -1354,6 +1366,13 @@ TSharedRef<SWidget> FCustomizableObjectEditor::GenerateCompileOptionsMenuContent
 	}
 	MenuBuilder.EndSection();
 
+	MenuBuilder.BeginSection("References", LOCTEXT("References", "References"));
+	{
+		MenuBuilder.AddMenuEntry(FCustomizableObjectEditorCommands::Get().CompileGatherReferences);
+		MenuBuilder.AddMenuEntry(FCustomizableObjectEditorCommands::Get().ClearGatheredReferences);
+	}
+	MenuBuilder.EndSection();
+
 	return MenuBuilder.MakeWidget();
 }
 
@@ -1450,7 +1469,7 @@ void FCustomizableObjectEditor::OnObjectModified(UObject* Object)
 }
 
 
-void FCustomizableObjectEditor::CompileObject(bool bOnlySelectedParameters)
+void FCustomizableObjectEditor::CompileObject(bool bOnlySelectedParameters, bool bGatherReferences)
 {
 	// Resetting viewport parameters
 	Viewport->SetDrawDefaultUVMaterial();
@@ -1473,6 +1492,7 @@ void FCustomizableObjectEditor::CompileObject(bool bOnlySelectedParameters)
 
 	TSharedRef<FCompilationRequest> CompileRequest = MakeShared<FCompilationRequest>(*CustomizableObject, true);
 	CompileRequest->GetCompileOptions().bSilentCompilation = false;
+	CompileRequest->GetCompileOptions().bGatherReferences = bGatherReferences;
 
 	if (bOnlySelectedParameters && PreviewInstance)
 	{
@@ -1505,6 +1525,13 @@ void FCustomizableObjectEditor::DebugObject() const
 }
 
 
+void FCustomizableObjectEditor::ClearGatheredReferences()
+{
+	CustomizableObject->GetPrivate()->References = {};
+	CustomizableObject->Modify();
+}
+
+
 void FCustomizableObjectEditor::ResetCompileOptions()
 {
 	const FScopedTransaction Transaction(LOCTEXT("ResetCompilationOptionsTransaction", "Reset Compilation Options"));
diff --git a/Engine/Plugins/Experimental/Mutable/Source/CustomizableObjectEditor/Private/MuCOE/CustomizableObjectEditor.h b/Engine/Plugins/Experimental/Mutable/Source/CustomizableObjectEditor/Private/MuCOE/CustomizableObjectEditor.h
--- a/Engine/Plugins/Experimental/Mutable/Source/CustomizableObjectEditor/Private/MuCOE/CustomizableObjectEditor.h
+++ b/Engine/Plugins/Experimental/Mutable/Source/CustomizableObjectEditor/Private/MuCOE/CustomizableObjectEditor.h
@@ -264,6 +264,9 @@ public:
 
 	/** Debug the object as a raw mutable data in the internal tools. */
 	void DebugObject() const;
+
+	/** Clear game asset references saved in the Customizable Object. */
+	void ClearGatheredReferences();
 	
 private:
 	TSharedRef<SDockTab> SpawnTab_Viewport(const FSpawnTabArgs& Args);
@@ -282,8 +285,11 @@ private:
 	/** Command list for the graph editor */
 	TSharedPtr<FUICommandList> GraphEditorCommands;
 
-	// Compile the customizable object.
-	void CompileObject(bool bOnlySelectedParameters);
+	/** Compile the Customizable Object.
+	 * 
+	 * @param bOnlySelectedParameters If true, compile only selected int parameters.
+	 * @param bGatherReferences If true, also gather asset references and save them in the Customizable Object. Marks the objects as modified. */
+	void CompileObject(bool bOnlySelectedParameters, bool bGatherReferences);
 	
 	// Compile options menu callbacks
 	TSharedRef<SWidget> GenerateCompileOptionsMenuContent(TSharedRef<FUICommandList> InCommandList);
diff --git a/Engine/Plugins/Experimental/Mutable/Source/CustomizableObjectEditor/Private/MuCOE/CustomizableObjectEditorActions.cpp b/Engine/Plugins/Experimental/Mutable/Source/CustomizableObjectEditor/Private/MuCOE/CustomizableObjectEditorActions.cpp
--- a/Engine/Plugins/Experimental/Mutable/Source/CustomizableObjectEditor/Private/MuCOE/CustomizableObjectEditorActions.cpp
+++ b/Engine/Plugins/Experimental/Mutable/Source/CustomizableObjectEditor/Private/MuCOE/CustomizableObjectEditorActions.cpp
@@ -29,6 +29,9 @@ void FCustomizableObjectEditorCommands::RegisterCommands()
 
 	UI_COMMAND(PerformanceReport, "Performance Report", "Open the Performance Report window to set up and perform all tests relevant to Customizable Objects and access worst cases data.", EUserInterfaceActionType::Button, FInputChord());
 	UI_COMMAND(TextureAnalyzer, "Texture Memory Analyzer", "Open the Texture Analyzer window to check all the information of the textures created by Mutable.", EUserInterfaceActionType::Button, FInputChord());
+
+	UI_COMMAND(CompileGatherReferences, "Compile and Gather References", "Compile and gather all asset references used in this Customizable Object. Marks the object as modified.", EUserInterfaceActionType::Button, FInputChord());
+	UI_COMMAND(ClearGatheredReferences, "Clear References", "Clear gathered asset references. Marks the object as modified.", EUserInterfaceActionType::Button, FInputChord());
 }
 
 
diff --git a/Engine/Plugins/Experimental/Mutable/Source/CustomizableObjectEditor/Private/MuCOE/CustomizableObjectEditorFunctionLibrary.cpp b/Engine/Plugins/Experimental/Mutable/Source/CustomizableObjectEditor/Private/MuCOE/CustomizableObjectEditorFunctionLibrary.cpp
--- a/Engine/Plugins/Experimental/Mutable/Source/CustomizableObjectEditor/Private/MuCOE/CustomizableObjectEditorFunctionLibrary.cpp
+++ b/Engine/Plugins/Experimental/Mutable/Source/CustomizableObjectEditor/Private/MuCOE/CustomizableObjectEditorFunctionLibrary.cpp
@@ -16,7 +16,8 @@
 ECustomizableObjectCompilationState UCustomizableObjectEditorFunctionLibrary::CompileCustomizableObjectSynchronously(
 	UCustomizableObject* CustomizableObject,
 	ECustomizableObjectOptimizationLevel InOptimizationLevel,
-	ECustomizableObjectTextureCompression InTextureCompression)
+	ECustomizableObjectTextureCompression InTextureCompression,
+	bool bGatherReferences)
 {
 	// store package dirty state so that we can restore it - compile is not an edit:
 	const bool bPackageWasDirty = CustomizableObject->GetOutermost()->IsDirty();
@@ -28,6 +29,7 @@ ECustomizableObjectCompilationState UCustomizableObjectEditorFunctionLibrary::Co
 	Options.OptimizationLevel = static_cast<int32>(InOptimizationLevel);
 	Options.TextureCompression = InTextureCompression;
 	Options.bSilentCompilation = false;
+	Options.bGatherReferences = bGatherReferences;
 	ICustomizableObjectEditorModule::GetChecked().CompileCustomizableObject(CompileRequest);
 
 	check(CompileRequest->GetCompilationState() == ECompilationStatePrivate::Completed);
diff --git a/Engine/Plugins/Experimental/Mutable/Source/CustomizableObjectEditor/Private/MuCOE/CustomizableObjectEditorModule.cpp b/Engine/Plugins/Experimental/Mutable/Source/CustomizableObjectEditor/Private/MuCOE/CustomizableObjectEditorModule.cpp
--- a/Engine/Plugins/Experimental/Mutable/Source/CustomizableObjectEditor/Private/MuCOE/CustomizableObjectEditorModule.cpp
+++ b/Engine/Plugins/Experimental/Mutable/Source/CustomizableObjectEditor/Private/MuCOE/CustomizableObjectEditorModule.cpp
@@ -549,7 +549,7 @@ bool FCustomizableObjectEditorModule::IsCompilationOutOfDate(const UCustomizable
 		}
 		else
 		{
-			return false;
+			return true;
 		}
 	}
 	
@@ -572,7 +572,7 @@ bool FCustomizableObjectEditorModule::IsCompilationOutOfDate(const UCustomizable
 				}
 				else
 				{
-					return false;
+					return true;
 				}
 			}
 		}
@@ -589,7 +589,7 @@ bool FCustomizableObjectEditorModule::IsCompilationOutOfDate(const UCustomizable
 				}
 				else
 				{
-					return false;
+					return true;
 				}
 			}
 
@@ -605,7 +605,7 @@ bool FCustomizableObjectEditorModule::IsCompilationOutOfDate(const UCustomizable
 				}
 				else
 				{
-					return false;
+					return true;
 				}
 			}
 		}
@@ -627,12 +627,19 @@ bool FCustomizableObjectEditorModule::IsCompilationOutOfDate(const UCustomizable
 			}
 			else
 			{
-				return false;
+				return true;
 			}
 		}
 	}
 
-	return !OutOfDatePackages->IsEmpty();	
+	if (OutOfDatePackages)
+	{
+		return !OutOfDatePackages->IsEmpty();	
+	}
+	else
+	{
+		return false;
+	}
 }
 
 
@@ -661,14 +668,6 @@ void FCustomizableObjectEditorModule::BakeCustomizableObjectInstance(UCustomizab
 }
 
 
-void CompileCustomizableObjectsSync(const TSharedRef<FCompilationRequest>& CompilationRequest)
-{
-	FCustomizableObjectCompiler* SyncCompiler = new FCustomizableObjectCompiler();
-	SyncCompiler->Compile(CompilationRequest);
-	delete SyncCompiler;
-}
-
-
 void FCustomizableObjectEditorModule::CompileCustomizableObject(const TSharedRef<FCompilationRequest>& InCompilationRequest, bool bForceRequest)
 {
 	if (IsRunningGame())
@@ -676,14 +675,7 @@ void FCustomizableObjectEditorModule::CompileCustomizableObject(const TSharedRef
 		return;
 	}
 
-	if (InCompilationRequest->IsAsyncCompilation())
-	{
-		CompileCustomizableObjects({ InCompilationRequest }, bForceRequest);
-	}
-	else
-	{
-		CompileCustomizableObjectsSync(InCompilationRequest);
-	}
+	CompileCustomizableObjects({ InCompilationRequest }, bForceRequest);
 }
 
 void FCustomizableObjectEditorModule::CompileCustomizableObjects(const TArray<TSharedRef<FCompilationRequest>>& InCompilationRequests, bool bForceRequests)
@@ -705,10 +697,12 @@ void FCustomizableObjectEditorModule::CompileCustomizableObjects(const TArray<TS
 		{
 			continue;
 		}
-
+		
 		if (!Request->IsAsyncCompilation())
 		{
-			CompileCustomizableObjectsSync(Request);
+			FCustomizableObjectCompiler* SyncCompiler = new FCustomizableObjectCompiler();
+			SyncCompiler->Compile(Request);
+			delete SyncCompiler;
 		}
 
 		else if (bForceRequests ||
diff --git a/Engine/Plugins/Experimental/Mutable/Source/CustomizableObjectEditor/Private/MuCOE/CustomizableObjectEditorStyle.cpp b/Engine/Plugins/Experimental/Mutable/Source/CustomizableObjectEditor/Private/MuCOE/CustomizableObjectEditorStyle.cpp
--- a/Engine/Plugins/Experimental/Mutable/Source/CustomizableObjectEditor/Private/MuCOE/CustomizableObjectEditorStyle.cpp
+++ b/Engine/Plugins/Experimental/Mutable/Source/CustomizableObjectEditor/Private/MuCOE/CustomizableObjectEditorStyle.cpp
@@ -71,89 +71,94 @@ TSharedRef<FSlateStyleSet> FCustomizableObjectEditorStyle::Create()
 	const FVector2D Icon64x64(64.0f, 64.0f);
 
 	TSharedRef< FSlateStyleSet > Style = MakeShareable(new FSlateStyleSet("CustomizableObjectEditorStyle"));
-	Style->SetCoreContentRoot(FPaths::EngineContentDir() / TEXT("Editor/Slate"));
+	Style->SetCoreContentRoot(FPaths::EngineContentDir());
 	Style->SetContentRoot(FPaths::EnginePluginsDir() / TEXT("Experimental/Mutable/Content"));
 	
-	Style->Set("CustomizableObjectEditor.Save", new IMAGE_BRUSH_SVG(Style, "Starship/Common/SaveCurrent", Icon40x40));
-	Style->Set("CustomizableObjectEditor.Save.Small", new IMAGE_BRUSH_SVG(Style, "Starship/Common/SaveCurrent", Icon20x20));
-	Style->Set("CustomizableObjectEditor.Compile", new IMAGE_BRUSH_SVG(Style, "Starship/MainToolbar/compile", Icon40x40));
-	Style->Set("CustomizableObjectEditor.Compile.Small", new IMAGE_BRUSH_SVG(Style, "Starship/MainToolbar/compile", Icon20x20));
-	Style->Set("CustomizableObjectEditor.CompileOnlySelected", new IMAGE_BRUSH_SVG(Style, "Starship/MainToolbar/compile", Icon40x40));
-	Style->Set("CustomizableObjectEditor.CompileOnlySelected.Small", new IMAGE_BRUSH_SVG(Style, "Starship/MainToolbar/compile", Icon20x20));
-	Style->Set("CustomizableObjectEditor.ResetCompileOptions", new IMAGE_BRUSH_SVG(Style, "Starship/MainToolbar/compile", Icon40x40));
-	Style->Set("CustomizableObjectEditor.ResetCompileOptions.Small", new IMAGE_BRUSH_SVG(Style, "Starship/MainToolbar/compile", Icon20x20));
-	Style->Set("CustomizableObjectEditor.Debug", new IMAGE_BRUSH_SVG(Style, "Starship/Common/Debug", Icon40x40));
-	Style->Set("CustomizableObjectEditor.Debug.Small", new IMAGE_BRUSH_SVG(Style, "Starship/Common/Debug", Icon20x20));
-	Style->Set("CustomizableObjectEditor.TextureAnalyzer", new IMAGE_BRUSH_SVG(Style, "Starship/Common/IssueTracker", Icon40x40));
-	Style->Set("CustomizableObjectEditor.TextureAnalyzer.Small", new IMAGE_BRUSH_SVG(Style, "Starship/Common/IssueTracker", Icon20x20));
-	Style->Set("CustomizableObjectEditor.PerformanceReport", new IMAGE_BRUSH_SVG(Style, "Starship/MainToolbar/simulate", Icon40x40));
-	Style->Set("CustomizableObjectEditor.PerformanceReport.Small", new IMAGE_BRUSH_SVG(Style, "Starship/MainToolbar/simulate", Icon20x20));
-
-	Style->Set("CustomizableObjectEditor.Tabs.Preview", new IMAGE_BRUSH_SVG(Style, "Starship/Persona/PersonaPreviewAnimation", Icon20x20));
-	Style->Set("CustomizableObjectEditor.Tabs.PreviewSettings", new IMAGE_BRUSH_SVG(Style, "Starship/AssetIcons/CameraActor_64", Icon20x20));
-	Style->Set("CustomizableObjectEditor.Tabs.CustomizableObjectProperties", new IMAGE_BRUSH_SVG(Style, "Starship/Persona/SectionSelection", Icon20x20));
-	Style->Set("CustomizableObjectEditor.Tabs.CustomizableInstanceProperties", new IMAGE_BRUSH_SVG(Style, "Starship/Common/Details", Icon20x20));
-	Style->Set("CustomizableObjectEditor.Tabs.NodeGraph", new IMAGE_BRUSH_SVG(Style, "Starship/Common/blueprint", Icon20x20));
-	Style->Set("CustomizableObjectEditor.Tabs.NodeProperties", new IMAGE_BRUSH_SVG(Style, "Starship/Blueprints/icon_BlueprintEditor_EventGraph", Icon20x20));
-
-	Style->Set("CustomizableObjectInstanceEditor.SetShowWireframe", new IMAGE_BRUSH(Style, "Icons/icon_StaticMeshEd_Wireframe_40x", Icon40x40));
-	Style->Set("CustomizableObjectInstanceEditor.SetShowWireframe.Small", new IMAGE_BRUSH(Style, "Icons/icon_StaticMeshEd_Wireframe_40x", Icon20x20));
-
-	Style->Set("CustomizableObjectInstanceEditor.Save", new IMAGE_BRUSH_SVG(Style, "Starship/Common/SaveCurrent", Icon40x40));
-	Style->Set("CustomizableObjectInstanceEditor.Save.Small", new IMAGE_BRUSH_SVG(Style, "Starship/Common/SaveCurrent", Icon20x20));
-	Style->Set("CustomizableObjectInstanceEditor.ShowParentCO", new IMAGE_BRUSH_SVG(Style, "Starship/MainToolbar/eject", Icon40x40));
-	Style->Set("CustomizableObjectInstanceEditor.ShowParentCO.Small", new IMAGE_BRUSH_SVG(Style, "Starship/MainToolbar/eject", Icon20x20));
-	Style->Set("CustomizableObjectInstanceEditor.EditParentCO", new IMAGE_BRUSH_SVG(Style, "Starship/Common/EditorModes", Icon40x40));
-	Style->Set("CustomizableObjectInstanceEditor.EditParentCO.Small", new IMAGE_BRUSH_SVG(Style, "Starship/Common/EditorModes", Icon20x20));
-	Style->Set("CustomizableObjectInstanceEditor.TextureAnalyzer", new IMAGE_BRUSH_SVG(Style, "Starship/Common/IssueTracker", Icon40x40));
-	Style->Set("CustomizableObjectInstanceEditor.TextureAnalyzer.Small", new IMAGE_BRUSH_SVG(Style, "Starship/Common/IssueTracker", Icon20x20));
-
-	Style->Set("CustomizableObjectEditorViewport.SetShowWireframe", new IMAGE_BRUSH(Style, "Icons/icon_StaticMeshEd_Wireframe_40x", Icon40x40));
-	Style->Set("CustomizableObjectEditorViewport.SetShowWireframe.Small", new IMAGE_BRUSH(Style, "Icons/icon_StaticMeshEd_Wireframe_40x", Icon20x20));
-	Style->Set("CustomizableObjectEditorViewport.SetRealtimePreview", new IMAGE_BRUSH(Style, "Icons/icon_MatEd_Realtime_40x", Icon40x40));
-	Style->Set("CustomizableObjectEditorViewport.SetRealtimePreview.Small", new IMAGE_BRUSH(Style, "Icons/icon_MatEd_Realtime_40x", Icon20x20));
-	Style->Set("CustomizableObjectEditorViewport.SetShowBounds", new IMAGE_BRUSH(Style, "Icons/icon_StaticMeshEd_Bounds_40x", Icon40x40));
-	Style->Set("CustomizableObjectEditorViewport.SetShowBounds.Small", new IMAGE_BRUSH(Style, "Icons/icon_StaticMeshEd_Bounds_40x", Icon20x20));
-	Style->Set("CustomizableObjectEditorViewport.SetShowNormals", new IMAGE_BRUSH(Style, "Icons/icon_StaticMeshEd_Normals_40x", Icon40x40));
-	Style->Set("CustomizableObjectEditorViewport.SetShowNormals.Small", new IMAGE_BRUSH(Style, "Icons/icon_StaticMeshEd_Normals_40x", Icon20x20));
-	Style->Set("CustomizableObjectEditorViewport.SetShowTangents", new IMAGE_BRUSH(Style, "Icons/icon_StaticMeshEd_Tangents_40x", Icon40x40));
-	Style->Set("CustomizableObjectEditorViewport.SetShowTangents.Small", new IMAGE_BRUSH(Style, "Icons/icon_StaticMeshEd_Tangents_40x", Icon20x20));
-	Style->Set("CustomizableObjectEditorViewport.SetShowBinormals", new IMAGE_BRUSH(Style, "Icons/icon_StaticMeshEd_Binormals_40x", Icon40x40));
-	Style->Set("CustomizableObjectEditorViewport.SetShowBinormals.Small", new IMAGE_BRUSH(Style, "Icons/icon_StaticMeshEd_Binormals_40x", Icon20x20));
-	Style->Set("CustomizableObjectEditorViewport.SaveThumbnail", new IMAGE_BRUSH(Style, "Icons/icon_StaticMeshEd_SaveThumbnail_40x", Icon40x40));
-	Style->Set("CustomizableObjectEditorViewport.SaveThumbnail.Small", new IMAGE_BRUSH(Style, "Icons/icon_StaticMeshEd_SaveThumbnail_40x", Icon20x20));
-	Style->Set("CustomizableObjectEditorViewport.SetShowPivot", new IMAGE_BRUSH(Style, "Icons/icon_StaticMeshEd_ShowPivot_40x", Icon40x40));
-	Style->Set("CustomizableObjectEditorViewport.SetShowPivot.Small", new IMAGE_BRUSH(Style, "Icons/icon_StaticMeshEd_ShowPivot_40x", Icon20x20));
-
-	Style->Set("CustomizableObjectEditorViewport.SetShowGrid", new IMAGE_BRUSH_SVG(Style, "Starship/Common/Grid", Icon20x20));
-	Style->Set("CustomizableObjectEditorViewport.SetShowSky", new IMAGE_BRUSH_SVG(Style, "Starship/Common/Reflections", Icon20x20));
-	Style->Set("CustomizableObjectEditorViewport.SetDrawUVs", new IMAGE_BRUSH_SVG(Style, "Starship/Common/SetDrawUVs", Icon20x20));
+	Style->Set("CustomizableObjectEditor.Save", new IMAGE_BRUSH_SVG(Style, "Editor/Slate/Starship/Common/SaveCurrent", Icon40x40));
+	Style->Set("CustomizableObjectEditor.Save.Small", new IMAGE_BRUSH_SVG(Style, "Editor/Slate/Starship/Common/SaveCurrent", Icon20x20));
+	Style->Set("CustomizableObjectEditor.Compile", new IMAGE_BRUSH_SVG(Style, "Editor/Slate/Starship/MainToolbar/compile", Icon40x40));
+	Style->Set("CustomizableObjectEditor.Compile.Small", new IMAGE_BRUSH_SVG(Style, "Editor/Slate/Starship/MainToolbar/compile", Icon20x20));
+	Style->Set("CustomizableObjectEditor.CompileOnlySelected", new IMAGE_BRUSH_SVG(Style, "Editor/Slate/Starship/MainToolbar/compile", Icon40x40));
+	Style->Set("CustomizableObjectEditor.CompileOnlySelected.Small", new IMAGE_BRUSH_SVG(Style, "Editor/Slate/Starship/MainToolbar/compile", Icon20x20));
+	Style->Set("CustomizableObjectEditor.ResetCompileOptions", new IMAGE_BRUSH_SVG(Style, "Editor/Slate/Starship/MainToolbar/compile", Icon40x40));
+	Style->Set("CustomizableObjectEditor.ResetCompileOptions.Small", new IMAGE_BRUSH_SVG(Style, "Editor/Slate/Starship/MainToolbar/compile", Icon20x20));
+	Style->Set("CustomizableObjectEditor.Debug", new IMAGE_BRUSH_SVG(Style, "Editor/Slate/Starship/Common/Debug", Icon40x40));
+	Style->Set("CustomizableObjectEditor.Debug.Small", new IMAGE_BRUSH_SVG(Style, "Editor/Slate/Starship/Common/Debug", Icon20x20));
+	Style->Set("CustomizableObjectEditor.TextureAnalyzer", new IMAGE_BRUSH_SVG(Style, "Editor/Slate/Starship/Common/IssueTracker", Icon40x40));
+	Style->Set("CustomizableObjectEditor.TextureAnalyzer.Small", new IMAGE_BRUSH_SVG(Style, "Editor/Slate/Starship/Common/IssueTracker", Icon20x20));
+	Style->Set("CustomizableObjectEditor.PerformanceReport", new IMAGE_BRUSH_SVG(Style, "Editor/Slate/Starship/MainToolbar/simulate", Icon40x40));
+	Style->Set("CustomizableObjectEditor.PerformanceReport.Small", new IMAGE_BRUSH_SVG(Style, "Editor/Slate/Starship/MainToolbar/simulate", Icon20x20));
+
+	Style->Set("CustomizableObjectEditor.CompileGatherReferences", new IMAGE_BRUSH_SVG(Style, "Editor/Slate/Starship/MainToolbar/compile", Icon40x40));
+	Style->Set("CustomizableObjectEditor.CompileGatherReferences.Small", new IMAGE_BRUSH_SVG(Style, "Editor/Slate/Starship/MainToolbar/compile", Icon20x20));
+	Style->Set("CustomizableObjectEditor.ClearGatheredReferences", new IMAGE_BRUSH_SVG(Style, "Slate/Starship/Common/Delete", Icon40x40));
+	Style->Set("CustomizableObjectEditor.ClearGatheredReferences.Small", new IMAGE_BRUSH_SVG(Style, "Slate/Starship/Common/Delete", Icon20x20));
+	
+	Style->Set("CustomizableObjectEditor.Tabs.Preview", new IMAGE_BRUSH_SVG(Style, "Editor/Slate/Starship/Persona/PersonaPreviewAnimation", Icon20x20));
+	Style->Set("CustomizableObjectEditor.Tabs.PreviewSettings", new IMAGE_BRUSH_SVG(Style, "Editor/Slate/Starship/AssetIcons/CameraActor_64", Icon20x20));
+	Style->Set("CustomizableObjectEditor.Tabs.CustomizableObjectProperties", new IMAGE_BRUSH_SVG(Style, "Editor/Slate/Starship/Persona/SectionSelection", Icon20x20));
+	Style->Set("CustomizableObjectEditor.Tabs.CustomizableInstanceProperties", new IMAGE_BRUSH_SVG(Style, "Editor/Slate/Starship/Common/Details", Icon20x20));
+	Style->Set("CustomizableObjectEditor.Tabs.NodeGraph", new IMAGE_BRUSH_SVG(Style, "Editor/Slate/Starship/Common/blueprint", Icon20x20));
+	Style->Set("CustomizableObjectEditor.Tabs.NodeProperties", new IMAGE_BRUSH_SVG(Style, "Editor/Slate/Starship/Blueprints/icon_BlueprintEditor_EventGraph", Icon20x20));
+
+	Style->Set("CustomizableObjectInstanceEditor.SetShowWireframe", new IMAGE_BRUSH(Style, "Editor/Slate/Icons/icon_StaticMeshEd_Wireframe_40x", Icon40x40));
+	Style->Set("CustomizableObjectInstanceEditor.SetShowWireframe.Small", new IMAGE_BRUSH(Style, "Editor/Slate/Icons/icon_StaticMeshEd_Wireframe_40x", Icon20x20));
+
+	Style->Set("CustomizableObjectInstanceEditor.Save", new IMAGE_BRUSH_SVG(Style, "Editor/Slate/Starship/Common/SaveCurrent", Icon40x40));
+	Style->Set("CustomizableObjectInstanceEditor.Save.Small", new IMAGE_BRUSH_SVG(Style, "Editor/Slate/Starship/Common/SaveCurrent", Icon20x20));
+	Style->Set("CustomizableObjectInstanceEditor.ShowParentCO", new IMAGE_BRUSH_SVG(Style, "Editor/Slate/Starship/MainToolbar/eject", Icon40x40));
+	Style->Set("CustomizableObjectInstanceEditor.ShowParentCO.Small", new IMAGE_BRUSH_SVG(Style, "Editor/Slate/Starship/MainToolbar/eject", Icon20x20));
+	Style->Set("CustomizableObjectInstanceEditor.EditParentCO", new IMAGE_BRUSH_SVG(Style, "Editor/Slate/Starship/Common/EditorModes", Icon40x40));
+	Style->Set("CustomizableObjectInstanceEditor.EditParentCO.Small", new IMAGE_BRUSH_SVG(Style, "Editor/Slate/Starship/Common/EditorModes", Icon20x20));
+	Style->Set("CustomizableObjectInstanceEditor.TextureAnalyzer", new IMAGE_BRUSH_SVG(Style, "Editor/Slate/Starship/Common/IssueTracker", Icon40x40));
+	Style->Set("CustomizableObjectInstanceEditor.TextureAnalyzer.Small", new IMAGE_BRUSH_SVG(Style, "Editor/Slate/Starship/Common/IssueTracker", Icon20x20));
+
+	Style->Set("CustomizableObjectEditorViewport.SetShowWireframe", new IMAGE_BRUSH(Style, "Editor/Slate/Icons/icon_StaticMeshEd_Wireframe_40x", Icon40x40));
+	Style->Set("CustomizableObjectEditorViewport.SetShowWireframe.Small", new IMAGE_BRUSH(Style, "Editor/Slate/Icons/icon_StaticMeshEd_Wireframe_40x", Icon20x20));
+	Style->Set("CustomizableObjectEditorViewport.SetRealtimePreview", new IMAGE_BRUSH(Style, "Editor/Slate/Icons/icon_MatEd_Realtime_40x", Icon40x40));
+	Style->Set("CustomizableObjectEditorViewport.SetRealtimePreview.Small", new IMAGE_BRUSH(Style, "Editor/Slate/Icons/icon_MatEd_Realtime_40x", Icon20x20));
+	Style->Set("CustomizableObjectEditorViewport.SetShowBounds", new IMAGE_BRUSH(Style, "Editor/Slate/Icons/icon_StaticMeshEd_Bounds_40x", Icon40x40));
+	Style->Set("CustomizableObjectEditorViewport.SetShowBounds.Small", new IMAGE_BRUSH(Style, "Editor/Slate/Icons/icon_StaticMeshEd_Bounds_40x", Icon20x20));
+	Style->Set("CustomizableObjectEditorViewport.SetShowNormals", new IMAGE_BRUSH(Style, "Editor/Slate/Icons/icon_StaticMeshEd_Normals_40x", Icon40x40));
+	Style->Set("CustomizableObjectEditorViewport.SetShowNormals.Small", new IMAGE_BRUSH(Style, "Editor/Slate/Icons/icon_StaticMeshEd_Normals_40x", Icon20x20));
+	Style->Set("CustomizableObjectEditorViewport.SetShowTangents", new IMAGE_BRUSH(Style, "Editor/Slate/Icons/icon_StaticMeshEd_Tangents_40x", Icon40x40));
+	Style->Set("CustomizableObjectEditorViewport.SetShowTangents.Small", new IMAGE_BRUSH(Style, "Editor/Slate/Icons/icon_StaticMeshEd_Tangents_40x", Icon20x20));
+	Style->Set("CustomizableObjectEditorViewport.SetShowBinormals", new IMAGE_BRUSH(Style, "Editor/Slate/Icons/icon_StaticMeshEd_Binormals_40x", Icon40x40));
+	Style->Set("CustomizableObjectEditorViewport.SetShowBinormals.Small", new IMAGE_BRUSH(Style, "Editor/Slate/Icons/icon_StaticMeshEd_Binormals_40x", Icon20x20));
+	Style->Set("CustomizableObjectEditorViewport.SaveThumbnail", new IMAGE_BRUSH(Style, "Editor/Slate/Icons/icon_StaticMeshEd_SaveThumbnail_40x", Icon40x40));
+	Style->Set("CustomizableObjectEditorViewport.SaveThumbnail.Small", new IMAGE_BRUSH(Style, "Editor/Slate/Icons/icon_StaticMeshEd_SaveThumbnail_40x", Icon20x20));
+	Style->Set("CustomizableObjectEditorViewport.SetShowPivot", new IMAGE_BRUSH(Style, "Editor/Slate/Icons/icon_StaticMeshEd_ShowPivot_40x", Icon40x40));
+	Style->Set("CustomizableObjectEditorViewport.SetShowPivot.Small", new IMAGE_BRUSH(Style, "Editor/Slate/Icons/icon_StaticMeshEd_ShowPivot_40x", Icon20x20));
+
+	Style->Set("CustomizableObjectEditorViewport.SetShowGrid", new IMAGE_BRUSH_SVG(Style, "Editor/Slate/Starship/Common/Grid", Icon20x20));
+	Style->Set("CustomizableObjectEditorViewport.SetShowSky", new IMAGE_BRUSH_SVG(Style, "Editor/Slate/Starship/Common/Reflections", Icon20x20));
+	Style->Set("CustomizableObjectEditorViewport.SetDrawUVs", new IMAGE_BRUSH_SVG(Style, "Editor/Slate/Starship/Common/SetDrawUVs", Icon20x20));
 	Style->Set("CustomizableObjectEditorViewport.BakeInstance", new IMAGE_PLUGIN_BRUSH(Style, "Icons/BakerIcon", Icon40x40));
 
-	Style->Set("COEditorViewportLODCommands.TranslateMode", new IMAGE_BRUSH_SVG(Style, "Starship/EditorViewport/translate", Icon20x20));
-	Style->Set("COEditorViewportLODCommands.RotateMode", new IMAGE_BRUSH_SVG(Style, "Starship/EditorViewport/rotate", Icon20x20));
-	Style->Set("COEditorViewportLODCommands.ScaleMode", new IMAGE_BRUSH_SVG(Style, "Starship/EditorViewport/scale", Icon20x20));
-	Style->Set("COEditorViewportLODCommands.RotationGridSnap", new IMAGE_BRUSH_SVG(Style, "Starship/EditorViewport/angle", Icon20x20));
-
-	Style->Set("LayoutEditorCommands.AddBlock", new IMAGE_BRUSH_SVG(Style, "Starship/Common/Layout", Icon40x40));
-	Style->Set("LayoutEditorCommands.AddBlock.Small", new IMAGE_BRUSH_SVG(Style, "Starship/Common/Layout", Icon20x20));
-	Style->Set("LayoutEditorCommands.RemoveBlock", new IMAGE_BRUSH_SVG(Style, "Starship/Common/LayoutRemove", Icon40x40));
-	Style->Set("LayoutEditorCommands.RemoveBlock.Small", new IMAGE_BRUSH_SVG(Style, "Starship/Common/LayoutRemove", Icon20x20));
-	Style->Set("LayoutEditorCommands.GenerateBlocks", new IMAGE_BRUSH_SVG(Style, "Starship/Persona/AnimationPreviewMesh", Icon40x40));
-	Style->Set("LayoutEditorCommands.GenerateBlocks.Small", new IMAGE_BRUSH_SVG(Style, "Starship/Persona/AnimationPreviewMesh", Icon20x20));
-
-	Style->Set("LayoutBlockSelector.SelectAll", new IMAGE_BRUSH_SVG(Style, "Starship/Common/LayoutLoad", Icon40x40));
-	Style->Set("LayoutBlockSelector.SelectAll.Small", new IMAGE_BRUSH_SVG(Style, "Starship/Common/LayoutLoad", Icon20x20));
-	Style->Set("LayoutBlockSelector.SelectNone", new IMAGE_BRUSH_SVG(Style, "Starship/Common/LayoutSave", Icon40x40));
-	Style->Set("LayoutBlockSelector.SelectNone.Small", new IMAGE_BRUSH_SVG(Style, "Starship/Common/LayoutSave", Icon20x20));
-
-	Style->Set("Nodes.ArrowUp", new IMAGE_BRUSH(Style, "Old/ArrowUp", Icon16x16));
-	Style->Set("Nodes.ArrowDown", new IMAGE_BRUSH(Style, "Old/ArrowDown", Icon16x16));
-
-	Style->Set("CustomizableObjectDebugger.GenerateMutableGraph", new IMAGE_BRUSH_SVG(Style, "Starship/Common/blueprint", Icon40x40));
-	Style->Set("CustomizableObjectDebugger.GenerateMutableGraph.Small", new IMAGE_BRUSH_SVG(Style, "Starship/Common/blueprint", Icon20x20));
-	Style->Set("CustomizableObjectDebugger.CompileMutableCode", new IMAGE_BRUSH_SVG(Style, "Starship/MainToolbar/Compile", Icon40x40));
-	Style->Set("CustomizableObjectDebugger.CompileMutableCode.Small", new IMAGE_BRUSH_SVG(Style, "Starship/MainToolbar/Compile", Icon20x20));
+	Style->Set("COEditorViewportLODCommands.TranslateMode", new IMAGE_BRUSH_SVG(Style, "Editor/Slate/Starship/EditorViewport/translate", Icon20x20));
+	Style->Set("COEditorViewportLODCommands.RotateMode", new IMAGE_BRUSH_SVG(Style, "Editor/Slate/Starship/EditorViewport/rotate", Icon20x20));
+	Style->Set("COEditorViewportLODCommands.ScaleMode", new IMAGE_BRUSH_SVG(Style, "Editor/Slate/Starship/EditorViewport/scale", Icon20x20));
+	Style->Set("COEditorViewportLODCommands.RotationGridSnap", new IMAGE_BRUSH_SVG(Style, "Editor/Slate/Starship/EditorViewport/angle", Icon20x20));
+
+	Style->Set("LayoutEditorCommands.AddBlock", new IMAGE_BRUSH_SVG(Style, "Editor/Slate/Starship/Common/Layout", Icon40x40));
+	Style->Set("LayoutEditorCommands.AddBlock.Small", new IMAGE_BRUSH_SVG(Style, "Editor/Slate/Starship/Common/Layout", Icon20x20));
+	Style->Set("LayoutEditorCommands.RemoveBlock", new IMAGE_BRUSH_SVG(Style, "Editor/Slate/Starship/Common/LayoutRemove", Icon40x40));
+	Style->Set("LayoutEditorCommands.RemoveBlock.Small", new IMAGE_BRUSH_SVG(Style, "Editor/Slate/Starship/Common/LayoutRemove", Icon20x20));
+	Style->Set("LayoutEditorCommands.GenerateBlocks", new IMAGE_BRUSH_SVG(Style, "Editor/Slate/Starship/Persona/AnimationPreviewMesh", Icon40x40));
+	Style->Set("LayoutEditorCommands.GenerateBlocks.Small", new IMAGE_BRUSH_SVG(Style, "Editor/Slate/Starship/Persona/AnimationPreviewMesh", Icon20x20));
+
+	Style->Set("LayoutBlockSelector.SelectAll", new IMAGE_BRUSH_SVG(Style, "Editor/Slate/Starship/Common/LayoutLoad", Icon40x40));
+	Style->Set("LayoutBlockSelector.SelectAll.Small", new IMAGE_BRUSH_SVG(Style, "Editor/Slate/Starship/Common/LayoutLoad", Icon20x20));
+	Style->Set("LayoutBlockSelector.SelectNone", new IMAGE_BRUSH_SVG(Style, "Editor/Slate/Starship/Common/LayoutSave", Icon40x40));
+	Style->Set("LayoutBlockSelector.SelectNone.Small", new IMAGE_BRUSH_SVG(Style, "Editor/Slate/Starship/Common/LayoutSave", Icon20x20));
+
+	Style->Set("Nodes.ArrowUp", new IMAGE_BRUSH(Style, "Editor/Slate/Old/ArrowUp", Icon16x16));
+	Style->Set("Nodes.ArrowDown", new IMAGE_BRUSH(Style, "Editor/Slate/Old/ArrowDown", Icon16x16));
+
+	Style->Set("CustomizableObjectDebugger.GenerateMutableGraph", new IMAGE_BRUSH_SVG(Style, "Editor/Slate/Starship/Common/blueprint", Icon40x40));
+	Style->Set("CustomizableObjectDebugger.GenerateMutableGraph.Small", new IMAGE_BRUSH_SVG(Style, "Editor/Slate/Starship/Common/blueprint", Icon20x20));
+	Style->Set("CustomizableObjectDebugger.CompileMutableCode", new IMAGE_BRUSH_SVG(Style, "Editor/Slate/Starship/MainToolbar/Compile", Icon40x40));
+	Style->Set("CustomizableObjectDebugger.CompileMutableCode.Small", new IMAGE_BRUSH_SVG(Style, "Editor/Slate/Starship/MainToolbar/Compile", Icon20x20));
 	
 	FTableRowStyle PerformanceReportRowStyle = FAppStyle::Get().GetWidgetStyle<FTableRowStyle>("TableView.Row");
 	PerformanceReportRowStyle.SetEvenRowBackgroundBrush(*FAppStyle::Get().GetBrush("Brushes.AccentBrown"));
diff --git a/Engine/Plugins/Experimental/Mutable/Source/CustomizableObjectEditor/Public/MuCOE/CustomizableObjectEditorActions.h b/Engine/Plugins/Experimental/Mutable/Source/CustomizableObjectEditor/Public/MuCOE/CustomizableObjectEditorActions.h
--- a/Engine/Plugins/Experimental/Mutable/Source/CustomizableObjectEditor/Public/MuCOE/CustomizableObjectEditorActions.h
+++ b/Engine/Plugins/Experimental/Mutable/Source/CustomizableObjectEditor/Public/MuCOE/CustomizableObjectEditorActions.h
@@ -11,7 +11,6 @@ class FUICommandInfo;
  */
 class FCustomizableObjectEditorCommands : public TCommands<FCustomizableObjectEditorCommands>
 {
-
 public:
 	FCustomizableObjectEditorCommands();
 	
@@ -25,15 +24,16 @@ public:
 	TSharedPtr< FUICommandInfo > ResetPerformanceReportOptions;
 	TSharedPtr< FUICommandInfo > TextureAnalyzer;
 
+	TSharedPtr< FUICommandInfo > CompileGatherReferences;
+	TSharedPtr< FUICommandInfo > ClearGatheredReferences;
+	
 	/**
 	 * Initialize commands
 	 */
 	virtual void RegisterCommands() override;
-
 };
 
 
-
 /**
  * 
  */
diff --git a/Engine/Plugins/Experimental/Mutable/Source/CustomizableObjectEditor/Public/MuCOE/CustomizableObjectEditorFunctionLibrary.h b/Engine/Plugins/Experimental/Mutable/Source/CustomizableObjectEditor/Public/MuCOE/CustomizableObjectEditorFunctionLibrary.h
--- a/Engine/Plugins/Experimental/Mutable/Source/CustomizableObjectEditor/Public/MuCOE/CustomizableObjectEditorFunctionLibrary.h
+++ b/Engine/Plugins/Experimental/Mutable/Source/CustomizableObjectEditor/Public/MuCOE/CustomizableObjectEditorFunctionLibrary.h
@@ -71,7 +71,8 @@ public:
 	static ECustomizableObjectCompilationState CompileCustomizableObjectSynchronously(
 		UCustomizableObject* CustomizableObject, 
 		ECustomizableObjectOptimizationLevel OptimizationLevel = ECustomizableObjectOptimizationLevel::Minimal, 
-		ECustomizableObjectTextureCompression TextureCompression = ECustomizableObjectTextureCompression::Fast);
+		ECustomizableObjectTextureCompression TextureCompression = ECustomizableObjectTextureCompression::Fast,
+		bool bGatherReferences = false);
 
 	/** Create a new Customizable Object inside a package. */
 	UFUNCTION(BlueprintCallable, Category = "CustomizableObject")
