[Mutable] Added pin connection changed notification propagation.

- PinConnectionListChanged and NodeConnectionListChanged now are propagated through Import/Export and Reroute nodes.
- NodeSwitch now removes all pins if it does not have an NodeEnumParameter connected.
- Fixed issue where FollowInputPin was not discovering all pins.
- Exposed FollowPinArray.
- Added ReverseFollowPinArray.

#rnx

diff --git a/Engine/Plugins/Experimental/Mutable/Source/CustomizableObjectEditor/Private/MuCOE/EdGraphSchema_CustomizableObject.cpp b/Engine/Plugins/Experimental/Mutable/Source/CustomizableObjectEditor/Private/MuCOE/EdGraphSchema_CustomizableObject.cpp
--- a/Engine/Plugins/Experimental/Mutable/Source/CustomizableObjectEditor/Private/MuCOE/EdGraphSchema_CustomizableObject.cpp
+++ b/Engine/Plugins/Experimental/Mutable/Source/CustomizableObjectEditor/Private/MuCOE/EdGraphSchema_CustomizableObject.cpp
@@ -84,6 +84,7 @@
 #include "MuCOE/Nodes/CustomizableObjectNodeReroute.h"
 #include "Toolkits/ToolkitManager.h"
 #include "PropertyEditorModule.h"
+#include "MuCOE/GraphTraversal.h"
 
 class IToolkit;
 
@@ -922,11 +923,38 @@ void UEdGraphSchema_CustomizableObject::GetContextMenuActions(UToolMenu* Menu, U
 }
 
 
+void UEdGraphSchema_CustomizableObject::BreakNodeLinks(UEdGraphNode& TargetNode) const
+{
+#if WITH_EDITOR
+	TArray<UEdGraphPin*> Pins;
+
+	for (UEdGraphPin* Pin : TargetNode.Pins)
+	{
+		Pins.Append(ReverseFollowPinArray(*Pin, false));
+	}
+#endif
+	
+	Super::BreakNodeLinks(TargetNode);
+
+#if WITH_EDITOR
+	NodePinConnectionListChanged(Pins);
+#endif
+}
+
+
 void UEdGraphSchema_CustomizableObject::BreakPinLinks(UEdGraphPin& TargetPin, bool bSendsNodeNotification) const
 {
 	const FScopedTransaction Transaction(NSLOCTEXT("UnrealEd", "GraphEd_BreakPinLinks", "Break Pin Links"));
+
+	TArray<UEdGraphPin*> Pins;
+	Pins.Append(FollowPinArray(TargetPin, false));
+	Pins.Append(ReverseFollowPinArray(TargetPin, false));
 	
 	Super::BreakPinLinks(TargetPin, bSendsNodeNotification);
+
+#if WITH_EDITOR
+	NodePinConnectionListChanged(Pins);
+#endif
 }
 
 
@@ -934,7 +962,15 @@ void UEdGraphSchema_CustomizableObject::BreakSinglePinLink(UEdGraphPin* SourcePi
 {
 	const FScopedTransaction Transaction(NSLOCTEXT("UnrealEd", "GraphEd_BreakSinglePinLink", "Break Pin Link"));
 
+	TArray<UEdGraphPin*> SourceConnectedPins = ReverseFollowPinArray(*SourcePin, false);
+	TArray<UEdGraphPin*> TargetConnectedPins = ReverseFollowPinArray(*TargetPin, false);
+	
 	Super::BreakSinglePinLink(SourcePin, TargetPin);
+
+#if WITH_EDITOR
+	NodePinConnectionListChanged(SourceConnectedPins);
+	NodePinConnectionListChanged(TargetConnectedPins);
+#endif
 }
 
 
@@ -1200,18 +1236,27 @@ void UEdGraphSchema_CustomizableObject::GetAssetsGraphHoverMessage(const TArray<
 
 bool UEdGraphSchema_CustomizableObject::TryCreateConnection(UEdGraphPin* PinA, UEdGraphPin* PinB) const
 {
-	bool Result = Super::TryCreateConnection(PinA, PinB);
+	UEdGraphNode* PinAOwningNode = PinA->GetOwningNode(); // TryCreateConnection can reconstruct the node invalidating the FromPin. Get the OwningNode before.
+	UEdGraphNode* PinBOwningNode = PinB->GetOwningNode();
+
+	bool bResult = Super::TryCreateConnection(PinA, PinB);
 
+	if (bResult)
+	{
+		PinAOwningNode->NodeConnectionListChanged();
+		PinBOwningNode->NodeConnectionListChanged();
+	}
+	
 	if (!PinA || PinA->bWasTrashed || !PinB || PinB->bWasTrashed)
 	{
-		return Result;
+		return bResult;
 	}
 	
 	UEdGraphPin* InputPin;
 	UEdGraphPin* OutputPin;
 	if (!CategorizePinsByDirection(PinA, PinB, InputPin, OutputPin))
 	{
-		return Result;
+		return bResult;
 	}
 
 	if (UCustomizableObjectNode* Node = Cast<UCustomizableObjectNode>(InputPin->GetOwningNode()))
@@ -1223,8 +1268,24 @@ bool UEdGraphSchema_CustomizableObject::TryCreateConnection(UEdGraphPin* PinA, U
 	{
 		Node->BreakExistingConnectionsPostConnection(InputPin, OutputPin);
 	}
+	
+#if WITH_EDITOR
+	if (bResult)
+	{
+		NodePinConnectionListChanged(ReverseFollowPinArray(*PinA, false));
+		NodePinConnectionListChanged(ReverseFollowPinArray(*PinB, false));
+	}
+#endif
+
+	return bResult;
+}
+
 
-	return Result;
+FPinConnectionResponse UEdGraphSchema_CustomizableObject::MovePinLinks(UEdGraphPin& MoveFromPin, UEdGraphPin& MoveToPin, bool bIsIntermediateMove, bool bNotifyLinkedNodes) const
+{
+    // Mutable graph and its super does not use it. If we ever want to use it we should call NotifyIndirectConnections.
+	unimplemented();
+	return {};
 }
 
 
diff --git a/Engine/Plugins/Experimental/Mutable/Source/CustomizableObjectEditor/Private/MuCOE/GraphTraversal.cpp b/Engine/Plugins/Experimental/Mutable/Source/CustomizableObjectEditor/Private/MuCOE/GraphTraversal.cpp
--- a/Engine/Plugins/Experimental/Mutable/Source/CustomizableObjectEditor/Private/MuCOE/GraphTraversal.cpp
+++ b/Engine/Plugins/Experimental/Mutable/Source/CustomizableObjectEditor/Private/MuCOE/GraphTraversal.cpp
@@ -28,18 +28,8 @@
 #include "UObject/UObjectIterator.h"
 
 
-/** Follow the given pin returning its connected pin.
- *
- * - Skips all orphan pins.
- * - Follows External Pin and Reroute nodes.
- *
- * @param Pin Pin to follow.
- * @param bOutCycleDetected If provided, it will set to true if a cycle has been found.
- * @param Direction Direction of the pin to explore.  */
-TArray<UEdGraphPin*> FollowPinArray(const UEdGraphPin& Pin, bool* bOutCycleDetected, const EEdGraphPinDirection Direction)
+TArray<UEdGraphPin*> FollowPinArray(const UEdGraphPin& Pin, bool bIgnoreOrphan, bool* bOutCycleDetected)
 {
-	check(Pin.Direction == Direction); // Can only follow input pins. To follow output pins see FollowOutputPin, but be aware of its limitations!
-
 	bool bCycleDetected = false;
 		
 	TArray<UEdGraphPin*> Result;
@@ -52,7 +42,7 @@ TArray<UEdGraphPin*> FollowPinArray(const UEdGraphPin& Pin, bool* bOutCycleDetec
 	{
 		const UEdGraphPin& CurrentPin = *PinsToVisit.Pop();
 
-		if (IsPinOrphan(CurrentPin))
+		if (!bIgnoreOrphan && IsPinOrphan(CurrentPin))
 		{
 			continue;
 		}
@@ -65,48 +55,43 @@ TArray<UEdGraphPin*> FollowPinArray(const UEdGraphPin& Pin, bool* bOutCycleDetec
 
 		for (UEdGraphPin* LinkedPin : CurrentPin.LinkedTo)
 		{
-			if (IsPinOrphan(*LinkedPin))
+			if (!bIgnoreOrphan && IsPinOrphan(*LinkedPin))
 			{
 				continue;
 			}
 			
-			if (const UCustomizableObjectNodeExposePin* ExposePinNode = Cast<UCustomizableObjectNodeExposePin>(LinkedPin->GetOwningNode()))
+			if (const UCustomizableObjectNodeExposePin* ExposePinNode = Cast<UCustomizableObjectNodeExposePin>(LinkedPin->GetOwningNodeUnchecked()))
 			{
-				check(Direction == EGPD_Output);
+				check(Pin.Direction == EGPD_Output);
 
-				const UCustomizableObjectNodeExternalPin* LinkedNode = nullptr;
 				for (TObjectIterator<UCustomizableObjectNodeExternalPin> It; It; ++It)
 				{
-					if (IsValid(*It) && (*It)->GetNodeExposePin() == ExposePinNode)
-					{
-						LinkedNode = *It;
-						break;
-					}
-				}
-	
-				if (LinkedNode)
-				{
-					if (const UEdGraphPin* ExternalPin = LinkedNode->GetExternalPin())
+					const UCustomizableObjectNodeExternalPin* LinkedNode = *It;
+					
+					if (IsValid(LinkedNode) &&
+						!LinkedNode->IsTemplate() &&
+						LinkedNode->GetNodeExposePin() == ExposePinNode)
 					{
+						const UEdGraphPin* ExternalPin = LinkedNode->GetExternalPin();
+						check(ExternalPin);
 						PinsToVisit.Add(ExternalPin);
 					}
 				}
 			}
-			else if (const UCustomizableObjectNodeExternalPin* ExternalPinNode = Cast<UCustomizableObjectNodeExternalPin>(LinkedPin->GetOwningNode()))
+			else if (const UCustomizableObjectNodeExternalPin* ExternalPinNode = Cast<UCustomizableObjectNodeExternalPin>(LinkedPin->GetOwningNodeUnchecked()))
 			{
-				check(Direction == EGPD_Input);
+				check(Pin.Direction == EGPD_Input);
 				
 				if (const UCustomizableObjectNodeExposePin* LinkedNode = ExternalPinNode->GetNodeExposePin())
 				{
-					if (const UEdGraphPin* ExposePin = LinkedNode->InputPin())
-					{
-						PinsToVisit.Add(ExposePin);
-					}
+					const UEdGraphPin* ExposePin = LinkedNode->InputPin();
+					check(ExposePin);
+					PinsToVisit.Add(ExposePin);
 				}
 			}
-			else if (const UCustomizableObjectNodeReroute* NodeReroute = Cast<UCustomizableObjectNodeReroute>(LinkedPin->GetOwningNode()))
+			else if (const UCustomizableObjectNodeReroute* NodeReroute = Cast<UCustomizableObjectNodeReroute>(LinkedPin->GetOwningNodeUnchecked()))
 			{
-				PinsToVisit.Add(Direction == EGPD_Input ? NodeReroute->GetInputPin() : NodeReroute->GetOutputPin());
+				PinsToVisit.Add(Pin.Direction == EGPD_Input ? NodeReroute->GetInputPin() : NodeReroute->GetOutputPin());
 			}
 			else
 			{
@@ -126,7 +111,8 @@ TArray<UEdGraphPin*> FollowPinArray(const UEdGraphPin& Pin, bool* bOutCycleDetec
 
 TArray<UEdGraphPin*> FollowInputPinArray(const UEdGraphPin& Pin, bool* bOutCycleDetected)
 {
-	return FollowPinArray(Pin, bOutCycleDetected, EGPD_Input);
+	check(Pin.Direction == EGPD_Input);
+	return FollowPinArray(Pin, true, bOutCycleDetected);
 }
 
 
@@ -148,7 +134,8 @@ UEdGraphPin* FollowInputPin(const UEdGraphPin& Pin, bool* CycleDetected)
 
 TArray<UEdGraphPin*> FollowOutputPinArray(const UEdGraphPin& Pin, bool* bOutCycleDetected)
 {
-	return FollowPinArray(Pin, bOutCycleDetected, EGPD_Output);
+	check(Pin.Direction == EGPD_Output);
+	return FollowPinArray(Pin, true, bOutCycleDetected);
 }
 
 
@@ -168,6 +155,95 @@ UEdGraphPin* FollowOutputPin(const UEdGraphPin& Pin, bool* CycleDetected)
 }
 
 
+TArray<UEdGraphPin*> ReverseFollowPinArray(const UEdGraphPin& Pin, bool bIgnoreOrphan, bool* bOutCycleDetected)
+{
+	bool bCycleDetected = false;
+		
+	TArray<UEdGraphPin*> Result;
+	
+	TSet<const UEdGraphPin*> Visited;
+
+	TArray<UEdGraphPin*> PinsToVisit;
+	PinsToVisit.Add(const_cast<UEdGraphPin*>(&Pin));
+	while (PinsToVisit.Num())
+	{
+		UEdGraphPin& CurrentPin = *PinsToVisit.Pop();
+
+		if (!bIgnoreOrphan && IsPinOrphan(CurrentPin))
+		{
+			continue;
+		}
+
+		Visited.FindOrAdd(&CurrentPin, &bCycleDetected);
+		if (bCycleDetected)
+		{
+			continue;
+		}
+		
+		if (const UCustomizableObjectNodeExposePin* ExposePinNode = Cast<UCustomizableObjectNodeExposePin>(CurrentPin.GetOwningNodeUnchecked()))
+		{
+			check(Pin.Direction == EGPD_Input);
+
+			for (TObjectIterator<UCustomizableObjectNodeExternalPin> It; It; ++It)
+			{
+				const UCustomizableObjectNodeExternalPin* LinkedNode = *It;
+
+				if (IsValid(LinkedNode) &&
+					!LinkedNode->IsTemplate() &&
+					LinkedNode->GetNodeExposePin() == ExposePinNode)
+				{
+					const UEdGraphPin* ExternalPin = LinkedNode->GetExternalPin();
+					check(ExternalPin);
+
+					for (UEdGraphPin* LinkedPin : ExternalPin->LinkedTo)
+					{
+						PinsToVisit.Add(LinkedPin);
+					}
+				}
+			}
+		}
+		else if (const UCustomizableObjectNodeExternalPin* ExternalPinNode = Cast<UCustomizableObjectNodeExternalPin>(CurrentPin.GetOwningNodeUnchecked()))
+		{
+			check(Pin.Direction == EGPD_Output);
+				
+			if (const UCustomizableObjectNodeExposePin* LinkedNode = ExternalPinNode->GetNodeExposePin())
+			{
+				const UEdGraphPin* ExposePin = LinkedNode->InputPin();
+				check(ExposePin);
+
+				for (UEdGraphPin* LinkedPin : ExposePin->LinkedTo)
+				{
+					PinsToVisit.Add(LinkedPin);
+				}
+			}
+		}
+		else if (const UCustomizableObjectNodeReroute* NodeReroute = Cast<UCustomizableObjectNodeReroute>(CurrentPin.GetOwningNodeUnchecked()))
+		{
+			UEdGraphPin* ReroutePin = Pin.Direction == EGPD_Output ? NodeReroute->GetInputPin() : NodeReroute->GetOutputPin();
+
+			for (UEdGraphPin* LinkedPin : ReroutePin->LinkedTo)
+			{
+				PinsToVisit.Add(LinkedPin);
+			}
+		}
+		else
+		{
+			if (bIgnoreOrphan || !IsPinOrphan(CurrentPin))
+			{
+				Result.Add(&CurrentPin);
+			}
+		}
+	}
+	
+	if (bOutCycleDetected)
+	{
+		*bOutCycleDetected = bCycleDetected;	
+	}
+	
+	return Result;
+}
+
+
 UCustomizableObjectNodeObject* GetRootNode(UCustomizableObject* Object, bool& bOutMultipleBaseObjectsFound)
 {
 	// Look for the base object node
@@ -623,3 +699,26 @@ namespace GraphTraversal
 		return Root && !Root->ParentObject;
 	}
 }
+
+
+void NodePinConnectionListChanged(const TArray<UEdGraphPin*>& Pins)
+{
+	TMap<UEdGraphNode*, TSet<UEdGraphPin*>> SortedPins;
+	for (UEdGraphPin* Pin : Pins) 
+	{
+		if (UEdGraphNode* Node = Pin->GetOwningNodeUnchecked())
+		{
+			SortedPins.FindOrAdd(Node).Add(Pin);
+		}
+	}
+
+	for (TTuple<UEdGraphNode*, TSet<UEdGraphPin*>> Pair : SortedPins)
+	{
+		for (UEdGraphPin* ConnectedPin : Pair.Value)
+		{
+			Pair.Key->PinConnectionListChanged(ConnectedPin);
+		}
+		
+		Pair.Key->NodeConnectionListChanged();
+	}
+}
diff --git a/Engine/Plugins/Experimental/Mutable/Source/CustomizableObjectEditor/Private/MuCOE/Nodes/CustomizableObjectNode.cpp b/Engine/Plugins/Experimental/Mutable/Source/CustomizableObjectEditor/Private/MuCOE/Nodes/CustomizableObjectNode.cpp
--- a/Engine/Plugins/Experimental/Mutable/Source/CustomizableObjectEditor/Private/MuCOE/Nodes/CustomizableObjectNode.cpp
+++ b/Engine/Plugins/Experimental/Mutable/Source/CustomizableObjectEditor/Private/MuCOE/Nodes/CustomizableObjectNode.cpp
@@ -66,6 +66,7 @@ bool UCustomizableObjectNode::ShouldBreakExistingConnections(const UEdGraphPin*
 	return IsSingleOutputNode();
 }
 
+
 void UCustomizableObjectNode::PostEditChangeProperty(FPropertyChangedEvent& PropertyChangedEvent)
 {
 	Super::PostEditChangeProperty(PropertyChangedEvent);
@@ -319,13 +320,8 @@ void UCustomizableObjectNode::AutowireNewNode(UEdGraphPin* FromPin)
 
 	for (UEdGraphPin* Pin : GetAllNonOrphanPins())
 	{	
-		UEdGraphNode* OwningNode = FromPin->GetOwningNode(); // TryCreateConnection can reconstruct the node invalidating the FromPin. Get the OwningNode before.
-
 		if (Schema->TryCreateConnection(FromPin, Pin))
 		{
-			OwningNode->NodeConnectionListChanged();
-			NodeConnectionListChanged();
-
 			break;
 		}		
 	}
diff --git a/Engine/Plugins/Experimental/Mutable/Source/CustomizableObjectEditor/Private/MuCOE/Nodes/CustomizableObjectNodeExternalPin.cpp b/Engine/Plugins/Experimental/Mutable/Source/CustomizableObjectEditor/Private/MuCOE/Nodes/CustomizableObjectNodeExternalPin.cpp
--- a/Engine/Plugins/Experimental/Mutable/Source/CustomizableObjectEditor/Private/MuCOE/Nodes/CustomizableObjectNodeExternalPin.cpp
+++ b/Engine/Plugins/Experimental/Mutable/Source/CustomizableObjectEditor/Private/MuCOE/Nodes/CustomizableObjectNodeExternalPin.cpp
@@ -6,6 +6,7 @@
 #include "MuCOE/CustomizableObjectGraph.h"
 #include "MuCOE/CustomizableObjectPin.h"
 #include "MuCOE/EdGraphSchema_CustomizableObject.h"
+#include "MuCOE/GraphTraversal.h"
 #include "MuCOE/Nodes/CustomizableObjectNodeExposePin.h"
 #include "MuCOE/RemapPins/CustomizableObjectNodeRemapPinsByPosition.h"
 
@@ -52,6 +53,8 @@ void UCustomizableObjectNodeExternalPin::PostBackwardsCompatibleFixup()
 
 void UCustomizableObjectNodeExternalPin::SetExternalObjectNodeId(FGuid Guid)
 {
+	PrePropagateConnectionChanged();
+
 	if (UCustomizableObjectNodeExposePin* NodeExposePin = GetNodeExposePin())
 	{
 		NodeExposePin->OnNameChangedDelegate.Remove(OnNameChangedDelegateHandle);
@@ -65,22 +68,16 @@ void UCustomizableObjectNodeExternalPin::SetExternalObjectNodeId(FGuid Guid)
 		OnNameChangedDelegateHandle = NodeExposePin->OnNameChangedDelegate.AddUObject(this, &Super::ReconstructNode);
 		DestroyNodeDelegateHandle = NodeExposePin->DestroyNodeDelegate.AddUObject(this, &Super::ReconstructNode);
 	}
-
+	
 	Super::ReconstructNode();
+
+	PropagateConnectionChanged();
 }
 
 
 UEdGraphPin* UCustomizableObjectNodeExternalPin::GetExternalPin() const
 {
-	const TArray<UEdGraphPin*> NonOrphanPins = GetAllNonOrphanPins();
-	if (NonOrphanPins.Num())
-	{
-		return NonOrphanPins[0];
-	}
-	else
-	{
-		return nullptr;
-	}
+	return Pins[0];
 }
 
 
@@ -90,6 +87,23 @@ UCustomizableObjectNodeExposePin* UCustomizableObjectNodeExternalPin::GetNodeExp
 }
 
 
+void UCustomizableObjectNodeExternalPin::PrePropagateConnectionChanged()
+{
+	PropagatePreviousPin = ReverseFollowPinArray(*Pins[0]);
+}
+
+
+void UCustomizableObjectNodeExternalPin::PropagateConnectionChanged()
+{
+	// Propagate new left.
+	PropagatePreviousPin.Append(ReverseFollowPinArray(*Pins[0])); // Notify old connections and new connections.
+	NodePinConnectionListChanged(PropagatePreviousPin); // This function avoids double notifications.
+	
+	// Propagate right.
+	NodePinConnectionListChanged(FollowOutputPinArray(*Pins[0]));
+}
+
+
 void UCustomizableObjectNodeExternalPin::AllocateDefaultPins(UCustomizableObjectNodeRemapPins* RemapPins)
 {
 	FName PinName = GetNodeExposePin() ? FName(GetNodeExposePin()->GetNodeName()) : FName("Object");
@@ -170,6 +184,34 @@ void UCustomizableObjectNodeExternalPin::UpdateReferencedNodeId(const FGuid& New
 }
 
 
+void UCustomizableObjectNodeExternalPin::PreEditChange(FProperty* PropertyAboutToChange)
+{
+	Super::PreEditChange(PropertyAboutToChange);
+	
+	if (PropertyAboutToChange)
+	{
+		if (PropertyAboutToChange->GetFName() == GET_MEMBER_NAME_CHECKED(UCustomizableObjectNodeExternalPin, ExternalObject))
+		{
+			PrePropagateConnectionChanged();
+		}
+	}
+}
+
+
+void UCustomizableObjectNodeExternalPin::PostEditChangeProperty(FPropertyChangedEvent& PropertyChangedEvent)
+{
+	Super::PostEditChangeProperty(PropertyChangedEvent);
+
+	if (const FProperty* PropertyThatChanged = PropertyChangedEvent.Property)
+	{
+		if (PropertyThatChanged->GetFName() == GET_MEMBER_NAME_CHECKED(UCustomizableObjectNodeExternalPin, ExternalObject))
+		{
+			PropagateConnectionChanged();
+		}
+	}
+}
+
+
 UCustomizableObjectNodeRemapPins* UCustomizableObjectNodeExternalPin::CreateRemapPinsDefault() const
 {
 	return NewObject<UCustomizableObjectNodeRemapPinsByPosition>();
diff --git a/Engine/Plugins/Experimental/Mutable/Source/CustomizableObjectEditor/Private/MuCOE/Nodes/CustomizableObjectNodeLayoutBlocks.cpp b/Engine/Plugins/Experimental/Mutable/Source/CustomizableObjectEditor/Private/MuCOE/Nodes/CustomizableObjectNodeLayoutBlocks.cpp
--- a/Engine/Plugins/Experimental/Mutable/Source/CustomizableObjectEditor/Private/MuCOE/Nodes/CustomizableObjectNodeLayoutBlocks.cpp
+++ b/Engine/Plugins/Experimental/Mutable/Source/CustomizableObjectEditor/Private/MuCOE/Nodes/CustomizableObjectNodeLayoutBlocks.cpp
@@ -75,6 +75,8 @@ FLinearColor UCustomizableObjectNodeLayoutBlocks::GetNodeTitleColor() const
 
 void UCustomizableObjectNodeLayoutBlocks::PinConnectionListChanged(UEdGraphPin* Pin)
 {
+	Super::PinConnectionListChanged(Pin);
+	
 	if (Layout && Pin == OutputPin())
 	{
 		LinkPostEditChangePropertyDelegate(Pin);
diff --git a/Engine/Plugins/Experimental/Mutable/Source/CustomizableObjectEditor/Private/MuCOE/Nodes/CustomizableObjectNodeMaterial.cpp b/Engine/Plugins/Experimental/Mutable/Source/CustomizableObjectEditor/Private/MuCOE/Nodes/CustomizableObjectNodeMaterial.cpp
--- a/Engine/Plugins/Experimental/Mutable/Source/CustomizableObjectEditor/Private/MuCOE/Nodes/CustomizableObjectNodeMaterial.cpp
+++ b/Engine/Plugins/Experimental/Mutable/Source/CustomizableObjectEditor/Private/MuCOE/Nodes/CustomizableObjectNodeMaterial.cpp
@@ -1297,6 +1297,8 @@ void UCustomizableObjectNodeMaterial::SetDefaultMaterial()
 
 void UCustomizableObjectNodeMaterial::PinConnectionListChanged(UEdGraphPin* Pin)
 {
+	Super::PinConnectionListChanged(Pin);
+	
 	if (Pin == GetMeshPin())
 	{
 		if (LastMeshNodeConnected.IsValid())
@@ -1396,7 +1398,10 @@ void UCustomizableObjectNodeMaterial::BreakExistingConnectionsPostConnection(UEd
 			}
 		}
 
-		OutputPin->BreakLinkTo(RemovePin); // Can not be called inside the range for loop. Can be called with a null value.
+		if (RemovePin)
+		{
+			GetSchema()->BreakSinglePinLink(OutputPin, RemovePin); // Can not be called inside the range for loop.
+		}
 	}
 }
 
diff --git a/Engine/Plugins/Experimental/Mutable/Source/CustomizableObjectEditor/Private/MuCOE/Nodes/CustomizableObjectNodeReroute.cpp b/Engine/Plugins/Experimental/Mutable/Source/CustomizableObjectEditor/Private/MuCOE/Nodes/CustomizableObjectNodeReroute.cpp
--- a/Engine/Plugins/Experimental/Mutable/Source/CustomizableObjectEditor/Private/MuCOE/Nodes/CustomizableObjectNodeReroute.cpp
+++ b/Engine/Plugins/Experimental/Mutable/Source/CustomizableObjectEditor/Private/MuCOE/Nodes/CustomizableObjectNodeReroute.cpp
@@ -4,6 +4,7 @@
 
 #include "SGraphNodeKnot.h"
 #include "MuCOE/EdGraphSchema_CustomizableObject.h"
+#include "MuCOE/GraphTraversal.h"
 #include "MuCOE/RemapPins/CustomizableObjectNodeRemapPinsByPosition.h"
 
 #define LOCTEXT_NAMESPACE "CustomizableObjectEditor"
@@ -85,15 +86,10 @@ TSharedPtr<SGraphNode> UCustomizableObjectNodeReroute::CreateVisualWidget()
 }
 
 
-void UCustomizableObjectNodeReroute::NotifyPinConnectionListChanged(UEdGraphPin* Pin)
+void UCustomizableObjectNodeReroute::NodeConnectionListChanged()
 {
-	PropagatePinType();
-}
-
-
-void UCustomizableObjectNodeReroute::PinConnectionListChanged(UEdGraphPin* Pin)
-{
-	Super::PinConnectionListChanged(Pin);
+	Super::NodeConnectionListChanged();
+	
 	PropagatePinType();
 }
 
diff --git a/Engine/Plugins/Experimental/Mutable/Source/CustomizableObjectEditor/Private/MuCOE/Nodes/CustomizableObjectNodeSwitchBase.cpp b/Engine/Plugins/Experimental/Mutable/Source/CustomizableObjectEditor/Private/MuCOE/Nodes/CustomizableObjectNodeSwitchBase.cpp
--- a/Engine/Plugins/Experimental/Mutable/Source/CustomizableObjectEditor/Private/MuCOE/Nodes/CustomizableObjectNodeSwitchBase.cpp
+++ b/Engine/Plugins/Experimental/Mutable/Source/CustomizableObjectEditor/Private/MuCOE/Nodes/CustomizableObjectNodeSwitchBase.cpp
@@ -13,6 +13,8 @@ class UCustomizableObjectNodeRemapPins;
 
 void UCustomizableObjectNodeSwitchBase::ReloadEnumParam()
 {
+	ReloadingElementsNames.Empty();
+
 	// Get the names of the enum parameter for the element pins
 	if (const UEdGraphPin* EnumPin = SwitchParameter())
 	{
@@ -20,7 +22,6 @@ void UCustomizableObjectNodeSwitchBase::ReloadEnumParam()
 		{
 			if (UCustomizableObjectNodeEnumParameter* EnumNode = Cast<UCustomizableObjectNodeEnumParameter>(LinkedPin->GetOwningNode()))
 			{
-				ReloadingElementsNames.Empty(EnumNode->Values.Num());
 				for (int i = 0; i < EnumNode->Values.Num(); ++i)
 				{
 					ReloadingElementsNames.Add(EnumNode->Values[i].Name);
@@ -193,12 +194,9 @@ void UCustomizableObjectNodeSwitchBase::LinkPostEditChangePropertyDelegate(const
 		{
 			LastNodeEnumParameterConnected->PostEditChangePropertyDelegate.AddUniqueDynamic(this, &UCustomizableObjectNodeSwitchBase::EnumParameterPostEditChangeProperty);
 		}
-
-		if (Cast<UCustomizableObjectNodeEnumParameter>(ConnectedPin->GetOwningNode()))
-		{
-			Super::ReconstructNode();
-		}
 	}
+
+	Super::ReconstructNode();
 }
 
 
diff --git a/Engine/Plugins/Experimental/Mutable/Source/CustomizableObjectEditor/Public/MuCOE/EdGraphSchema_CustomizableObject.h b/Engine/Plugins/Experimental/Mutable/Source/CustomizableObjectEditor/Public/MuCOE/EdGraphSchema_CustomizableObject.h
--- a/Engine/Plugins/Experimental/Mutable/Source/CustomizableObjectEditor/Public/MuCOE/EdGraphSchema_CustomizableObject.h
+++ b/Engine/Plugins/Experimental/Mutable/Source/CustomizableObjectEditor/Public/MuCOE/EdGraphSchema_CustomizableObject.h
@@ -114,17 +114,18 @@ public:
 	virtual bool ShouldHidePinDefaultValue(UEdGraphPin* Pin) const override;
 	virtual void DroppedAssetsOnGraph(const TArray<FAssetData>& Assets, const FVector2D& GraphPosition, UEdGraph* Graph) const override;
 	virtual void OnPinConnectionDoubleCicked(UEdGraphPin* PinA, UEdGraphPin* PinB, const FVector2D& GraphPosition) const override;
-	virtual FConnectionDrawingPolicy* CreateConnectionDrawingPolicy(int32 InBackLayerID, int32 InFrontLayerID, float InZoomFactor, const FSlateRect& InClippingRect, FSlateWindowElementList& InDrawElements, UEdGraph* InGraphObj) const override;
-
-	static FLinearColor GetPinTypeColor(const FName& PinType);
-
-	void GetContextMenuActionsReconstructAllChildNodes(UToolMenu* Menu, TWeakObjectPtr<UGraphNodeContextMenuContext> Context) const;
 	virtual void GetContextMenuActions(UToolMenu* Menu, UGraphNodeContextMenuContext* Context) const override;
-
+	virtual void BreakNodeLinks(UEdGraphNode& TargetNode) const override;
 	virtual void BreakPinLinks(UEdGraphPin& TargetPin, bool bSendsNodeNotification) const override;
 	virtual void BreakSinglePinLink(UEdGraphPin* SourcePin, UEdGraphPin* TargetPin) const override;
+	virtual bool TryCreateConnection(UEdGraphPin* PinA, UEdGraphPin* PinB) const override;
+	virtual FPinConnectionResponse MovePinLinks(UEdGraphPin& MoveFromPin, UEdGraphPin& MoveToPin, bool bIsIntermediateMove, bool bNotifyLinkedNodes) const override;
+	virtual FConnectionDrawingPolicy* CreateConnectionDrawingPolicy(int32 InBackLayerID, int32 InFrontLayerID, float InZoomFactor, const FSlateRect& InClippingRect, FSlateWindowElementList& InDrawElements, UEdGraph* InGraphObj) const override;
 
 	//  Own interface
+	static FLinearColor GetPinTypeColor(const FName& PinType);
+
+	void GetContextMenuActionsReconstructAllChildNodes(UToolMenu* Menu, TWeakObjectPtr<UGraphNodeContextMenuContext> Context) const;
 
 	/** Utility to create a node creating action */
 	static TSharedPtr<FCustomizableObjectSchemaAction_NewNode> AddNewNodeAction(FGraphActionListBuilderBase& ContextMenuBuilder, const FString& Category, const FText& MenuDesc, const FText& Tooltip, const int32 Grouping = 0, const FString& Keywords = FString());
@@ -161,13 +162,7 @@ public:
 
 	static TSharedPtr<class ICustomizableObjectEditor> GetCustomizableObjectEditor(const class UEdGraph* ParentGraph);
 	UEdGraphNode* AddComment(class UEdGraph* ParentGraph, UEdGraphPin* FromPin, const FVector2D Location, bool bSelectNewNode = true) const;
-
-	/**
-	 * Break non allowed connections once the connection has been created.
-	 * Allows the nodes to have greater connectivity granularity than simply forbidding to create a connection with UEdGraphSchema_CustomizableObject::CanCreateConnection.
-	 */ 
-	virtual bool TryCreateConnection(UEdGraphPin* PinA, UEdGraphPin* PinB) const override;
-
+	
 	/** Given a pin category name, get its friendly name (user readable). */
 	static FText GetPinCategoryName(const FName& PinCategory);
 
diff --git a/Engine/Plugins/Experimental/Mutable/Source/CustomizableObjectEditor/Public/MuCOE/GraphTraversal.h b/Engine/Plugins/Experimental/Mutable/Source/CustomizableObjectEditor/Public/MuCOE/GraphTraversal.h
--- a/Engine/Plugins/Experimental/Mutable/Source/CustomizableObjectEditor/Public/MuCOE/GraphTraversal.h
+++ b/Engine/Plugins/Experimental/Mutable/Source/CustomizableObjectEditor/Public/MuCOE/GraphTraversal.h
@@ -10,6 +10,17 @@ class UEdGraph;
 class UEdGraphPin;
 class UCustomizableObjectNode;
 
+
+/** Follow the given pin returning its connected pin.
+ *
+ * - Skips all orphan pins.
+ * - Follows External Pin and Reroute nodes.
+ *
+ * @param Pin Pin to follow.
+ * @param bIgnoreOrphan If true, it will not follow orphan pins.
+ * @param bOutCycleDetected If provided, it will set to true if a cycle has been found. */
+TArray<UEdGraphPin*> FollowPinArray(const UEdGraphPin& Pin, bool bIgnoreOrphan = false, bool* bOutCycleDetected = nullptr);
+
 /** Follow the given input pin returning the output connected pin.
  *
  * - Skips all orphan pins.
@@ -37,6 +48,10 @@ TArray<UEdGraphPin*> FollowOutputPinArray(const UEdGraphPin& Pin, bool* bOutCycl
 /** Non-array version of FollowOutputPinArray. The pin can only have one connected pin. */
 UEdGraphPin* FollowOutputPin(const UEdGraphPin& Pin, bool* CycleDetected = nullptr);
 
+/** See FollowPinArray.
+ * Given a pin, follow it in reverse (through the owning node instead of the linked pins). */
+TArray<UEdGraphPin*> ReverseFollowPinArray(const UEdGraphPin& Pin, bool bIgnoreOrphan = false, bool* bOutCycleDetected = nullptr);
+
 /** Returns the root Object Node of the Customizable Object's graph */
 UCustomizableObjectNodeObject* GetRootNode(UCustomizableObject* Object, bool& bOutMultipleBaseObjectsFound);
 
@@ -79,3 +94,7 @@ namespace GraphTraversal
     /** Return ture if the given Customizable Object is Root Object (not a Child Object). */
     bool IsRootObject(const UCustomizableObject& Object); 
 }
+
+/** For each given pin, call PinConnectionListChanged and NodeConnectionListChanged in the correct order.
+ *  Order: For each node, first call all PinConnectionListChanged, then NodeConnectionListChanged. */
+void NodePinConnectionListChanged(const TArray<UEdGraphPin*>& Pins);
\ No newline at end of file
diff --git a/Engine/Plugins/Experimental/Mutable/Source/CustomizableObjectEditor/Public/MuCOE/Nodes/CustomizableObjectNodeExposePin.h b/Engine/Plugins/Experimental/Mutable/Source/CustomizableObjectEditor/Public/MuCOE/Nodes/CustomizableObjectNodeExposePin.h
--- a/Engine/Plugins/Experimental/Mutable/Source/CustomizableObjectEditor/Public/MuCOE/Nodes/CustomizableObjectNodeExposePin.h
+++ b/Engine/Plugins/Experimental/Mutable/Source/CustomizableObjectEditor/Public/MuCOE/Nodes/CustomizableObjectNodeExposePin.h
@@ -16,6 +16,8 @@ struct FPropertyChangedEvent;
 
 DECLARE_MULTICAST_DELEGATE(FOnNameChangedDelegate);
 
+
+/** Export Node. */
 UCLASS()
 class CUSTOMIZABLEOBJECTEDITOR_API UCustomizableObjectNodeExposePin : public UCustomizableObjectNode
 {
diff --git a/Engine/Plugins/Experimental/Mutable/Source/CustomizableObjectEditor/Public/MuCOE/Nodes/CustomizableObjectNodeExternalPin.h b/Engine/Plugins/Experimental/Mutable/Source/CustomizableObjectEditor/Public/MuCOE/Nodes/CustomizableObjectNodeExternalPin.h
--- a/Engine/Plugins/Experimental/Mutable/Source/CustomizableObjectEditor/Public/MuCOE/Nodes/CustomizableObjectNodeExternalPin.h
+++ b/Engine/Plugins/Experimental/Mutable/Source/CustomizableObjectEditor/Public/MuCOE/Nodes/CustomizableObjectNodeExternalPin.h
@@ -17,6 +17,7 @@ class UEdGraphPin;
 class UObject;
 
 
+/** Import Node. */
 UCLASS()
 class CUSTOMIZABLEOBJECTEDITOR_API UCustomizableObjectNodeExternalPin : public UCustomizableObjectNode
 {
@@ -27,21 +28,23 @@ public:
 	virtual void Serialize(FArchive& Ar) override;
 	
 	// EdGraphNode interface 
-	FText GetNodeTitle(ENodeTitleType::Type TitleType) const override;
-	FLinearColor GetNodeTitleColor() const override;
-	FText GetTooltipText() const override;
+	virtual FText GetNodeTitle(ENodeTitleType::Type TitleType) const override;
+	virtual FLinearColor GetNodeTitleColor() const override;
+	virtual FText GetTooltipText() const override;
 	
 	// UCustomizableObjectNode interface
 	virtual void BackwardsCompatibleFixup() override;
 	virtual void PostBackwardsCompatibleFixup() override;
-	void AllocateDefaultPins(UCustomizableObjectNodeRemapPins* RemapPins) override;
-	UCustomizableObjectNodeRemapPins* CreateRemapPinsDefault() const override;
-	void BeginPostDuplicate(bool bDuplicateForPIE) override;
+	virtual void AllocateDefaultPins(UCustomizableObjectNodeRemapPins* RemapPins) override;
+	virtual UCustomizableObjectNodeRemapPins* CreateRemapPinsDefault() const override;
+	virtual void BeginPostDuplicate(bool bDuplicateForPIE) override;
 	virtual bool CanConnect(const UEdGraphPin* InOwnedInputPin, const UEdGraphPin* InOutputPin, bool& bOutIsOtherNodeBlocklisted, bool& bOutArePinsCompatible) const override;
-	
 	virtual void ReconstructNode(UCustomizableObjectNodeRemapPins* RemapPinsMode) override;
-
-	void UpdateReferencedNodeId(const FGuid& NewGuid) override;
+	virtual void UpdateReferencedNodeId(const FGuid& NewGuid) override;
+	virtual void PreEditChange(FProperty* PropertyAboutToChange) override;
+	virtual void PostEditChangeProperty(FPropertyChangedEvent& PropertyChangedEvent) override;
+	
+	// Own interface
 	
 	/** Set the linked Node Expose Pin node guid. */
 	void SetExternalObjectNodeId(FGuid Guid);
@@ -52,6 +55,11 @@ public:
 	/** Return the linked Expose Pin node. Return nullptr if not set. */
 	UCustomizableObjectNodeExposePin* GetNodeExposePin() const;
 
+private:
+	void PrePropagateConnectionChanged();
+	void PropagateConnectionChanged();
+	
+public:
 	// This is actually PinCategory
 	UPROPERTY()
 	FName PinType;
@@ -67,5 +75,8 @@ private:
 	
 	FDelegateHandle OnNameChangedDelegateHandle;
 	FDelegateHandle DestroyNodeDelegateHandle;
+
+	/** Connected pins (pins connected to the Export Node pin) before changing the import/export implicit connection. */
+	TArray<UEdGraphPin*> PropagatePreviousPin;
 };
 
diff --git a/Engine/Plugins/Experimental/Mutable/Source/CustomizableObjectEditor/Public/MuCOE/Nodes/CustomizableObjectNodeReroute.h b/Engine/Plugins/Experimental/Mutable/Source/CustomizableObjectEditor/Public/MuCOE/Nodes/CustomizableObjectNodeReroute.h
--- a/Engine/Plugins/Experimental/Mutable/Source/CustomizableObjectEditor/Public/MuCOE/Nodes/CustomizableObjectNodeReroute.h
+++ b/Engine/Plugins/Experimental/Mutable/Source/CustomizableObjectEditor/Public/MuCOE/Nodes/CustomizableObjectNodeReroute.h
@@ -22,8 +22,7 @@ public:
 	virtual void OnRenameNode(const FString& NewName) override;
 	virtual bool ShouldDrawNodeAsControlPointOnly(int32& OutInputPinIndex, int32& OutOutputPinIndex) const override;
 	virtual TSharedPtr<SGraphNode> CreateVisualWidget() override;
-	void NotifyPinConnectionListChanged(UEdGraphPin* Pin);
-	virtual void PinConnectionListChanged(UEdGraphPin* Pin) override;
+	virtual void NodeConnectionListChanged() override;
 	virtual UCustomizableObjectNodeRemapPins* CreateRemapPinsDefault() const override;
 
 	// UCustomizableObjectNode interface
