[Mutable] Fix TryCreateConnection static analysis.

#rnx

diff --git a/Engine/Plugins/Experimental/Mutable/Source/CustomizableObjectEditor/Private/MuCOE/EdGraphSchema_CustomizableObject.cpp b/Engine/Plugins/Experimental/Mutable/Source/CustomizableObjectEditor/Private/MuCOE/EdGraphSchema_CustomizableObject.cpp
--- a/Engine/Plugins/Experimental/Mutable/Source/CustomizableObjectEditor/Private/MuCOE/EdGraphSchema_CustomizableObject.cpp
+++ b/Engine/Plugins/Experimental/Mutable/Source/CustomizableObjectEditor/Private/MuCOE/EdGraphSchema_CustomizableObject.cpp
@@ -1236,6 +1236,11 @@ void UEdGraphSchema_CustomizableObject::GetAssetsGraphHoverMessage(const TArray<
 
 bool UEdGraphSchema_CustomizableObject::TryCreateConnection(UEdGraphPin* PinA, UEdGraphPin* PinB) const
 {
+	if (!PinA || !PinB)
+	{
+		return false;
+	}
+	
 	UEdGraphNode* PinAOwningNode = PinA->GetOwningNode(); // TryCreateConnection can reconstruct the node invalidating the FromPin. Get the OwningNode before.
 	UEdGraphNode* PinBOwningNode = PinB->GetOwningNode();
 
