[mutable] Fixed double update of override materials in editor. Unified all code relating to set skeletal mesh and override materials in CustomizableObjectInstanceUsage.
#tests PIE
#rnx

diff --git a/Engine/Plugins/Experimental/Mutable/Source/CustomizableObject/Private/MuCO/CustomizableObjectInstanceUsage.cpp b/Engine/Plugins/Experimental/Mutable/Source/CustomizableObject/Private/MuCO/CustomizableObjectInstanceUsage.cpp
--- a/Engine/Plugins/Experimental/Mutable/Source/CustomizableObject/Private/MuCO/CustomizableObjectInstanceUsage.cpp
+++ b/Engine/Plugins/Experimental/Mutable/Source/CustomizableObject/Private/MuCO/CustomizableObjectInstanceUsage.cpp
@@ -232,57 +232,75 @@ bool RequiresReinitPose(USkeletalMesh* CurrentSkeletalMesh, USkeletalMesh* Skele
 }
 
 
-void UCustomizableObjectInstanceUsage::SetSkeletalMesh(USkeletalMesh* SkeletalMesh, bool* bOutSkeletalMeshUpdated, bool* bOutMaterialsUpdated)
+void UCustomizableObjectInstanceUsage::SetSkeletalMeshAndOverrideMaterials(USkeletalMeshComponent& Parent, USkeletalMesh* SkeletalMesh, 
+	const UCustomizableObjectInstance& CustomizableObjectInstance, const bool bInstanceGenerated, bool* bOutSkeletalMeshUpdated, bool* bOutMaterialsUpdated)
 {
-	USkeletalMeshComponent* Parent = Cast<USkeletalMeshComponent>(GetAttachParent());
-
-	if (Parent)
+	if (SkeletalMesh != Parent.GetSkeletalMeshAsset())
 	{
-		if (SkeletalMesh != Parent->GetSkeletalMeshAsset())
-		{
-			Parent->SetSkeletalMesh(SkeletalMesh, RequiresReinitPose(Parent->GetSkeletalMeshAsset(), SkeletalMesh));
+		Parent.SetSkeletalMesh(SkeletalMesh, RequiresReinitPose(Parent.GetSkeletalMeshAsset(), SkeletalMesh));
 
-			if (bOutSkeletalMeshUpdated)
-			{
-				*bOutSkeletalMeshUpdated = true;
-			}
-		}
-
-		TArray<TObjectPtr<UMaterialInterface>> OldOverridenMaterials = Parent->OverrideMaterials;
-		
-		if (Parent->HasOverrideMaterials())
+		if (bOutSkeletalMeshUpdated)
 		{
-			// For some reason the reference skeletal mesh materials are added as override materials, clear them if necessary
-			Parent->EmptyOverrideMaterials();
+			*bOutSkeletalMeshUpdated = true;
 		}
+	}
+
+	SetPendingSetSkeletalMesh(false);
 
-		const UCustomizableObjectInstance* Instance = GetCustomizableObjectInstance();
-		const UCustomizableObject* CustomizableObject = Instance ? Instance->GetCustomizableObject() : nullptr;
-		if (!CustomizableObject)
+	TArray<TObjectPtr<UMaterialInterface>> OldOverridenMaterials = Parent.OverrideMaterials;
+
+	if (Parent.HasOverrideMaterials())
+	{
+#if !(UE_BUILD_SHIPPING || UE_BUILD_TEST)
+		if (Parent.GetClass()->GetFName() != FName(TEXT("SkeletalMeshComponentBudgeted"))) // Reduce unnecessary logging
 		{
-			return;
+			UE_LOG(LogMutable, Warning, TEXT("Attaching Customizable Skeletal Component to Skeletal Mesh Component with overriden materials! Deleting overrides."));
 		}
-		
-		const bool bIsTransientMesh = SkeletalMesh ? static_cast<bool>(SkeletalMesh->HasAllFlags(EObjectFlags::RF_Transient)) : false;
-		const bool bUseOverrideMaterials = !bIsTransientMesh
-			||
-			(CustomizableObject->bEnableMeshCache && CVarEnableMeshCache.GetValueOnAnyThread());
+#endif
+
+		// For some reason the reference skeletal mesh materials are added as override materials, clear them if necessary
+		Parent.EmptyOverrideMaterials();
+	}
+
+	const UCustomizableObject* CustomizableObject = CustomizableObjectInstance.GetCustomizableObject();
+	if (!CustomizableObject)
+	{
+		return;
+	}
+	
+	const bool bIsTransientMesh = SkeletalMesh ? static_cast<bool>(SkeletalMesh->HasAllFlags(EObjectFlags::RF_Transient)) : false;
+	const bool bUseOverrideMaterials = !bIsTransientMesh
+		||
+		(CustomizableObject->bEnableMeshCache && CVarEnableMeshCache.GetValueOnAnyThread());
 
-		if (bUseOverrideMaterials)
+	if (bUseOverrideMaterials && bInstanceGenerated)
+	{
+		if (FCustomizableInstanceComponentData* ComponentData = CustomizableObjectInstance.GetPrivate()->GetComponentData(GetComponentIndex()))
 		{
-			if (FCustomizableInstanceComponentData* ComponentData = Instance->GetPrivate()->GetComponentData(GetComponentIndex()))
+			for (int32 Index = 0; Index < ComponentData->OverrideMaterials.Num(); ++Index)
 			{
-				for (int32 Index = 0; Index < ComponentData->OverrideMaterials.Num(); ++Index)
-				{
-					Parent->SetMaterial(Index, ComponentData->OverrideMaterials[Index]);
-				}
+				Parent.SetMaterial(Index, ComponentData->OverrideMaterials[Index]);
 			}
 		}
+	}
 
-		if (bOutMaterialsUpdated)
-		{
-			*bOutMaterialsUpdated = OldOverridenMaterials != Parent->OverrideMaterials;
-		}
+	if (bOutMaterialsUpdated)
+	{
+		*bOutMaterialsUpdated = OldOverridenMaterials != Parent.OverrideMaterials;
+	}
+}
+
+
+void UCustomizableObjectInstanceUsage::SetSkeletalMesh(USkeletalMesh* SkeletalMesh, bool* bOutSkeletalMeshUpdated, bool* bOutMaterialsUpdated)
+{
+	USkeletalMeshComponent* Parent = Cast<USkeletalMeshComponent>(GetAttachParent());
+	const UCustomizableObjectInstance* CustomizableObjectInstance = GetCustomizableObjectInstance();
+
+	if (Parent && CustomizableObjectInstance)
+	{
+		const bool bInstanceGenerated = CustomizableObjectInstance->GetPrivate()->SkeletalMeshStatus == ESkeletalMeshStatus::Success;
+
+		SetSkeletalMeshAndOverrideMaterials(*Parent, SkeletalMesh, *CustomizableObjectInstance, bInstanceGenerated, bOutSkeletalMeshUpdated, bOutMaterialsUpdated);
 	}
 }
 
@@ -548,7 +566,7 @@ void UCustomizableObjectInstanceUsage::Tick(float DeltaTime)
 
 		USkeletalMesh* SkeletalMesh = nullptr;
 
-		const bool bInstanceGenerated = CustomizableObjectInstance->GetPrivate()->SkeletalMeshStatus == ESkeletalMeshStatus::Success;		
+		const bool bInstanceGenerated = CustomizableObjectInstance->GetPrivate()->SkeletalMeshStatus == ESkeletalMeshStatus::Success;
 		if (bInstanceGenerated)
 		{
 			// Generated SkeletalMesh to set, can be null if the component is empty
@@ -568,40 +586,7 @@ void UCustomizableObjectInstanceUsage::Tick(float DeltaTime)
 		// Set SkeletalMesh
 		if (bInstanceGenerated || SkeletalMesh)
 		{
-			Parent->SetSkeletalMesh(SkeletalMesh, RequiresReinitPose(Parent->GetSkeletalMeshAsset(), SkeletalMesh));
-
-			if (Parent->HasOverrideMaterials())
-			{
-#if !(UE_BUILD_SHIPPING || UE_BUILD_TEST)
-				if (Parent->GetClass()->GetFName() != FName(TEXT("SkeletalMeshComponentBudgeted"))) // Reduce unnecessary logging
-				{
-					UE_LOG(LogMutable, Warning, TEXT("Attaching Customizable Skeletal Component to Skeletal Mesh Component with overriden materials! Deleting overrides."));
-				}
-#endif
-
-				Parent->EmptyOverrideMaterials();
-			}
-
-			// 
-			const bool bIsTransientMesh = SkeletalMesh ? static_cast<bool>(SkeletalMesh->HasAllFlags(EObjectFlags::RF_Transient)) : false;
-			const bool bUseOverrideMaterials = !bIsTransientMesh
-				||
-				(CustomizableObject->bEnableMeshCache && CVarEnableMeshCache.GetValueOnAnyThread());
-
-			if (CustomizableObject &&
-				bInstanceGenerated &&
-				bUseOverrideMaterials)
-			{
-				if (FCustomizableInstanceComponentData* ComponentData = CustomizableObjectInstance->GetPrivate()->GetComponentData(GetComponentIndex()))
-				{
-					for (int32 Index = 0; Index < ComponentData->OverrideMaterials.Num(); ++Index)
-					{
-						Parent->SetMaterial(Index, ComponentData->OverrideMaterials[Index]);
-					}
-				}
-			}
-
-			SetPendingSetSkeletalMesh(false);
+			SetSkeletalMeshAndOverrideMaterials(*Parent, SkeletalMesh, *CustomizableObjectInstance, bInstanceGenerated, nullptr, nullptr);
 		}
 	}
 }
diff --git a/Engine/Plugins/Experimental/Mutable/Source/CustomizableObject/Public/MuCO/CustomizableObjectInstanceUsage.h b/Engine/Plugins/Experimental/Mutable/Source/CustomizableObject/Public/MuCO/CustomizableObjectInstanceUsage.h
--- a/Engine/Plugins/Experimental/Mutable/Source/CustomizableObject/Public/MuCO/CustomizableObjectInstanceUsage.h
+++ b/Engine/Plugins/Experimental/Mutable/Source/CustomizableObject/Public/MuCO/CustomizableObjectInstanceUsage.h
@@ -115,6 +115,9 @@ private:
 	virtual bool IsTickable() const override;
 	// End FTickableGameObject
 
+	void SetSkeletalMeshAndOverrideMaterials(USkeletalMeshComponent& Parent, USkeletalMesh* SkeletalMesh, const UCustomizableObjectInstance& CustomizableObjectInstance, 
+											 const bool bInstanceGenerated, bool* bOutSkeletalMeshUpdated, bool* bOutMaterialsUpdated);
+
 	friend UCustomizableSkeletalComponent;
 };
 
