[Mutable] Fixed swap of layout pointers.

#rnx

diff --git a/Engine/Plugins/Experimental/Mutable/Source/MutableTools/Private/MuT/CodeGenerator.cpp b/Engine/Plugins/Experimental/Mutable/Source/MutableTools/Private/MuT/CodeGenerator.cpp
--- a/Engine/Plugins/Experimental/Mutable/Source/MutableTools/Private/MuT/CodeGenerator.cpp
+++ b/Engine/Plugins/Experimental/Mutable/Source/MutableTools/Private/MuT/CodeGenerator.cpp
@@ -1407,8 +1407,8 @@ namespace mu
 
 												// Block in layout grid units
 												box< UE::Math::TIntVector2<uint16> > rectInCells;
-												rectInCells.min = pLayout->Blocks[b].Min;
-												rectInCells.size = pLayout->Blocks[b].Size;
+												rectInCells.min = pExtendLayout->Blocks[b].Min;
+												rectInCells.size = pExtendLayout->Blocks[b].Size;
 
 												FImageDesc ExtendDesc = fragmentAd->GetImageDesc();
 
