[mutable] Fix color precision error in image pixel format and slightly optimize luminance calculations.
- Investigating some suspicious white textures with luminance 254 unearthed a color precision loss in conversion from rgb to L.
#rnx

diff --git a/Engine/Plugins/Experimental/Mutable/Source/MutableRuntime/Private/MuR/OpImageLuminance.h b/Engine/Plugins/Experimental/Mutable/Source/MutableRuntime/Private/MuR/OpImageLuminance.h
--- a/Engine/Plugins/Experimental/Mutable/Source/MutableRuntime/Private/MuR/OpImageLuminance.h
+++ b/Engine/Plugins/Experimental/Mutable/Source/MutableRuntime/Private/MuR/OpImageLuminance.h
@@ -36,8 +36,8 @@ namespace mu
             {
                 for (int32 I = 0; I < NumElems; ++I)
                 {
-                    const uint16 L = (76*ABuf[3*I + 0] + 150*ABuf[3*I + 1] + 29*ABuf[3*I + 2]);
-                    DestBuf[I] = uint8(L / 255);
+                    const uint16 L = (77*ABuf[3*I + 0] + 150*ABuf[3*I + 1] + 29*ABuf[3*I + 2]);
+                    DestBuf[I] = uint8(L >> 8);
                 }
                 break;
             }
@@ -46,8 +46,8 @@ namespace mu
             {
                 for (int32 I = 0; I < NumElems; ++I)
                 {
-                    const uint16 L = (76*ABuf[4*I + 0] + 150*ABuf[4*I + 1] + 29*ABuf[4*I + 2]);
-                    DestBuf[I] = uint8(L / 255);
+                    const uint16 L = (77*ABuf[4*I + 0] + 150*ABuf[4*I + 1] + 29*ABuf[4*I + 2]);
+                    DestBuf[I] = uint8(L >> 8);
                 }
                 break;
             }
@@ -56,8 +56,8 @@ namespace mu
             {
                 for (int32 I = 0; I < NumElems; ++I)
                 {
-                    const uint16 L = (76*ABuf[4*I + 2] + 150*ABuf[4*I + 1] + 29*ABuf[4*I + 0]);
-                    DestBuf[I] = uint8(L / 255);
+                    const uint16 L = (77*ABuf[4*I + 2] + 150*ABuf[4*I + 1] + 29*ABuf[4*I + 0]);
+                    DestBuf[I] = uint8(L >> 8);
                 }
                 break;
             }
diff --git a/Engine/Plugins/Experimental/Mutable/Source/MutableRuntime/Private/MuR/OpImagePixelFormat.cpp b/Engine/Plugins/Experimental/Mutable/Source/MutableRuntime/Private/MuR/OpImagePixelFormat.cpp
--- a/Engine/Plugins/Experimental/Mutable/Source/MutableRuntime/Private/MuR/OpImagePixelFormat.cpp
+++ b/Engine/Plugins/Experimental/Mutable/Source/MutableRuntime/Private/MuR/OpImagePixelFormat.cpp
@@ -254,7 +254,7 @@ void FImageOperator::ImagePixelFormat(bool& bOutSuccess, int32 CompressionQualit
 				{
 					for (int32 I = 0; I < NumElems; ++I)
 					{
-						uint32 ResultValue = 76 * BaseBuf[3*I + 0] + 150 * BaseBuf[3*I + 1] + 29 * BaseBuf[3*I + 2];
+						uint32 ResultValue = 77 * BaseBuf[3*I + 0] + 150 * BaseBuf[3*I + 1] + 29 * BaseBuf[3*I + 2];
 						DestBuf[I] = (uint8)FMath::Min(255u, ResultValue >> 8);
 					}
 				});
@@ -271,7 +271,7 @@ void FImageOperator::ImagePixelFormat(bool& bOutSuccess, int32 CompressionQualit
 				{
 					for (int32 I = 0; I < NumElems; ++I)
 					{
-						uint32 ResultValue = 76 * BaseBuf[3*I + 0] + 150 * BaseBuf[3*I + 1] + 29 * BaseBuf[3*I + 2];
+						uint32 ResultValue = 77 * BaseBuf[3*I + 0] + 150 * BaseBuf[3*I + 1] + 29 * BaseBuf[3*I + 2];
 						DestBuf[I] = (uint8)FMath::Min(255u, ResultValue >> 8);
 					}
 				});
@@ -287,7 +287,7 @@ void FImageOperator::ImagePixelFormat(bool& bOutSuccess, int32 CompressionQualit
 				{
 					for (int32 I = 0; I < NumElems; ++I)
 					{
-						uint32 ResultValue = 76 * BaseBuf[4*I + 0] + 150 * BaseBuf[4*I + 1] + 29 * BaseBuf[4*I + 2];
+						uint32 ResultValue = 77 * BaseBuf[4*I + 0] + 150 * BaseBuf[4*I + 1] + 29 * BaseBuf[4*I + 2];
 						DestBuf[I] = (uint8)FMath::Min(255u, ResultValue >> 8);
 					}
 				});
@@ -303,8 +303,8 @@ void FImageOperator::ImagePixelFormat(bool& bOutSuccess, int32 CompressionQualit
 				{
 					for (int32 I = 0; I < NumElems; ++I)
 					{
-						uint32 ResultValue = 76 * BaseBuf[4*I + 0] + 150 * BaseBuf[4*I + 1] + 29 * BaseBuf[4*I + 2];
-						DestBuf[I] = (uint8)FMath::Min(255u, ResultValue >> 8);
+						uint32 ResultValue = 77 * BaseBuf[4*I + 0] + 150 * BaseBuf[4*I + 1] + 29 * BaseBuf[4*I + 2];
+						DestBuf[I] = uint8( ResultValue >> 8);
 					}
 				});
 
@@ -319,8 +319,8 @@ void FImageOperator::ImagePixelFormat(bool& bOutSuccess, int32 CompressionQualit
 				{
 					for (int32 I = 0; I < NumElems; ++I)
 					{
-						uint32 ResultValue = 76 * BaseBuf[4*I + 2] + 150 * BaseBuf[4*I + 1] + 29 * BaseBuf[4*I + 0];
-						DestBuf[I] = (uint8)FMath::Min(255u, ResultValue >> 8);
+						uint32 ResultValue = 77 * BaseBuf[4*I + 2] + 150 * BaseBuf[4*I + 1] + 29 * BaseBuf[4*I + 0];
+						DestBuf[I] = uint8(ResultValue >> 8);
 					}
 				});
 
@@ -337,8 +337,8 @@ void FImageOperator::ImagePixelFormat(bool& bOutSuccess, int32 CompressionQualit
 				{
 					for (int32 I = 0; I < NumElems; ++I)
 					{
-						uint32 ResultValue = 76 * BaseBuf[3*I + 0] + 150 * BaseBuf[3*I + 1] + 29 * BaseBuf[3*I + 2];
-						DestBuf[I] = (uint8)FMath::Min(255u, ResultValue >> 8);
+						uint32 ResultValue = 77 * BaseBuf[3*I + 0] + 150 * BaseBuf[3*I + 1] + 29 * BaseBuf[3*I + 2];
+						DestBuf[I] = uint8(ResultValue >> 8);
 					}
 				});
 
