[mutable] Ensure layout pod is fully initialized to prevent binary indeterminism in serialized data.
#rnx

diff --git a/Engine/Plugins/Experimental/Mutable/Source/MutableRuntime/Public/MuR/Layout.h b/Engine/Plugins/Experimental/Mutable/Source/MutableRuntime/Public/MuR/Layout.h
--- a/Engine/Plugins/Experimental/Mutable/Source/MutableRuntime/Public/MuR/Layout.h
+++ b/Engine/Plugins/Experimental/Mutable/Source/MutableRuntime/Public/MuR/Layout.h
@@ -135,6 +135,7 @@ namespace mu
 				Priority = 0;
 				bReduceBothAxes = false;
 				bReduceByTwo = false;
+				UnusedPadding = 0;
 			}
 
 			UE::Math::TIntVector2<uint16> Min = UE::Math::TIntVector2<uint16>(0,0);
@@ -152,6 +153,9 @@ namespace mu
 			//! Value to control if a block has to be reduced by two in an unitary reduction strategy
 			uint32 bReduceByTwo : 1;
 
+			/** Explicit padding to prevent uninitialized memory in this POD. */
+			uint32 UnusedPadding : 30;
+
 
 			//!
 			inline bool operator==(const FBlock& o) const
