[mutable] Added a new check to table nodes that uses the same Data Table with compilation restrictions.
+ Added new section to the table node details "Compilation Restrictions"
+ Added new warning that is triggered when two Table nodes using the same Data Table have different Compliation Restrictions.

#rnx

diff --git a/Engine/Plugins/Experimental/Mutable/Source/CustomizableObjectEditor/Private/MuCOE/GenerateMutableSource/GenerateMutableSource.cpp b/Engine/Plugins/Experimental/Mutable/Source/CustomizableObjectEditor/Private/MuCOE/GenerateMutableSource/GenerateMutableSource.cpp
--- a/Engine/Plugins/Experimental/Mutable/Source/CustomizableObjectEditor/Private/MuCOE/GenerateMutableSource/GenerateMutableSource.cpp
+++ b/Engine/Plugins/Experimental/Mutable/Source/CustomizableObjectEditor/Private/MuCOE/GenerateMutableSource/GenerateMutableSource.cpp
@@ -2134,7 +2134,14 @@ FMutableGraphGenerationContext::FSharedSurface::FSharedSurface(uint8 InLOD, cons
 	NodeSurfaceNew = InNodeSurfaceNew;
 }
 
+
 bool FMutableGraphGenerationContext::FSharedSurface::operator==(const FSharedSurface& o) const
 {
 	return NodeModifierIDs == o.NodeModifierIDs;
 }
+
+
+bool FMutableGraphGenerationContext::FGeneratedDataTablesData::HasSameSettings(const UCustomizableObjectNodeTable* Node) const
+{
+	return Node->bDisableCheckedRows == bDisableCheckedRows && Node->VersionColumn == VersionColumn;
+}
diff --git a/Engine/Plugins/Experimental/Mutable/Source/CustomizableObjectEditor/Private/MuCOE/GenerateMutableSource/GenerateMutableSource.h b/Engine/Plugins/Experimental/Mutable/Source/CustomizableObjectEditor/Private/MuCOE/GenerateMutableSource/GenerateMutableSource.h
--- a/Engine/Plugins/Experimental/Mutable/Source/CustomizableObjectEditor/Private/MuCOE/GenerateMutableSource/GenerateMutableSource.h
+++ b/Engine/Plugins/Experimental/Mutable/Source/CustomizableObjectEditor/Private/MuCOE/GenerateMutableSource/GenerateMutableSource.h
@@ -555,23 +555,48 @@ struct FMutableGraphGenerationContext
 	/** Set of all generated nodes. */
 	TSet<UCustomizableObjectNode*> GeneratedNodes;
 
+	/** Struct that stores the relevant information of a data table generated during the compilation. 
+	e.g. all data tables must have the same compilation restrictions */
+	struct FGeneratedDataTablesData
+	{
+		// Pointer to the generated mutable Table
+		mu::Ptr<mu::Table> GeneratedTable;
+
+		// Table Node used to fill this info
+		const UCustomizableObjectNodeTable* ReferenceNode;
+
+		// Stores the names of the rows that will be compiled
+		TArray<FName> RowNames;
+
+		// Compilation Restrictions:
+		// If there is a bool column in the table, checked rows will not be compiled
+		bool bDisableCheckedRows;
+
+		// Name of the column that determines de version control
+		FName VersionColumn;
+
+		// Compare the stored compilation settings with the compilation settings of a Table Node
+		// return true if the compilation settings are equal
+		bool HasSameSettings(const UCustomizableObjectNodeTable* Node) const;
+	};
+
 	// Cache of generated Node Tables
-	TMap<FString, mu::TablePtr> GeneratedTables;
+	TMap<FString, FGeneratedDataTablesData> GeneratedTables;
 
-	struct FGeneratedDataTablesData
+	struct FGeneratedCompositeDataTablesData
 	{
 		UScriptStruct* ParentStruct = nullptr;
 		TArray<FName> FilterPaths;
 		UCompositeDataTable* GeneratedDataTable = nullptr;
 
-		bool operator==(const FGeneratedDataTablesData& Other) const
+		bool operator==(const FGeneratedCompositeDataTablesData& Other) const
 		{
 			return ParentStruct == Other.ParentStruct && FilterPaths == Other.FilterPaths;
 		}
 	};
 
 	// Cache of generated Composited Data Tables
-	TArray<FGeneratedDataTablesData> GeneratedCompositeDataTables;
+	TArray<FGeneratedCompositeDataTablesData> GeneratedCompositeDataTables;
 
 	// Cache of generated images, because sometimes they are reused by LOD, we use this as a second
 	// level cache
@@ -833,9 +858,6 @@ struct FMutableGraphGenerationContext
 	// Current material parameter id to find the corresponding column in a mutable table
 	FString CurrentMaterialTableParameterId;
 
-	// Stores the parameters generated in the node tables
-	TMap<const class UCustomizableObjectNodeTable*, TArray<FGuid>> GeneratedParametersInTables;
-
 	struct FSharedSurface
 	{
 		FSharedSurface(uint8 InLOD, const mu::NodeSurfaceNewPtr& InNodeSurfaceNew);
diff --git a/Engine/Plugins/Experimental/Mutable/Source/CustomizableObjectEditor/Private/MuCOE/GenerateMutableSource/GenerateMutableSourceMesh.cpp b/Engine/Plugins/Experimental/Mutable/Source/CustomizableObjectEditor/Private/MuCOE/GenerateMutableSource/GenerateMutableSourceMesh.cpp
--- a/Engine/Plugins/Experimental/Mutable/Source/CustomizableObjectEditor/Private/MuCOE/GenerateMutableSource/GenerateMutableSourceMesh.cpp
+++ b/Engine/Plugins/Experimental/Mutable/Source/CustomizableObjectEditor/Private/MuCOE/GenerateMutableSource/GenerateMutableSourceMesh.cpp
@@ -2547,7 +2547,7 @@ TArray<TTuple<USkeletalMesh*, TSoftClassPtr<UAnimInstance>>> GetSkeletalMeshesIn
 
 		if (DataTable)
 		{
-			for (const FName& RowName : GetRowsToCompile(*DataTable, *TableNode,GenerationContext))
+			for (const FName& RowName : GetRowsToCompile(*DataTable, *TableNode, GenerationContext))
 			{
 				USkeletalMesh* SkeletalMesh = TableNode->GetSkeletalMeshAt(SourceMeshPin, DataTable, RowName);
 				TSoftClassPtr<UAnimInstance> MeshAnimInstance = TableNode->GetAnimInstanceAt(SourceMeshPin, DataTable, RowName);
@@ -2557,7 +2557,7 @@ TArray<TTuple<USkeletalMesh*, TSoftClassPtr<UAnimInstance>>> GetSkeletalMeshesIn
 					SkeletalMeshesInfo.Emplace(SkeletalMesh, MeshAnimInstance);
 				}
 			}
-		}	
+		}
 	}
 	else
 	{
@@ -3014,7 +3014,8 @@ mu::NodeMeshPtr GenerateMorphMesh(const UEdGraphPin* Pin,
 		int32 NumRows = RowNames.Num();
 
 		// Should exist
-		mu::TablePtr Table = GenerationContext.GeneratedTables[DataTable->GetName()];
+		mu::TablePtr Table = GenerationContext.GeneratedTables[DataTable->GetName()].GeneratedTable;
+		check(Table);
 
 		FString ColumnName = TableColumnName + TypedNodeMorphs[MorphIndex].MorphTargetName;
 		int32 ColumnIndex = INDEX_NONE;
diff --git a/Engine/Plugins/Experimental/Mutable/Source/CustomizableObjectEditor/Private/MuCOE/GenerateMutableSource/GenerateMutableSourceTable.cpp b/Engine/Plugins/Experimental/Mutable/Source/CustomizableObjectEditor/Private/MuCOE/GenerateMutableSource/GenerateMutableSourceTable.cpp
--- a/Engine/Plugins/Experimental/Mutable/Source/CustomizableObjectEditor/Private/MuCOE/GenerateMutableSource/GenerateMutableSourceTable.cpp
+++ b/Engine/Plugins/Experimental/Mutable/Source/CustomizableObjectEditor/Private/MuCOE/GenerateMutableSource/GenerateMutableSourceTable.cpp
@@ -797,15 +797,22 @@ void RestrictRowContentByVersion( TArray<FName>& InOutRowNames, const UDataTable
 
 TArray<FName> GetRowsToCompile(const UDataTable& DataTable, const UCustomizableObjectNodeTable& TableNode, FMutableGraphGenerationContext& GenerationContext)
 {
-	TArray<FName> RowNames = GetEnabledRows(DataTable, TableNode);
-
-	if (!RowNames.IsEmpty())
+	if (FMutableGraphGenerationContext::FGeneratedDataTablesData* Result = GenerationContext.GeneratedTables.Find(DataTable.GetName()))
 	{
-		RestrictRowNamesToSelectedOption(RowNames, TableNode, GenerationContext);
-		RestrictRowContentByVersion(RowNames, DataTable, TableNode, GenerationContext);
+		return Result->RowNames;
 	}
+	else
+	{
+		TArray<FName> RowNames = GetEnabledRows(DataTable, TableNode);
 
-	return RowNames;
+		if (!RowNames.IsEmpty())
+		{
+			RestrictRowNamesToSelectedOption(RowNames, TableNode, GenerationContext);
+			RestrictRowContentByVersion(RowNames, DataTable, TableNode, GenerationContext);
+		}
+
+		return RowNames;
+	}
 }
 
 
@@ -928,15 +935,27 @@ mu::TablePtr GenerateMutableSourceTable(const UDataTable* DataTable, const UCust
 {
 	check(DataTable && TableNode);
 
-	// Checking if the table in the cache
+	// Checking if the table is in the cache
 	const FString TableName = DataTable->GetName();
 
-	if (mu::TablePtr* Result = GenerationContext.GeneratedTables.Find(TableName))
+	if (FMutableGraphGenerationContext::FGeneratedDataTablesData* CachedTable = GenerationContext.GeneratedTables.Find(TableName))
 	{
 		// Generating Parameter Metadata for parameters that reuse a Table
 		GenerateTableParameterUIData(DataTable, TableNode, GenerationContext);
 
-		return *Result;
+		FMutableGraphGenerationContext::FGeneratedDataTablesData ;
+
+		if (!CachedTable->HasSameSettings(TableNode))
+		{
+			TArray<const UObject*> Nodes;
+			Nodes.Add(TableNode);
+			Nodes.Add(CachedTable->ReferenceNode);
+
+			GenerationContext.Compiler->CompilerLog(LOCTEXT("TableNodesCompilationRestrictionError",
+				"Found one or more Table Nodes with the same data table but different Compilation Restrictions."), Nodes);
+		}
+
+		return CachedTable->GeneratedTable;
 	}
 
 	mu::TablePtr MutableTable = new mu::Table();
@@ -957,6 +976,16 @@ mu::TablePtr GenerateMutableSourceTable(const UDataTable* DataTable, const UCust
 
 		// Generating Parameter Metadata for new table parameters
 		GenerateTableParameterUIData(DataTable, TableNode, GenerationContext);
+
+		FMutableGraphGenerationContext::FGeneratedDataTablesData GeneratedTable;
+		GeneratedTable.GeneratedTable = MutableTable;
+		GeneratedTable.bDisableCheckedRows = TableNode->bDisableCheckedRows;
+		GeneratedTable.VersionColumn = TableNode->VersionColumn;
+		GeneratedTable.RowNames = RowNames;
+		GeneratedTable.ReferenceNode = TableNode;
+
+		// Add table to cache
+		GenerationContext.GeneratedTables.Add(TableName, GeneratedTable);
 	}
 	else
 	{
@@ -966,8 +995,6 @@ mu::TablePtr GenerateMutableSourceTable(const UDataTable* DataTable, const UCust
 		return nullptr;
 	}
 
-	GenerationContext.GeneratedTables.Add(TableName, MutableTable);
-
 	return MutableTable;
 }
 
@@ -1013,7 +1040,7 @@ UDataTable* GenerateDataTableFromStruct(const UCustomizableObjectNodeTable* Tabl
 		return nullptr;
 	}
 
-	FMutableGraphGenerationContext::FGeneratedDataTablesData DataTableData;
+	FMutableGraphGenerationContext::FGeneratedCompositeDataTablesData DataTableData;
 	DataTableData.ParentStruct = TableNode->Structure;
 	DataTableData.FilterPaths = TableNode->FilterPaths;
 	
diff --git a/Engine/Plugins/Experimental/Mutable/Source/CustomizableObjectEditor/Private/MuCOE/Nodes/CustomizableObjectNodeTableDetails.cpp b/Engine/Plugins/Experimental/Mutable/Source/CustomizableObjectEditor/Private/MuCOE/Nodes/CustomizableObjectNodeTableDetails.cpp
--- a/Engine/Plugins/Experimental/Mutable/Source/CustomizableObjectEditor/Private/MuCOE/Nodes/CustomizableObjectNodeTableDetails.cpp
+++ b/Engine/Plugins/Experimental/Mutable/Source/CustomizableObjectEditor/Private/MuCOE/Nodes/CustomizableObjectNodeTableDetails.cpp
@@ -51,6 +51,7 @@ void FCustomizableObjectNodeTableDetails::CustomizeDetails(const TSharedPtr<IDet
 	if (Node.IsValid())
 	{
 		IDetailCategoryBuilder& CustomizableObjectCategory = DetailBuilder->EditCategory("TableProperties");
+		IDetailCategoryBuilder& CompilationRestrictionsCategory = DetailBuilder->EditCategory("CompilationRestrictions");
 		DetailBuilder->HideProperty("VersionColumn");
 		IDetailCategoryBuilder& UICategory = DetailBuilder->EditCategory("UI");
 		DetailBuilder->HideProperty("ParamUIMetadataColumn");
@@ -67,7 +68,8 @@ void FCustomizableObjectNodeTableDetails::CustomizeDetails(const TSharedPtr<IDet
 		TSharedPtr<FString> CurrentThumbnailColumn = GenerateThumbnailColumnComboBoxOptions();
 
 		CustomizableObjectCategory.AddProperty("ParameterName");
-		CustomizableObjectCategory.AddCustomRow(LOCTEXT("VersionColumn_Selector","VersionColumn"))
+
+		CompilationRestrictionsCategory.AddCustomRow(LOCTEXT("VersionColumn_Selector","VersionColumn"))
 		.NameContent()
 		[
 			SNew(STextBlock)
diff --git a/Engine/Plugins/Experimental/Mutable/Source/CustomizableObjectEditor/Public/MuCOE/Nodes/CustomizableObjectNodeTable.h b/Engine/Plugins/Experimental/Mutable/Source/CustomizableObjectEditor/Public/MuCOE/Nodes/CustomizableObjectNodeTable.h
--- a/Engine/Plugins/Experimental/Mutable/Source/CustomizableObjectEditor/Public/MuCOE/Nodes/CustomizableObjectNodeTable.h
+++ b/Engine/Plugins/Experimental/Mutable/Source/CustomizableObjectEditor/Public/MuCOE/Nodes/CustomizableObjectNodeTable.h
@@ -189,7 +189,7 @@ public:
 	bool bAddNoneOption;
 
 	/** If there is a bool column in the table, checked rows will not be compiled */
-	UPROPERTY(EditAnywhere, Category = TableProperties)
+	UPROPERTY(EditAnywhere, Category = CompilationRestrictions)
 	bool bDisableCheckedRows = true;
 
 	/** Source where table gathers the data */
@@ -208,7 +208,7 @@ public:
 	TArray<FName> FilterPaths;
 	
 	/** Name of the column that contains the Version options. */
-	UPROPERTY(EditAnywhere, Category = TableProperties)
+	UPROPERTY(EditAnywhere, Category = CompilationRestrictions)
 	FName VersionColumn;
 
 	/** Name of the row that will be used as default value. */
