[mutable] Removed unused declaration.

#rnx

diff --git a/Engine/Plugins/Experimental/Mutable/Source/CustomizableObjectEditor/Private/MuCOE/GenerateMutableSource/GenerateMutableSourceTable.cpp b/Engine/Plugins/Experimental/Mutable/Source/CustomizableObjectEditor/Private/MuCOE/GenerateMutableSource/GenerateMutableSourceTable.cpp
--- a/Engine/Plugins/Experimental/Mutable/Source/CustomizableObjectEditor/Private/MuCOE/GenerateMutableSource/GenerateMutableSourceTable.cpp
+++ b/Engine/Plugins/Experimental/Mutable/Source/CustomizableObjectEditor/Private/MuCOE/GenerateMutableSource/GenerateMutableSourceTable.cpp
@@ -943,8 +943,6 @@ mu::TablePtr GenerateMutableSourceTable(const UDataTable* DataTable, const UCust
 		// Generating Parameter Metadata for parameters that reuse a Table
 		GenerateTableParameterUIData(DataTable, TableNode, GenerationContext);
 
-		FMutableGraphGenerationContext::FGeneratedDataTablesData ;
-
 		if (!CachedTable->HasSameSettings(TableNode))
 		{
 			TArray<const UObject*> Nodes;
