[mutable] Fixed stalled updates until the main editor window was refocused. When changing a Mutable parameter with a Color Picker, it entered responsive mode and Async Asset Loads didn't complete until the editor was clicked on.
#tests Editor
#rnx

diff --git a/Engine/Plugins/Experimental/Mutable/Source/CustomizableObject/Private/MuCO/CustomizableObjectInstance.cpp b/Engine/Plugins/Experimental/Mutable/Source/CustomizableObject/Private/MuCO/CustomizableObjectInstance.cpp
--- a/Engine/Plugins/Experimental/Mutable/Source/CustomizableObject/Private/MuCO/CustomizableObjectInstance.cpp
+++ b/Engine/Plugins/Experimental/Mutable/Source/CustomizableObject/Private/MuCO/CustomizableObjectInstance.cpp
@@ -60,6 +60,7 @@
 #include "UnrealEdMisc.h"
 #include "Subsystems/AssetEditorSubsystem.h"
 #include "Misc/TransactionObjectEvent.h"
+#include "Application/ThrottleManager.h"
 #endif
 
 namespace
@@ -5378,6 +5379,13 @@ UE::Tasks::FTask UCustomizableInstancePrivate::LoadAdditionalAssetsAndData(
 	{	
 		UE::Tasks::FTaskEvent AssetAsyncLoadCompletionEvent = StreamingCompletionEvents.Emplace_GetRef(TEXT("AssetAsyncLoadCompletionEvent"));
 
+#if WITH_EDITOR
+		// TODO: Remove with UE-217665 when the underlying bug in the ColorPicker is solved
+		// Disable the Slate throttling, otherwise the AsyncLoad may not complete until the editor window is clicked on due to a bug in
+		// some widgets such as the ColorPicker's throttling handling
+		FSlateThrottleManager::Get().DisableThrottle(true);
+#endif
+
 		if (bAsync)
 		{
 			StreamingHandle = StreamableManager.RequestAsyncLoad(
@@ -5641,6 +5649,12 @@ void UCustomizableInstancePrivate::AdditionalAssetsAsyncLoaded(UE::Tasks::FTaskE
 	CompletionEvent.Trigger(); // TODO: we know it is game thread?
 
 	StreamingHandle = nullptr;
+
+#if WITH_EDITOR
+	// TODO: Remove with UE-217665 when the underlying bug in the ColorPicker is solved
+	// Reenable the throttling which disabled when launching the Async Load
+	FSlateThrottleManager::Get().DisableThrottle(false);
+#endif
 }
 
 
