[mutable] Added flags to compiled stream data chunks to track "high-res" data.
- Minor optimizations in the rom generation stage.
#rnx

diff --git a/Engine/Plugins/Experimental/Mutable/Source/CustomizableObject/Internal/MuCO/CustomizableObjectCompilerTypes.h b/Engine/Plugins/Experimental/Mutable/Source/CustomizableObject/Internal/MuCO/CustomizableObjectCompilerTypes.h
--- a/Engine/Plugins/Experimental/Mutable/Source/CustomizableObject/Internal/MuCO/CustomizableObjectCompilerTypes.h
+++ b/Engine/Plugins/Experimental/Mutable/Source/CustomizableObject/Internal/MuCO/CustomizableObjectCompilerTypes.h
@@ -60,6 +60,14 @@ struct FCompilationOptions
 	/** High (inclusive) limit of the size in bytes of a data block to be included into the compiled object directly instead of stored in a streamable file. */
 	uint64 EmbeddedDataBytesLimit = 1024;
 
+	/** Number of minimum mipmaps that we want to always be available in disk regardless of NumHighResImageLODs. */
+	int32 MinDiskMips = 7;
+
+	/** Number of image mipmaps that will be flagged as high-res data (possibly to store separately).
+	* This is only used if the total mips in the source image is above the MinDiskMips.
+	*/
+	int32 NumHighResImageMips = 2;
+
 	// Did we have the extra bones enabled when we compiled?
 	ECustomizableObjectNumBoneInfluences CustomizableObjectNumBoneInfluences = ECustomizableObjectNumBoneInfluences::Four;
 
diff --git a/Engine/Plugins/Experimental/Mutable/Source/CustomizableObject/Internal/MuCO/CustomizableObjectPrivate.h b/Engine/Plugins/Experimental/Mutable/Source/CustomizableObject/Internal/MuCO/CustomizableObjectPrivate.h
--- a/Engine/Plugins/Experimental/Mutable/Source/CustomizableObject/Internal/MuCO/CustomizableObjectPrivate.h
+++ b/Engine/Plugins/Experimental/Mutable/Source/CustomizableObject/Internal/MuCO/CustomizableObjectPrivate.h
@@ -1087,6 +1087,8 @@ public:
 
 		LayoutToPOD,
 
+		AddedRomFlags,
+
 		// -----<new versions can be added above this line>--------
 		LastCustomizableObjectVersion
 	};
diff --git a/Engine/Plugins/Experimental/Mutable/Source/CustomizableObject/Private/MuCO/CustomizableObject.cpp b/Engine/Plugins/Experimental/Mutable/Source/CustomizableObject/Private/MuCO/CustomizableObject.cpp
--- a/Engine/Plugins/Experimental/Mutable/Source/CustomizableObject/Private/MuCO/CustomizableObject.cpp
+++ b/Engine/Plugins/Experimental/Mutable/Source/CustomizableObject/Private/MuCO/CustomizableObject.cpp
@@ -2151,11 +2151,15 @@ void UCustomizableObjectBulk::PostLoad()
 	BulkFilePrefix = PackageFilename;
 }
 
-TUniquePtr<IAsyncReadFileHandle> UCustomizableObjectBulk::OpenFileAsyncRead(uint32 FileId) const
+TUniquePtr<IAsyncReadFileHandle> UCustomizableObjectBulk::OpenFileAsyncRead(uint32 FileId, uint32 Flags) const
 {
 	check(IsInGameThread());
 
 	FString FilePath = FString::Printf(TEXT("%s-%08x.mut"), *BulkFilePrefix, FileId);
+	if (Flags == uint32(mu::ERomFlags::HighRes))
+	{
+		FilePath += TEXT(".high");
+	}
 
 	IAsyncReadFileHandle* Result = FPlatformFileManager::Get().GetPlatformFile().OpenAsyncRead(*FilePath);
 	
@@ -2248,6 +2252,14 @@ void UCustomizableObjectBulk::CookAdditionalFilesOverride(const TCHAR* PackageFi
 		const FString CookedFilePath = FPaths::GetPath(PackageFilename);
 
 		FString CookedBulkFileName = FString::Printf(TEXT("%s/%s-%08x.mut"), *CookedFilePath, *CustomizableObject->GetName(), CurrentFile.Id);
+
+		if (CurrentFile.Flags == uint32(mu::ERomFlags::HighRes))
+		{
+			// We can do something different here for high-res data.
+			// For example: change the file name. We also need to detect it when generating the file name for loading.
+			CookedBulkFileName += TEXT(".high");
+		}
+
 		WriteAdditionalFile(*CookedBulkFileName, FileBulkData.GetData(), FileBulkData.Num());
 	}
 }
@@ -2284,19 +2296,30 @@ void UCustomizableObjectBulk::PrepareBulkData(UCustomizableObject* InOuter, cons
 		//uint32 Depth=0;
 	};
 
-	FClassifyNode RootNode;
-	// Create blocks data, filtering out the ones that are too big and will go on its own file in any case.
-	{
-		const int32 NumRoms = Model->GetRomCount();
-		RootNode.Blocks.Reserve(NumRoms);
-		
+	// Root nodes by flags.
+	const int32 NumRoms = Model->GetRomCount();
+	TMap<uint32,FClassifyNode> RootNode;
+	auto AddNode = [&RootNode,NumRoms]( const FBlock& Block )
+		{
+			FClassifyNode& Root = RootNode.FindOrAdd(Block.Flags);
+			if (Root.Blocks.IsEmpty())
+			{
+				Root.Blocks.Reserve(NumRoms);
+			}
+
+			Root.Blocks.Add(Block);
+		};
+
+	// Create blocks data.
+	{		
 		for (int32 RomIndex = 0; RomIndex < NumRoms; ++RomIndex)
 		{
 			uint32 BlockId = Model->GetRomId(RomIndex);
 			const uint32 BlockSize = Model->GetRomSize(RomIndex);
+			const mu::ERomFlags BlockFlags = Model->GetRomFlags(RomIndex);
 
-			FBlock CurrentBlock = { EDataType::Model, BlockId, BlockSize, 0 };
-			RootNode.Blocks.Add(CurrentBlock);
+			FBlock CurrentBlock = { EDataType::Model, BlockId, BlockSize, uint32(BlockFlags), 0 };
+			AddNode(CurrentBlock);
 		}
 	}
 
@@ -2315,8 +2338,9 @@ void UCustomizableObjectBulk::PrepareBulkData(UCustomizableObject* InOuter, cons
 			const uint32 BlockSize = MorphStreamable.Value.Size;
 			
 			check(SourceOffset == MorphStreamable.Value.Block.Offset);
-			FBlock CurrentBlock = { EDataType::RealTimeMorph, MorphStreamable.Key, BlockSize, SourceOffset };
-			RootNode.Blocks.Add(CurrentBlock);
+			uint32 Flags = 0;
+			FBlock CurrentBlock = { EDataType::RealTimeMorph, MorphStreamable.Key, BlockSize, Flags, SourceOffset };
+			AddNode(CurrentBlock);
 
 			SourceOffset += BlockSize;
 		}
@@ -2337,47 +2361,57 @@ void UCustomizableObjectBulk::PrepareBulkData(UCustomizableObject* InOuter, cons
 			const uint32 BlockSize = ClothStreamable.Value.Size;
 			
 			check(SourceOffset == ClothStreamable.Value.Block.Offset);
-			FBlock CurrentBlock = { EDataType::Clothing, ClothStreamable.Key, BlockSize, SourceOffset };
-			RootNode.Blocks.Add(CurrentBlock);
+			uint32 Flags = 0;
+			FBlock CurrentBlock = { EDataType::Clothing, ClothStreamable.Key, BlockSize, Flags, SourceOffset };
+			AddNode(CurrentBlock);
 
 			SourceOffset += BlockSize;
 		}
 	}
 
 
-	// Temp: Group by order in the array
-	for (int32 BlockIndex = 0; BlockIndex < RootNode.Blocks.Num(); )
+	for (int32 FlagClassIndex = 0; FlagClassIndex < RootNode.Num(); ++FlagClassIndex)
 	{
-		int32 CurrentFileSize = 0;
-		
-		FFile CurrentFile;
-		CurrentFile.DataType = RootNode.Blocks[BlockIndex].DataType;
+		FClassifyNode& Root = RootNode[FlagClassIndex];
 
-		while(BlockIndex < RootNode.Blocks.Num())
+		// Temp: Group by order in the array
+		for (int32 BlockIndex = 0; BlockIndex < Root.Blocks.Num(); )
 		{
-			FBlock CurrentBlock = RootNode.Blocks[BlockIndex];
+			int32 CurrentFileSize = 0;
 
-			// Next file?
-			// Store different data types in different files. Blocks should be sorted by DataType so data is properly packeted
-			if (CurrentFile.DataType != CurrentBlock.DataType)
-			{
-				break;
-			}
-			
-			if (CurrentFileSize > 0 && CurrentFileSize + CurrentBlock.Size > TargetBulkDataFileBytes)
+			FFile CurrentFile;
+			CurrentFile.DataType = Root.Blocks[BlockIndex].DataType;
+			CurrentFile.Flags = Root.Blocks[BlockIndex].Flags;
+
+			while (BlockIndex < Root.Blocks.Num())
 			{
-				break;
-			}
+				FBlock CurrentBlock = Root.Blocks[BlockIndex];
 
-			// Add the block to the current file
-			CurrentFile.Blocks.Add(CurrentBlock);
-			CurrentFileSize += CurrentBlock.Size;
+				// Next file?
+				// Store different data types in different files. Blocks should be sorted by DataType so data is properly packeted
+				if (CurrentFile.DataType != CurrentBlock.DataType)
+				{
+					break;
+				}
 
-			// Next block
-			++BlockIndex;			
-		}
+				// Different flags go to different files
+				check (CurrentFile.Flags == CurrentBlock.Flags)
 
-		BulkDataFiles.Add(MoveTemp(CurrentFile));
+				if (CurrentFileSize > 0 && CurrentFileSize + CurrentBlock.Size > TargetBulkDataFileBytes)
+				{
+					break;
+				}
+
+				// Add the block to the current file
+				CurrentFile.Blocks.Add(CurrentBlock);
+				CurrentFileSize += CurrentBlock.Size;
+
+				// Next block
+				++BlockIndex;
+			}
+
+			BulkDataFiles.Add(MoveTemp(CurrentFile));
+		}
 	}
 
 	// Create the file list
@@ -2422,6 +2456,7 @@ void UCustomizableObjectBulk::PrepareBulkData(UCustomizableObject* InOuter, cons
 				FMutableStreamableBlock* StreamableBlock = ModelResources.HashToStreamableBlock.Find(ThisBlock.Id);
 				StreamableBlock->FileId = FileId;
 				StreamableBlock->Offset = OffsetInFile;
+				check(StreamableBlock->Flags==CurrentFile.Flags);
 				OffsetInFile += ThisBlock.Size;
 			}
 		}
diff --git a/Engine/Plugins/Experimental/Mutable/Source/CustomizableObject/Private/MuCO/CustomizableObjectInstance.cpp b/Engine/Plugins/Experimental/Mutable/Source/CustomizableObject/Private/MuCO/CustomizableObjectInstance.cpp
--- a/Engine/Plugins/Experimental/Mutable/Source/CustomizableObject/Private/MuCO/CustomizableObjectInstance.cpp
+++ b/Engine/Plugins/Experimental/Mutable/Source/CustomizableObject/Private/MuCO/CustomizableObjectInstance.cpp
@@ -5508,7 +5508,7 @@ UE::Tasks::FTask UCustomizableInstancePrivate::LoadAdditionalAssetsAndData(
 			int32 FileHandleIndex = OpenFilesIds.Find(Block.FileId);
 			if (FileHandleIndex == INDEX_NONE && BulkData)
 			{
-				TUniquePtr<IAsyncReadFileHandle> ReadFileHandle = BulkData->OpenFileAsyncRead(Block.FileId);
+				TUniquePtr<IAsyncReadFileHandle> ReadFileHandle = BulkData->OpenFileAsyncRead(Block.FileId,Block.Flags);
 
 				OpenFileHandles.Emplace(MoveTemp(ReadFileHandle));
 				FileHandleIndex = OpenFilesIds.Add(Block.FileId);
@@ -5553,7 +5553,7 @@ UE::Tasks::FTask UCustomizableInstancePrivate::LoadAdditionalAssetsAndData(
 			int32 FileHandleIndex = OpenFilesIds.Find(Block.FileId);
 			if (FileHandleIndex == INDEX_NONE && BulkData)
 			{
-				TUniquePtr<IAsyncReadFileHandle> ReadFileHandle = BulkData->OpenFileAsyncRead(Block.FileId);
+				TUniquePtr<IAsyncReadFileHandle> ReadFileHandle = BulkData->OpenFileAsyncRead(Block.FileId, Block.Flags);
 
 				OpenFileHandles.Emplace(MoveTemp(ReadFileHandle));
 				FileHandleIndex = OpenFilesIds.Add(Block.FileId);
diff --git a/Engine/Plugins/Experimental/Mutable/Source/CustomizableObject/Private/MuCO/UnrealMutableModelDiskStreamer.cpp b/Engine/Plugins/Experimental/Mutable/Source/CustomizableObject/Private/MuCO/UnrealMutableModelDiskStreamer.cpp
--- a/Engine/Plugins/Experimental/Mutable/Source/CustomizableObject/Private/MuCO/UnrealMutableModelDiskStreamer.cpp
+++ b/Engine/Plugins/Experimental/Mutable/Source/CustomizableObject/Private/MuCO/UnrealMutableModelDiskStreamer.cpp
@@ -256,6 +256,10 @@ mu::ModelReader::OPERATION_ID FUnrealMutableModelBulkReader::BeginReadBlock(cons
 			FString FilePath = ObjectData->BulkFilePrefix;
 #else
 			FString FilePath = FString::Printf(TEXT("%s-%08x.mut"), *ObjectData->BulkFilePrefix, Block->FileId);
+			if (Block->Flags == uint32(mu::ERomFlags::HighRes))
+			{
+				FilePath += TEXT(".high");
+			}
 #endif
 
 			Found = MakeShareable(FPlatformFileManager::Get().GetPlatformFile().OpenAsyncRead(*FilePath));
diff --git a/Engine/Plugins/Experimental/Mutable/Source/CustomizableObject/Public/MuCO/CustomizableObject.h b/Engine/Plugins/Experimental/Mutable/Source/CustomizableObject/Public/MuCO/CustomizableObject.h
--- a/Engine/Plugins/Experimental/Mutable/Source/CustomizableObject/Public/MuCO/CustomizableObject.h
+++ b/Engine/Plugins/Experimental/Mutable/Source/CustomizableObject/Public/MuCO/CustomizableObject.h
@@ -206,15 +206,19 @@ struct CUSTOMIZABLEOBJECT_API FMutableStreamableBlock
 
 	UPROPERTY()
 	uint32 FileId = 0;
-	
+
+	/** Used to store properties of the data, necessary for its recovery. For instance if it is high-res. */
+	UPROPERTY()
+	uint32 Flags = 0;
+
 	UPROPERTY()
 	uint64 Offset = 0;
 
 	friend FArchive& operator<<(FArchive& Ar, FMutableStreamableBlock& Data)
 	{
 		Ar << Data.FileId;
+		Ar << Data.Flags;
 		Ar << Data.Offset;
-
 		return Ar;
 	}
 };
@@ -268,7 +272,7 @@ public:
 	/**  */
 	const FString& GetBulkFilePrefix() const { return BulkFilePrefix; }
 	
-	TUniquePtr<IAsyncReadFileHandle> OpenFileAsyncRead(uint32 FileId) const;
+	TUniquePtr<IAsyncReadFileHandle> OpenFileAsyncRead(uint32 FileId, uint32 Flags) const;
 
 #if WITH_EDITOR
 
@@ -306,6 +310,9 @@ private:
 		/** Size of the data block. */
 		uint32 Size;
 
+		/** Data flags, like "high-res". */
+		uint32 Flags;
+
 		/** Offset in the full source streamed data file that is created when compiling. */
 		uint64 Offset;
 	};
@@ -317,6 +324,9 @@ private:
 		/** Id generated from a hash of the file content + offset to avoid collisions. */
 		uint32 Id;
 
+		/** Common flags of the data stored in this file. See mu::ERomFlags. */
+		uint32 Flags=0;
+
 		/** List of blocks that are contained in the file, in order. */
 		TArray<FBlock> Blocks;
 
diff --git a/Engine/Plugins/Experimental/Mutable/Source/CustomizableObjectEditor/Private/MuCOE/CustomizableObjectCompiler.cpp b/Engine/Plugins/Experimental/Mutable/Source/CustomizableObjectEditor/Private/MuCOE/CustomizableObjectCompiler.cpp
--- a/Engine/Plugins/Experimental/Mutable/Source/CustomizableObjectEditor/Private/MuCOE/CustomizableObjectCompiler.cpp
+++ b/Engine/Plugins/Experimental/Mutable/Source/CustomizableObjectEditor/Private/MuCOE/CustomizableObjectCompiler.cpp
@@ -998,7 +998,8 @@ void FCustomizableObjectCompiler::CompileInternal(bool bAsync)
 
 			ResourceMeshData.NameResolutionMap = MeshData.Value.NameResolutionMap;
 			ResourceMeshData.Size = DataSizeInBytes;
-			ResourceMeshData.Block = FMutableStreamableBlock { uint32(0), RealTimeMorphDataOffsetInBytes };
+			mu::ERomFlags Flags = mu::ERomFlags::None;
+			ResourceMeshData.Block = FMutableStreamableBlock { uint32(0), uint32(Flags), RealTimeMorphDataOffsetInBytes };
 
 			RealTimeMorphDataOffsetInBytes += DataSizeInBytes;
 			ModelResources.EditorOnlyMorphTargetReconstructionData.Append(MeshData.Value.Data);
@@ -1028,7 +1029,8 @@ void FCustomizableObjectCompiler::CompileInternal(bool bAsync)
 			ResourceMeshData.ClothingAssetLOD = MeshData.Value.ClothingAssetLOD;
 			ResourceMeshData.PhysicsAssetIndex = MeshData.Value.PhysicsAssetIndex;
 			ResourceMeshData.Size = DataSizeInBytes;
-			ResourceMeshData.Block = FMutableStreamableBlock { uint32(0), ClothingDataOffsetInBytes };
+			mu::ERomFlags Flags = mu::ERomFlags::None;
+			ResourceMeshData.Block = FMutableStreamableBlock{ uint32(0), uint32(Flags), ClothingDataOffsetInBytes };
 
 			ClothingDataOffsetInBytes += DataSizeInBytes;
 			ModelResources.EditorOnlyClothingMeshToMeshVertData.Append(MeshData.Value.Data);
@@ -1379,8 +1381,8 @@ void FCustomizableObjectCompiler::FinishCompilationTask()
 		{
 			const uint32 ResourceId = CurrentModel->GetRomId(FileIndex);
 			const uint32 ResourceSize = CurrentModel->GetRomSize(FileIndex);
-
-			HashToStreamableBlock->Add(ResourceId, FMutableStreamableBlock{ 0, Offset });
+			mu::ERomFlags Flags = CurrentModel->GetRomFlags(FileIndex);
+			HashToStreamableBlock->Add(ResourceId, FMutableStreamableBlock{ 0, uint32(Flags), Offset });
 			Offset += ResourceSize;
 		}
 
@@ -1397,7 +1399,7 @@ void FCustomizableObjectCompiler::FinishCompilationTask()
 	UE_LOG(LogMutable, Verbose, TEXT("PROFILE: [ %16.8f ] Finishing Compilation task for CO [%s]."), FPlatformTime::Seconds(), *CurrentObject->GetName());
 	TRACE_END_REGION(UE_MUTABLE_COMPILE_REGION);
 
-	// Create  SaveDD task
+	// Create SaveDD task
 	TRACE_BEGIN_REGION(UE_MUTABLE_SAVEDD_REGION);
 	SaveDDTask = MakeShareable(new FCustomizableObjectSaveDDRunnable(CurrentObject, CurrentOptions, CurrentModel));
 }
diff --git a/Engine/Plugins/Experimental/Mutable/Source/CustomizableObjectEditor/Private/MuCOE/GenerateMutableSource/GenerateMutableSourceColor.cpp b/Engine/Plugins/Experimental/Mutable/Source/CustomizableObjectEditor/Private/MuCOE/GenerateMutableSource/GenerateMutableSourceColor.cpp
--- a/Engine/Plugins/Experimental/Mutable/Source/CustomizableObjectEditor/Private/MuCOE/GenerateMutableSource/GenerateMutableSourceColor.cpp
+++ b/Engine/Plugins/Experimental/Mutable/Source/CustomizableObjectEditor/Private/MuCOE/GenerateMutableSource/GenerateMutableSourceColor.cpp
@@ -320,7 +320,7 @@ mu::NodeColourPtr GenerateMutableSourceColor(const UEdGraphPin* Pin, FMutableGra
 					{
 						Result = ColorTableNode;
 
-						ColorTableNode->SetTable(Table);
+						ColorTableNode->Table = Table;
 						ColorTableNode->SetColumn(ColumnName);
 						ColorTableNode->SetParameterName(TypedNodeTable->ParameterName);
 						ColorTableNode->SetNoneOption(TypedNodeTable->bAddNoneOption);
diff --git a/Engine/Plugins/Experimental/Mutable/Source/CustomizableObjectEditor/Private/MuCOE/GenerateMutableSource/GenerateMutableSourceFloat.cpp b/Engine/Plugins/Experimental/Mutable/Source/CustomizableObjectEditor/Private/MuCOE/GenerateMutableSource/GenerateMutableSourceFloat.cpp
--- a/Engine/Plugins/Experimental/Mutable/Source/CustomizableObjectEditor/Private/MuCOE/GenerateMutableSource/GenerateMutableSourceFloat.cpp
+++ b/Engine/Plugins/Experimental/Mutable/Source/CustomizableObjectEditor/Private/MuCOE/GenerateMutableSource/GenerateMutableSourceFloat.cpp
@@ -326,7 +326,7 @@ mu::NodeScalarPtr GenerateMutableSourceFloat(const UEdGraphPin* Pin, FMutableGra
 					{
 						Result = ScalarTableNode;
 
-						ScalarTableNode->SetTable(Table);
+						ScalarTableNode->Table = Table;
 						ScalarTableNode->SetColumn(ColumnName);
 						ScalarTableNode->SetParameterName(TypedNodeTable->ParameterName);
 						ScalarTableNode->SetNoneOption(TypedNodeTable->bAddNoneOption);
diff --git a/Engine/Plugins/Experimental/Mutable/Source/CustomizableObjectEditor/Private/MuCOE/GenerateMutableSource/GenerateMutableSourceGroupProjector.cpp b/Engine/Plugins/Experimental/Mutable/Source/CustomizableObjectEditor/Private/MuCOE/GenerateMutableSource/GenerateMutableSourceGroupProjector.cpp
--- a/Engine/Plugins/Experimental/Mutable/Source/CustomizableObjectEditor/Private/MuCOE/GenerateMutableSource/GenerateMutableSourceGroupProjector.cpp
+++ b/Engine/Plugins/Experimental/Mutable/Source/CustomizableObjectEditor/Private/MuCOE/GenerateMutableSource/GenerateMutableSourceGroupProjector.cpp
@@ -482,11 +482,20 @@ bool GenerateMutableSourceGroupProjector(const UEdGraphPin* Pin, FMutableGraphGe
 				{
 					mu::Ptr<mu::Image> ImageConstant = GenerateImageConstant(ArrayOptionImage[SelectorIndex].OptionImage, GenerationContext, false);
 
-					mu::NodeImageConstantPtr ImageNode = new mu::NodeImageConstant();
+					mu::Ptr<mu::NodeImageConstant> ImageNode = new mu::NodeImageConstant();
 					ImageNode->SetValue(ImageConstant.get());
 
 					const uint32 MipsToSkip = ComputeLODBiasForTexture(GenerationContext, *Texture, ProjParamNode->ReferenceTexture) + AdditionalLODBias;
 					SwitchNode->SetOption(SelectorIndex, ResizeTextureByNumMips(ImageNode, MipsToSkip));
+
+					// Calculate the number of mips to tag as high res for this image.
+					if (ProjParamNode->ReferenceTexture)
+					{
+						int32 TotalMips = mu::Image::GetMipmapCount(ProjParamNode->ReferenceTexture->GetSizeX(), ProjParamNode->ReferenceTexture->GetSizeY());
+						int32 NumMipsBeyondMin = FMath::Max(0, TotalMips - GenerationContext.Options.MinDiskMips);
+						int32 HighResMipsForThisImage = FMath::Min(NumMipsBeyondMin, GenerationContext.Options.NumHighResImageMips);
+						ImageNode->SourceDataDescriptor.SourceHighResMips = HighResMipsForThisImage;
+					}
 				}
 				else
 				{
diff --git a/Engine/Plugins/Experimental/Mutable/Source/CustomizableObjectEditor/Private/MuCOE/GenerateMutableSource/GenerateMutableSourceImage.cpp b/Engine/Plugins/Experimental/Mutable/Source/CustomizableObjectEditor/Private/MuCOE/GenerateMutableSource/GenerateMutableSourceImage.cpp
--- a/Engine/Plugins/Experimental/Mutable/Source/CustomizableObjectEditor/Private/MuCOE/GenerateMutableSource/GenerateMutableSourceImage.cpp
+++ b/Engine/Plugins/Experimental/Mutable/Source/CustomizableObjectEditor/Private/MuCOE/GenerateMutableSource/GenerateMutableSourceImage.cpp
@@ -109,6 +109,8 @@ mu::NodeImagePtr GenerateMutableSourceImage(const UEdGraphPin* Pin, FMutableGrap
 			FGeneratedImageKey ImageKey = FGeneratedImageKey(Pin);
 			mu::NodeImagePtr ImageNode;
 			
+			const uint32 MipsToSkip = ComputeLODBiasForTexture(GenerationContext, *BaseTexture, nullptr, ReferenceTextureSize);
+
 			if (mu::NodeImagePtr* Cached = GenerationContext.GeneratedImages.Find(ImageKey))
 			{
 				ImageNode = *Cached;
@@ -117,14 +119,19 @@ mu::NodeImagePtr GenerateMutableSourceImage(const UEdGraphPin* Pin, FMutableGrap
 			{
 				mu::Ptr<mu::Image> ImageConstant = GenerateImageConstant(BaseTexture, GenerationContext, false);
 
-				mu::Ptr<mu::NodeImageConstant> ReferenceImageNode = new mu::NodeImageConstant;
-				ReferenceImageNode->SetValue( ImageConstant.get() );
-				ImageNode = ReferenceImageNode;
+				mu::Ptr<mu::NodeImageConstant> ConstantImageNode = new mu::NodeImageConstant;
+				ConstantImageNode->SetValue( ImageConstant.get() );
+				
+				int32 TotalMips = mu::Image::GetMipmapCount(ImageConstant->GetSizeX(), ImageConstant->GetSizeY());
+				int32 NumMipsBeyondMin = FMath::Max(0, TotalMips - int32(MipsToSkip) - GenerationContext.Options.MinDiskMips);
+				int32 HighResMipsForThisImage = FMath::Min(NumMipsBeyondMin, GenerationContext.Options.NumHighResImageMips);
+				ConstantImageNode->SourceDataDescriptor.SourceHighResMips = HighResMipsForThisImage;
+
+				ImageNode = ConstantImageNode;
 
 				GenerationContext.GeneratedImages.Add(ImageKey, ImageNode);
 			}
 
-			const uint32 MipsToSkip = ComputeLODBiasForTexture(GenerationContext, *BaseTexture, nullptr, ReferenceTextureSize);
 			Result = ResizeTextureByNumMips(ImageNode, MipsToSkip);
 		}
 		else
@@ -203,17 +210,23 @@ mu::NodeImagePtr GenerateMutableSourceImage(const UEdGraphPin* Pin, FMutableGrap
 
 	else if (const UCustomizableObjectNodeMesh* TypedNodeMesh = Cast<UCustomizableObjectNodeMesh>(Node))
 	{
-		mu::NodeImageConstantPtr ImageNode = new mu::NodeImageConstant();
+		mu::Ptr<mu::NodeImageConstant> ImageNode = new mu::NodeImageConstant();
 		Result = ImageNode;
 
 		UTexture2D* Texture = TypedNodeMesh->FindTextureForPin(Pin);
 
 		if (Texture)
 		{
-			ImageNode->SetValue(GenerateImageConstant(Texture, GenerationContext, false).get());
+			mu::Ptr<mu::Image> ImageConstant = GenerateImageConstant(Texture, GenerationContext, false);
+			ImageNode->SetValue(ImageConstant.get());
 
 			const uint32 MipsToSkip = ComputeLODBiasForTexture(GenerationContext, *Texture, nullptr, ReferenceTextureSize);
 			Result = ResizeTextureByNumMips(ImageNode, MipsToSkip);
+
+			int32 TotalMips = mu::Image::GetMipmapCount(ImageConstant->GetSizeX(), ImageConstant->GetSizeY());
+			int32 NumMipsBeyondMin = FMath::Max(0, TotalMips - int32(MipsToSkip) - GenerationContext.Options.MinDiskMips);
+			int32 HighResMipsForThisImage = FMath::Min(NumMipsBeyondMin, GenerationContext.Options.NumHighResImageMips);
+			ImageNode->SourceDataDescriptor.SourceHighResMips = HighResMipsForThisImage;
 		}
 		else
 		{
@@ -879,7 +892,7 @@ mu::NodeImagePtr GenerateMutableSourceImage(const UEdGraphPin* Pin, FMutableGrap
 		UTexture* BaseTexture = TypedNodePassThroughTex->PassThroughTexture;
 		if (BaseTexture)
 		{
-			mu::NodeImageConstantPtr ImageNode = new mu::NodeImageConstant();
+			mu::Ptr<mu::NodeImageConstant> ImageNode = new mu::NodeImageConstant();
 			Result = ImageNode;
 
 			ImageNode->SetValue(GenerateImageConstant(BaseTexture, GenerationContext, true).get());
@@ -982,8 +995,8 @@ mu::NodeImagePtr GenerateMutableSourceImage(const UEdGraphPin* Pin, FMutableGrap
 		}
 		else
 		{
-			//This node will add a checker texture in case of error
-			mu::NodeImageConstantPtr EmptyNode = new mu::NodeImageConstant();
+			// This node will add a checker texture in case of error
+			mu::Ptr<mu::NodeImageConstant> EmptyNode = new mu::NodeImageConstant();
 			Result = EmptyNode;
 
 			bool bSuccess = true;
@@ -1027,7 +1040,7 @@ mu::NodeImagePtr GenerateMutableSourceImage(const UEdGraphPin* Pin, FMutableGrap
 
 					if (Table)
 					{
-						mu::NodeImageTablePtr ImageTableNode = new mu::NodeImageTable();
+						mu::Ptr<mu::NodeImageTable> ImageTableNode = new mu::NodeImageTable();
 
 						if (Pin->PinType.PinCategory == Schema->PC_MaterialAsset)
 						{
@@ -1081,11 +1094,11 @@ mu::NodeImagePtr GenerateMutableSourceImage(const UEdGraphPin* Pin, FMutableGrap
 						{
 							Result = ImageTableNode;
 
-							ImageTableNode->SetTable(Table);
-							ImageTableNode->SetColumn(ColumnName);
-							ImageTableNode->SetParameterName(TypedNodeTable->ParameterName);
-							ImageTableNode->SetNoneOption(TypedNodeTable->bAddNoneOption);
-							ImageTableNode->SetDefaultRowName(TypedNodeTable->DefaultRowName.ToString());
+							ImageTableNode->Table = Table;
+							ImageTableNode->ColumnName = ColumnName;
+							ImageTableNode->ParameterName = TypedNodeTable->ParameterName;
+							ImageTableNode->bNoneOption = TypedNodeTable->bAddNoneOption;
+							ImageTableNode->DefaultRowName = TypedNodeTable->DefaultRowName.ToString();
 
 							if (UTexture2D* DefaultTexture2D = Cast<UTexture2D>(DefaultTexture))
 							{
@@ -1096,8 +1109,14 @@ mu::NodeImagePtr GenerateMutableSourceImage(const UEdGraphPin* Pin, FMutableGrap
 								ImageDesc.m_size[1] = ImageDesc.m_size[1] >> LODBias;
 								
 								const uint16 MaxTextureSize = FMath::Max3(ImageDesc.m_size[0], ImageDesc.m_size[1], (uint16)1);
-								ImageTableNode->SetMaxTextureSize(MaxTextureSize);
-								ImageTableNode->SetReferenceImageDescriptor(ImageDesc);
+								ImageTableNode->MaxTextureSize = MaxTextureSize;
+								ImageTableNode->ReferenceImageDesc = ImageDesc;
+
+								// Calculate the number of mips to tag as high res for this image.
+								int32 TotalMips = mu::Image::GetMipmapCount(ImageDesc.m_size[0], ImageDesc.m_size[1]);
+								int32 NumMipsBeyondMin = FMath::Max(0, TotalMips - GenerationContext.Options.MinDiskMips);
+								int32 HighResMipsForThisImage = FMath::Min(NumMipsBeyondMin, GenerationContext.Options.NumHighResImageMips);
+								ImageTableNode->SourceDataDescriptor.SourceHighResMips = HighResMipsForThisImage;
 							}
 						}
 					}
diff --git a/Engine/Plugins/Experimental/Mutable/Source/CustomizableObjectEditor/Private/MuCOE/GenerateMutableSource/GenerateMutableSourceImage.h b/Engine/Plugins/Experimental/Mutable/Source/CustomizableObjectEditor/Private/MuCOE/GenerateMutableSource/GenerateMutableSourceImage.h
--- a/Engine/Plugins/Experimental/Mutable/Source/CustomizableObjectEditor/Private/MuCOE/GenerateMutableSource/GenerateMutableSourceImage.h
+++ b/Engine/Plugins/Experimental/Mutable/Source/CustomizableObjectEditor/Private/MuCOE/GenerateMutableSource/GenerateMutableSourceImage.h
@@ -16,7 +16,6 @@ namespace mu
 {
 	class Image;
 	class NodeImage;
-	class NodeImageConstant;
 }
 
 
diff --git a/Engine/Plugins/Experimental/Mutable/Source/CustomizableObjectEditor/Private/MuCOE/GenerateMutableSource/GenerateMutableSourceMesh.cpp b/Engine/Plugins/Experimental/Mutable/Source/CustomizableObjectEditor/Private/MuCOE/GenerateMutableSource/GenerateMutableSourceMesh.cpp
--- a/Engine/Plugins/Experimental/Mutable/Source/CustomizableObjectEditor/Private/MuCOE/GenerateMutableSource/GenerateMutableSourceMesh.cpp
+++ b/Engine/Plugins/Experimental/Mutable/Source/CustomizableObjectEditor/Private/MuCOE/GenerateMutableSource/GenerateMutableSourceMesh.cpp
@@ -3040,7 +3040,7 @@ mu::NodeMeshPtr GenerateMorphMesh(const UEdGraphPin* Pin,
 			bSuccess = true;
 
 			mu::NodeMeshTablePtr MorphedSourceMeshNodeTable = new mu::NodeMeshTable;
-			MorphedSourceMeshNodeTable->SetTable(Table);
+			MorphedSourceMeshNodeTable->Table = Table;
 			MorphedSourceMeshNodeTable->SetColumn(ColumnName);
 			MorphedSourceMeshNodeTable->SetParameterName(TypedNodeTable->ParameterName);
 			MorphedSourceMeshNodeTable->SetMessageContext(MorphNode);
@@ -4041,7 +4041,7 @@ mu::NodeMeshPtr GenerateMutableSourceMesh(const UEdGraphPin* Pin,
 					{
 						Result = MeshTableNode;
 
-						MeshTableNode->SetTable(Table);
+						MeshTableNode->Table = Table;
 						MeshTableNode->SetColumn(MutableColumnName);
 						MeshTableNode->SetParameterName(TypedNodeTable->ParameterName);
 						MeshTableNode->SetNoneOption(TypedNodeTable->bAddNoneOption);
diff --git a/Engine/Plugins/Experimental/Mutable/Source/CustomizableObjectEditor/Private/MuCOE/GenerateMutableSource/GenerateMutableSourceSurface.cpp b/Engine/Plugins/Experimental/Mutable/Source/CustomizableObjectEditor/Private/MuCOE/GenerateMutableSource/GenerateMutableSourceSurface.cpp
--- a/Engine/Plugins/Experimental/Mutable/Source/CustomizableObjectEditor/Private/MuCOE/GenerateMutableSource/GenerateMutableSourceSurface.cpp
+++ b/Engine/Plugins/Experimental/Mutable/Source/CustomizableObjectEditor/Private/MuCOE/GenerateMutableSource/GenerateMutableSourceSurface.cpp
@@ -605,12 +605,20 @@ mu::Ptr<mu::NodeSurface> GenerateMutableSourceSurface(const UEdGraphPin * Pin, F
 
 									if (Texture2D)
 									{
-										const mu::NodeImageConstantPtr ConstImageNode = new mu::NodeImageConstant();
+										mu::Ptr<mu::NodeImageConstant> ConstImageNode = new mu::NodeImageConstant();
 										mu::Ptr<mu::Image> ImageConstant = GenerateImageConstant(Texture2D, GenerationContext, false);
 										ConstImageNode->SetValue(ImageConstant.get());
 
 										const uint32 MipsToSkip = ComputeLODBiasForTexture(GenerationContext, *Texture2D, nullptr, Props.TextureSize);
-										return ResizeTextureByNumMips(ConstImageNode, MipsToSkip);
+										mu::Ptr<mu::NodeImage> Result =  ResizeTextureByNumMips(ConstImageNode, MipsToSkip);
+
+										// Calculate the number of mips to tag as high res for this image.
+										int32 TotalMips = mu::Image::GetMipmapCount(ImageConstant->GetSizeX(), ImageConstant->GetSizeY());
+										int32 NumMipsBeyondMin = FMath::Max(0, TotalMips - int32(MipsToSkip) - GenerationContext.Options.MinDiskMips);
+										int32 HighResMipsForThisImage = FMath::Min(NumMipsBeyondMin, GenerationContext.Options.NumHighResImageMips);
+										ConstImageNode->SourceDataDescriptor.SourceHighResMips = HighResMipsForThisImage;
+
+										return Result;
 									}
 									else
 									{
@@ -684,7 +692,7 @@ mu::Ptr<mu::NodeSurface> GenerateMutableSourceSurface(const UEdGraphPin * Pin, F
 
 								CompositedImage->SetMode(CompositeImageMode);
 
-								mu::NodeImageConstantPtr CompositeNormalImage = new mu::NodeImageConstant();
+								mu::Ptr<mu::NodeImageConstant> CompositeNormalImage = new mu::NodeImageConstant();
 
 								UTexture2D* ReferenceCompositeNormalTexture = Cast<UTexture2D>(ReferenceTexture->GetCompositeTexture());
 								if (ReferenceCompositeNormalTexture)
@@ -701,6 +709,11 @@ mu::Ptr<mu::NodeSurface> GenerateMutableSourceSurface(const UEdGraphPin * Pin, F
 									NormalCompositeMipmapImage->SetMipmapGenerationSettings(mu::EMipmapFilterType::SimpleAverage, mu::EAddressMode::None);
 
 									CompositedImage->SetNormal(NormalCompositeMipmapImage);
+
+									int32 TotalMips = mu::Image::GetMipmapCount(ImageConstant->GetSizeX(), ImageConstant->GetSizeY());
+									int32 NumMipsBeyondMin = FMath::Max(0, TotalMips - int32(MipsToSkip) - GenerationContext.Options.MinDiskMips);
+									int32 HighResMipsForThisImage = FMath::Min(NumMipsBeyondMin, GenerationContext.Options.NumHighResImageMips);
+									CompositeNormalImage->SourceDataDescriptor.SourceHighResMips = HighResMipsForThisImage;
 								}
 
 								LastImage = CompositedImage;
diff --git a/Engine/Plugins/Experimental/Mutable/Source/CustomizableObjectEditor/Private/MuCOE/SMutableConstantsWidget.cpp b/Engine/Plugins/Experimental/Mutable/Source/CustomizableObjectEditor/Private/MuCOE/SMutableConstantsWidget.cpp
--- a/Engine/Plugins/Experimental/Mutable/Source/CustomizableObjectEditor/Private/MuCOE/SMutableConstantsWidget.cpp
+++ b/Engine/Plugins/Experimental/Mutable/Source/CustomizableObjectEditor/Private/MuCOE/SMutableConstantsWidget.cpp
@@ -1120,7 +1120,7 @@ void SMutableConstantsWidget::LoadConstantImages()
 {
 	check (MutableProgramPtr);
 
-	const int32 ConstantsCount = MutableProgramPtr->m_constantImages.Num();
+	const int32 ConstantsCount = MutableProgramPtr->ConstantImages.Num();
 	ConstantImageElements.Empty(ConstantsCount);
 	
 	uint64 ConstantImagesAccumulatedSize = 0;
diff --git a/Engine/Plugins/Experimental/Mutable/Source/CustomizableObjectEditor/Private/MuCOE/SMutableGraphViewer.cpp b/Engine/Plugins/Experimental/Mutable/Source/CustomizableObjectEditor/Private/MuCOE/SMutableGraphViewer.cpp
--- a/Engine/Plugins/Experimental/Mutable/Source/CustomizableObjectEditor/Private/MuCOE/SMutableGraphViewer.cpp
+++ b/Engine/Plugins/Experimental/Mutable/Source/CustomizableObjectEditor/Private/MuCOE/SMutableGraphViewer.cpp
@@ -94,7 +94,6 @@
 #include "MuT/NodeImageInvertPrivate.h"
 #include "MuT/NodeImageSwizzlePrivate.h"
 #include "MuT/NodeImageMultiLayerPrivate.h"
-#include "MuT/NodeMeshTablePrivate.h"
 #include "MuT/NodeColourFromScalarsPrivate.h"
 #include "MuT/NodeScalarSwitchPrivate.h"
 
@@ -750,10 +749,9 @@ void SMutableGraphViewer::GetChildrenForInfo(TSharedPtr<FMutableGraphTreeElement
 	else if (ParentNode->GetType() == mu::NodeMeshTable::GetStaticType())
 	{
 		mu::NodeMeshTable* MeshTableVar = StaticCast<mu::NodeMeshTable*>(ParentNode);
-		mu::NodeMeshTable::Private* Private = MeshTableVar->GetPrivate();
-		for (int32 LayoutIndex = 0; LayoutIndex < Private->Layouts.Num(); ++LayoutIndex)
+		for (int32 LayoutIndex = 0; LayoutIndex < MeshTableVar->Layouts.Num(); ++LayoutIndex)
 		{
-			AddChildFunc(Private->Layouts[LayoutIndex].get(), FString::Printf(TEXT("LAYOUT [%d]"), LayoutIndex));
+			AddChildFunc(MeshTableVar->Layouts[LayoutIndex].get(), FString::Printf(TEXT("LAYOUT [%d]"), LayoutIndex));
 		}
 	}
 
diff --git a/Engine/Plugins/Experimental/Mutable/Source/MutableRuntime/Internal/MuR/ModelPrivate.h b/Engine/Plugins/Experimental/Mutable/Source/MutableRuntime/Internal/MuR/ModelPrivate.h
--- a/Engine/Plugins/Experimental/Mutable/Source/MutableRuntime/Internal/MuR/ModelPrivate.h
+++ b/Engine/Plugins/Experimental/Mutable/Source/MutableRuntime/Internal/MuR/ModelPrivate.h
@@ -30,16 +30,19 @@ namespace mu
     struct FRomData
     {
         //! This is used to identify a ROM file. It is usually a hash from its data.
-        uint32 Id;
+        uint32 Id=0;
 
 		//! Size of the rom
-		uint32 Size;
+		uint32 Size=0;
 
 		//! Index of the resource in its type-specific array
-		uint32 ResourceIndex;
+		uint32 ResourceIndex=0;
 
 		//! Index of the resource in its type-specific array (one of DATATYPE values)
-		uint32 ResourceType;        
+		uint16 ResourceType=0;
+
+		/** Properties of the rom data. */
+		ERomFlags Flags = ERomFlags::None;
     };
 
     MUTABLE_DEFINE_POD_SERIALISABLE(FRomData);
@@ -226,10 +229,10 @@ namespace mu
 		TArray<TPair<int32, Ptr<const Image>>> ConstantImageLODs;
 
 		//! Constant image mip chain indices: ranges in this array are defined in FImageLODRange and the indices here refer to ConstantImageLODs.
-		TArray<uint32> m_constantImageLODIndices;
+		TArray<uint32> ConstantImageLODIndices;
 
 		//! Constant image data.
-		TArray<FImageLODRange> m_constantImages;
+		TArray<FImageLODRange> ConstantImages;
 
         //! Constant mesh data: the first is the index in m_roms for each mesh or -1 if it is always loaded.
 		TArray<TPair<int32, Ptr<const Mesh>>> ConstantMeshes;
@@ -286,8 +289,8 @@ namespace mu
             arch << m_states;
 			arch << m_roms;
 			arch << ConstantImageLODs;
-			arch << m_constantImageLODIndices;
-			arch << m_constantImages;
+			arch << ConstantImageLODIndices;
+			arch << ConstantImages;
 			arch << ConstantMeshes;
 			arch << m_constantExtensionData;
 			arch << m_constantStrings;
@@ -311,8 +314,8 @@ namespace mu
             arch >> m_states;
 			arch >> m_roms;
 			arch >> ConstantImageLODs;
-			arch >> m_constantImageLODIndices;
-			arch >> m_constantImages;
+			arch >> ConstantImageLODIndices;
+			arch >> ConstantImages;
 			arch >> ConstantMeshes;
 			arch >> m_constantExtensionData;
 			arch >> m_constantStrings;
@@ -563,14 +566,14 @@ namespace mu
 		template <typename CreateImageFunc>
         void GetConstant( int32 ConstantIndex, ImagePtrConst& res, int32 MipsToSkip, const CreateImageFunc& CreateImage) const
         {
-			int32 ReallySkippedLODs = FMath::Min(m_constantImages[ConstantIndex].LODCount - 1, MipsToSkip);
-			int32 FirstLODIndexIndex = m_constantImages[ConstantIndex].FirstIndex;
+			int32 ReallySkippedLODs = FMath::Min(ConstantImages[ConstantIndex].LODCount - 1, MipsToSkip);
+			int32 FirstLODIndexIndex = ConstantImages[ConstantIndex].FirstIndex;
 			int32 ResultLODIndexIndex = FirstLODIndexIndex + ReallySkippedLODs;
-			int32 FinalLODs = m_constantImages[ConstantIndex].LODCount - ReallySkippedLODs;
+			int32 FinalLODs = ConstantImages[ConstantIndex].LODCount - ReallySkippedLODs;
 			check(FinalLODs > 0);
 
 			// Get the first mip
-			int32 ResultLODIndex = m_constantImageLODIndices[ResultLODIndexIndex];
+			int32 ResultLODIndex = ConstantImageLODIndices[ResultLODIndexIndex];
 			Ptr<const Image> CurrentMip = ConstantImageLODs[ResultLODIndex].Value;
 			check(CurrentMip);
 				
@@ -593,7 +596,7 @@ namespace mu
 				{
 					for (int32 LOD = 0; LOD < FinalLODs; ++LOD)
 					{
-						int32 LODIndex = m_constantImageLODIndices[ResultLODIndexIndex + LOD];
+						int32 LODIndex = ConstantImageLODIndices[ResultLODIndexIndex + LOD];
 						int32 MipSizeBytes = ConstantImageLODs[LODIndex].Value->GetLODDataSize(0);
 						Result->DataStorage.ResizeLOD(LOD, MipSizeBytes);
 					}
@@ -613,7 +616,7 @@ namespace mu
 
 					if (LOD + 1 < FinalLODs)
 					{
-						ResultLODIndex = m_constantImageLODIndices[ResultLODIndexIndex + LOD + 1];
+						ResultLODIndex = ConstantImageLODIndices[ResultLODIndexIndex + LOD + 1];
 						CurrentMip = ConstantImageLODs[ResultLODIndex].Value;
 						check(CurrentMip);
 					}
diff --git a/Engine/Plugins/Experimental/Mutable/Source/MutableRuntime/Internal/MuR/Operations.h b/Engine/Plugins/Experimental/Mutable/Source/MutableRuntime/Internal/MuR/Operations.h
--- a/Engine/Plugins/Experimental/Mutable/Source/MutableRuntime/Internal/MuR/Operations.h
+++ b/Engine/Plugins/Experimental/Mutable/Source/MutableRuntime/Internal/MuR/Operations.h
@@ -1198,6 +1198,8 @@ namespace mu
 		DT_COUNT
 	} DATATYPE;
 
+	static_assert(DT_COUNT < 256, "Unexpected enum size.");
+
 
 	// Generic data about a Mutable runtime operation.
     struct FOpDesc
diff --git a/Engine/Plugins/Experimental/Mutable/Source/MutableRuntime/Internal/MuR/SerialisationPrivate.h b/Engine/Plugins/Experimental/Mutable/Source/MutableRuntime/Internal/MuR/SerialisationPrivate.h
--- a/Engine/Plugins/Experimental/Mutable/Source/MutableRuntime/Internal/MuR/SerialisationPrivate.h
+++ b/Engine/Plugins/Experimental/Mutable/Source/MutableRuntime/Internal/MuR/SerialisationPrivate.h
@@ -54,37 +54,6 @@ namespace mu
     };
 
 
-    //! This stream doesn-t store any data, it just counts de amonut of data serialised.
-    class MUTABLERUNTIME_API OutputSizeStream : public OutputStream
-    {
-    public:
-
-        //-----------------------------------------------------------------------------------------
-        // Life cycle
-        //-----------------------------------------------------------------------------------------
-
-        //!
-        OutputSizeStream();
-
-        //-----------------------------------------------------------------------------------------
-        // OutputStream interface
-        //-----------------------------------------------------------------------------------------
-        void Write( const void* pData, uint64 size ) override;
-
-        //-----------------------------------------------------------------------------------------
-        // Own interface
-        //-----------------------------------------------------------------------------------------
-
-        //! Get the amount of data serialised, in bytes.
-        uint64 GetBufferSize() const;
-
-    private:
-
-        uint64 WrittenBytes;
-
-    };
-
-
 #define MUTABLE_IMPLEMENT_POD_SERIALISABLE(Type)				     \
     void DLLEXPORT operator<<(OutputArchive& Arch, const Type& T)    \
     {																 \
diff --git a/Engine/Plugins/Experimental/Mutable/Source/MutableRuntime/Private/MuR/CodeRunner.cpp b/Engine/Plugins/Experimental/Mutable/Source/MutableRuntime/Private/MuR/CodeRunner.cpp
--- a/Engine/Plugins/Experimental/Mutable/Source/MutableRuntime/Private/MuR/CodeRunner.cpp
+++ b/Engine/Plugins/Experimental/Mutable/Source/MutableRuntime/Private/MuR/CodeRunner.cpp
@@ -6165,10 +6165,10 @@ namespace mu
 			int32 ImageIndex = args.value;
 
 			FImageDesc& Result = m_heapImageDesc[item.CustomState];
-			Result.m_format = Program.m_constantImages[ImageIndex].ImageFormat;
-			Result.m_size[0] = Program.m_constantImages[ImageIndex].ImageSizeX;
-			Result.m_size[1] = Program.m_constantImages[ImageIndex].ImageSizeY;
-			Result.m_lods = Program.m_constantImages[ImageIndex].LODCount;
+			Result.m_format = Program.ConstantImages[ImageIndex].ImageFormat;
+			Result.m_size[0] = Program.ConstantImages[ImageIndex].ImageSizeX;
+			Result.m_size[1] = Program.ConstantImages[ImageIndex].ImageSizeY;
+			Result.m_lods = Program.ConstantImages[ImageIndex].LODCount;
 			StoreValidDesc(item);
 			break;
 		}
diff --git a/Engine/Plugins/Experimental/Mutable/Source/MutableRuntime/Private/MuR/CodeRunnerTasked.cpp b/Engine/Plugins/Experimental/Mutable/Source/MutableRuntime/Private/MuR/CodeRunnerTasked.cpp
--- a/Engine/Plugins/Experimental/Mutable/Source/MutableRuntime/Private/MuR/CodeRunnerTasked.cpp
+++ b/Engine/Plugins/Experimental/Mutable/Source/MutableRuntime/Private/MuR/CodeRunnerTasked.cpp
@@ -2256,7 +2256,7 @@ namespace mu
 		for (int32 LODIndex = 0; LODIndex < LODIndexCount; ++LODIndex)
 		{
 			const int32 CurrentIndexIndex = LODIndexIndex + LODIndex;
-			const int32 CurrentIndex = program.m_constantImageLODIndices[CurrentIndexIndex];
+			const int32 CurrentIndex = program.ConstantImageLODIndices[CurrentIndexIndex];
 
 			if (program.ConstantImageLODs[CurrentIndex].Key < 0)
 			{
@@ -2377,7 +2377,7 @@ namespace mu
 		for (int32 LODIndex = 0; LODIndex < LODIndexCount; ++LODIndex)
 		{
 			int32 CurrentIndexIndex = LODIndexIndex + LODIndex;
-			int32 CurrentIndex = Program.m_constantImageLODIndices[CurrentIndexIndex];
+			int32 CurrentIndex = Program.ConstantImageLODIndices[CurrentIndexIndex];
 
 			if (Program.ConstantImageLODs[CurrentIndex].Key < 0)
 			{
@@ -2581,9 +2581,9 @@ namespace mu
 			OP::ResourceConstantArgs args = program.GetOpArgs<OP::ResourceConstantArgs>(item.At);
 			int32 MipsToSkip = item.ExecutionOptions;
 			int32 ImageIndex = args.value;
-			int32 ReallySkip = FMath::Min(MipsToSkip, program.m_constantImages[ImageIndex].LODCount - 1);
-			int32 LODIndexIndex = program.m_constantImages[ImageIndex].FirstIndex + ReallySkip;
-			int32 LODIndexCount = program.m_constantImages[ImageIndex].LODCount - ReallySkip;
+			int32 ReallySkip = FMath::Min(MipsToSkip, program.ConstantImages[ImageIndex].LODCount - 1);
+			int32 LODIndexIndex = program.ConstantImages[ImageIndex].FirstIndex + ReallySkip;
+			int32 LODIndexCount = program.ConstantImages[ImageIndex].LODCount - ReallySkip;
 			check(LODIndexCount > 0);
 
 			// We always need to follow this path, or roms may not be protected for long enough and might be unloaded 
@@ -2592,7 +2592,7 @@ namespace mu
 			//bool bAnyMissing = false;
 			//for (int32 i=0; i<LODIndexCount; ++i)
 			//{
-			//	uint32 LODIndex = program.m_constantImageLODIndices[LODIndexIndex+i];
+			//	uint32 LODIndex = program.ConstantImageLODIndices[LODIndexIndex+i];
 			//	if ( !program.ConstantImageLODs[LODIndex].Value )
 			//	{
 			//		bAnyMissing = true;
diff --git a/Engine/Plugins/Experimental/Mutable/Source/MutableRuntime/Private/MuR/Model.cpp b/Engine/Plugins/Experimental/Mutable/Source/MutableRuntime/Private/MuR/Model.cpp
--- a/Engine/Plugins/Experimental/Mutable/Source/MutableRuntime/Private/MuR/Model.cpp
+++ b/Engine/Plugins/Experimental/Mutable/Source/MutableRuntime/Private/MuR/Model.cpp
@@ -430,19 +430,22 @@ namespace mu
     }
 
 
-    //---------------------------------------------------------------------------------------------
     uint32 Model::GetRomId(int32 Index) const
     {
     	return m_pD->m_program.m_roms[Index].Id;
     }
 
 
-    //---------------------------------------------------------------------------------------------
     uint32 Model::GetRomSize(int32 Index) const
     {
     	return m_pD->m_program.m_roms[Index].Size;
     }
 
+	ERomFlags Model::GetRomFlags(int32 Index) const
+	{
+		return m_pD->m_program.m_roms[Index].Flags;
+	}
+
 
     //---------------------------------------------------------------------------------------------
     ParametersPtr Model::NewParameters(TSharedPtr<const Model> Model, const Parameters* pOld )
diff --git a/Engine/Plugins/Experimental/Mutable/Source/MutableRuntime/Public/MuR/Model.h b/Engine/Plugins/Experimental/Mutable/Source/MutableRuntime/Public/MuR/Model.h
--- a/Engine/Plugins/Experimental/Mutable/Source/MutableRuntime/Public/MuR/Model.h
+++ b/Engine/Plugins/Experimental/Mutable/Source/MutableRuntime/Public/MuR/Model.h
@@ -37,6 +37,17 @@ namespace mu
     typedef Ptr<const System> SystemPtrConst;
 
 
+	/** Flags for the model data chunks. */
+	enum class ERomFlags : uint16
+	{
+		/** Standard data. */
+		None = 0,
+
+		/** Bigger mips and mesh lods that are optional in some devices of a platform. */
+		HighRes = 1 << 0,
+	};
+
+
     //! \brief A Model represents a customisable object with any number of parameters.
     //!
     //! When values are given to the parameters, specific Instances can be built, which hold the
@@ -154,6 +165,8 @@ namespace mu
 		uint32 GetRomId(int32 Index) const;
 
 		uint32 GetRomSize(int32 Index) const;
+
+		ERomFlags GetRomFlags(int32 Index) const;
     	
 		//-----------------------------------------------------------------------------------------
 		// Interface pattern
diff --git a/Engine/Plugins/Experimental/Mutable/Source/MutableRuntime/Public/MuR/Serialisation.h b/Engine/Plugins/Experimental/Mutable/Source/MutableRuntime/Public/MuR/Serialisation.h
--- a/Engine/Plugins/Experimental/Mutable/Source/MutableRuntime/Public/MuR/Serialisation.h
+++ b/Engine/Plugins/Experimental/Mutable/Source/MutableRuntime/Public/MuR/Serialisation.h
@@ -303,6 +303,37 @@ namespace mu
     };
 
 
+	//! This stream doesn't store any data, it just counts the amount of data serialised.
+	class MUTABLERUNTIME_API OutputSizeStream : public OutputStream
+	{
+	public:
+
+		//-----------------------------------------------------------------------------------------
+		// Life cycle
+		//-----------------------------------------------------------------------------------------
+
+		//!
+		OutputSizeStream();
+
+		//-----------------------------------------------------------------------------------------
+		// OutputStream interface
+		//-----------------------------------------------------------------------------------------
+		void Write(const void* pData, uint64 size) override;
+
+		//-----------------------------------------------------------------------------------------
+		// Own interface
+		//-----------------------------------------------------------------------------------------
+
+		//! Get the amount of data serialised, in bytes.
+		uint64 GetBufferSize() const;
+
+	private:
+
+		uint64 WrittenBytes;
+
+	};
+
+
 	template<typename Type>
 	void operator<<(OutputArchive& Arch, const Type& Value)
 	{
diff --git a/Engine/Plugins/Experimental/Mutable/Source/MutableTools/Internal/MuT/AST.h b/Engine/Plugins/Experimental/Mutable/Source/MutableTools/Internal/MuT/AST.h
--- a/Engine/Plugins/Experimental/Mutable/Source/MutableTools/Internal/MuT/AST.h
+++ b/Engine/Plugins/Experimental/Mutable/Source/MutableTools/Internal/MuT/AST.h
@@ -4,6 +4,7 @@
 
 #include "MuT/Platform.h"
 #include "MuT/Compiler.h"
+#include "MuT/NodeImageConstant.h"
 #include "MuR/Image.h"
 #include "MuR/Mesh.h"
 #include "MuR/ModelPrivate.h"
@@ -286,6 +287,17 @@ namespace mu
 
 		/** Image operation functions, so that they can be overriden. */
 		FImageOperator& ImageOperator;
+
+		/** Store for additional data generated during compilation, but not necessary for the runtime. */
+		struct FAdditionalData
+		{
+			/** Source data descriptor for every image constant that has been generated.
+			* It must have the same size than the Program::ConstantImages array.
+			*/
+			TArray<FSourceDataDescriptor> SourceImagePerConstant;
+		};
+
+		FAdditionalData AdditionalData;
 	};
 
 
@@ -793,6 +805,16 @@ namespace mu
         //!
         virtual FImageDesc GetImageDesc( bool returnBestOption=false, FGetImageDescContext* context=nullptr ) const;
 
+		/** */
+		class FGetSourceDataDescriptorContext
+		{
+		public:
+			TMap<const ASTOp*, FSourceDataDescriptor> Cache;
+		};
+
+		/** */
+		virtual FSourceDataDescriptor GetSourceDataDescriptor(FGetSourceDataDescriptorContext* = nullptr) const;
+
         //! Optional cache struct to use int he method below.
         using FBlockLayoutSizeCache=TMap< const TPair<ASTOp*,uint64>, TPair<int32,int32>>;
 
@@ -1083,21 +1105,22 @@ namespace mu
 
         OP_TYPE GetOpType() const override { return op.type; }
 
-        Ptr<ASTOp> Clone( MapChildFuncRef mapChild ) const override;
-        void ForEachChild( const TFunctionRef<void(ASTChild&)> f ) override;
-        void Link( FProgram& program, FLinkerOptions* Options) override;
-        bool IsEqual(const ASTOp& otherUntyped) const override;
-		uint64 Hash() const override;
-		FImageDesc GetImageDesc( bool returnBestOption, FGetImageDescContext* context ) const override;
-        void GetLayoutBlockSize( int* pBlockX, int* pBlockY ) override;
-        Ptr<ASTOp> OptimiseSize() const override;
-        Ptr<ASTOp> OptimiseSemantic(const FModelOptimizationOptions&, int32 Pass) const override;
-        Ptr<ASTOp> OptimiseSink(const FModelOptimizationOptions&, FOptimizeSinkContext&) const override;
-        FBoolEvalResult EvaluateBool( ASTOpList& facts, FEvaluateBoolCache* cache ) const override;
-        int EvaluateInt( ASTOpList& facts, bool &unknown ) const override;
-        bool IsImagePlainConstant(FVector4f& colour ) const override;
-        bool IsColourConstant(FVector4f& colour ) const override;
-        Ptr<ImageSizeExpression> GetImageSizeExpression() const override;
+		virtual Ptr<ASTOp> Clone( MapChildFuncRef mapChild ) const override;
+		virtual void ForEachChild( const TFunctionRef<void(ASTChild&)> f ) override;
+		virtual void Link( FProgram& program, FLinkerOptions* Options) override;
+		virtual bool IsEqual(const ASTOp& otherUntyped) const override;
+		virtual uint64 Hash() const override;
+		virtual FImageDesc GetImageDesc( bool returnBestOption, FGetImageDescContext* context ) const override;
+		virtual void GetLayoutBlockSize( int* pBlockX, int* pBlockY ) override;
+		virtual Ptr<ASTOp> OptimiseSize() const override;
+		virtual Ptr<ASTOp> OptimiseSemantic(const FModelOptimizationOptions&, int32 Pass) const override;
+		virtual Ptr<ASTOp> OptimiseSink(const FModelOptimizationOptions&, FOptimizeSinkContext&) const override;
+		virtual FBoolEvalResult EvaluateBool( ASTOpList& facts, FEvaluateBoolCache* cache ) const override;
+		virtual int32 EvaluateInt( ASTOpList& facts, bool &unknown ) const override;
+		virtual bool IsImagePlainConstant(FVector4f& colour ) const override;
+		virtual bool IsColourConstant(FVector4f& colour ) const override;
+		virtual Ptr<ImageSizeExpression> GetImageSizeExpression() const override;
+		virtual FSourceDataDescriptor GetSourceDataDescriptor(FGetSourceDataDescriptorContext* Context) const override;
 
         //---------------------------------------------------------------------------------------------
         //! Add child for OP fixed nodes.
@@ -1148,6 +1171,19 @@ namespace mu
             return res;
         }
 
+		//
+		inline FSourceDataDescriptor GetSourceDataDescriptor(OP::ADDRESS at, FGetSourceDataDescriptorContext* Context = nullptr) const
+		{
+			FSourceDataDescriptor Result;
+
+			if (children[at])
+			{
+				Result = children[at]->GetSourceDataDescriptor(Context);
+			}
+
+			return Result;
+		}
+
         inline void GetLayoutBlockSize( OP::ADDRESS at, int* pBlockX, int* pBlockY )
         {
             if (children[at])
diff --git a/Engine/Plugins/Experimental/Mutable/Source/MutableTools/Internal/MuT/NodeMeshTablePrivate.h b/Engine/Plugins/Experimental/Mutable/Source/MutableTools/Internal/MuT/NodeMeshTablePrivate.h
deleted file mode 100644
diff --git a/Engine/Plugins/Experimental/Mutable/Source/MutableTools/Private/MuT/AST.cpp b/Engine/Plugins/Experimental/Mutable/Source/MutableTools/Private/MuT/AST.cpp
--- a/Engine/Plugins/Experimental/Mutable/Source/MutableTools/Private/MuT/AST.cpp
+++ b/Engine/Plugins/Experimental/Mutable/Source/MutableTools/Private/MuT/AST.cpp
@@ -865,6 +865,13 @@ FImageDesc ASTOp::GetImageDesc( bool, FGetImageDescContext* ) const
 }
 
 
+FSourceDataDescriptor ASTOp::GetSourceDataDescriptor(FGetSourceDataDescriptorContext*) const
+{
+	ensure(false);
+	return {};
+}
+
+
 bool ASTOp::IsImagePlainConstant(FVector4f&) const
 {
 	// Some image operations don't have this implemented and hit here.
@@ -1081,7 +1088,7 @@ FImageDesc ASTOpFixed::GetImageDesc( bool returnBestOption, FGetImageDescContext
         }
     }
 
-    OP_TYPE type = (OP_TYPE) op.type;
+    OP_TYPE type = op.type;
 
     switch ( type )
     {
@@ -1618,4 +1625,97 @@ mu::Ptr<ImageSizeExpression> ASTOpFixed::GetImageSizeExpression() const
     return pRes;
 }
 
+
+
+
+//!
+FSourceDataDescriptor ASTOpFixed::GetSourceDataDescriptor(FGetSourceDataDescriptorContext* Context) const
+{
+	FSourceDataDescriptor  Result;
+
+	// Cache management
+	TUniquePtr<FGetSourceDataDescriptorContext> LocalContext;
+	if (!Context)
+	{
+		LocalContext.Reset(new FGetSourceDataDescriptorContext);
+		Context = LocalContext.Get();
+	}
+
+	FSourceDataDescriptor* Found = Context->Cache.Find(this);
+	if (Found)
+	{
+		return *Found;
+	}
+
+	// Calculate
+	OP_TYPE Type = op.type;
+
+	switch (Type)
+	{
+
+	case OP_TYPE::NONE:
+	case OP_TYPE::IM_PLAINCOLOUR:
+	case OP_TYPE::IM_GRADIENT:
+	case OP_TYPE::IM_BLANKLAYOUT:
+		break;
+
+	case OP_TYPE::IM_SATURATE:
+		Result = GetSourceDataDescriptor(op.args.ImageSaturate.base, Context);
+		break;
+
+	case OP_TYPE::IM_LUMINANCE:
+		Result = GetSourceDataDescriptor(op.args.ImageLuminance.base, Context);
+		break;
+
+	case OP_TYPE::IM_INTERPOLATE:
+	{
+		for (int32 SourceIndex = 0; SourceIndex < MUTABLE_OP_MAX_INTERPOLATE_COUNT; ++SourceIndex)
+		{
+			FSourceDataDescriptor SourceDesc = GetSourceDataDescriptor(op.args.ImageInterpolate.targets[SourceIndex], Context);
+			Result.CombineWith(SourceDesc);
+		}
+		break;
+	}
+
+	case OP_TYPE::IM_RESIZE:
+		Result = GetSourceDataDescriptor(op.args.ImageResize.source, Context);
+		break;
+
+	case OP_TYPE::IM_RESIZEREL:
+		Result = GetSourceDataDescriptor(op.args.ImageResizeRel.source, Context);
+		break;
+
+	case OP_TYPE::IM_RESIZELIKE:
+		Result = GetSourceDataDescriptor(op.args.ImageResizeLike.source, Context);
+		break;
+
+	case OP_TYPE::IM_BINARISE:
+		Result = GetSourceDataDescriptor(op.args.ImageBinarise.base, Context);
+		break;
+
+	case OP_TYPE::IM_DISPLACE:
+		Result = GetSourceDataDescriptor(op.args.ImageDisplace.source, Context);
+		break;
+
+	case OP_TYPE::IM_INVERT:
+		Result = GetSourceDataDescriptor(op.args.ImageInvert.base, Context);
+		break;
+
+	case OP_TYPE::IM_COLOURMAP:
+		Result = GetSourceDataDescriptor(op.args.ImageColourMap.base, Context);
+		break;
+
+	default:
+		// This should probably be implemented
+		ensure(false);
+		break;
+	}
+
+	// Cache the result
+	Context->Cache.Add(this, Result);
+
+	return Result;
+}
+
+
 }
diff --git a/Engine/Plugins/Experimental/Mutable/Source/MutableTools/Private/MuT/ASTOpConditional.cpp b/Engine/Plugins/Experimental/Mutable/Source/MutableTools/Private/MuT/ASTOpConditional.cpp
--- a/Engine/Plugins/Experimental/Mutable/Source/MutableTools/Private/MuT/ASTOpConditional.cpp
+++ b/Engine/Plugins/Experimental/Mutable/Source/MutableTools/Private/MuT/ASTOpConditional.cpp
@@ -358,4 +358,41 @@ namespace mu
 		return pRes;
 	}
 
+
+	FSourceDataDescriptor ASTOpConditional::GetSourceDataDescriptor(FGetSourceDataDescriptorContext* Context) const
+	{
+		// Cache management
+		TUniquePtr<FGetSourceDataDescriptorContext> LocalContext;
+		if (!Context)
+		{
+			LocalContext.Reset(new FGetSourceDataDescriptorContext);
+			Context = LocalContext.Get();
+		}
+
+		FSourceDataDescriptor* Found = Context->Cache.Find(this);
+		if (Found)
+		{
+			return *Found;
+		}
+
+		// Not cached: calculate
+		FSourceDataDescriptor Result;
+
+		if (yes)
+		{
+			FSourceDataDescriptor SourceDesc = yes->GetSourceDataDescriptor(Context);
+			Result.CombineWith(SourceDesc);
+		}
+
+		if (no)
+		{
+			FSourceDataDescriptor SourceDesc = no->GetSourceDataDescriptor(Context);
+			Result.CombineWith(SourceDesc);
+		}
+
+		Context->Cache.Add(this, Result);
+
+		return Result;
+	}
+
 }
\ No newline at end of file
diff --git a/Engine/Plugins/Experimental/Mutable/Source/MutableTools/Private/MuT/ASTOpConditional.h b/Engine/Plugins/Experimental/Mutable/Source/MutableTools/Private/MuT/ASTOpConditional.h
--- a/Engine/Plugins/Experimental/Mutable/Source/MutableTools/Private/MuT/ASTOpConditional.h
+++ b/Engine/Plugins/Experimental/Mutable/Source/MutableTools/Private/MuT/ASTOpConditional.h
@@ -36,21 +36,21 @@ struct FProgram;
 		ASTOpConditional(const ASTOpConditional&) = delete;
 		~ASTOpConditional() override;
 
-		OP_TYPE GetOpType() const override { return type; }
-
-		bool IsEqual(const ASTOp& otherUntyped) const override;
-		Ptr<ASTOp> Clone(MapChildFuncRef mapChild) const override;
-		uint64 Hash() const override;
-		void Assert() override;
-		void ForEachChild(const TFunctionRef<void(ASTChild&)> f) override;
-		void Link(FProgram& program, FLinkerOptions* Options) override;
-		FImageDesc GetImageDesc(bool returnBestOption, class FGetImageDescContext* context) const override;
-		void GetLayoutBlockSize(int32* pBlockX, int32* pBlockY) override;
-		void GetBlockLayoutSize(uint64 BlockId, int32* pBlockX, int32* pBlockY, FBlockLayoutSizeCache* cache) override;
-		Ptr<ASTOp> OptimiseSemantic(const FModelOptimizationOptions&, int32 Pass) const override;
-		bool GetNonBlackRect(FImageRect& maskUsage) const override;
-		Ptr<ImageSizeExpression> GetImageSizeExpression() const override;
+		virtual OP_TYPE GetOpType() const override { return type; }
+		virtual bool IsEqual(const ASTOp& otherUntyped) const override;
+		virtual Ptr<ASTOp> Clone(MapChildFuncRef mapChild) const override;
+		virtual uint64 Hash() const override;
+		virtual void Assert() override;
+		virtual void ForEachChild(const TFunctionRef<void(ASTChild&)> f) override;
+		virtual void Link(FProgram& program, FLinkerOptions* Options) override;
+		virtual FImageDesc GetImageDesc(bool returnBestOption, class FGetImageDescContext* context) const override;
+		virtual void GetLayoutBlockSize(int32* pBlockX, int32* pBlockY) override;
+		virtual void GetBlockLayoutSize(uint64 BlockId, int32* pBlockX, int32* pBlockY, FBlockLayoutSizeCache* cache) override;
+		virtual Ptr<ASTOp> OptimiseSemantic(const FModelOptimizationOptions&, int32 Pass) const override;
+		virtual bool GetNonBlackRect(FImageRect& maskUsage) const override;
+		virtual Ptr<ImageSizeExpression> GetImageSizeExpression() const override;
 		virtual bool IsConditional() const override { return true; }
+		virtual FSourceDataDescriptor GetSourceDataDescriptor(FGetSourceDataDescriptorContext*) const override;
 	};
 
 
diff --git a/Engine/Plugins/Experimental/Mutable/Source/MutableTools/Private/MuT/ASTOpConstantResource.cpp b/Engine/Plugins/Experimental/Mutable/Source/MutableTools/Private/MuT/ASTOpConstantResource.cpp
--- a/Engine/Plugins/Experimental/Mutable/Source/MutableTools/Private/MuT/ASTOpConstantResource.cpp
+++ b/Engine/Plugins/Experimental/Mutable/Source/MutableTools/Private/MuT/ASTOpConstantResource.cpp
@@ -248,7 +248,8 @@ namespace mu
 		{
 			const ASTOpConstantResource* Other = static_cast<const ASTOpConstantResource*>(&OtherUntyped);
 			return Type == Other->Type && ValueHash == Other->ValueHash &&
-				LoadedValue == Other->LoadedValue && Proxy == Other->Proxy;
+				LoadedValue == Other->LoadedValue && Proxy == Other->Proxy
+				&& SourceDataDescriptor == Other->SourceDataDescriptor;
 		}
 		return false;
 	}
@@ -262,6 +263,7 @@ namespace mu
 		n->Proxy = Proxy;
 		n->LoadedValue = LoadedValue;
 		n->ValueHash = ValueHash;
+		n->SourceDataDescriptor = SourceDataDescriptor; 
 		return n;
 	}
 
@@ -287,7 +289,7 @@ namespace mu
 			// Mips to store
 			int32 MipsToStore = 1;
 
-			int32 FirstLODIndexIndex = Program.m_constantImageLODIndices.Num();
+			int32 FirstLODIndexIndex = Program.ConstantImageLODIndices.Num();
 
 			FImageOperator& ImOp = Options.ImageOperator;
 			Ptr<const Image> pMip;
@@ -350,7 +352,7 @@ namespace mu
 					Options.ImageConstantMipMap.Add(pMip, MipIndex);
 				}
 
-				Program.m_constantImageLODIndices.Add(uint32(MipIndex));
+				Program.ConstantImageLODIndices.Add(uint32(MipIndex));
 
 				// Generate next mip if necessary
 				if (Mip + 1 < MipsToStore)
@@ -393,7 +395,7 @@ namespace mu
 			LODRange.ImageFormat = pImage->GetFormat();
 			LODRange.ImageSizeX = pImage->GetSizeX();
 			LODRange.ImageSizeY = pImage->GetSizeY();
-			int32 ImageIndex = Program.m_constantImages.Add(LODRange);
+			int32 ImageIndex = Program.ConstantImages.Add(LODRange);
 			return ImageIndex;
 		}
 	}
@@ -467,6 +469,9 @@ namespace mu
 					else
 					{
 						args.value = AddConstantImage( program, pTyped, *Options);
+
+						int32 DataDescIndex = Options->AdditionalData.SourceImagePerConstant.Add(SourceDataDescriptor);
+						check(DataDescIndex == args.value);
 					}
 
 					break;
@@ -744,7 +749,6 @@ namespace mu
 	}
 
 
-	//-------------------------------------------------------------------------------------------------
 	mu::Ptr<ImageSizeExpression> ASTOpConstantResource::GetImageSizeExpression() const
 	{
 		if (Type==OP_TYPE::IM_CONSTANT)
@@ -759,4 +763,10 @@ namespace mu
 		return nullptr;
 	}
 
+
+	FSourceDataDescriptor ASTOpConstantResource::GetSourceDataDescriptor(FGetSourceDataDescriptorContext*) const
+	{
+		return SourceDataDescriptor;
+	}
+
 }
diff --git a/Engine/Plugins/Experimental/Mutable/Source/MutableTools/Private/MuT/ASTOpConstantResource.h b/Engine/Plugins/Experimental/Mutable/Source/MutableTools/Private/MuT/ASTOpConstantResource.h
--- a/Engine/Plugins/Experimental/Mutable/Source/MutableTools/Private/MuT/ASTOpConstantResource.h
+++ b/Engine/Plugins/Experimental/Mutable/Source/MutableTools/Private/MuT/ASTOpConstantResource.h
@@ -32,9 +32,12 @@ namespace mu
 
 	public:
 
-		//! Type of switch
+		//! Type of constant
 		OP_TYPE Type;
 
+		/** Source data descriptor. */
+		FSourceDataDescriptor SourceDataDescriptor;
+
 	public:
 
 		~ASTOpConstantResource() override;
@@ -54,18 +57,19 @@ namespace mu
 
 
 		// ASTOp interface
-		OP_TYPE GetOpType() const override { return Type; }
-		void ForEachChild(const TFunctionRef<void(ASTChild&)>) override;
-		bool IsEqual(const ASTOp& otherUntyped) const override;
-		Ptr<ASTOp> Clone(MapChildFuncRef mapChild) const override;
-		uint64 Hash() const override;
-		void Link(FProgram& program, FLinkerOptions*) override;
-		FImageDesc GetImageDesc(bool, class FGetImageDescContext*) const override;
-		void GetBlockLayoutSize(uint64 BlockId, int32* pBlockX, int32* pBlockY, FBlockLayoutSizeCache* cache) override;
-		void GetLayoutBlockSize(int32* pBlockX, int32* pBlockY) override;
-		bool GetNonBlackRect(FImageRect& maskUsage) const override;
-		bool IsImagePlainConstant(FVector4f& colour) const override;
-		Ptr<ImageSizeExpression> GetImageSizeExpression() const override;
+		virtual OP_TYPE GetOpType() const override { return Type; }
+		virtual void ForEachChild(const TFunctionRef<void(ASTChild&)>) override;
+		virtual bool IsEqual(const ASTOp& otherUntyped) const override;
+		virtual Ptr<ASTOp> Clone(MapChildFuncRef mapChild) const override;
+		virtual uint64 Hash() const override;
+		virtual void Link(FProgram& program, FLinkerOptions*) override;
+		virtual FImageDesc GetImageDesc(bool, class FGetImageDescContext*) const override;
+		virtual void GetBlockLayoutSize(uint64 BlockId, int32* pBlockX, int32* pBlockY, FBlockLayoutSizeCache* cache) override;
+		virtual void GetLayoutBlockSize(int32* pBlockX, int32* pBlockY) override;
+		virtual bool GetNonBlackRect(FImageRect& maskUsage) const override;
+		virtual bool IsImagePlainConstant(FVector4f& colour) const override;
+		virtual Ptr<ImageSizeExpression> GetImageSizeExpression() const override;
+		virtual FSourceDataDescriptor GetSourceDataDescriptor(FGetSourceDataDescriptorContext*) const override;
 	};
 
 
diff --git a/Engine/Plugins/Experimental/Mutable/Source/MutableTools/Private/MuT/ASTOpImageCompose.cpp b/Engine/Plugins/Experimental/Mutable/Source/MutableTools/Private/MuT/ASTOpImageCompose.cpp
--- a/Engine/Plugins/Experimental/Mutable/Source/MutableTools/Private/MuT/ASTOpImageCompose.cpp
+++ b/Engine/Plugins/Experimental/Mutable/Source/MutableTools/Private/MuT/ASTOpImageCompose.cpp
@@ -292,4 +292,48 @@ mu::Ptr<ASTOp> ASTOpImageCompose::OptimiseSemantic(const FModelOptimizationOptio
 	return at;
 }
 
+
+FSourceDataDescriptor ASTOpImageCompose::GetSourceDataDescriptor(FGetSourceDataDescriptorContext* Context) const
+{
+	// Cache management
+	TUniquePtr<FGetSourceDataDescriptorContext> LocalContext;
+	if (!Context)
+	{
+		LocalContext.Reset(new FGetSourceDataDescriptorContext);
+		Context = LocalContext.Get();
+	}
+
+	FSourceDataDescriptor* Found = Context->Cache.Find(this);
+	if (Found)
+	{
+		return *Found;
+	}
+
+	// Not cached: calculate
+	FSourceDataDescriptor Result;
+
+	if (Base)
+	{
+		FSourceDataDescriptor SourceDesc = Base->GetSourceDataDescriptor(Context);
+		Result.CombineWith(SourceDesc);
+	}
+
+	if (BlockImage)
+	{
+		FSourceDataDescriptor SourceDesc = BlockImage->GetSourceDataDescriptor(Context);
+		Result.CombineWith(SourceDesc);
+	}
+
+	if (Mask)
+	{
+		FSourceDataDescriptor SourceDesc = Mask->GetSourceDataDescriptor(Context);
+		Result.CombineWith(SourceDesc);
+	}
+
+	Context->Cache.Add(this, Result);
+
+	return Result;
+}
+
+
 }
\ No newline at end of file
diff --git a/Engine/Plugins/Experimental/Mutable/Source/MutableTools/Private/MuT/ASTOpImageCompose.h b/Engine/Plugins/Experimental/Mutable/Source/MutableTools/Private/MuT/ASTOpImageCompose.h
--- a/Engine/Plugins/Experimental/Mutable/Source/MutableTools/Private/MuT/ASTOpImageCompose.h
+++ b/Engine/Plugins/Experimental/Mutable/Source/MutableTools/Private/MuT/ASTOpImageCompose.h
@@ -29,18 +29,19 @@ namespace mu
 		ASTOpImageCompose(const ASTOpImageCompose&) = delete;
 		~ASTOpImageCompose();
 
-		OP_TYPE GetOpType() const override { return OP_TYPE::IM_COMPOSE; }
-		uint64 Hash() const override;
-		bool IsEqual(const ASTOp& otherUntyped) const override;
-		Ptr<ASTOp> Clone(MapChildFuncRef mapChild) const override;
-		void ForEachChild(const TFunctionRef<void(ASTChild&)>) override;
-		void Link(FProgram& program, FLinkerOptions* Options) override;
-		//TODO: Ptr<ASTOp> OptimiseSink(const FModelOptimizationOptions& options, FOptimizeSinkContext& context) const override;
-		Ptr<ASTOp> OptimiseSemantic(const FModelOptimizationOptions& options, int32 Pass) const;
-		FImageDesc GetImageDesc(bool returnBestOption, FGetImageDescContext* context) const override;
-		void GetLayoutBlockSize(int32* OutBlockX, int32* OutBlockY) override;
-		bool IsImagePlainConstant(FVector4f& OutColor) const override;
-		Ptr<ImageSizeExpression> GetImageSizeExpression() const override;
+		virtual OP_TYPE GetOpType() const override { return OP_TYPE::IM_COMPOSE; }
+		virtual uint64 Hash() const override;
+		virtual bool IsEqual(const ASTOp& otherUntyped) const override;
+		virtual Ptr<ASTOp> Clone(MapChildFuncRef mapChild) const override;
+		virtual void ForEachChild(const TFunctionRef<void(ASTChild&)>) override;
+		virtual void Link(FProgram& program, FLinkerOptions* Options) override;
+		//TODO: virtual Ptr<ASTOp> OptimiseSink(const FModelOptimizationOptions& options, FOptimizeSinkContext& context) const override;
+		virtual Ptr<ASTOp> OptimiseSemantic(const FModelOptimizationOptions& options, int32 Pass) const;
+		virtual FImageDesc GetImageDesc(bool returnBestOption, FGetImageDescContext* context) const override;
+		virtual void GetLayoutBlockSize(int32* OutBlockX, int32* OutBlockY) override;
+		virtual bool IsImagePlainConstant(FVector4f& OutColor) const override;
+		virtual Ptr<ImageSizeExpression> GetImageSizeExpression() const override;
+		virtual FSourceDataDescriptor GetSourceDataDescriptor(FGetSourceDataDescriptorContext*) const override;
 
 	};
 
diff --git a/Engine/Plugins/Experimental/Mutable/Source/MutableTools/Private/MuT/ASTOpImageCrop.cpp b/Engine/Plugins/Experimental/Mutable/Source/MutableTools/Private/MuT/ASTOpImageCrop.cpp
--- a/Engine/Plugins/Experimental/Mutable/Source/MutableTools/Private/MuT/ASTOpImageCrop.cpp
+++ b/Engine/Plugins/Experimental/Mutable/Source/MutableTools/Private/MuT/ASTOpImageCrop.cpp
@@ -159,6 +159,17 @@ namespace mu
 	}
 
 
+	FSourceDataDescriptor ASTOpImageCrop::GetSourceDataDescriptor(FGetSourceDataDescriptorContext* Context) const
+	{
+		if (Source)
+		{
+			return Source->GetSourceDataDescriptor(Context);
+		}
+
+		return {};
+	}
+
+
 	Ptr<ASTOp> ASTOpImageCrop::OptimiseSemantic(const FModelOptimizationOptions& options, int32 Pass) const
 	{
 		Ptr<ASTOp> at;
diff --git a/Engine/Plugins/Experimental/Mutable/Source/MutableTools/Private/MuT/ASTOpImageCrop.h b/Engine/Plugins/Experimental/Mutable/Source/MutableTools/Private/MuT/ASTOpImageCrop.h
--- a/Engine/Plugins/Experimental/Mutable/Source/MutableTools/Private/MuT/ASTOpImageCrop.h
+++ b/Engine/Plugins/Experimental/Mutable/Source/MutableTools/Private/MuT/ASTOpImageCrop.h
@@ -40,6 +40,7 @@ namespace mu
 		//TODO: virtual bool IsImagePlainConstant(FVector4f& colour) const override;
 		virtual Ptr<ASTOp> OptimiseSemantic(const FModelOptimizationOptions& options, int32 Pass) const override;
 		virtual Ptr<ASTOp> OptimiseSink(const FModelOptimizationOptions& options, FOptimizeSinkContext& context) const override;
+		virtual FSourceDataDescriptor GetSourceDataDescriptor(FGetSourceDataDescriptorContext*) const override;
 
 	};
 
diff --git a/Engine/Plugins/Experimental/Mutable/Source/MutableTools/Private/MuT/ASTOpImageLayer.cpp b/Engine/Plugins/Experimental/Mutable/Source/MutableTools/Private/MuT/ASTOpImageLayer.cpp
--- a/Engine/Plugins/Experimental/Mutable/Source/MutableTools/Private/MuT/ASTOpImageLayer.cpp
+++ b/Engine/Plugins/Experimental/Mutable/Source/MutableTools/Private/MuT/ASTOpImageLayer.cpp
@@ -783,4 +783,46 @@ namespace mu
 	}
 
 
+	FSourceDataDescriptor ASTOpImageLayer::GetSourceDataDescriptor(FGetSourceDataDescriptorContext* Context) const
+	{
+		// Cache management
+		TUniquePtr<FGetSourceDataDescriptorContext> LocalContext;
+		if (!Context)
+		{
+			LocalContext.Reset(new FGetSourceDataDescriptorContext);
+			Context = LocalContext.Get();
+		}
+
+		FSourceDataDescriptor* Found = Context->Cache.Find(this);
+		if (Found)
+		{
+			return *Found;
+		}
+
+		// Not cached: calculate
+		FSourceDataDescriptor Result;
+
+		if (base)
+		{
+			FSourceDataDescriptor SourceDesc = base->GetSourceDataDescriptor(Context);
+			Result.CombineWith(SourceDesc);
+		}
+
+		if (blend)
+		{
+			FSourceDataDescriptor SourceDesc = blend->GetSourceDataDescriptor(Context);
+			Result.CombineWith(SourceDesc);
+		}
+
+		if (mask)
+		{
+			FSourceDataDescriptor SourceDesc = mask->GetSourceDataDescriptor(Context);
+			Result.CombineWith(SourceDesc);
+		}
+
+		Context->Cache.Add(this, Result);
+
+		return Result;
+	}
+
 }
diff --git a/Engine/Plugins/Experimental/Mutable/Source/MutableTools/Private/MuT/ASTOpImageLayer.h b/Engine/Plugins/Experimental/Mutable/Source/MutableTools/Private/MuT/ASTOpImageLayer.h
--- a/Engine/Plugins/Experimental/Mutable/Source/MutableTools/Private/MuT/ASTOpImageLayer.h
+++ b/Engine/Plugins/Experimental/Mutable/Source/MutableTools/Private/MuT/ASTOpImageLayer.h
@@ -43,17 +43,18 @@ namespace mu
 		ASTOpImageLayer(const ASTOpImageLayer&) = delete;
 		~ASTOpImageLayer();
 
-		OP_TYPE GetOpType() const override { return OP_TYPE::IM_LAYER; }
-		uint64 Hash() const override;
-		bool IsEqual(const ASTOp& otherUntyped) const override;
-		Ptr<ASTOp> Clone(MapChildFuncRef mapChild) const override;
-		void ForEachChild(const TFunctionRef<void(ASTChild&)>) override;
-		void Link(FProgram& program, FLinkerOptions*) override;
-		FImageDesc GetImageDesc(bool returnBestOption, FGetImageDescContext*) const override;
-		void GetLayoutBlockSize(int* pBlockX, int* pBlockY) override;
-		Ptr<ImageSizeExpression> GetImageSizeExpression() const override;
-		Ptr<ASTOp> OptimiseSemantic(const FModelOptimizationOptions&, int32 Pass) const override;
-		Ptr<ASTOp> OptimiseSink(const FModelOptimizationOptions&, FOptimizeSinkContext&) const;
+		virtual OP_TYPE GetOpType() const override { return OP_TYPE::IM_LAYER; }
+		virtual uint64 Hash() const override;
+		virtual bool IsEqual(const ASTOp& otherUntyped) const override;
+		virtual Ptr<ASTOp> Clone(MapChildFuncRef mapChild) const override;
+		virtual void ForEachChild(const TFunctionRef<void(ASTChild&)>) override;
+		virtual void Link(FProgram& program, FLinkerOptions*) override;
+		virtual FImageDesc GetImageDesc(bool returnBestOption, FGetImageDescContext*) const override;
+		virtual void GetLayoutBlockSize(int* pBlockX, int* pBlockY) override;
+		virtual Ptr<ImageSizeExpression> GetImageSizeExpression() const override;
+		virtual Ptr<ASTOp> OptimiseSemantic(const FModelOptimizationOptions&, int32 Pass) const override;
+		virtual Ptr<ASTOp> OptimiseSink(const FModelOptimizationOptions&, FOptimizeSinkContext&) const;
+		virtual FSourceDataDescriptor GetSourceDataDescriptor(FGetSourceDataDescriptorContext*) const override;
 
 	};
 
diff --git a/Engine/Plugins/Experimental/Mutable/Source/MutableTools/Private/MuT/ASTOpImageLayerColor.cpp b/Engine/Plugins/Experimental/Mutable/Source/MutableTools/Private/MuT/ASTOpImageLayerColor.cpp
--- a/Engine/Plugins/Experimental/Mutable/Source/MutableTools/Private/MuT/ASTOpImageLayerColor.cpp
+++ b/Engine/Plugins/Experimental/Mutable/Source/MutableTools/Private/MuT/ASTOpImageLayerColor.cpp
@@ -438,4 +438,41 @@ namespace mu
 		return at;
 	}
 
+
+	FSourceDataDescriptor ASTOpImageLayerColor::GetSourceDataDescriptor(FGetSourceDataDescriptorContext* Context) const
+	{
+		// Cache management
+		TUniquePtr<FGetSourceDataDescriptorContext> LocalContext;
+		if (!Context)
+		{
+			LocalContext.Reset(new FGetSourceDataDescriptorContext);
+			Context = LocalContext.Get();
+		}
+
+		FSourceDataDescriptor* Found = Context->Cache.Find(this);
+		if (Found)
+		{
+			return *Found;
+		}
+
+		// Not cached: calculate
+		FSourceDataDescriptor Result;
+
+		if (base)
+		{
+			FSourceDataDescriptor SourceDesc = base->GetSourceDataDescriptor(Context);
+			Result.CombineWith(SourceDesc);
+		}
+
+		if (mask)
+		{
+			FSourceDataDescriptor SourceDesc = mask->GetSourceDataDescriptor(Context);
+			Result.CombineWith(SourceDesc);
+		}
+
+		Context->Cache.Add(this, Result);
+
+		return Result;
+	}
+
 }
diff --git a/Engine/Plugins/Experimental/Mutable/Source/MutableTools/Private/MuT/ASTOpImageLayerColor.h b/Engine/Plugins/Experimental/Mutable/Source/MutableTools/Private/MuT/ASTOpImageLayerColor.h
--- a/Engine/Plugins/Experimental/Mutable/Source/MutableTools/Private/MuT/ASTOpImageLayerColor.h
+++ b/Engine/Plugins/Experimental/Mutable/Source/MutableTools/Private/MuT/ASTOpImageLayerColor.h
@@ -43,17 +43,18 @@ namespace mu
 		ASTOpImageLayerColor(const ASTOpImageLayerColor&) = delete;
 		~ASTOpImageLayerColor();
 
-		OP_TYPE GetOpType() const override { return OP_TYPE::IM_LAYERCOLOUR; }
-		uint64 Hash() const override;
-		bool IsEqual(const ASTOp& otherUntyped) const override;
-		Ptr<ASTOp> Clone(MapChildFuncRef mapChild) const override;
-		void ForEachChild(const TFunctionRef<void(ASTChild&)>) override;
-		void Link(FProgram& program, FLinkerOptions* Options) override;
-		FImageDesc GetImageDesc(bool returnBestOption, FGetImageDescContext* context) const override;
-		void GetLayoutBlockSize(int* pBlockX, int* pBlockY) override;
-		Ptr<ImageSizeExpression> GetImageSizeExpression() const override;
-		Ptr<ASTOp> OptimiseSemantic(const FModelOptimizationOptions& options, int32 Pass) const;
-		Ptr<ASTOp> OptimiseSink(const FModelOptimizationOptions&, FOptimizeSinkContext&) const;
+		virtual OP_TYPE GetOpType() const override { return OP_TYPE::IM_LAYERCOLOUR; }
+		virtual uint64 Hash() const override;
+		virtual bool IsEqual(const ASTOp& otherUntyped) const override;
+		virtual Ptr<ASTOp> Clone(MapChildFuncRef mapChild) const override;
+		virtual void ForEachChild(const TFunctionRef<void(ASTChild&)>) override;
+		virtual void Link(FProgram& program, FLinkerOptions* Options) override;
+		virtual FImageDesc GetImageDesc(bool returnBestOption, FGetImageDescContext* context) const override;
+		virtual void GetLayoutBlockSize(int* pBlockX, int* pBlockY) override;
+		virtual Ptr<ImageSizeExpression> GetImageSizeExpression() const override;
+		virtual Ptr<ASTOp> OptimiseSemantic(const FModelOptimizationOptions& options, int32 Pass) const;
+		virtual Ptr<ASTOp> OptimiseSink(const FModelOptimizationOptions&, FOptimizeSinkContext&) const;
+		virtual FSourceDataDescriptor GetSourceDataDescriptor(FGetSourceDataDescriptorContext*) const override;
 
 	};
 
diff --git a/Engine/Plugins/Experimental/Mutable/Source/MutableTools/Private/MuT/ASTOpImageMakeGrowMap.cpp b/Engine/Plugins/Experimental/Mutable/Source/MutableTools/Private/MuT/ASTOpImageMakeGrowMap.cpp
--- a/Engine/Plugins/Experimental/Mutable/Source/MutableTools/Private/MuT/ASTOpImageMakeGrowMap.cpp
+++ b/Engine/Plugins/Experimental/Mutable/Source/MutableTools/Private/MuT/ASTOpImageMakeGrowMap.cpp
@@ -219,5 +219,16 @@ namespace mu
 		return pRes;
 	}
 
+
+	FSourceDataDescriptor ASTOpImageMakeGrowMap::GetSourceDataDescriptor(FGetSourceDataDescriptorContext* Context) const
+	{
+		if (Mask)
+		{
+			return Mask->GetSourceDataDescriptor(Context);
+		}
+
+		return {};
+	}
+
 }
 
diff --git a/Engine/Plugins/Experimental/Mutable/Source/MutableTools/Private/MuT/ASTOpImageMakeGrowMap.h b/Engine/Plugins/Experimental/Mutable/Source/MutableTools/Private/MuT/ASTOpImageMakeGrowMap.h
--- a/Engine/Plugins/Experimental/Mutable/Source/MutableTools/Private/MuT/ASTOpImageMakeGrowMap.h
+++ b/Engine/Plugins/Experimental/Mutable/Source/MutableTools/Private/MuT/ASTOpImageMakeGrowMap.h
@@ -27,17 +27,18 @@ namespace mu
 		ASTOpImageMakeGrowMap(const ASTOpImageMakeGrowMap&) = delete;
 		~ASTOpImageMakeGrowMap();
 
-		OP_TYPE GetOpType() const override { return OP_TYPE::IM_MAKEGROWMAP; }
-		uint64 Hash() const override;
-		bool IsEqual(const ASTOp& otherUntyped) const override;
-		Ptr<ASTOp> Clone(MapChildFuncRef mapChild) const override;
-		void ForEachChild(const TFunctionRef<void(ASTChild&)>) override;
-		void Link(FProgram& program, FLinkerOptions* Options) override;
-		Ptr<ASTOp> OptimiseSink(const FModelOptimizationOptions& options, FOptimizeSinkContext& context) const override;
-		FImageDesc GetImageDesc(bool returnBestOption, FGetImageDescContext* context) const override;
-		void GetLayoutBlockSize(int* pBlockX, int* pBlockY) override;
-		bool IsImagePlainConstant(FVector4f& colour) const override;
-		Ptr<ImageSizeExpression> GetImageSizeExpression() const override;
+		virtual OP_TYPE GetOpType() const override { return OP_TYPE::IM_MAKEGROWMAP; }
+		virtual uint64 Hash() const override;
+		virtual bool IsEqual(const ASTOp& otherUntyped) const override;
+		virtual Ptr<ASTOp> Clone(MapChildFuncRef mapChild) const override;
+		virtual void ForEachChild(const TFunctionRef<void(ASTChild&)>) override;
+		virtual void Link(FProgram& program, FLinkerOptions* Options) override;
+		virtual Ptr<ASTOp> OptimiseSink(const FModelOptimizationOptions& options, FOptimizeSinkContext& context) const override;
+		virtual FImageDesc GetImageDesc(bool returnBestOption, FGetImageDescContext* context) const override;
+		virtual void GetLayoutBlockSize(int* pBlockX, int* pBlockY) override;
+		virtual bool IsImagePlainConstant(FVector4f& colour) const override;
+		virtual Ptr<ImageSizeExpression> GetImageSizeExpression() const override;
+		virtual FSourceDataDescriptor GetSourceDataDescriptor(FGetSourceDataDescriptorContext*) const override;
 
 	};
 
diff --git a/Engine/Plugins/Experimental/Mutable/Source/MutableTools/Private/MuT/ASTOpImageMipmap.cpp b/Engine/Plugins/Experimental/Mutable/Source/MutableTools/Private/MuT/ASTOpImageMipmap.cpp
--- a/Engine/Plugins/Experimental/Mutable/Source/MutableTools/Private/MuT/ASTOpImageMipmap.cpp
+++ b/Engine/Plugins/Experimental/Mutable/Source/MutableTools/Private/MuT/ASTOpImageMipmap.cpp
@@ -270,6 +270,17 @@ namespace mu
 	}
 
 
+	FSourceDataDescriptor ASTOpImageMipmap::GetSourceDataDescriptor(FGetSourceDataDescriptorContext* Context) const
+	{
+		if (Source)
+		{
+			return Source->GetSourceDataDescriptor(Context);
+		}
+
+		return {};
+	}
+
+
 	//---------------------------------------------------------------------------------------------
 	//---------------------------------------------------------------------------------------------
 	//---------------------------------------------------------------------------------------------
diff --git a/Engine/Plugins/Experimental/Mutable/Source/MutableTools/Private/MuT/ASTOpImageMipmap.h b/Engine/Plugins/Experimental/Mutable/Source/MutableTools/Private/MuT/ASTOpImageMipmap.h
--- a/Engine/Plugins/Experimental/Mutable/Source/MutableTools/Private/MuT/ASTOpImageMipmap.h
+++ b/Engine/Plugins/Experimental/Mutable/Source/MutableTools/Private/MuT/ASTOpImageMipmap.h
@@ -41,18 +41,19 @@ namespace mu
 		ASTOpImageMipmap(const ASTOpImageMipmap&) = delete;
 		~ASTOpImageMipmap();
 
-		OP_TYPE GetOpType() const override { return OP_TYPE::IM_MIPMAP; }
-		uint64 Hash() const override;
-		bool IsEqual(const ASTOp& otherUntyped) const override;
-		Ptr<ASTOp> Clone(MapChildFuncRef mapChild) const override;
-		void ForEachChild(const TFunctionRef<void(ASTChild&)>) override;
-		void Link(FProgram& program, FLinkerOptions* Options) override;
-		mu::Ptr<ASTOp> OptimiseSemantic(const FModelOptimizationOptions& options, int32 Pass) const override;
-		Ptr<ASTOp> OptimiseSink(const FModelOptimizationOptions& options, FOptimizeSinkContext& context) const override;
-		FImageDesc GetImageDesc(bool returnBestOption, FGetImageDescContext* context) const override;
-		void GetLayoutBlockSize(int* pBlockX, int* pBlockY) override;
-		bool IsImagePlainConstant(FVector4f& colour) const override;
-		Ptr<ImageSizeExpression> GetImageSizeExpression() const override;
+		virtual OP_TYPE GetOpType() const override { return OP_TYPE::IM_MIPMAP; }
+		virtual uint64 Hash() const override;
+		virtual bool IsEqual(const ASTOp& otherUntyped) const override;
+		virtual Ptr<ASTOp> Clone(MapChildFuncRef mapChild) const override;
+		virtual void ForEachChild(const TFunctionRef<void(ASTChild&)>) override;
+		virtual void Link(FProgram& program, FLinkerOptions* Options) override;
+		virtual mu::Ptr<ASTOp> OptimiseSemantic(const FModelOptimizationOptions& options, int32 Pass) const override;
+		virtual Ptr<ASTOp> OptimiseSink(const FModelOptimizationOptions& options, FOptimizeSinkContext& context) const override;
+		virtual FImageDesc GetImageDesc(bool returnBestOption, FGetImageDescContext* context) const override;
+		virtual void GetLayoutBlockSize(int* pBlockX, int* pBlockY) override;
+		virtual bool IsImagePlainConstant(FVector4f& colour) const override;
+		virtual Ptr<ImageSizeExpression> GetImageSizeExpression() const override;
+		virtual FSourceDataDescriptor GetSourceDataDescriptor(FGetSourceDataDescriptorContext*) const override;
 
 	};
 
diff --git a/Engine/Plugins/Experimental/Mutable/Source/MutableTools/Private/MuT/ASTOpImagePatch.cpp b/Engine/Plugins/Experimental/Mutable/Source/MutableTools/Private/MuT/ASTOpImagePatch.cpp
--- a/Engine/Plugins/Experimental/Mutable/Source/MutableTools/Private/MuT/ASTOpImagePatch.cpp
+++ b/Engine/Plugins/Experimental/Mutable/Source/MutableTools/Private/MuT/ASTOpImagePatch.cpp
@@ -161,4 +161,41 @@ namespace mu
 		}
 	}
 
+
+	FSourceDataDescriptor ASTOpImagePatch::GetSourceDataDescriptor(FGetSourceDataDescriptorContext* Context) const
+	{
+		// Cache management
+		TUniquePtr<FGetSourceDataDescriptorContext> LocalContext;
+		if (!Context)
+		{
+			LocalContext.Reset(new FGetSourceDataDescriptorContext);
+			Context = LocalContext.Get();
+		}
+
+		FSourceDataDescriptor* Found = Context->Cache.Find(this);
+		if (Found)
+		{
+			return *Found;
+		}
+
+		// Not cached: calculate
+		FSourceDataDescriptor Result;
+
+		if (base)
+		{
+			FSourceDataDescriptor SourceDesc = base->GetSourceDataDescriptor(Context);
+			Result.CombineWith(SourceDesc);
+		}
+
+		if (patch)
+		{
+			FSourceDataDescriptor SourceDesc = patch->GetSourceDataDescriptor(Context);
+			Result.CombineWith(SourceDesc);
+		}
+
+		Context->Cache.Add(this, Result);
+
+		return Result;
+	}
+
 }
diff --git a/Engine/Plugins/Experimental/Mutable/Source/MutableTools/Private/MuT/ASTOpImagePatch.h b/Engine/Plugins/Experimental/Mutable/Source/MutableTools/Private/MuT/ASTOpImagePatch.h
--- a/Engine/Plugins/Experimental/Mutable/Source/MutableTools/Private/MuT/ASTOpImagePatch.h
+++ b/Engine/Plugins/Experimental/Mutable/Source/MutableTools/Private/MuT/ASTOpImagePatch.h
@@ -28,16 +28,18 @@ struct FProgram;
 		ASTOpImagePatch(const ASTOpImagePatch&) = delete;
 		~ASTOpImagePatch();
 
-		OP_TYPE GetOpType() const override { return OP_TYPE::IM_PATCH; }
-		uint64 Hash() const override;
-		bool IsEqual(const ASTOp& otherUntyped) const override;
-		Ptr<ASTOp> Clone(MapChildFuncRef mapChild) const override;
-		void ForEachChild(const TFunctionRef<void(ASTChild&)>) override;
-		void Link(FProgram& program, FLinkerOptions* Options) override;
-		FImageDesc GetImageDesc(bool returnBestOption, FGetImageDescContext* context) const override;
-		Ptr<ImageSizeExpression> GetImageSizeExpression() const override;
-		void GetLayoutBlockSize(int* pBlockX, int* pBlockY) override;
-		//TODO: bool IsImagePlainConstant(FVector4f& colour) const override;
+		virtual OP_TYPE GetOpType() const override { return OP_TYPE::IM_PATCH; }
+		virtual uint64 Hash() const override;
+		virtual bool IsEqual(const ASTOp& otherUntyped) const override;
+		virtual Ptr<ASTOp> Clone(MapChildFuncRef mapChild) const override;
+		virtual void ForEachChild(const TFunctionRef<void(ASTChild&)>) override;
+		virtual void Link(FProgram& program, FLinkerOptions* Options) override;
+		virtual FImageDesc GetImageDesc(bool returnBestOption, FGetImageDescContext* context) const override;
+		virtual Ptr<ImageSizeExpression> GetImageSizeExpression() const override;
+		virtual void GetLayoutBlockSize(int* pBlockX, int* pBlockY) override;
+		//TODO: virtual bool IsImagePlainConstant(FVector4f& colour) const override;
+		virtual FSourceDataDescriptor GetSourceDataDescriptor(FGetSourceDataDescriptorContext*) const override;
+
 	};
 
 
diff --git a/Engine/Plugins/Experimental/Mutable/Source/MutableTools/Private/MuT/ASTOpImagePixelFormat.cpp b/Engine/Plugins/Experimental/Mutable/Source/MutableTools/Private/MuT/ASTOpImagePixelFormat.cpp
--- a/Engine/Plugins/Experimental/Mutable/Source/MutableTools/Private/MuT/ASTOpImagePixelFormat.cpp
+++ b/Engine/Plugins/Experimental/Mutable/Source/MutableTools/Private/MuT/ASTOpImagePixelFormat.cpp
@@ -322,6 +322,17 @@ namespace mu
 	}
 
 
+	FSourceDataDescriptor ASTOpImagePixelFormat::GetSourceDataDescriptor(FGetSourceDataDescriptorContext* Context) const
+	{
+		if (Source)
+		{
+			return Source->GetSourceDataDescriptor(Context);
+		}
+
+		return {};
+	}
+
+
 	//---------------------------------------------------------------------------------------------
 	//---------------------------------------------------------------------------------------------
 	//---------------------------------------------------------------------------------------------
diff --git a/Engine/Plugins/Experimental/Mutable/Source/MutableTools/Private/MuT/ASTOpImagePixelFormat.h b/Engine/Plugins/Experimental/Mutable/Source/MutableTools/Private/MuT/ASTOpImagePixelFormat.h
--- a/Engine/Plugins/Experimental/Mutable/Source/MutableTools/Private/MuT/ASTOpImagePixelFormat.h
+++ b/Engine/Plugins/Experimental/Mutable/Source/MutableTools/Private/MuT/ASTOpImagePixelFormat.h
@@ -27,18 +27,19 @@ namespace mu
 		ASTOpImagePixelFormat(const ASTOpImagePixelFormat&) = delete;
 		~ASTOpImagePixelFormat();
 
-		OP_TYPE GetOpType() const override { return OP_TYPE::IM_PIXELFORMAT; }
-		uint64 Hash() const override;
-		bool IsEqual(const ASTOp& otherUntyped) const override;
-		Ptr<ASTOp> Clone(MapChildFuncRef mapChild) const override;
-		void ForEachChild(const TFunctionRef<void(ASTChild&)>) override;
-		void Link(FProgram& program, FLinkerOptions* Options) override;
-		Ptr<ASTOp> OptimiseSemantic(const FModelOptimizationOptions& options, int32 Pass) const override;
-		Ptr<ASTOp> OptimiseSink(const FModelOptimizationOptions& options, FOptimizeSinkContext& context) const override;
-		FImageDesc GetImageDesc(bool returnBestOption, FGetImageDescContext* context) const override;
-		void GetLayoutBlockSize(int* pBlockX, int* pBlockY) override;
-		bool IsImagePlainConstant(FVector4f& colour) const override;
-		Ptr<ImageSizeExpression> GetImageSizeExpression() const override;
+		virtual OP_TYPE GetOpType() const override { return OP_TYPE::IM_PIXELFORMAT; }
+		virtual uint64 Hash() const override;
+		virtual bool IsEqual(const ASTOp& otherUntyped) const override;
+		virtual Ptr<ASTOp> Clone(MapChildFuncRef mapChild) const override;
+		virtual void ForEachChild(const TFunctionRef<void(ASTChild&)>) override;
+		virtual void Link(FProgram& program, FLinkerOptions* Options) override;
+		virtual Ptr<ASTOp> OptimiseSemantic(const FModelOptimizationOptions& options, int32 Pass) const override;
+		virtual Ptr<ASTOp> OptimiseSink(const FModelOptimizationOptions& options, FOptimizeSinkContext& context) const override;
+		virtual FImageDesc GetImageDesc(bool returnBestOption, FGetImageDescContext* context) const override;
+		virtual void GetLayoutBlockSize(int* pBlockX, int* pBlockY) override;
+		virtual bool IsImagePlainConstant(FVector4f& colour) const override;
+		virtual Ptr<ImageSizeExpression> GetImageSizeExpression() const override;
+		virtual FSourceDataDescriptor GetSourceDataDescriptor(FGetSourceDataDescriptorContext*) const override;
 
 	};
 
diff --git a/Engine/Plugins/Experimental/Mutable/Source/MutableTools/Private/MuT/ASTOpImageRasterMesh.cpp b/Engine/Plugins/Experimental/Mutable/Source/MutableTools/Private/MuT/ASTOpImageRasterMesh.cpp
--- a/Engine/Plugins/Experimental/Mutable/Source/MutableTools/Private/MuT/ASTOpImageRasterMesh.cpp
+++ b/Engine/Plugins/Experimental/Mutable/Source/MutableTools/Private/MuT/ASTOpImageRasterMesh.cpp
@@ -401,4 +401,15 @@ namespace mu
 		return at;
 	}
 
+
+	FSourceDataDescriptor ASTOpImageRasterMesh::GetSourceDataDescriptor(FGetSourceDataDescriptorContext* Context) const
+	{
+		if (image)
+		{
+			return image->GetSourceDataDescriptor(Context);
+		}
+
+		return {};
+	}
+
 }
diff --git a/Engine/Plugins/Experimental/Mutable/Source/MutableTools/Private/MuT/ASTOpImageRasterMesh.h b/Engine/Plugins/Experimental/Mutable/Source/MutableTools/Private/MuT/ASTOpImageRasterMesh.h
--- a/Engine/Plugins/Experimental/Mutable/Source/MutableTools/Private/MuT/ASTOpImageRasterMesh.h
+++ b/Engine/Plugins/Experimental/Mutable/Source/MutableTools/Private/MuT/ASTOpImageRasterMesh.h
@@ -50,16 +50,17 @@ namespace mu
 		ASTOpImageRasterMesh(const ASTOpImageRasterMesh&) = delete;
 		~ASTOpImageRasterMesh();
 
-		OP_TYPE GetOpType() const override { return OP_TYPE::IM_RASTERMESH; }
-		uint64 Hash() const override;
-		bool IsEqual(const ASTOp& otherUntyped) const override;
-		Ptr<ASTOp> Clone(MapChildFuncRef mapChild) const override;
-		void ForEachChild(const TFunctionRef<void(ASTChild&)>) override;
-		void Link(FProgram& program, FLinkerOptions* Options) override;
-		FImageDesc GetImageDesc(bool returnBestOption, FGetImageDescContext* context) const override;
-		Ptr<ImageSizeExpression> GetImageSizeExpression() const override;
-		Ptr<ASTOp> OptimiseSemantic(const FModelOptimizationOptions& options, int32 Pass) const;
-		Ptr<ASTOp> OptimiseSink(const FModelOptimizationOptions&, FOptimizeSinkContext&) const;
+		virtual OP_TYPE GetOpType() const override { return OP_TYPE::IM_RASTERMESH; }
+		virtual uint64 Hash() const override;
+		virtual bool IsEqual(const ASTOp& otherUntyped) const override;
+		virtual Ptr<ASTOp> Clone(MapChildFuncRef mapChild) const override;
+		virtual void ForEachChild(const TFunctionRef<void(ASTChild&)>) override;
+		virtual void Link(FProgram& program, FLinkerOptions* Options) override;
+		virtual FImageDesc GetImageDesc(bool returnBestOption, FGetImageDescContext* context) const override;
+		virtual Ptr<ImageSizeExpression> GetImageSizeExpression() const override;
+		virtual Ptr<ASTOp> OptimiseSemantic(const FModelOptimizationOptions& options, int32 Pass) const;
+		virtual Ptr<ASTOp> OptimiseSink(const FModelOptimizationOptions&, FOptimizeSinkContext&) const;
+		virtual FSourceDataDescriptor GetSourceDataDescriptor(FGetSourceDataDescriptorContext*) const override;
 
 	};
 
diff --git a/Engine/Plugins/Experimental/Mutable/Source/MutableTools/Private/MuT/ASTOpImageSwizzle.cpp b/Engine/Plugins/Experimental/Mutable/Source/MutableTools/Private/MuT/ASTOpImageSwizzle.cpp
--- a/Engine/Plugins/Experimental/Mutable/Source/MutableTools/Private/MuT/ASTOpImageSwizzle.cpp
+++ b/Engine/Plugins/Experimental/Mutable/Source/MutableTools/Private/MuT/ASTOpImageSwizzle.cpp
@@ -1079,6 +1079,39 @@ namespace mu
 	}
 
 
+	FSourceDataDescriptor ASTOpImageSwizzle::GetSourceDataDescriptor(FGetSourceDataDescriptorContext* Context) const
+	{
+		// Cache management
+		TUniquePtr<FGetSourceDataDescriptorContext> LocalContext;
+		if (!Context)
+		{
+			LocalContext.Reset(new FGetSourceDataDescriptorContext);
+			Context = LocalContext.Get();
+		}
+
+		FSourceDataDescriptor* Found= Context->Cache.Find(this);
+		if (Found)
+		{
+			return *Found;
+		}
+
+		// Not cached: calculate
+		FSourceDataDescriptor Result;
+
+		for( int32 SourceIndex=0; SourceIndex<MUTABLE_OP_MAX_SWIZZLE_CHANNELS; ++SourceIndex )
+		{ 
+			if (Sources[SourceIndex])
+			{
+				FSourceDataDescriptor SourceDesc = Sources[SourceIndex]->GetSourceDataDescriptor(Context);				
+				Result.CombineWith(SourceDesc);
+			}
+		}
+
+		Context->Cache.Add(this,Result);
+
+		return Result;
+	}
+
 
 
 	//---------------------------------------------------------------------------------------------
diff --git a/Engine/Plugins/Experimental/Mutable/Source/MutableTools/Private/MuT/ASTOpImageSwizzle.h b/Engine/Plugins/Experimental/Mutable/Source/MutableTools/Private/MuT/ASTOpImageSwizzle.h
--- a/Engine/Plugins/Experimental/Mutable/Source/MutableTools/Private/MuT/ASTOpImageSwizzle.h
+++ b/Engine/Plugins/Experimental/Mutable/Source/MutableTools/Private/MuT/ASTOpImageSwizzle.h
@@ -29,18 +29,19 @@ namespace mu
 		ASTOpImageSwizzle(const ASTOpImageSwizzle&) = delete;
 		~ASTOpImageSwizzle();
 
-		OP_TYPE GetOpType() const override { return OP_TYPE::IM_SWIZZLE; }
-		uint64 Hash() const override;
-		bool IsEqual(const ASTOp& otherUntyped) const override;
-		Ptr<ASTOp> Clone(MapChildFuncRef mapChild) const override;
-		void ForEachChild(const TFunctionRef<void(ASTChild&)>) override;
-		void Link(FProgram& program, FLinkerOptions* Options) override;
-		Ptr<ASTOp> OptimiseSemantic(const FModelOptimizationOptions&, int32 Pass) const override;
-		Ptr<ASTOp> OptimiseSink(const FModelOptimizationOptions& options, FOptimizeSinkContext& context) const override;
-		FImageDesc GetImageDesc(bool returnBestOption, FGetImageDescContext* context) const override;
-		void GetLayoutBlockSize(int* pBlockX, int* pBlockY) override;
-		bool IsImagePlainConstant(FVector4f& colour) const override;
-		Ptr<ImageSizeExpression> GetImageSizeExpression() const override;
+		virtual OP_TYPE GetOpType() const override { return OP_TYPE::IM_SWIZZLE; }
+		virtual uint64 Hash() const override;
+		virtual bool IsEqual(const ASTOp& otherUntyped) const override;
+		virtual Ptr<ASTOp> Clone(MapChildFuncRef mapChild) const override;
+		virtual void ForEachChild(const TFunctionRef<void(ASTChild&)>) override;
+		virtual void Link(FProgram& program, FLinkerOptions* Options) override;
+		virtual Ptr<ASTOp> OptimiseSemantic(const FModelOptimizationOptions&, int32 Pass) const override;
+		virtual Ptr<ASTOp> OptimiseSink(const FModelOptimizationOptions& options, FOptimizeSinkContext& context) const override;
+		virtual FImageDesc GetImageDesc(bool returnBestOption, FGetImageDescContext* context) const override;
+		virtual void GetLayoutBlockSize(int* pBlockX, int* pBlockY) override;
+		virtual bool IsImagePlainConstant(FVector4f& colour) const override;
+		virtual Ptr<ImageSizeExpression> GetImageSizeExpression() const override;
+		virtual FSourceDataDescriptor GetSourceDataDescriptor(FGetSourceDataDescriptorContext*) const override;
 
 	};
 
diff --git a/Engine/Plugins/Experimental/Mutable/Source/MutableTools/Private/MuT/ASTOpImageTransform.cpp b/Engine/Plugins/Experimental/Mutable/Source/MutableTools/Private/MuT/ASTOpImageTransform.cpp
--- a/Engine/Plugins/Experimental/Mutable/Source/MutableTools/Private/MuT/ASTOpImageTransform.cpp
+++ b/Engine/Plugins/Experimental/Mutable/Source/MutableTools/Private/MuT/ASTOpImageTransform.cpp
@@ -217,4 +217,15 @@ namespace mu
 		return nullptr;
 	}
 
+
+	FSourceDataDescriptor ASTOpImageTransform::GetSourceDataDescriptor(FGetSourceDataDescriptorContext* Context) const
+	{
+		if (Base)
+		{
+			return Base->GetSourceDataDescriptor(Context);
+		}
+
+		return {};
+	}
+
 }
diff --git a/Engine/Plugins/Experimental/Mutable/Source/MutableTools/Private/MuT/ASTOpImageTransform.h b/Engine/Plugins/Experimental/Mutable/Source/MutableTools/Private/MuT/ASTOpImageTransform.h
--- a/Engine/Plugins/Experimental/Mutable/Source/MutableTools/Private/MuT/ASTOpImageTransform.h
+++ b/Engine/Plugins/Experimental/Mutable/Source/MutableTools/Private/MuT/ASTOpImageTransform.h
@@ -43,15 +43,16 @@ struct FProgram;
 		ASTOpImageTransform(const ASTOpImageTransform&) = delete;
 		~ASTOpImageTransform();
 
-		OP_TYPE GetOpType() const override { return OP_TYPE::IM_TRANSFORM; }
-		uint64 Hash() const override;
-		bool IsEqual(const ASTOp& OtherUntyped) const override;
-		Ptr<ASTOp> Clone(MapChildFuncRef MapChild) const override;
-		void ForEachChild(const TFunctionRef<void(ASTChild&)>) override;
-		void Link(FProgram& Program, FLinkerOptions* Options) override;
-		FImageDesc GetImageDesc(bool bReturnBestOption, FGetImageDescContext* Context) const override;
-		void GetLayoutBlockSize(int* pBlockX, int* pBlockY) override;
-		Ptr<ImageSizeExpression> GetImageSizeExpression() const override;
+		virtual OP_TYPE GetOpType() const override { return OP_TYPE::IM_TRANSFORM; }
+		virtual uint64 Hash() const override;
+		virtual bool IsEqual(const ASTOp& OtherUntyped) const override;
+		virtual Ptr<ASTOp> Clone(MapChildFuncRef MapChild) const override;
+		virtual void ForEachChild(const TFunctionRef<void(ASTChild&)>) override;
+		virtual void Link(FProgram& Program, FLinkerOptions* Options) override;
+		virtual FImageDesc GetImageDesc(bool bReturnBestOption, FGetImageDescContext* Context) const override;
+		virtual void GetLayoutBlockSize(int* pBlockX, int* pBlockY) override;
+		virtual Ptr<ImageSizeExpression> GetImageSizeExpression() const override;
+		virtual FSourceDataDescriptor GetSourceDataDescriptor(FGetSourceDataDescriptorContext*) const override;
 	};
 
 }
diff --git a/Engine/Plugins/Experimental/Mutable/Source/MutableTools/Private/MuT/ASTOpReferenceResource.cpp b/Engine/Plugins/Experimental/Mutable/Source/MutableTools/Private/MuT/ASTOpReferenceResource.cpp
--- a/Engine/Plugins/Experimental/Mutable/Source/MutableTools/Private/MuT/ASTOpReferenceResource.cpp
+++ b/Engine/Plugins/Experimental/Mutable/Source/MutableTools/Private/MuT/ASTOpReferenceResource.cpp
@@ -122,7 +122,6 @@ namespace mu
 	}
 
 
-	//-------------------------------------------------------------------------------------------------
 	mu::Ptr<ImageSizeExpression> ASTOpReferenceResource::GetImageSizeExpression() const
 	{
 		Ptr<ImageSizeExpression> pRes = new ImageSizeExpression;
@@ -132,4 +131,10 @@ namespace mu
 	}
 
 
+	FSourceDataDescriptor ASTOpReferenceResource::GetSourceDataDescriptor(FGetSourceDataDescriptorContext*) const
+	{
+		return SourceDataDescriptor;
+	}
+
+
 }
\ No newline at end of file
diff --git a/Engine/Plugins/Experimental/Mutable/Source/MutableTools/Private/MuT/ASTOpReferenceResource.h b/Engine/Plugins/Experimental/Mutable/Source/MutableTools/Private/MuT/ASTOpReferenceResource.h
--- a/Engine/Plugins/Experimental/Mutable/Source/MutableTools/Private/MuT/ASTOpReferenceResource.h
+++ b/Engine/Plugins/Experimental/Mutable/Source/MutableTools/Private/MuT/ASTOpReferenceResource.h
@@ -25,31 +25,36 @@ namespace mu
 
 		//! Type of switch
 		OP_TYPE type = OP_TYPE::NONE;
-
-		//!
-		uint32 ID = 0;
 		
 		/** */
 		bool bForceLoad = false;
 
+		/** Externally provided ID to identify the resource. */
+		uint32 ID = 0;
+
+		/** */
 		FImageDesc ImageDesc;
 
+		/** Source data descriptor. */
+		FSourceDataDescriptor SourceDataDescriptor;
+
 	public:
 
 		// Own interface
 
 		// ASTOp interface
-		OP_TYPE GetOpType() const override { return type; }
-		void ForEachChild(const TFunctionRef<void(ASTChild&)>) override;
-		bool IsEqual(const ASTOp& otherUntyped) const override;
-		Ptr<ASTOp> Clone(MapChildFuncRef mapChild) const override;
-		uint64 Hash() const override;
-		void Link(FProgram& program, FLinkerOptions*) override;
-		FImageDesc GetImageDesc(bool, class FGetImageDescContext*) const override;
-		void GetBlockLayoutSize(uint64 BlockId, int32* pBlockX, int32* pBlockY, FBlockLayoutSizeCache* cache) override;
-		void GetLayoutBlockSize(int32* pBlockX, int32* pBlockY) override;
-		bool GetNonBlackRect(FImageRect& maskUsage) const override;
-		Ptr<ImageSizeExpression> GetImageSizeExpression() const override;
+		virtual OP_TYPE GetOpType() const override { return type; }
+		virtual void ForEachChild(const TFunctionRef<void(ASTChild&)>) override;
+		virtual bool IsEqual(const ASTOp& otherUntyped) const override;
+		virtual Ptr<ASTOp> Clone(MapChildFuncRef mapChild) const override;
+		virtual uint64 Hash() const override;
+		virtual void Link(FProgram& program, FLinkerOptions*) override;
+		virtual FImageDesc GetImageDesc(bool, class FGetImageDescContext*) const override;
+		virtual void GetBlockLayoutSize(uint64 BlockId, int32* pBlockX, int32* pBlockY, FBlockLayoutSizeCache* cache) override;
+		virtual void GetLayoutBlockSize(int32* pBlockX, int32* pBlockY) override;
+		virtual bool GetNonBlackRect(FImageRect& maskUsage) const override;
+		virtual Ptr<ImageSizeExpression> GetImageSizeExpression() const override;
+		virtual FSourceDataDescriptor GetSourceDataDescriptor(FGetSourceDataDescriptorContext*) const;
 	};
 
 
diff --git a/Engine/Plugins/Experimental/Mutable/Source/MutableTools/Private/MuT/ASTOpSwitch.cpp b/Engine/Plugins/Experimental/Mutable/Source/MutableTools/Private/MuT/ASTOpSwitch.cpp
--- a/Engine/Plugins/Experimental/Mutable/Source/MutableTools/Private/MuT/ASTOpSwitch.cpp
+++ b/Engine/Plugins/Experimental/Mutable/Source/MutableTools/Private/MuT/ASTOpSwitch.cpp
@@ -597,7 +597,7 @@ namespace mu
 		Ptr<ImageSizeExpression> pRes = new ImageSizeExpression;
 
 		bool first = true;
-		for (const auto& c : cases)
+		for (const FCase& c : cases)
 		{
 			if (c.branch)
 			{
@@ -621,4 +621,38 @@ namespace mu
 	}
 
 
+	FSourceDataDescriptor ASTOpSwitch::GetSourceDataDescriptor(FGetSourceDataDescriptorContext* Context) const
+	{
+		// Cache management
+		TUniquePtr<FGetSourceDataDescriptorContext> LocalContext;
+		if (!Context)
+		{
+			LocalContext.Reset(new FGetSourceDataDescriptorContext);
+			Context = LocalContext.Get();
+		}
+
+		FSourceDataDescriptor* Found = Context->Cache.Find(this);
+		if (Found)
+		{
+			return *Found;
+		}
+
+		// Not cached: calculate
+		FSourceDataDescriptor Result;
+
+		for (const FCase& Case : cases)
+		{
+			if (Case.branch)
+			{
+				FSourceDataDescriptor SourceDesc = Case.branch->GetSourceDataDescriptor(Context);
+				Result.CombineWith(SourceDesc);
+			}
+		}
+
+		Context->Cache.Add(this, Result);
+
+		return Result;
+	}
+
+
 }
diff --git a/Engine/Plugins/Experimental/Mutable/Source/MutableTools/Private/MuT/ASTOpSwitch.h b/Engine/Plugins/Experimental/Mutable/Source/MutableTools/Private/MuT/ASTOpSwitch.h
--- a/Engine/Plugins/Experimental/Mutable/Source/MutableTools/Private/MuT/ASTOpSwitch.h
+++ b/Engine/Plugins/Experimental/Mutable/Source/MutableTools/Private/MuT/ASTOpSwitch.h
@@ -59,21 +59,22 @@ template <class SCALAR> class vec4;
 		~ASTOpSwitch() override;
 
 		// ASTOp interface
-		OP_TYPE GetOpType() const override { return type; }
-		uint64 Hash() const override;
-		void ForEachChild(const TFunctionRef<void(ASTChild&)> f) override;
-		bool IsEqual(const ASTOp& otherUntyped) const override;
-		Ptr<ASTOp> Clone(MapChildFuncRef mapChild) const override;
-		void Assert() override;
-		void Link(FProgram& program, FLinkerOptions* Options) override;
-		FImageDesc GetImageDesc(bool returnBestOption, class FGetImageDescContext* context) const  override;
-		void GetBlockLayoutSize(uint64 BlockId, int32* pBlockX, int32* pBlockY, FBlockLayoutSizeCache* cache) override;
-		void GetLayoutBlockSize(int32* pBlockX, int32* pBlockY) override;
-		Ptr<ASTOp> OptimiseSemantic(const FModelOptimizationOptions&, int32 Pass) const override;
-		bool GetNonBlackRect(FImageRect& maskUsage) const override;
-		bool IsImagePlainConstant(FVector4f& colour) const override;
-		Ptr<ImageSizeExpression> GetImageSizeExpression() const override;
+		virtual OP_TYPE GetOpType() const override { return type; }
+		virtual uint64 Hash() const override;
+		virtual void ForEachChild(const TFunctionRef<void(ASTChild&)> f) override;
+		virtual bool IsEqual(const ASTOp& otherUntyped) const override;
+		virtual Ptr<ASTOp> Clone(MapChildFuncRef mapChild) const override;
+		virtual void Assert() override;
+		virtual void Link(FProgram& program, FLinkerOptions* Options) override;
+		virtual FImageDesc GetImageDesc(bool returnBestOption, class FGetImageDescContext* context) const  override;
+		virtual void GetBlockLayoutSize(uint64 BlockId, int32* pBlockX, int32* pBlockY, FBlockLayoutSizeCache* cache) override;
+		virtual void GetLayoutBlockSize(int32* pBlockX, int32* pBlockY) override;
+		virtual Ptr<ASTOp> OptimiseSemantic(const FModelOptimizationOptions&, int32 Pass) const override;
+		virtual bool GetNonBlackRect(FImageRect& maskUsage) const override;
+		virtual bool IsImagePlainConstant(FVector4f& colour) const override;
+		virtual Ptr<ImageSizeExpression> GetImageSizeExpression() const override;
 		virtual bool IsSwitch() const override { return true; }
+		virtual FSourceDataDescriptor GetSourceDataDescriptor(FGetSourceDataDescriptorContext*) const override;
 
 		// Own interface
 
diff --git a/Engine/Plugins/Experimental/Mutable/Source/MutableTools/Private/MuT/CodeGenerator.cpp b/Engine/Plugins/Experimental/Mutable/Source/MutableTools/Private/MuT/CodeGenerator.cpp
--- a/Engine/Plugins/Experimental/Mutable/Source/MutableTools/Private/MuT/CodeGenerator.cpp
+++ b/Engine/Plugins/Experimental/Mutable/Source/MutableTools/Private/MuT/CodeGenerator.cpp
@@ -308,7 +308,7 @@ namespace mu
 	}
 
 
-	Ptr<ASTOp> CodeGenerator::GenerateTableVariable(Ptr<const Node> InNode, const FTableCacheKey& CacheKey, bool bAddNoneOption, const FString& DefaultRowName)
+	Ptr<ASTOp> CodeGenerator::GenerateTableVariable( Ptr<const Node> InNode, const FTableCacheKey& CacheKey, bool bAddNoneOption, const FString& DefaultRowName)
 	{
 		Ptr<ASTOp> result;
 
@@ -327,7 +327,7 @@ namespace mu
 		{
 			// See if there is a string column. If there is one, we will use it as names for the
 			// options. Only the first string column will be used.
-			int nameCol = -1;
+			int32 nameCol = -1;
 			int32 cols = CacheKey.Table->GetPrivate()->Columns.Num();
 			for (int32 c = 0; c < cols && nameCol < 0; ++c)
 			{
diff --git a/Engine/Plugins/Experimental/Mutable/Source/MutableTools/Private/MuT/CodeGenerator.h b/Engine/Plugins/Experimental/Mutable/Source/MutableTools/Private/MuT/CodeGenerator.h
--- a/Engine/Plugins/Experimental/Mutable/Source/MutableTools/Private/MuT/CodeGenerator.h
+++ b/Engine/Plugins/Experimental/Mutable/Source/MutableTools/Private/MuT/CodeGenerator.h
@@ -367,8 +367,8 @@ namespace mu
         Ptr<ASTOp> GenerateMissingBoolCode(const TCHAR* strWhere, bool value, const void* errorContext );
 
         //!
-		template<class NODE_TABLE_PRIVATE, ETableColumnType TYPE, OP_TYPE OPTYPE, typename F>
-		Ptr<ASTOp> GenerateTableSwitch( const NODE_TABLE_PRIVATE& node, F&& GenerateOption );
+		template<class NODE_TABLE, ETableColumnType TYPE, OP_TYPE OPTYPE, typename F>
+		Ptr<ASTOp> GenerateTableSwitch( const NODE_TABLE& node, F&& GenerateOption );
 
 
 		//-----------------------------------------------------------------------------------------
@@ -768,8 +768,8 @@ namespace mu
 
 	
     //---------------------------------------------------------------------------------------------
-    template<class NODE_TABLE_PRIVATE, ETableColumnType TYPE, OP_TYPE OPTYPE, typename F>
-    Ptr<ASTOp> CodeGenerator::GenerateTableSwitch( const NODE_TABLE_PRIVATE& node, F&& GenerateOption )
+    template<class NODE_TABLE, ETableColumnType TYPE, OP_TYPE OPTYPE, typename F>
+    Ptr<ASTOp> CodeGenerator::GenerateTableSwitch( const NODE_TABLE& node, F&& GenerateOption )
     {
         Ptr<const Table> NodeTable = node.Table;
         Ptr<ASTOp> Variable;
@@ -784,7 +784,7 @@ namespace mu
         if ( !Variable)
         {
             // Create the table variable expression
-            Variable = GenerateTableVariable( node.m_pNode, CacheKey, node.bNoneOption, node.DefaultRowName);
+            Variable = GenerateTableVariable( &node, CacheKey, node.bNoneOption, node.DefaultRowName);
 
             GeneratedTables.Add(CacheKey, Variable );
         }
@@ -796,18 +796,18 @@ namespace mu
 
 		if (NumRows == 0)
 		{
-			m_pErrorLog->GetPrivate()->Add("The table has no rows.", ELMT_ERROR, node.m_pNode->GetMessageContext());
+			m_pErrorLog->GetPrivate()->Add("The table has no rows.", ELMT_ERROR, node.GetMessageContext());
 			return nullptr;
 		}
         else if (ColIndex < 0)
         {
-            m_pErrorLog->GetPrivate()->Add("Table column not found.", ELMT_ERROR, node.m_pNode->GetMessageContext());
+            m_pErrorLog->GetPrivate()->Add("Table column not found.", ELMT_ERROR, node.GetMessageContext());
             return nullptr;
         }
 
         if (NodeTable->GetPrivate()->Columns[ ColIndex ].Type != TYPE )
         {
-            m_pErrorLog->GetPrivate()->Add("Table column type is not the right type.", ELMT_ERROR, node.m_pNode->GetMessageContext());
+            m_pErrorLog->GetPrivate()->Add("Table column type is not the right type.", ELMT_ERROR, node.GetMessageContext());
             return nullptr;
         }
 
diff --git a/Engine/Plugins/Experimental/Mutable/Source/MutableTools/Private/MuT/CodeGenerator_Color.cpp b/Engine/Plugins/Experimental/Mutable/Source/MutableTools/Private/MuT/CodeGenerator_Color.cpp
--- a/Engine/Plugins/Experimental/Mutable/Source/MutableTools/Private/MuT/CodeGenerator_Color.cpp
+++ b/Engine/Plugins/Experimental/Mutable/Source/MutableTools/Private/MuT/CodeGenerator_Color.cpp
@@ -35,7 +35,6 @@
 #include "MuT/NodeColourSwitch.h"
 #include "MuT/NodeColourSwitchPrivate.h"
 #include "MuT/NodeColourTable.h"
-#include "MuT/NodeColourTablePrivate.h"
 #include "MuT/NodeColourVariation.h"
 #include "MuT/NodeColourVariationPrivate.h"
 #include "MuT/NodeImage.h"
@@ -495,10 +494,10 @@ namespace mu
 	//-------------------------------------------------------------------------------------------------
 	void CodeGenerator::GenerateColor_Table(FColorGenerationResult& result, const FGenericGenerationOptions& Options, const Ptr<const NodeColourTable>& Typed)
 	{
-		const NodeColourTable::Private& node = *Typed->GetPrivate();
+		const NodeColourTable& node = *Typed;
 
-		result.op = GenerateTableSwitch<NodeColourTable::Private, ETableColumnType::Color, OP_TYPE::CO_SWITCH>(node,
-			[this, &Options](const NodeColourTable::Private& node, int colIndex, int row, ErrorLog* pErrorLog)
+		result.op = GenerateTableSwitch<NodeColourTable, ETableColumnType::Color, OP_TYPE::CO_SWITCH>(node,
+			[this, &Options](const NodeColourTable& node, int colIndex, int row, ErrorLog* pErrorLog)
 			{
 				Ptr<NodeColourConstant> CellData = new NodeColourConstant();
 				FVector4f Colour = node.Table->GetPrivate()->Rows[row].Values[colIndex].Color;
diff --git a/Engine/Plugins/Experimental/Mutable/Source/MutableTools/Private/MuT/CodeGenerator_Image.cpp b/Engine/Plugins/Experimental/Mutable/Source/MutableTools/Private/MuT/CodeGenerator_Image.cpp
--- a/Engine/Plugins/Experimental/Mutable/Source/MutableTools/Private/MuT/CodeGenerator_Image.cpp
+++ b/Engine/Plugins/Experimental/Mutable/Source/MutableTools/Private/MuT/CodeGenerator_Image.cpp
@@ -55,7 +55,6 @@
 #include "MuT/NodeImageConditional.h"
 #include "MuT/NodeImageConditionalPrivate.h"
 #include "MuT/NodeImageConstant.h"
-#include "MuT/NodeImageConstantPrivate.h"
 #include "MuT/NodeImageFormat.h"
 #include "MuT/NodeImageFormatPrivate.h"
 #include "MuT/NodeImageGradient.h"
@@ -91,7 +90,6 @@
 #include "MuT/NodeImageSwizzle.h"
 #include "MuT/NodeImageSwizzlePrivate.h"
 #include "MuT/NodeImageTable.h"
-#include "MuT/NodeImageTablePrivate.h"
 #include "MuT/NodeImageTransform.h"
 #include "MuT/NodeImageTransformPrivate.h"
 #include "MuT/NodeImageVariation.h"
@@ -171,13 +169,13 @@ namespace mu
     //---------------------------------------------------------------------------------------------
     void CodeGenerator::GenerateImage_Constant(const FImageGenerationOptions& Options, FImageGenerationResult& Result, const NodeImageConstant* InNode)
     {
-		const NodeImageConstant::Private& node = *InNode->GetPrivate();
+		const NodeImageConstant& node = *InNode;
 		
         // TODO: check duplicates
         Ptr<const Image> pImage;
-		if (node.m_pProxy)
+		if (node.Proxy)
 		{
-			pImage = node.m_pProxy->Get();
+			pImage = node.Proxy->Get();
 		}
 
         if (!pImage)
@@ -196,6 +194,7 @@ namespace mu
 			ReferenceOp->type = OP_TYPE::IM_REFERENCE;
 			ReferenceOp->ID = pImage->GetReferencedTexture();
 			ReferenceOp->bForceLoad = pImage->IsForceLoad();
+			ReferenceOp->SourceDataDescriptor = InNode->SourceDataDescriptor;
 
 			// Don't store the format. Format can vary between loaded constant image and reference and cause
 			// code optimization bugs.
@@ -211,6 +210,7 @@ namespace mu
 			Ptr<ASTOpConstantResource> op = new ASTOpConstantResource();
 			op->Type = OP_TYPE::IM_CONSTANT;
 			op->SetValue(pImage, m_compilerOptions->OptimisationOptions.DiskCacheContext);
+			op->SourceDataDescriptor = InNode->SourceDataDescriptor;
 			Result.op = op;
 		}
 
@@ -1846,13 +1846,13 @@ namespace mu
 	//---------------------------------------------------------------------------------------------
 	void CodeGenerator::GenerateImage_Table(const FImageGenerationOptions& Options, FImageGenerationResult& Result, const NodeImageTable* InNode)
 	{
-		const NodeImageTable::Private& node = *InNode->GetPrivate();
+		const NodeImageTable& node = *InNode;
 
-		Result.op = GenerateTableSwitch<NodeImageTable::Private, ETableColumnType::Image, OP_TYPE::IM_SWITCH>(node,
-			[this, InNode, Options](const NodeImageTable::Private& node, int colIndex, int row, ErrorLog* pErrorLog)
+		Result.op = GenerateTableSwitch<NodeImageTable, ETableColumnType::Image, OP_TYPE::IM_SWITCH>(node,
+			[this, InNode, Options](const NodeImageTable& node, int colIndex, int row, ErrorLog* pErrorLog)
 			{
 				const FTableValue& CellData = node.Table->GetPrivate()->Rows[row].Values[colIndex];
-				ImagePtrConst pImage = nullptr;
+				Ptr<const Image> pImage = nullptr;
 
 				if (Ptr<ResourceProxy<Image>> pProxyImage = CellData.ProxyImage)
 				{
@@ -1870,9 +1870,12 @@ namespace mu
 				}
 				else
 				{
-					NodeImageConstantPtr ImageConst = new NodeImageConstant();
+					Ptr<NodeImageConstant> ImageConst = new NodeImageConstant();
 					ImageConst->SetValue(pImage.get());
 
+					// TODO: We probably want to get the data tags from the table row.
+					ImageConst->SourceDataDescriptor = InNode->SourceDataDescriptor;
+
 					FImageGenerationResult Result;
 					GenerateImage(Options, Result, ImageConst);
 					ImageOp = Result.op;
@@ -1996,7 +1999,7 @@ namespace mu
 		// Make a checkered debug image
 		Ptr<Image> GeneratedImage = GenerateMissingImage( format );
 
-		NodeImageConstantPtr pNode = new NodeImageConstant();
+		Ptr<NodeImageConstant> pNode = new NodeImageConstant();
 		pNode->SetValue(GeneratedImage.get());
 
 		FImageGenerationResult Result;
diff --git a/Engine/Plugins/Experimental/Mutable/Source/MutableTools/Private/MuT/CodeGenerator_Mesh.cpp b/Engine/Plugins/Experimental/Mutable/Source/MutableTools/Private/MuT/CodeGenerator_Mesh.cpp
--- a/Engine/Plugins/Experimental/Mutable/Source/MutableTools/Private/MuT/CodeGenerator_Mesh.cpp
+++ b/Engine/Plugins/Experimental/Mutable/Source/MutableTools/Private/MuT/CodeGenerator_Mesh.cpp
@@ -69,7 +69,6 @@
 #include "MuT/NodeMeshSwitch.h"
 #include "MuT/NodeMeshSwitchPrivate.h"
 #include "MuT/NodeMeshTable.h"
-#include "MuT/NodeMeshTablePrivate.h"
 #include "MuT/NodeMeshTransform.h"
 #include "MuT/NodeMeshTransformPrivate.h"
 #include "MuT/NodeMeshVariation.h"
@@ -1139,8 +1138,8 @@ namespace mu
 		int t = 0;
 		bool bFirstRowGenerated = false;
 
-		Ptr<ASTOp> Op = GenerateTableSwitch<NodeMeshTable::Private, ETableColumnType::Mesh, OP_TYPE::ME_SWITCH>(*TableNode->GetPrivate(),
-			[this, &NewResult, &bFirstRowGenerated, &InOptions] (const NodeMeshTable::Private& node, int colIndex, int row, ErrorLog* pErrorLog)
+		Ptr<ASTOp> Op = GenerateTableSwitch<NodeMeshTable, ETableColumnType::Mesh, OP_TYPE::ME_SWITCH>(*TableNode,
+			[this, &NewResult, &bFirstRowGenerated, &InOptions] (const NodeMeshTable& node, int colIndex, int row, ErrorLog* pErrorLog)
 			{
 				mu::Ptr<mu::Mesh> pMesh = node.Table->GetPrivate()->Rows[row].Values[colIndex].Mesh;
 				FMeshGenerationResult BranchResults;
diff --git a/Engine/Plugins/Experimental/Mutable/Source/MutableTools/Private/MuT/CodeGenerator_Scalar.cpp b/Engine/Plugins/Experimental/Mutable/Source/MutableTools/Private/MuT/CodeGenerator_Scalar.cpp
--- a/Engine/Plugins/Experimental/Mutable/Source/MutableTools/Private/MuT/CodeGenerator_Scalar.cpp
+++ b/Engine/Plugins/Experimental/Mutable/Source/MutableTools/Private/MuT/CodeGenerator_Scalar.cpp
@@ -38,7 +38,6 @@
 #include "MuT/NodeScalarSwitch.h"
 #include "MuT/NodeScalarSwitchPrivate.h"
 #include "MuT/NodeScalarTable.h"
-#include "MuT/NodeScalarTablePrivate.h"
 #include "MuT/NodeScalarVariation.h"
 #include "MuT/NodeScalarVariationPrivate.h"
 #include "MuT/Table.h"
@@ -409,10 +408,10 @@ namespace mu
 	//-------------------------------------------------------------------------------------------------
 	void CodeGenerator::GenerateScalar_Table(FScalarGenerationResult& result, const FGenericGenerationOptions& Options, const Ptr<const NodeScalarTable>& Typed)
 	{
-		const NodeScalarTable::Private& node = *Typed->GetPrivate();
+		const NodeScalarTable& node = *Typed;
 
-		Ptr<ASTOp> Op = GenerateTableSwitch<NodeScalarTable::Private, ETableColumnType::Scalar, OP_TYPE::SC_SWITCH>(node,
-			[this,&Options](const NodeScalarTable::Private& node, int colIndex, int row, ErrorLog* pErrorLog)
+		Ptr<ASTOp> Op = GenerateTableSwitch<NodeScalarTable, ETableColumnType::Scalar, OP_TYPE::SC_SWITCH>(node,
+			[this,&Options](const NodeScalarTable& node, int colIndex, int row, ErrorLog* pErrorLog)
 			{
 				NodeScalarConstantPtr pCell = new NodeScalarConstant();
 				float scalar = node.Table->GetPrivate()->Rows[row].Values[colIndex].Scalar;
diff --git a/Engine/Plugins/Experimental/Mutable/Source/MutableTools/Private/MuT/CodeOptimiser.cpp b/Engine/Plugins/Experimental/Mutable/Source/MutableTools/Private/MuT/CodeOptimiser.cpp
--- a/Engine/Plugins/Experimental/Mutable/Source/MutableTools/Private/MuT/CodeOptimiser.cpp
+++ b/Engine/Plugins/Experimental/Mutable/Source/MutableTools/Private/MuT/CodeOptimiser.cpp
@@ -471,7 +471,7 @@ namespace mu
 			// This runs in a worker thread
 
 			OP_TYPE type = SourceCloned->GetOpType();
-			DATATYPE dtype = GetOpDataType(type);
+			DATATYPE DataType = GetOpDataType(type);
 
 			Ptr<Settings> pSettings = new Settings;
 			pSettings->SetProfile( false );
@@ -480,6 +480,13 @@ namespace mu
 
 			pSystem->GetPrivate()->ImagePixelFormatOverride = ImOp.FormatImageOverride;
 
+			FSourceDataDescriptor SourceDataDescriptor;
+			if (DataType == DT_IMAGE)
+			{
+				SourceDataDescriptor = SourceCloned->GetSourceDataDescriptor();
+				check(!SourceDataDescriptor.IsInvalid());
+			}
+
 			// Don't generate mips during linking here.
 			FLinkerOptions LinkerOptions(ImOp);
 			LinkerOptions.MinTextureResidentMipCount = 255;
@@ -496,7 +503,7 @@ namespace mu
 			pSystem->GetPrivate()->BeginBuild( model );
 
 			// Calculate the value and replace this op by a constant
-			switch( dtype )
+			switch( DataType )
 			{
 			case DT_MESH:
 			{
@@ -522,10 +529,11 @@ namespace mu
 
 				if (pImage)
 				{
-					mu::Ptr<ASTOpConstantResource> constantOp = new ASTOpConstantResource();
-					constantOp->Type = OP_TYPE::IM_CONSTANT;
-					constantOp->SetValue( pImage, DiskCacheContext );
-					Result = constantOp;
+					mu::Ptr<ASTOpConstantResource> ConstantOp = new ASTOpConstantResource();
+					ConstantOp->SourceDataDescriptor = SourceDataDescriptor;
+					ConstantOp->Type = OP_TYPE::IM_CONSTANT;
+					ConstantOp->SetValue( pImage, DiskCacheContext );
+					Result = ConstantOp;
 				}
 				break;
 			}
@@ -709,8 +717,8 @@ namespace mu
 					SubgraphRoot->bHasSpecialOpInSubgraph = bHasSpecialOpInSubgraph;
 
 					bool bIsDataTypeThanCanTurnIntoConst = false;
-					DATATYPE dtype = GetOpDataType(SubgraphType);
-					switch (dtype)
+					DATATYPE DataType = GetOpDataType(SubgraphType);
+					switch (DataType)
 					{
 					case DT_MESH:
 					case DT_IMAGE:
@@ -723,15 +731,32 @@ namespace mu
 						break;
 					}
 
+
 					// See if it is worth generating this as constant
 					// ---------------------------------------------
-					if (SubgraphRoot->bIsConstantSubgraph
+					bool bWorthGenerating = SubgraphRoot->bIsConstantSubgraph
 						&& !SubgraphRoot->bHasSpecialOpInSubgraph
 						&& !SubgraphRoot->IsConstantOp()
-						&& bIsDataTypeThanCanTurnIntoConst
-						)
+						&& bIsDataTypeThanCanTurnIntoConst;
+
+					if (bWorthGenerating)						 
 					{
-						ConstantSubgraphs.Add({ SubgraphRoot, UE::Tasks::FTaskEvent(TEXT("MutableConstantSubgraph")) });
+						bool bCanBeGenerated = true;
+
+						// Check source data incompatiblities: when generating constants don't mix data that has different source descriptors (tags and other properties).
+						if (DataType == DT_IMAGE)
+						{
+							FSourceDataDescriptor SourceDescriptor = SubgraphRoot->GetSourceDataDescriptor();
+							if (SourceDescriptor.IsInvalid())
+							{
+								bCanBeGenerated = false;
+							}
+						}
+
+						if (bCanBeGenerated)
+						{
+							ConstantSubgraphs.Add({ SubgraphRoot, UE::Tasks::FTaskEvent(TEXT("MutableConstantSubgraph")) });
+						}
 					}
 				});
 		}
@@ -808,6 +833,7 @@ namespace mu
 									MUTABLE_CPUPROFILER_SCOPE(MutableResolveComplete);
 
 									Ptr<ASTOpConstantResource> ConstantOp = new ASTOpConstantResource;
+									ConstantOp->SourceDataDescriptor = SubgraphRoot->GetSourceDataDescriptor();
 									ConstantOp->Type = OP_TYPE::IM_CONSTANT;
 									ConstantOp->SetValue(ResolveImage->get(), InOptions->OptimisationOptions.DiskCacheContext);
 
@@ -924,6 +950,7 @@ namespace mu
 								UE::Tasks::FTask ReferenceCompletion = InOptions->OptimisationOptions.ReferencedResourceProvider(ImageID, ResolveImage, bRunImmediatlyIfPossible);
 								ReferenceCompletion.Wait();
 								Ptr<ASTOpConstantResource> ConstantOp = new ASTOpConstantResource;
+								ConstantOp->SourceDataDescriptor = SubgraphRoot->GetSourceDataDescriptor();
 								ConstantOp->Type = OP_TYPE::IM_CONSTANT;
 								ConstantOp->SetValue(ResolveImage->get(), InOptions->OptimisationOptions.DiskCacheContext);
 								ASTOp::Replace(SubgraphRoot, ConstantOp);
diff --git a/Engine/Plugins/Experimental/Mutable/Source/MutableTools/Private/MuT/Compiler.cpp b/Engine/Plugins/Experimental/Mutable/Source/MutableTools/Private/MuT/Compiler.cpp
--- a/Engine/Plugins/Experimental/Mutable/Source/MutableTools/Private/MuT/Compiler.cpp
+++ b/Engine/Plugins/Experimental/Mutable/Source/MutableTools/Private/MuT/Compiler.cpp
@@ -3,6 +3,16 @@
 
 #include "MuT/Compiler.h"
 
+#include "MuT/AST.h"
+#include "MuT/ASTOpParameter.h"
+#include "MuT/CodeGenerator.h"
+#include "MuT/CodeOptimiser.h"
+#include "MuT/CompilerPrivate.h"
+#include "MuT/ErrorLog.h"
+#include "MuT/ErrorLogPrivate.h"
+#include "MuT/Node.h"
+#include "MuT/NodePrivate.h"
+#include "MuT/Table.h"
 #include "MuR/Image.h"
 #include "MuR/Mesh.h"
 #include "MuR/Model.h"
@@ -14,16 +24,6 @@
 #include "MuR/Serialisation.h"
 #include "MuR/MutableRuntimeModule.h"
 #include "MuR/System.h"
-#include "MuT/AST.h"
-#include "MuT/ASTOpParameter.h"
-#include "MuT/CodeGenerator.h"
-#include "MuT/CodeOptimiser.h"
-#include "MuT/CompilerPrivate.h"
-#include "MuT/ErrorLog.h"
-#include "MuT/ErrorLogPrivate.h"
-#include "MuT/Node.h"
-#include "MuT/NodePrivate.h"
-#include "MuT/Table.h"
 
 #include "Trace/Detail/Channel.h"
 #include "Containers/Array.h"
@@ -33,6 +33,7 @@
 #include "Hash/CityHash.h"
 #include "Logging/LogCategory.h"
 #include "Logging/LogMacros.h"
+#include "Misc/EnumClassFlags.h"
 #include "Misc/AssertionMacros.h"
 #include "Async/ParallelFor.h"
 
@@ -283,16 +284,16 @@ namespace mu
             optimiser.OptimiseAST( );
         }
 
-        // Link the program and generate state data.
+        // Link the Program and generate state data.
 		TSharedPtr<Model> pResult = MakeShared<Model>();
-        FProgram& program = pResult->GetPrivate()->m_program;
+        FProgram& Program = pResult->GetPrivate()->m_program;
 
-		check(program.m_parameters.IsEmpty());
-		program.m_parameters = Parameters;
+		check(Program.m_parameters.IsEmpty());
+		Program.m_parameters = Parameters;
 
 		// Preallocate ample memory
-		program.m_byteCode.Reserve(16 * 1024 * 1024);
-		program.m_opAddress.Reserve(1024 * 1024);
+		Program.m_byteCode.Reserve(16 * 1024 * 1024);
+		Program.m_opAddress.Reserve(1024 * 1024);
 
 		// Keep the link options outside the scope because it is also used to cache constant data that has already been 
 		// added and could be reused across states.
@@ -305,7 +306,7 @@ namespace mu
 
             if (s.root)
             {
-				s.state.m_root = ASTOp::FullLink(s.root, program, &LinkerOptions);
+				s.state.m_root = ASTOp::FullLink(s.root, Program, &LinkerOptions);
             }
             else
             {
@@ -313,8 +314,8 @@ namespace mu
             }
         }
 
-		program.m_byteCode.Shrink();
-		program.m_opAddress.Shrink();
+		Program.m_byteCode.Shrink();
+		Program.m_opAddress.Shrink();
 
         // Set the runtime parameter indices.
         for(FStateCompilationData& s: states )
@@ -322,9 +323,9 @@ namespace mu
             for ( int32 p=0; p<s.nodeState.RuntimeParams.Num(); ++p )
             {
                 int32 paramIndex = -1;
-                for ( int32 i=0; paramIndex<0 && i<program.m_parameters.Num(); ++i )
+                for ( int32 i=0; paramIndex<0 && i<Program.m_parameters.Num(); ++i )
                 {
-                    if ( program.m_parameters[i].m_name
+                    if ( Program.m_parameters[i].m_name
                          ==
                          s.nodeState.RuntimeParams[p] )
                     {
@@ -364,9 +365,9 @@ namespace mu
                 {
                     // Find the index in the model parameter list
                     int paramIndex = -1;
-                    for ( int32 i=0; paramIndex<0 && i<program.m_parameters.Num(); ++i )
+                    for ( int32 i=0; paramIndex<0 && i<Program.m_parameters.Num(); ++i )
                     {
-                        if ( program.m_parameters[i].m_name == b )
+                        if ( Program.m_parameters[i].m_name == b )
                         {
                             paramIndex = (int)i;
                         }
@@ -395,20 +396,19 @@ namespace mu
             // Sort for performance and determinism
             s.state.m_dynamicResources.Sort();
 
-            program.m_states.Add(s.state);
+            Program.m_states.Add(s.state);
         }
 
-		UE_LOG(LogMutableCore, Verbose, TEXT("(int) %s : %ld"), TEXT("program size"), int64(program.m_opAddress.Num()));
+		UE_LOG(LogMutableCore, Verbose, TEXT("(int) %s : %ld"), TEXT("Program size"), int64(Program.m_opAddress.Num()));
 
         // Merge the log in the right order
         genErrorLog->Merge( m_pD->m_pErrorLog.get() );
         m_pD->m_pErrorLog = genErrorLog.get();
 
 		// Pack data
-		uint64 EmbeddedDataFileBytesLimit = m_pD->m_options->GetPrivate()->EmbeddedDataBytesLimit;
-		m_pD->GenerateRoms(pResult.Get(), EmbeddedDataFileBytesLimit);
+		m_pD->GenerateRoms(pResult.Get(), m_pD->m_options.get(), LinkerOptions.AdditionalData );
 
-		UE_LOG(LogMutableCore, Verbose, TEXT("(int) %s : %ld"), TEXT("program size"), int64(program.m_opAddress.Num()));
+		UE_LOG(LogMutableCore, Verbose, TEXT("(int) %s : %ld"), TEXT("Program size"), int64(Program.m_opAddress.Num()));
 
         return pResult;
     }
@@ -445,41 +445,52 @@ namespace mu
 
 
 	//---------------------------------------------------------------------------------------------
-	void Compiler::Private::GenerateRoms(Model* p, int32 EmbeddedDataBytesLimit)
+	void Compiler::Private::GenerateRoms(Model* p, const CompilerOptions* Options, const FLinkerOptions::FAdditionalData& AdditionalData )
 	{
 		LLM_SCOPE_BYNAME(TEXT("MutableRuntime"));
 		MUTABLE_CPUPROFILER_SCOPE(GenerateRoms);
 
-		mu::FProgram& program = p->GetPrivate()->m_program;
+		uint64 EmbeddedDataBytesLimit = Options->GetPrivate()->EmbeddedDataBytesLimit;
+
+		mu::FProgram& Program = p->GetPrivate()->m_program;
 
+		// These are used for logging only.
 		int32 NumRoms = 0;
+		uint64 NumRomsBytes = 0;
 		int32 NumEmbedded = 0;
+		uint64 NumEmbeddedBytes = 0;
+		int32 NumHighRes = 0;
+		uint64 NumHighResBytes = 0;
+
+		// Maximum number of roms
+		int32 MaxRomCount = Program.ConstantImageLODs.Num() + Program.ConstantMeshes.Num();
+		Program.m_roms.Reserve(MaxRomCount);
 
 		TSet<uint32> UsedIds;
-		UsedIds.Reserve(program.ConstantImageLODs.Num()+ program.ConstantMeshes.Num());
+		UsedIds.Reserve(MaxRomCount);
 
 		TArray<FRomData> RomDatas;
-		RomDatas.SetNumZeroed(FMath::Max(program.ConstantImageLODs.Num(), program.ConstantMeshes.Num()));
+		RomDatas.SetNumZeroed(FMath::Max(Program.ConstantImageLODs.Num(), Program.ConstantMeshes.Num()));
 
-		// Save images and unload from memory
+		// Images
 		{
 			MUTABLE_CPUPROFILER_SCOPE(GenerateRoms_ImageIds);
 
-			ParallelFor(program.ConstantImageLODs.Num(),
-				[EmbeddedDataBytesLimit, &program, &RomDatas](uint32 ResourceIndex)
+			ParallelFor(Program.ConstantImageLODs.Num(),
+				[EmbeddedDataBytesLimit, &Program, &RomDatas](uint32 ResourceIndex)
 				{
-					TPair<int32, Ptr<const Image>>& ResData = program.ConstantImageLODs[ResourceIndex];
+					TPair<int32, Ptr<const Image>>& ResData = Program.ConstantImageLODs[ResourceIndex];
 
 					// This shouldn't have been serialised with rom support before.
 					check(ResData.Key < 0);
 
-					// Serialize to memory, to find out final size of this rom
-					OutputMemoryStream MemStream(1024 * 1024);
-					OutputArchive MemoryArch(&MemStream);
+					// Serialize to find out final size of this rom
+					OutputSizeStream SizeStream;
+					OutputArchive MemoryArch(&SizeStream);
 					Image::Serialise(ResData.Value.get(), MemoryArch);
 
 					// If the resource uses less memory than the threshold, don't save it in a separate rom.
-					if (MemStream.GetBufferSize() <= EmbeddedDataBytesLimit)
+					if (SizeStream.GetBufferSize() <= EmbeddedDataBytesLimit)
 					{
 						return;
 					}
@@ -487,17 +498,59 @@ namespace mu
 					FRomData& RomData = RomDatas[ResourceIndex];
 					RomData.ResourceType = DT_IMAGE;
 					RomData.ResourceIndex = ResourceIndex;
-					RomData.Size = MemStream.GetBufferSize();
+					RomData.Size = SizeStream.GetBufferSize();
+					RomData.Flags = ERomFlags::None;
 
-					uint64 FullHash = CityHash64(static_cast<const char*>(MemStream.GetBuffer()), MemStream.GetBufferSize());
+					const Image* Resource = ResData.Value.get();
+					uint64 FullHash = CityHash64(reinterpret_cast<const char*>(Resource->GetLODData(0)), Resource->GetLODDataSize(0));
 					RomData.Id = GetTypeHash(FullHash);
 				});
 		}
 
+		// Generate the high-res flags for images
+		{
+			// Initially all are high-res because if at least one reference to a mip is not we will set it to not-high-res
+			TArray<bool> IsLODHighRes;
+			IsLODHighRes.Init(true, Program.ConstantImageLODs.Num());
+
+			for (int32 ImageIndex = 0; ImageIndex < Program.ConstantImages.Num(); ++ImageIndex)
+			{
+				const FImageLODRange& LODRange = Program.ConstantImages[ImageIndex];
+
+				int32 NumHighResMips = FMath::Max(0,AdditionalData.SourceImagePerConstant[ImageIndex].SourceHighResMips);
+
+				for (int32 LODRangeIndex = NumHighResMips; LODRangeIndex < LODRange.LODCount; ++LODRangeIndex)
+				{
+					int32 LODIndex = Program.ConstantImageLODIndices[LODRange.FirstIndex+LODRangeIndex];
+					IsLODHighRes[LODIndex] = false;
+				}
+
+				// Moreover, at least one mip of each image has to be non-highres
+				if (LODRange.LODCount>0)
+				{
+					int32 LastLODIndex = Program.ConstantImageLODIndices[LODRange.FirstIndex + LODRange.LODCount - 1];
+					IsLODHighRes[LastLODIndex] = false;
+				}
+			}
+
+			for (int32 ResourceIndex = 0; ResourceIndex < Program.ConstantImageLODs.Num(); ++ResourceIndex)
+			{
+				// If this mip represents a high-quality mip, flag the rom as such
+				if (IsLODHighRes[ResourceIndex])
+				{
+					FRomData& RomData = RomDatas[ResourceIndex];
+					EnumAddFlags(RomData.Flags, ERomFlags::HighRes);
+					
+					++NumHighRes;
+					NumHighResBytes += RomData.Size;
+				}
+			}
+		}
+
 		{
 			MUTABLE_CPUPROFILER_SCOPE(GenerateRoms_ImageIdsUnique);
 
-			for (int32 ResourceIndex = 0; ResourceIndex < program.ConstantImageLODs.Num(); ++ResourceIndex)
+			for (int32 ResourceIndex = 0; ResourceIndex < Program.ConstantImageLODs.Num(); ++ResourceIndex)
 			{
 				FRomData& RomData = RomDatas[ResourceIndex];
 
@@ -505,48 +558,54 @@ namespace mu
 				if (int32(RomData.Size) <= EmbeddedDataBytesLimit)
 				{
 					NumEmbedded++;
+					NumEmbeddedBytes += RomData.Size;
 					continue;
 				}
 				NumRoms++;
+				NumRomsBytes += RomData.Size;
 
 				// Ensure that the Id is unique
 				EnsureUniqueRomId(UsedIds, RomData.Id);
 
-				int32 RomIndex = program.m_roms.Add(RomData);
+				int32 RomIndex = Program.m_roms.Add(RomData);
 
-				TPair<int32, Ptr<const Image>>& ResData = program.ConstantImageLODs[ResourceIndex];
+				TPair<int32, Ptr<const Image>>& ResData = Program.ConstantImageLODs[ResourceIndex];
 				ResData.Key = RomIndex;
 			}
 		}
 
-		// Save meshes and unload from memory
+		// Meshes
 		FMemory::Memzero( RomDatas.GetData(), RomDatas.GetAllocatedSize() );
 		{
 			MUTABLE_CPUPROFILER_SCOPE(GenerateRoms_MeshIds);
 
-			ParallelFor(program.ConstantMeshes.Num(),
-				[EmbeddedDataBytesLimit, &program, &RomDatas](uint32 ResourceIndex)
+			ParallelFor(Program.ConstantMeshes.Num(),
+				[EmbeddedDataBytesLimit, &Program, &RomDatas](uint32 ResourceIndex)
 				{
-					TPair<int32, Ptr<const Mesh>>& ResData = program.ConstantMeshes[ResourceIndex];
+					TPair<int32, Ptr<const Mesh>>& ResData = Program.ConstantMeshes[ResourceIndex];
 
 					// This shouldn't have been serialised with rom support before.
 					check(ResData.Key < 0);
 
 					// Serialize to memory, to find out final size of this rom
-					OutputMemoryStream MemStream(1024 * 1024);
-					OutputArchive MemoryArch(&MemStream);
-					Mesh::Serialise(ResData.Value.get(), MemoryArch);
+					const Mesh* Resource = ResData.Value.get();
+					int32 ApproximateSize = Resource->GetDataSize();
 
 					// If the resource uses less memory than the threshold, don't save it in a separate rom.
-					if (MemStream.GetBufferSize() <= EmbeddedDataBytesLimit)
+					if (ApproximateSize <= EmbeddedDataBytesLimit)
 					{
 						return;
 					}
 
+					OutputMemoryStream MemStream(ApproximateSize + 64 * 1024);
+					OutputArchive MemoryArch(&MemStream);
+					Mesh::Serialise(ResData.Value.get(), MemoryArch);
+
 					FRomData& RomData = RomDatas[ResourceIndex];
 					RomData.ResourceType = DT_MESH;
 					RomData.ResourceIndex = ResourceIndex;
 					RomData.Size = MemStream.GetBufferSize();
+					RomData.Flags = ERomFlags::None;
 
 					// Ensure that the Id is unique
 					uint64 FullHash = CityHash64(static_cast<const char*>(MemStream.GetBuffer()), MemStream.GetBufferSize());
@@ -557,7 +616,7 @@ namespace mu
 		{
 			MUTABLE_CPUPROFILER_SCOPE(GenerateRoms_MeshIdsUnique);
 
-			for (int32 ResourceIndex = 0; ResourceIndex < program.ConstantMeshes.Num(); ++ResourceIndex)
+			for (int32 ResourceIndex = 0; ResourceIndex < Program.ConstantMeshes.Num(); ++ResourceIndex)
 			{
 				FRomData& RomData = RomDatas[ResourceIndex];
 
@@ -565,21 +624,24 @@ namespace mu
 				if (int32(RomData.Size) <= EmbeddedDataBytesLimit)
 				{
 					NumEmbedded++;
+					NumEmbeddedBytes += RomData.Size;
 					continue;
 				}
 				NumRoms++;
+				NumRomsBytes += RomData.Size;
 
 				// Ensure that the Id is unique
 				EnsureUniqueRomId(UsedIds, RomData.Id);
 
-				int32 RomIndex = program.m_roms.Add(RomData);
+				int32 RomIndex = Program.m_roms.Add(RomData);
 
-				TPair<int32, Ptr<const Mesh>>& ResData = program.ConstantMeshes[ResourceIndex];
+				TPair<int32, Ptr<const Mesh>>& ResData = Program.ConstantMeshes[ResourceIndex];
 				ResData.Key = RomIndex;
 			}
 		}
 
-		UE_LOG(LogMutableCore, Log, TEXT("Generated roms for model with %d embedded constants and %d in roms."), NumEmbedded, NumRoms);
+		UE_LOG(LogMutableCore, Log, TEXT("Generated roms: %d (%d KB) are embedded, %d (%d KB) are streamed of which %d (%d KB) are high-res."), 
+			NumEmbedded, uint32(NumEmbeddedBytes/1024), NumRoms, uint32(NumRomsBytes/1024), NumHighRes, uint32(NumHighResBytes/1024));
 	}
 
 
diff --git a/Engine/Plugins/Experimental/Mutable/Source/MutableTools/Private/MuT/CompilerPrivate.h b/Engine/Plugins/Experimental/Mutable/Source/MutableTools/Private/MuT/CompilerPrivate.h
--- a/Engine/Plugins/Experimental/Mutable/Source/MutableTools/Private/MuT/CompilerPrivate.h
+++ b/Engine/Plugins/Experimental/Mutable/Source/MutableTools/Private/MuT/CompilerPrivate.h
@@ -93,12 +93,20 @@ namespace mu
 
         Ptr<ErrorLog> m_pErrorLog;
 
-        //! Detailed options
+        /** */
         Ptr<CompilerOptions> m_options;
 
+		/** Store for additional data generated during compilation, but not necessary for the runtime. */
+		struct FAdditionalData
+		{
+			/** Source data descriptor for every image constant that has been generated. 
+			* It must have the same size than the Program::ConstantImages array.
+			*/
+			TArray<FSourceDataDescriptor> SourceImagePerConstant;
+		};
 
 		//!
-		void GenerateRoms(Model* p, int32 MinRomSize);
+		void GenerateRoms(Model*, const CompilerOptions*, const FLinkerOptions::FAdditionalData&);
 
     };
 
diff --git a/Engine/Plugins/Experimental/Mutable/Source/MutableTools/Private/MuT/Node.cpp b/Engine/Plugins/Experimental/Mutable/Source/MutableTools/Private/MuT/Node.cpp
--- a/Engine/Plugins/Experimental/Mutable/Source/MutableTools/Private/MuT/Node.cpp
+++ b/Engine/Plugins/Experimental/Mutable/Source/MutableTools/Private/MuT/Node.cpp
@@ -10,6 +10,7 @@
 #include "MuT/NodeLOD.h"
 #include "MuT/NodeExtensionData.h"
 #include "MuT/NodeExtensionDataConstant.h"
+#include "MuT/NodeImageTable.h"
 
 namespace mu
 {
@@ -17,9 +18,10 @@ namespace mu
 	FNodeType NodeComponent::StaticType = FNodeType(Node::EType::Component, Node::GetStaticType());
 	FNodeType NodeComponentNew::StaticType = FNodeType(Node::EType::ComponentNew, NodeComponent::GetStaticType());
 	FNodeType NodeLOD::StaticType = FNodeType(Node::EType::LOD, Node::GetStaticType());
-
 	FNodeType NodeExtensionData::StaticType = FNodeType(Node::EType::ExtensionData, Node::GetStaticType());
 	FNodeType NodeExtensionDataConstant::StaticType = FNodeType(Node::EType::ExtensionDataConstant, NodeExtensionData::GetStaticType());
+	FNodeType NodeImageTable::StaticType = FNodeType(Node::EType::ImageTable, NodeImage::GetStaticType());
+
 
 	// Static initialisation
 	static FNodeType s_nodeType = FNodeType(Node::EType::Node, nullptr );
diff --git a/Engine/Plugins/Experimental/Mutable/Source/MutableTools/Private/MuT/NodeColourTable.cpp b/Engine/Plugins/Experimental/Mutable/Source/MutableTools/Private/MuT/NodeColourTable.cpp
--- a/Engine/Plugins/Experimental/Mutable/Source/MutableTools/Private/MuT/NodeColourTable.cpp
+++ b/Engine/Plugins/Experimental/Mutable/Source/MutableTools/Private/MuT/NodeColourTable.cpp
@@ -4,12 +4,10 @@
 #include "MuT/NodeColourTable.h"
 
 #include "Misc/AssertionMacros.h"
-#include "MuT/NodeColourTablePrivate.h"
 #include "MuT/NodePrivate.h"
 #include "MuT/Table.h"
 
 
-
 namespace mu
 {
 
@@ -17,14 +15,7 @@ namespace mu
 	//---------------------------------------------------------------------------------------------
 	// Static initialisation
 	//---------------------------------------------------------------------------------------------
-	FNodeType NodeColourTable::Private::s_type = FNodeType(Node::EType::ColorTable, NodeColour::GetStaticType() );
-
-
-	//---------------------------------------------------------------------------------------------
-	//!
-	//---------------------------------------------------------------------------------------------
-
-	MUTABLE_IMPLEMENT_NODE( NodeColourTable )
+	FNodeType NodeColourTable::StaticType = FNodeType(Node::EType::ColorTable, NodeColour::GetStaticType() );
 
 
 	//---------------------------------------------------------------------------------------------
@@ -32,42 +23,28 @@ namespace mu
 	//---------------------------------------------------------------------------------------------
 	void NodeColourTable::SetColumn( const FString& strName )
 	{
-		m_pD->ColumnName = strName;
+		ColumnName = strName;
 	}
 
 
 	//---------------------------------------------------------------------------------------------
 	void NodeColourTable::SetParameterName( const FString& strName )
 	{
-		m_pD->ParameterName = strName;
-	}
-
-
-	//---------------------------------------------------------------------------------------------
-	void NodeColourTable::SetTable( TablePtr pTable )
-	{
-		m_pD->Table = pTable;
-	}
-
-
-	//---------------------------------------------------------------------------------------------
-	TablePtr NodeColourTable::GetTable() const
-	{
-		return m_pD->Table;
+		ParameterName = strName;
 	}
 
 
 	//---------------------------------------------------------------------------------------------
 	void NodeColourTable::SetNoneOption(bool bAddNoneOption)
 	{
-		m_pD->bNoneOption = bAddNoneOption;
+		bNoneOption = bAddNoneOption;
 	}
 
 	
 	//---------------------------------------------------------------------------------------------
 	void NodeColourTable::SetDefaultRowName(const FString RowName)
 	{
-		m_pD->DefaultRowName = RowName;
+		DefaultRowName = RowName;
 	}
 
 }
diff --git a/Engine/Plugins/Experimental/Mutable/Source/MutableTools/Private/MuT/NodeColourTablePrivate.h b/Engine/Plugins/Experimental/Mutable/Source/MutableTools/Private/MuT/NodeColourTablePrivate.h
deleted file mode 100644
diff --git a/Engine/Plugins/Experimental/Mutable/Source/MutableTools/Private/MuT/NodeImageConstant.cpp b/Engine/Plugins/Experimental/Mutable/Source/MutableTools/Private/MuT/NodeImageConstant.cpp
--- a/Engine/Plugins/Experimental/Mutable/Source/MutableTools/Private/MuT/NodeImageConstant.cpp
+++ b/Engine/Plugins/Experimental/Mutable/Source/MutableTools/Private/MuT/NodeImageConstant.cpp
@@ -6,7 +6,6 @@
 #include "Misc/AssertionMacros.h"
 #include "MuR/Image.h"
 #include "MuR/Serialisation.h"
-#include "MuT/NodeImageConstantPrivate.h"
 #include "MuT/NodePrivate.h"
 
 
@@ -16,33 +15,20 @@ namespace mu
 	//---------------------------------------------------------------------------------------------
 	// Static initialisation
 	//---------------------------------------------------------------------------------------------
-	FNodeType NodeImageConstant::Private::s_type = FNodeType(Node::EType::ImageConstant, NodeImage::GetStaticType());
-
-
-	//---------------------------------------------------------------------------------------------
-	MUTABLE_IMPLEMENT_NODE(NodeImageConstant);
-
-
-	//---------------------------------------------------------------------------------------------
-	ImagePtrConst NodeImageConstant::GetValue() const
-	{
-		Ptr<const Image> pImage;
-		if (m_pD->m_pProxy) pImage = m_pD->m_pProxy->Get();
-		return pImage;
-	}
+	FNodeType NodeImageConstant::StaticType = FNodeType(Node::EType::ImageConstant, NodeImage::GetStaticType());
 
 
 	//---------------------------------------------------------------------------------------------
 	void NodeImageConstant::SetValue(const Image* Value)
 	{
-		m_pD->m_pProxy = new ResourceProxyMemory<Image>(Value);
+		Proxy = new ResourceProxyMemory<Image>(Value);
 	}
 
 
 	//---------------------------------------------------------------------------------------------
-	void NodeImageConstant::SetValue(Ptr<ResourceProxy<Image>> pImage)
+	void NodeImageConstant::SetValue(Ptr<ResourceProxy<Image>> InProxy)
 	{
-		m_pD->m_pProxy = pImage;
+		Proxy = InProxy;
 	}
 
 }
diff --git a/Engine/Plugins/Experimental/Mutable/Source/MutableTools/Private/MuT/NodeImageConstantPrivate.h b/Engine/Plugins/Experimental/Mutable/Source/MutableTools/Private/MuT/NodeImageConstantPrivate.h
deleted file mode 100644
diff --git a/Engine/Plugins/Experimental/Mutable/Source/MutableTools/Private/MuT/NodeImageTable.cpp b/Engine/Plugins/Experimental/Mutable/Source/MutableTools/Private/MuT/NodeImageTable.cpp
deleted file mode 100644
diff --git a/Engine/Plugins/Experimental/Mutable/Source/MutableTools/Private/MuT/NodeImageTablePrivate.h b/Engine/Plugins/Experimental/Mutable/Source/MutableTools/Private/MuT/NodeImageTablePrivate.h
deleted file mode 100644
diff --git a/Engine/Plugins/Experimental/Mutable/Source/MutableTools/Private/MuT/NodeMeshTable.cpp b/Engine/Plugins/Experimental/Mutable/Source/MutableTools/Private/MuT/NodeMeshTable.cpp
--- a/Engine/Plugins/Experimental/Mutable/Source/MutableTools/Private/MuT/NodeMeshTable.cpp
+++ b/Engine/Plugins/Experimental/Mutable/Source/MutableTools/Private/MuT/NodeMeshTable.cpp
@@ -5,7 +5,6 @@
 
 #include "Misc/AssertionMacros.h"
 #include "MuT/NodeLayout.h"
-#include "MuT/NodeMeshTablePrivate.h"
 #include "MuT/NodePrivate.h"
 #include "MuT/Table.h"
 
@@ -13,18 +12,10 @@
 namespace mu
 {
 
-
 	//---------------------------------------------------------------------------------------------
 	// Static initialisation
 	//---------------------------------------------------------------------------------------------
-	FNodeType NodeMeshTable::Private::s_type = FNodeType(Node::EType::MeshTable, NodeMesh::GetStaticType() );
-
-
-	//---------------------------------------------------------------------------------------------
-	//!
-	//---------------------------------------------------------------------------------------------
-
-	MUTABLE_IMPLEMENT_NODE( NodeMeshTable );
+	FNodeType NodeMeshTable::StaticType = FNodeType(Node::EType::MeshTable, NodeMesh::GetStaticType() );
 
 
 	//---------------------------------------------------------------------------------------------
@@ -32,42 +23,28 @@ namespace mu
 	//---------------------------------------------------------------------------------------------
 	void NodeMeshTable::SetColumn( const FString& strName )
 	{
-		m_pD->ColumnName = strName;
+		ColumnName = strName;
 	}
 
 
 	//---------------------------------------------------------------------------------------------
 	void NodeMeshTable::SetParameterName( const FString& strName )
 	{
-		m_pD->ParameterName = strName;
-	}
-
-
-	//---------------------------------------------------------------------------------------------
-	void NodeMeshTable::SetTable( TablePtr pTable )
-	{
-		m_pD->Table = pTable;
-	}
-
-
-	//---------------------------------------------------------------------------------------------
-	TablePtr NodeMeshTable::GetTable() const
-	{
-		return m_pD->Table;
+		ParameterName = strName;
 	}
 
 
 	//---------------------------------------------------------------------------------------------
 	int NodeMeshTable::GetLayoutCount() const
 	{
-		return m_pD->Layouts.Num();
+		return Layouts.Num();
 	}
 
 
 	//---------------------------------------------------------------------------------------------
 	void NodeMeshTable::SetLayoutCount( int i )
 	{
-		m_pD->Layouts.SetNum( i );
+		Layouts.SetNum( i );
 	}
 
 
@@ -76,9 +53,9 @@ namespace mu
 	{
 		NodeLayoutPtr pResult;
 
-		if (i >= 0 && i < m_pD->Layouts.Num())
+		if (i >= 0 && i < Layouts.Num())
 		{
-			pResult = m_pD->Layouts[i];
+			pResult = Layouts[i];
 		}
 
 		return pResult;
@@ -89,21 +66,21 @@ namespace mu
 	void NodeMeshTable::SetLayout( int i, NodeLayoutPtr pLayout )
 	{
 		check( i>=0 && i<GetLayoutCount() );
-		m_pD->Layouts[i] = pLayout;
+		Layouts[i] = pLayout;
 	}
 
 
 	//---------------------------------------------------------------------------------------------
 	void NodeMeshTable::SetNoneOption(bool bAddNoneOption)
 	{
-		m_pD->bNoneOption = bAddNoneOption;
+		bNoneOption = bAddNoneOption;
 	}
 
 
 	//---------------------------------------------------------------------------------------------
 	void NodeMeshTable::SetDefaultRowName(const FString& RowName)
 	{
-		m_pD->DefaultRowName = RowName;
+		DefaultRowName = RowName;
 	}
 
 }
diff --git a/Engine/Plugins/Experimental/Mutable/Source/MutableTools/Private/MuT/NodeScalarTable.cpp b/Engine/Plugins/Experimental/Mutable/Source/MutableTools/Private/MuT/NodeScalarTable.cpp
--- a/Engine/Plugins/Experimental/Mutable/Source/MutableTools/Private/MuT/NodeScalarTable.cpp
+++ b/Engine/Plugins/Experimental/Mutable/Source/MutableTools/Private/MuT/NodeScalarTable.cpp
@@ -5,7 +5,6 @@
 
 #include "Misc/AssertionMacros.h"
 #include "MuT/NodePrivate.h"
-#include "MuT/NodeScalarTablePrivate.h"
 #include "MuT/Table.h"
 
 
@@ -16,56 +15,36 @@ namespace mu
 	//---------------------------------------------------------------------------------------------
 	// Static initialisation
 	//---------------------------------------------------------------------------------------------
-	FNodeType NodeScalarTable::Private::s_type = FNodeType(Node::EType::ScalarTable, NodeScalar::GetStaticType() );
+	FNodeType NodeScalarTable::StaticType = FNodeType(Node::EType::ScalarTable, NodeScalar::GetStaticType() );
 
 
-	//---------------------------------------------------------------------------------------------
-	//!
-	//---------------------------------------------------------------------------------------------
-
-	MUTABLE_IMPLEMENT_NODE(NodeScalarTable )
-
 
 	//---------------------------------------------------------------------------------------------
 	// Own Interface
 	//---------------------------------------------------------------------------------------------
 	void NodeScalarTable::SetColumn( const FString& strName )
 	{
-		m_pD->ColumnName = strName;
+		ColumnName = strName;
 	}
 
 
 	//---------------------------------------------------------------------------------------------
 	void NodeScalarTable::SetParameterName( const FString& strName )
 	{
-		m_pD->ParameterName = strName;
-	}
-
-
-	//---------------------------------------------------------------------------------------------
-	void NodeScalarTable::SetTable( TablePtr pTable )
-	{
-		m_pD->Table = pTable;
-	}
-
-
-	//---------------------------------------------------------------------------------------------
-	TablePtr NodeScalarTable::GetTable() const
-	{
-		return m_pD->Table;
+		ParameterName = strName;
 	}
 
 
 	//---------------------------------------------------------------------------------------------
 	void NodeScalarTable::SetNoneOption(bool bAddNoneOption)
 	{
-		m_pD->bNoneOption = bAddNoneOption;
+		bNoneOption = bAddNoneOption;
 	}
 
 	//---------------------------------------------------------------------------------------------
 	void NodeScalarTable::SetDefaultRowName(const FString& RowName)
 	{
-		m_pD->DefaultRowName = RowName;
+		DefaultRowName = RowName;
 	}
 }
 
diff --git a/Engine/Plugins/Experimental/Mutable/Source/MutableTools/Private/MuT/NodeScalarTablePrivate.h b/Engine/Plugins/Experimental/Mutable/Source/MutableTools/Private/MuT/NodeScalarTablePrivate.h
deleted file mode 100644
diff --git a/Engine/Plugins/Experimental/Mutable/Source/MutableTools/Public/MuT/NodeColourTable.h b/Engine/Plugins/Experimental/Mutable/Source/MutableTools/Public/MuT/NodeColourTable.h
--- a/Engine/Plugins/Experimental/Mutable/Source/MutableTools/Public/MuT/NodeColourTable.h
+++ b/Engine/Plugins/Experimental/Mutable/Source/MutableTools/Public/MuT/NodeColourTable.h
@@ -6,6 +6,7 @@
 #include "MuR/RefCounted.h"
 #include "MuT/Node.h"
 #include "MuT/NodeColour.h"
+#include "MuT/Table.h"
 
 
 namespace mu
@@ -16,10 +17,6 @@ namespace mu
 	typedef Ptr<NodeColourTable> NodeColourTablePtr;
 	typedef Ptr<const NodeColourTable> NodeColourTablePtrConst;
 
-	class Table;
-	typedef Ptr<Table> TablePtr;
-	typedef Ptr<const Table> TablePtrConst;
-
 
 	//! This node provides the meshes stored in the column of a table.
 	//! \ingroup transform
@@ -27,14 +24,17 @@ namespace mu
 	{
 	public:
 
-		NodeColourTable();
+		FString ParameterName;
+		Ptr<Table> Table;
+		FString ColumnName;
+		bool bNoneOption = false;
+		FString DefaultRowName;
 
+	public:
 
-		//-----------------------------------------------------------------------------------------
-		// Node Interface
-		//-----------------------------------------------------------------------------------------
-		const FNodeType* GetType() const override;
-		static const FNodeType* GetStaticType();
+		// Node interface
+		virtual const FNodeType* GetType() const override { return GetStaticType(); }
+		static const FNodeType* GetStaticType() { return &StaticType; }
 
 		//-----------------------------------------------------------------------------------------
 		// Own Interface
@@ -43,10 +43,6 @@ namespace mu
 		//! Set the name of the implicit table parameter.
 		void SetParameterName( const FString& strName );
 
-		//!
-		TablePtr GetTable() const;
-		void SetTable( TablePtr );
-
 		//!
 		void SetColumn( const FString& strName );
 
@@ -56,21 +52,15 @@ namespace mu
 		//! Set the row name to be used as default value
 		void SetDefaultRowName(const FString RowName);
 
-		//-----------------------------------------------------------------------------------------
-		// Interface pattern
-		//-----------------------------------------------------------------------------------------
-		class Private;
-		Private* GetPrivate() const;
-
 
 	protected:
 
 		//! Forbidden. Manage with the Ptr<> template.
-		~NodeColourTable();
+		~NodeColourTable() {}
 
 	private:
 
-		Private* m_pD;
+		static FNodeType StaticType;
 
 	};
 
diff --git a/Engine/Plugins/Experimental/Mutable/Source/MutableTools/Public/MuT/NodeImage.h b/Engine/Plugins/Experimental/Mutable/Source/MutableTools/Public/MuT/NodeImage.h
--- a/Engine/Plugins/Experimental/Mutable/Source/MutableTools/Public/MuT/NodeImage.h
+++ b/Engine/Plugins/Experimental/Mutable/Source/MutableTools/Public/MuT/NodeImage.h
@@ -15,9 +15,75 @@ namespace mu
 	typedef Ptr<NodeImage> NodeImagePtr;
 	typedef Ptr<const NodeImage> NodeImagePtrConst;
 
+	/** Data related to the a source image that is necessary to classify the final image fragments
+	* that are derived from this source.
+	*/
+	struct FSourceDataDescriptor
+	{
+		enum ESpecialValues
+		{
+			EInvalid = -2,
+			ENeutral = -1,
+		};
+
+		/** Number of mips in the source texture that are considered high resolution.
+		* If -1 it means the entire descriptor is neutral.
+		* If -2 it means the descriptor is invalid (results of an operation that shouldn't happen).
+		*/
+		int32 SourceHighResMips = ENeutral;
+
+		/** Source tags that mark this data and prevent it from mixing with other data at compile time. */
+		TArray<FString> Tags;
+
+		inline bool operator==(const FSourceDataDescriptor& Other) const
+		{
+			return SourceHighResMips == Other.SourceHighResMips && Tags == Other.Tags;
+		}
+
+		inline bool IsInvalid() const
+		{
+			return SourceHighResMips == EInvalid;
+		}
+
+		inline bool IsNeutral() const
+		{
+			return SourceHighResMips == ENeutral;
+		}
+
+		void CombineWith(const FSourceDataDescriptor& Other)
+		{
+			if (IsInvalid() || Other.IsInvalid())
+			{
+				SourceHighResMips = EInvalid;
+				Tags.Empty();
+				return;
+			}
+
+			if (Other.IsNeutral())
+			{
+				return;
+			}
+
+			if (IsNeutral())
+			{
+				SourceHighResMips = Other.SourceHighResMips;
+				Tags = Other.Tags;
+				return;
+			}
+
+			if (!(*this==Other))
+			{
+				SourceHighResMips = EInvalid;
+				Tags.Empty();
+				return;
+			}
+
+			return;
+		}
+	};
+
 
-    //! %Base class of any node that outputs an image.
-	//! \ingroup model
+    /** Base class of any node that outputs an image. */
 	class MUTABLETOOLS_API NodeImage : public Node
 	{
 	public:
diff --git a/Engine/Plugins/Experimental/Mutable/Source/MutableTools/Public/MuT/NodeImageConstant.h b/Engine/Plugins/Experimental/Mutable/Source/MutableTools/Public/MuT/NodeImageConstant.h
--- a/Engine/Plugins/Experimental/Mutable/Source/MutableTools/Public/MuT/NodeImageConstant.h
+++ b/Engine/Plugins/Experimental/Mutable/Source/MutableTools/Public/MuT/NodeImageConstant.h
@@ -12,12 +12,6 @@
 namespace mu
 {
 
-	// Forward definitions
-    class NodeImageConstant;
-    typedef Ptr<NodeImageConstant> NodeImageConstantPtr;
-    typedef Ptr<const NodeImageConstant> NodeImageConstantPtrConst;
-
-
 	/** Node that outputs a constant image.
 	* This node also supports "image references".
 	*/
@@ -25,42 +19,39 @@ namespace mu
 	{
 	public:
 
-		NodeImageConstant();
+		/** */
+		Ptr<ResourceProxy<Image>> Proxy;
+		
+		/** */
+		FSourceDataDescriptor SourceDataDescriptor;
+
+	public:
 
 		//-----------------------------------------------------------------------------------------
 		// Node Interface
 		//-----------------------------------------------------------------------------------------
 
-        const FNodeType* GetType() const override;
-		static const FNodeType* GetStaticType();
+		virtual const FNodeType* GetType() const override { return GetStaticType(); }
+		static const FNodeType* GetStaticType() { return &StaticType; }
 
 		//-----------------------------------------------------------------------------------------
 		// Own Interface
 		//-----------------------------------------------------------------------------------------
 
-		//! Get the image that will be output by this node
-        Ptr<const Image> GetValue() const;
-
         //! Set the image to be output by this node
         void SetValue( const Image* );
 
         //! Set the image proxy that will provide the image for this node when necessary
         void SetValue( Ptr<ResourceProxy<Image>> pImage );
 
-		//-----------------------------------------------------------------------------------------
-		// Interface pattern
-		//-----------------------------------------------------------------------------------------
-		class Private;
-		Private* GetPrivate() const;
-
 	protected:
 
 		//! Forbidden. Manage with the Ptr<> template.
-		~NodeImageConstant();
+		~NodeImageConstant() {}
 
 	private:
 
-		Private* m_pD;
+		static FNodeType StaticType;
 
 	};
 
diff --git a/Engine/Plugins/Experimental/Mutable/Source/MutableTools/Public/MuT/NodeImageTable.h b/Engine/Plugins/Experimental/Mutable/Source/MutableTools/Public/MuT/NodeImageTable.h
--- a/Engine/Plugins/Experimental/Mutable/Source/MutableTools/Public/MuT/NodeImageTable.h
+++ b/Engine/Plugins/Experimental/Mutable/Source/MutableTools/Public/MuT/NodeImageTable.h
@@ -4,6 +4,8 @@
 
 #include "MuR/Ptr.h"
 #include "MuR/RefCounted.h"
+#include "MuR/ImageTypes.h"
+#include "MuT/Table.h"
 #include "MuT/Node.h"
 #include "MuT/NodeImage.h"
 
@@ -11,74 +13,44 @@
 namespace mu
 {
 
-	// Forward definitions
-	class NodeImageTable;
-	typedef Ptr<NodeImageTable> NodeImageTablePtr;
-	typedef Ptr<const NodeImageTable> NodeImageTablePtrConst;
 
-	class Table;
-	typedef Ptr<Table> TablePtr;
-	typedef Ptr<const Table> TablePtrConst;
-
-	struct FImageDesc;
-
-
-	//! This node provides the meshes stored in the column of a table.
+	/** This node provides the meshes stored in the column of a table. */
 	class MUTABLETOOLS_API NodeImageTable : public NodeImage
 	{
 	public:
 
-		NodeImageTable();
-
-		//-----------------------------------------------------------------------------------------
-		// Node Interface
-		//-----------------------------------------------------------------------------------------
+		FString ParameterName;
 
-        const FNodeType* GetType() const override;
-		static const FNodeType* GetStaticType();
+		Ptr<Table> Table;
 
-		//-----------------------------------------------------------------------------------------
-		// Own Interface
-		//-----------------------------------------------------------------------------------------
+		FString ColumnName;
 
-		//! Set the name of the implicit table parameter.
-		void SetParameterName( const FString& strName );
+		uint16 MaxTextureSize = 0;
 
-		//!
-		TablePtr GetTable() const;
-		void SetTable( TablePtr );
+		FImageDesc ReferenceImageDesc;
 
-		//!
-		void SetColumn( const FString& strName );
+		bool bNoneOption = false;
 
-		//!
-		void SetMaxTextureSize(uint16 Size);
-		uint16 GetMaxTextureSize();
+		FString DefaultRowName;
 
-		//! Adds the "None" option to the parameter that represents this table column
-		void SetNoneOption(bool bAddOption);
-		
-		//!
-		void SetReferenceImageDescriptor(const FImageDesc& Descriptor);
+		/** */
+		FSourceDataDescriptor SourceDataDescriptor;
 
-		//! Set the row name to be used as default value
-		void SetDefaultRowName(const FString& RowName);
-
-		//-----------------------------------------------------------------------------------------
-		// Interface pattern
-		//-----------------------------------------------------------------------------------------
-		class Private;
-		Private* GetPrivate() const;
+	public:
 
+		// Node interface
+		virtual const FNodeType* GetType() const override { return GetStaticType(); }
+		static const FNodeType* GetStaticType() { return &StaticType; }
 
 	protected:
 
 		//! Forbidden. Manage with the Ptr<> template.
-		~NodeImageTable();
+		~NodeImageTable() {}
 
 	private:
 
-		Private* m_pD;
+		static FNodeType StaticType;
+
 
 	};
 
diff --git a/Engine/Plugins/Experimental/Mutable/Source/MutableTools/Public/MuT/NodeMeshTable.h b/Engine/Plugins/Experimental/Mutable/Source/MutableTools/Public/MuT/NodeMeshTable.h
--- a/Engine/Plugins/Experimental/Mutable/Source/MutableTools/Public/MuT/NodeMeshTable.h
+++ b/Engine/Plugins/Experimental/Mutable/Source/MutableTools/Public/MuT/NodeMeshTable.h
@@ -6,6 +6,8 @@
 #include "MuR/RefCounted.h"
 #include "MuT/Node.h"
 #include "MuT/NodeMesh.h"
+#include "MuT/NodeLayout.h"
+#include "MuT/Table.h"
 
 
 namespace mu
@@ -16,29 +18,25 @@ namespace mu
 	typedef Ptr<NodeMeshTable> NodeMeshTablePtr;
 	typedef Ptr<const NodeMeshTable> NodeMeshTablePtrConst;
 
-	class NodeLayout;
-	typedef Ptr<NodeLayout> NodeLayoutPtr;
-	typedef Ptr<const NodeLayout> NodeLayoutPtrConst;
-
-	class Table;
-	typedef Ptr<Table> TablePtr;
-	typedef Ptr<const Table> TablePtrConst;
-
 
 	//! This node provides the meshes stored in the column of a table.
 	class MUTABLETOOLS_API NodeMeshTable : public NodeMesh
 	{
 	public:
 
-		NodeMeshTable();
+		FString ParameterName;
+		Ptr<Table> Table;
+		FString ColumnName;
+		bool bNoneOption = false;
+		FString DefaultRowName;
 
+		TArray<NodeLayoutPtr> Layouts;
 
-		//-----------------------------------------------------------------------------------------
-		// Node Interface
-		//-----------------------------------------------------------------------------------------
+	public:
 
-        const FNodeType* GetType() const override;
-		static const FNodeType* GetStaticType();
+		// Node interface
+		virtual const FNodeType* GetType() const override { return GetStaticType(); }
+		static const FNodeType* GetStaticType() { return &StaticType; }
 
 		//-----------------------------------------------------------------------------------------
 		// Own Interface
@@ -47,10 +45,6 @@ namespace mu
 		//! Set the name of the implicit table parameter.
 		void SetParameterName( const FString& strName );
 
-		//!
-		TablePtr GetTable() const;
-		void SetTable( TablePtr );
-
 		//!
 		void SetColumn( const FString& strName );
 
@@ -68,21 +62,15 @@ namespace mu
 		//! Set the row name to be used as default value
 		void SetDefaultRowName(const FString& RowName);
 
-		//-----------------------------------------------------------------------------------------
-		// Interface pattern
-		//-----------------------------------------------------------------------------------------
-		class Private;
-		Private* GetPrivate() const;
-
 
 	protected:
 
 		//! Forbidden. Manage with the Ptr<> template.
-		~NodeMeshTable();
+		~NodeMeshTable() {}
 
 	private:
 
-		Private* m_pD;
+		static FNodeType StaticType;
 
 	};
 
diff --git a/Engine/Plugins/Experimental/Mutable/Source/MutableTools/Public/MuT/NodeScalarTable.h b/Engine/Plugins/Experimental/Mutable/Source/MutableTools/Public/MuT/NodeScalarTable.h
--- a/Engine/Plugins/Experimental/Mutable/Source/MutableTools/Public/MuT/NodeScalarTable.h
+++ b/Engine/Plugins/Experimental/Mutable/Source/MutableTools/Public/MuT/NodeScalarTable.h
@@ -6,6 +6,7 @@
 #include "MuR/RefCounted.h"
 #include "MuT/Node.h"
 #include "MuT/NodeScalar.h"
+#include "MuT/Table.h"
 
 
 namespace mu
@@ -16,24 +17,23 @@ namespace mu
 	typedef Ptr<NodeScalarTable> NodeScalarTablePtr;
 	typedef Ptr<const NodeScalarTable> NodeScalarTablePtrConst;
 
-	class Table;
-	typedef Ptr<Table> TablePtr;
-	typedef Ptr<const Table> TablePtrConst;
-
 
 	//! This node provides the meshes stored in the column of a table.
 	class MUTABLETOOLS_API NodeScalarTable : public NodeScalar
 	{
 	public:
 
-		NodeScalarTable();
+		FString ParameterName;
+		Ptr<Table> Table;
+		FString ColumnName;
+		bool bNoneOption = false;
+		FString DefaultRowName;
 
-		//-----------------------------------------------------------------------------------------
-		// Node Interface
-		//-----------------------------------------------------------------------------------------
+	public:
 
-		const FNodeType* GetType() const override;
-		static const FNodeType* GetStaticType();
+		// Node interface
+		virtual const FNodeType* GetType() const override { return GetStaticType(); }
+		static const FNodeType* GetStaticType() { return &StaticType; }
 
 		//-----------------------------------------------------------------------------------------
 		// Own Interface
@@ -42,10 +42,6 @@ namespace mu
 		//! Set the name of the implicit table parameter.
 		void SetParameterName( const FString& strName );
 
-		//!
-		TablePtr GetTable() const;
-		void SetTable( TablePtr );
-
 		//!
 		void SetColumn( const FString& strName );
 
@@ -55,21 +51,15 @@ namespace mu
 		//! Set the row name to be used as default value
 		void SetDefaultRowName(const FString& RowName);
 
-		//-----------------------------------------------------------------------------------------
-		// Interface pattern
-		//-----------------------------------------------------------------------------------------
-		class Private;
-		Private* GetPrivate() const;
-
 
 	protected:
 
 		//! Forbidden. Manage with the Ptr<> template.
-		~NodeScalarTable();
+		~NodeScalarTable() {}
 
 	private:
 
-		Private* m_pD;
+		static FNodeType StaticType;
 
 	};
 
