[mutable] Fixed crash when connecting a Group Node directly as a child of another Group node of type "Toggle".
#tests Editor
#rnx

diff --git a/Engine/Plugins/Experimental/Mutable/Source/MutableTools/Private/MuT/CodeGenerator_FirstPass.cpp b/Engine/Plugins/Experimental/Mutable/Source/MutableTools/Private/MuT/CodeGenerator_FirstPass.cpp
--- a/Engine/Plugins/Experimental/Mutable/Source/MutableTools/Private/MuT/CodeGenerator_FirstPass.cpp
+++ b/Engine/Plugins/Experimental/Mutable/Source/MutableTools/Private/MuT/CodeGenerator_FirstPass.cpp
@@ -607,20 +607,30 @@ namespace mu
                 switch (Private->m_type )
                 {
                     case NodeObjectGroup::CS_TOGGLE_EACH:
-                    {
-                        // Create a new boolean parameter
-                        Ptr<ASTOpParameter> op = new ASTOpParameter();
-                        op->type = OP_TYPE::BO_PARAMETER;
-
-                        op->parameter.m_name = pChildNode->GetName();
-                   		const TCHAR* CStr = ToCStr(pChildNode->GetUid());
-                   		op->parameter.m_uid.ImportTextItem(CStr, 0, nullptr, nullptr);
-                        op->parameter.m_type = PARAMETER_TYPE::T_BOOL;
-                        op->parameter.m_defaultValue.Set<ParamBoolType>(false);
+                    {           
+						if (pChildNode->GetType() == NodeObjectGroup::GetStaticType())
+						{
+							FString Msg = FString::Printf(TEXT("The Group Node [%s] has type Toggle and its direct child is a Group node, which is not allowed. Change the type or add a Child Object node in between them."),
+								*Private->Name);
+							m_pErrorLog->GetPrivate()->Add(Msg, ELMT_ERROR, InNode->GetMessageContext());
+						}
+						else
+						{
+							// Create a new boolean parameter
+							Ptr<ASTOpParameter> op = new ASTOpParameter();
+							op->type = OP_TYPE::BO_PARAMETER;
+
+							op->parameter.m_name = pChildNode->GetName();
+							const TCHAR* CStr = ToCStr(pChildNode->GetUid());
+							op->parameter.m_uid.ImportTextItem(CStr, 0, nullptr, nullptr);
+							op->parameter.m_type = PARAMETER_TYPE::T_BOOL;
+							op->parameter.m_defaultValue.Set<ParamBoolType>(false);
+
+							ParameterNodes.Add(pChildNode, op);
+
+							paramOp = op;
+						}
 
-						ParameterNodes.Add(pChildNode, op);
-
-                        paramOp = op;
                         break;
                     }
 
