[mutable] Fixed the new Mesh Component nodes so that they refresh automatically when a mesh is selected and they show the material pins without having to manually reconstruct the node. Otherwise they would crash when compiled.
#tests Editor
#rnx

diff --git a/Engine/Plugins/Experimental/Mutable/Source/CustomizableObjectEditor/Private/MuCOE/Nodes/CustomizableObjectNodeComponentMesh.cpp b/Engine/Plugins/Experimental/Mutable/Source/CustomizableObjectEditor/Private/MuCOE/Nodes/CustomizableObjectNodeComponentMesh.cpp
--- a/Engine/Plugins/Experimental/Mutable/Source/CustomizableObjectEditor/Private/MuCOE/Nodes/CustomizableObjectNodeComponentMesh.cpp
+++ b/Engine/Plugins/Experimental/Mutable/Source/CustomizableObjectEditor/Private/MuCOE/Nodes/CustomizableObjectNodeComponentMesh.cpp
@@ -66,7 +66,7 @@ void UCustomizableObjectNodeComponentMesh::PostEditChangeProperty(FPropertyChang
 	Super::PostEditChangeProperty(PropertyChangedEvent);
 
 	FProperty* PropertyThatChanged = PropertyChangedEvent.Property;
-	if (PropertyThatChanged && PropertyThatChanged->GetName() == TEXT("SkeletalMesh"))
+	if (PropertyThatChanged && PropertyThatChanged->GetName() == TEXT("Mesh"))
 	{
 		ReconstructNode();
 	}
