[mutable] Workarounds for static analysis complaints.
#rnx

diff --git a/Engine/Plugins/Experimental/Mutable/Source/CustomizableObject/Private/MuCO/CustomizableObject.cpp b/Engine/Plugins/Experimental/Mutable/Source/CustomizableObject/Private/MuCO/CustomizableObject.cpp
--- a/Engine/Plugins/Experimental/Mutable/Source/CustomizableObject/Private/MuCO/CustomizableObject.cpp
+++ b/Engine/Plugins/Experimental/Mutable/Source/CustomizableObject/Private/MuCO/CustomizableObject.cpp
@@ -2264,6 +2264,34 @@ void UCustomizableObjectBulk::CookAdditionalFilesOverride(const TCHAR* PackageFi
 	}
 }
 
+
+namespace MutablePrivate
+{
+	// TODO:
+	// To avoid influence of the order of the streamed data (their index), classify it recursively based on hash values
+	// until the tree leaves have either a single block, or a sum of blocks below the desired file size.
+	struct FClassifyNode
+	{
+		TArray<UCustomizableObjectBulk::FBlock> Blocks;
+
+		//TSharedPtr<FClassifyNode> Child0;
+		//TSharedPtr<FClassifyNode> Child1;
+		//uint32 Depth=0;
+	};
+
+	void AddNode(TMap<uint32, FClassifyNode>& Nodes, int32 Slack, const UCustomizableObjectBulk::FBlock& Block)
+	{
+		FClassifyNode& Root = Nodes.FindOrAdd(Block.Flags);
+		if (Root.Blocks.IsEmpty())
+		{
+			Root.Blocks.Reserve(Slack);
+		}
+
+		Root.Blocks.Add(Block);
+	}
+}
+
+
 void UCustomizableObjectBulk::PrepareBulkData(UCustomizableObject* InOuter, const ITargetPlatform* TargetPlatform)
 {
 	CustomizableObject = InOuter;
@@ -2283,32 +2311,9 @@ void UCustomizableObjectBulk::PrepareBulkData(UCustomizableObject* InOuter, cons
 	const uint64 MaxChunkSize = UCustomizableObjectSystem::GetInstance()->GetMaxChunkSizeForPlatform(TargetPlatform);
 	TargetBulkDataFileBytes = FMath::Min(TargetBulkDataFileBytes, MaxChunkSize);
 
-
-	// TODO:
-	// To avoid influence of the order of the streamed data (their index), classify it recursively based on hash values
-	// until the tree leaves have either a single block, or a sum of blocks below the desired file size.
-	struct FClassifyNode
-	{
-		TArray<FBlock> Blocks;
-
-		//TSharedPtr<FClassifyNode> Child0;
-		//TSharedPtr<FClassifyNode> Child1;
-		//uint32 Depth=0;
-	};
-
 	// Root nodes by flags.
 	const int32 NumRoms = Model->GetRomCount();
-	TMap<uint32,FClassifyNode> RootNode;
-	auto AddNode = [&RootNode,NumRoms]( const FBlock& Block )
-		{
-			FClassifyNode& Root = RootNode.FindOrAdd(Block.Flags);
-			if (Root.Blocks.IsEmpty())
-			{
-				Root.Blocks.Reserve(NumRoms);
-			}
-
-			Root.Blocks.Add(Block);
-		};
+	TMap<uint32, MutablePrivate::FClassifyNode> RootNode;
 
 	// Create blocks data.
 	{		
@@ -2319,7 +2324,7 @@ void UCustomizableObjectBulk::PrepareBulkData(UCustomizableObject* InOuter, cons
 			const mu::ERomFlags BlockFlags = Model->GetRomFlags(RomIndex);
 
 			FBlock CurrentBlock = { EDataType::Model, BlockId, BlockSize, uint32(BlockFlags), 0 };
-			AddNode(CurrentBlock);
+			MutablePrivate::AddNode(RootNode, NumRoms, CurrentBlock);
 		}
 	}
 
@@ -2340,7 +2345,7 @@ void UCustomizableObjectBulk::PrepareBulkData(UCustomizableObject* InOuter, cons
 			check(SourceOffset == MorphStreamable.Value.Block.Offset);
 			uint32 Flags = 0;
 			FBlock CurrentBlock = { EDataType::RealTimeMorph, MorphStreamable.Key, BlockSize, Flags, SourceOffset };
-			AddNode(CurrentBlock);
+			MutablePrivate::AddNode(RootNode, NumRoms, CurrentBlock);
 
 			SourceOffset += BlockSize;
 		}
@@ -2363,7 +2368,7 @@ void UCustomizableObjectBulk::PrepareBulkData(UCustomizableObject* InOuter, cons
 			check(SourceOffset == ClothStreamable.Value.Block.Offset);
 			uint32 Flags = 0;
 			FBlock CurrentBlock = { EDataType::Clothing, ClothStreamable.Key, BlockSize, Flags, SourceOffset };
-			AddNode(CurrentBlock);
+			MutablePrivate::AddNode(RootNode, NumRoms, CurrentBlock);
 
 			SourceOffset += BlockSize;
 		}
@@ -2372,7 +2377,7 @@ void UCustomizableObjectBulk::PrepareBulkData(UCustomizableObject* InOuter, cons
 
 	for (int32 FlagClassIndex = 0; FlagClassIndex < RootNode.Num(); ++FlagClassIndex)
 	{
-		FClassifyNode& Root = RootNode[FlagClassIndex];
+		MutablePrivate::FClassifyNode& Root = RootNode[FlagClassIndex];
 
 		// Temp: Group by order in the array
 		for (int32 BlockIndex = 0; BlockIndex < Root.Blocks.Num(); )
diff --git a/Engine/Plugins/Experimental/Mutable/Source/MutableTools/Private/MuT/ASTOpImageCompose.h b/Engine/Plugins/Experimental/Mutable/Source/MutableTools/Private/MuT/ASTOpImageCompose.h
--- a/Engine/Plugins/Experimental/Mutable/Source/MutableTools/Private/MuT/ASTOpImageCompose.h
+++ b/Engine/Plugins/Experimental/Mutable/Source/MutableTools/Private/MuT/ASTOpImageCompose.h
@@ -36,7 +36,7 @@ namespace mu
 		virtual void ForEachChild(const TFunctionRef<void(ASTChild&)>) override;
 		virtual void Link(FProgram& program, FLinkerOptions* Options) override;
 		//TODO: virtual Ptr<ASTOp> OptimiseSink(const FModelOptimizationOptions& options, FOptimizeSinkContext& context) const override;
-		virtual Ptr<ASTOp> OptimiseSemantic(const FModelOptimizationOptions& options, int32 Pass) const;
+		virtual Ptr<ASTOp> OptimiseSemantic(const FModelOptimizationOptions& options, int32 Pass) const override;
 		virtual FImageDesc GetImageDesc(bool returnBestOption, FGetImageDescContext* context) const override;
 		virtual void GetLayoutBlockSize(int32* OutBlockX, int32* OutBlockY) override;
 		virtual bool IsImagePlainConstant(FVector4f& OutColor) const override;
diff --git a/Engine/Plugins/Experimental/Mutable/Source/MutableTools/Private/MuT/ASTOpImageLayer.h b/Engine/Plugins/Experimental/Mutable/Source/MutableTools/Private/MuT/ASTOpImageLayer.h
--- a/Engine/Plugins/Experimental/Mutable/Source/MutableTools/Private/MuT/ASTOpImageLayer.h
+++ b/Engine/Plugins/Experimental/Mutable/Source/MutableTools/Private/MuT/ASTOpImageLayer.h
@@ -53,7 +53,7 @@ namespace mu
 		virtual void GetLayoutBlockSize(int* pBlockX, int* pBlockY) override;
 		virtual Ptr<ImageSizeExpression> GetImageSizeExpression() const override;
 		virtual Ptr<ASTOp> OptimiseSemantic(const FModelOptimizationOptions&, int32 Pass) const override;
-		virtual Ptr<ASTOp> OptimiseSink(const FModelOptimizationOptions&, FOptimizeSinkContext&) const;
+		virtual Ptr<ASTOp> OptimiseSink(const FModelOptimizationOptions&, FOptimizeSinkContext&) const override;
 		virtual FSourceDataDescriptor GetSourceDataDescriptor(FGetSourceDataDescriptorContext*) const override;
 
 	};
diff --git a/Engine/Plugins/Experimental/Mutable/Source/MutableTools/Private/MuT/ASTOpImageLayerColor.h b/Engine/Plugins/Experimental/Mutable/Source/MutableTools/Private/MuT/ASTOpImageLayerColor.h
--- a/Engine/Plugins/Experimental/Mutable/Source/MutableTools/Private/MuT/ASTOpImageLayerColor.h
+++ b/Engine/Plugins/Experimental/Mutable/Source/MutableTools/Private/MuT/ASTOpImageLayerColor.h
@@ -52,8 +52,8 @@ namespace mu
 		virtual FImageDesc GetImageDesc(bool returnBestOption, FGetImageDescContext* context) const override;
 		virtual void GetLayoutBlockSize(int* pBlockX, int* pBlockY) override;
 		virtual Ptr<ImageSizeExpression> GetImageSizeExpression() const override;
-		virtual Ptr<ASTOp> OptimiseSemantic(const FModelOptimizationOptions& options, int32 Pass) const;
-		virtual Ptr<ASTOp> OptimiseSink(const FModelOptimizationOptions&, FOptimizeSinkContext&) const;
+		virtual Ptr<ASTOp> OptimiseSemantic(const FModelOptimizationOptions& options, int32 Pass) const override;
+		virtual Ptr<ASTOp> OptimiseSink(const FModelOptimizationOptions&, FOptimizeSinkContext&) const override;
 		virtual FSourceDataDescriptor GetSourceDataDescriptor(FGetSourceDataDescriptorContext*) const override;
 
 	};
diff --git a/Engine/Plugins/Experimental/Mutable/Source/MutableTools/Private/MuT/ASTOpImageRasterMesh.h b/Engine/Plugins/Experimental/Mutable/Source/MutableTools/Private/MuT/ASTOpImageRasterMesh.h
--- a/Engine/Plugins/Experimental/Mutable/Source/MutableTools/Private/MuT/ASTOpImageRasterMesh.h
+++ b/Engine/Plugins/Experimental/Mutable/Source/MutableTools/Private/MuT/ASTOpImageRasterMesh.h
@@ -58,8 +58,8 @@ namespace mu
 		virtual void Link(FProgram& program, FLinkerOptions* Options) override;
 		virtual FImageDesc GetImageDesc(bool returnBestOption, FGetImageDescContext* context) const override;
 		virtual Ptr<ImageSizeExpression> GetImageSizeExpression() const override;
-		virtual Ptr<ASTOp> OptimiseSemantic(const FModelOptimizationOptions& options, int32 Pass) const;
-		virtual Ptr<ASTOp> OptimiseSink(const FModelOptimizationOptions&, FOptimizeSinkContext&) const;
+		virtual Ptr<ASTOp> OptimiseSemantic(const FModelOptimizationOptions& options, int32 Pass) const override;
+		virtual Ptr<ASTOp> OptimiseSink(const FModelOptimizationOptions&, FOptimizeSinkContext&) const override;
 		virtual FSourceDataDescriptor GetSourceDataDescriptor(FGetSourceDataDescriptorContext*) const override;
 
 	};
diff --git a/Engine/Plugins/Experimental/Mutable/Source/MutableTools/Private/MuT/ASTOpReferenceResource.h b/Engine/Plugins/Experimental/Mutable/Source/MutableTools/Private/MuT/ASTOpReferenceResource.h
--- a/Engine/Plugins/Experimental/Mutable/Source/MutableTools/Private/MuT/ASTOpReferenceResource.h
+++ b/Engine/Plugins/Experimental/Mutable/Source/MutableTools/Private/MuT/ASTOpReferenceResource.h
@@ -54,7 +54,7 @@ namespace mu
 		virtual void GetLayoutBlockSize(int32* pBlockX, int32* pBlockY) override;
 		virtual bool GetNonBlackRect(FImageRect& maskUsage) const override;
 		virtual Ptr<ImageSizeExpression> GetImageSizeExpression() const override;
-		virtual FSourceDataDescriptor GetSourceDataDescriptor(FGetSourceDataDescriptorContext*) const;
+		virtual FSourceDataDescriptor GetSourceDataDescriptor(FGetSourceDataDescriptorContext*) const override;
 	};
 
 
