[mutable] Fixed crash when compiling a CustomizableObject while PIE.
#tests PIE
#rnx

diff --git a/Engine/Plugins/Experimental/Mutable/Source/CustomizableObjectEditor/Private/MuCOE/CustomizableObjectEditor.cpp b/Engine/Plugins/Experimental/Mutable/Source/CustomizableObjectEditor/Private/MuCOE/CustomizableObjectEditor.cpp
--- a/Engine/Plugins/Experimental/Mutable/Source/CustomizableObjectEditor/Private/MuCOE/CustomizableObjectEditor.cpp
+++ b/Engine/Plugins/Experimental/Mutable/Source/CustomizableObjectEditor/Private/MuCOE/CustomizableObjectEditor.cpp
@@ -2340,7 +2340,7 @@ void FCustomizableObjectEditor::CreatePreviewActor()
 {
 	if (Actor)
 	{
-		ViewportClient->GetWorld()->RemoveActor(Actor.Get(), false);
+		Actor->Destroy();
 	}
 	
 	Actor.Reset(ViewportClient->GetWorld()->SpawnActor<ASkeletalMeshActor>());
