[mutable] Cleanup and optimize the mesh merging operation.
- Detect the fast path to merge on a per-buffer basis.
- Avoid copies doing in-place reformatting when possible.
#rnx

diff --git a/Engine/Plugins/Experimental/Mutable/Source/MutableRuntime/Internal/MuR/OpMeshFormat.h b/Engine/Plugins/Experimental/Mutable/Source/MutableRuntime/Internal/MuR/OpMeshFormat.h
--- a/Engine/Plugins/Experimental/Mutable/Source/MutableRuntime/Internal/MuR/OpMeshFormat.h
+++ b/Engine/Plugins/Experimental/Mutable/Source/MutableRuntime/Internal/MuR/OpMeshFormat.h
@@ -8,6 +8,7 @@ namespace mu
 {
 	class Mesh;
 	class FMeshBufferSet;
+	struct FMeshBuffer;
 
 	/** Convert a mesh format into another one.
 	* Slow implementation, but it should never happen at run-time
@@ -27,13 +28,8 @@ namespace mu
 		);
 
 
-    /** Fill a mesh buffer from result with the current data in source but keeping the format already in the result. */
-    void MeshFormatBuffer
-        (
-            const FMeshBufferSet& Source,
-            FMeshBufferSet& Result,
-            int32 BufferIndex
-        );
+    /** Fill a mesh buffer in Result (with an optional offset) with the current data in Source but keeping the format already in the Result. */
+    void MeshFormatBuffer( const FMeshBufferSet& Source, FMeshBuffer& Result, int32 ResultElementOffset, bool bHasSpecialSemantics, uint32 IDPrefix);
 
 	/** Try to reduce the mesh size by reducing the component count, and data type of some buffers. */
 	void MeshOptimizeBuffers( Mesh* );
diff --git a/Engine/Plugins/Experimental/Mutable/Source/MutableRuntime/Internal/MuR/OpMeshMerge.h b/Engine/Plugins/Experimental/Mutable/Source/MutableRuntime/Internal/MuR/OpMeshMerge.h
--- a/Engine/Plugins/Experimental/Mutable/Source/MutableRuntime/Internal/MuR/OpMeshMerge.h
+++ b/Engine/Plugins/Experimental/Mutable/Source/MutableRuntime/Internal/MuR/OpMeshMerge.h
@@ -333,19 +333,18 @@ namespace mu
 			MUTABLE_CPUPROFILER_SCOPE(Layouts);
 
 			Result->Layouts.SetNum(pFirst->Layouts.Num());
-			for (int32 i = 0; i < pFirst->Layouts.Num(); ++i)
+			for (int32 LayoutIndex = 0; LayoutIndex < pFirst->Layouts.Num(); ++LayoutIndex)
 			{
-				const Layout* pF = pFirst->Layouts[i].get();
-				LayoutPtr pR = pF->Clone();
+				const Layout* pF = pFirst->Layouts[LayoutIndex].get();
+				Ptr<Layout> pR = pF->Clone();
 
-				if (i < pSecond->Layouts.Num())
+				if (LayoutIndex < pSecond->Layouts.Num())
 				{
-					const Layout* pS = pSecond->Layouts[i].get();
-
+					const Layout* pS = pSecond->Layouts[LayoutIndex].get();
 					pR->Blocks.Append(pS->Blocks);
 				}
 
-				Result->Layouts[i] = pR;
+				Result->Layouts[LayoutIndex] = pR;
 			}
 		}
 
@@ -692,69 +691,116 @@ namespace mu
 		{
             MUTABLE_CPUPROFILER_SCOPE(Vertices);
 
-            const int32 FirstCount = pFirst->GetVertexBuffers().GetElementCount();
-			const int32 SecondCount = pSecond->GetVertexBuffers().GetElementCount();
-
-			// TODO: fast path should be per-buffer
-			// TODO: when formats match, which at runtime should be often.
-			bool bFastPath = pFirst->GetVertexBuffers().HasSameFormat( pSecond->GetVertexBuffers() );
-			bFastPath = bFastPath && !bNeedsExplicitIds;
+			const int32 FirstBufferCount = pFirst->VertexBuffers.m_buffers.Num();
+			const int32 SecondBufferCount = pSecond->VertexBuffers.m_buffers.Num();
+			const int32 FirstVertexCount = pFirst->GetVertexBuffers().GetElementCount();
+			const int32 SecondVertexCount = pSecond->GetVertexBuffers().GetElementCount();
 
 			// Check if the format of the BoneIndex buffer has to change
 			bool bChangeBoneIndicesFormat = false;
 			EMeshBufferFormat BoneIndexFormat = MaxNumBonesInBoneMaps > MAX_uint8 ? MBF_UINT16 : MBF_UINT8;
+			bChangeBoneIndicesFormat = pFirst->GetVertexBuffers().HasAnySemanticWithDifferentFormat(MBS_BONEINDICES, BoneIndexFormat);
+			if (!bChangeBoneIndicesFormat)
+			{
+				bChangeBoneIndicesFormat = pSecond->GetVertexBuffers().HasAnySemanticWithDifferentFormat(MBS_BONEINDICES, BoneIndexFormat);
+			}
 
-			// Iterate all vertex buffers to check if we need to format bone indices
+			// Step 1: Set up the initial vertex buffer structure of the result mesh.
+			//-----------------------------------------------------------------------
 			{
-				bChangeBoneIndicesFormat = pFirst->GetVertexBuffers().HasAnySemanticWithDifferentFormat(MBS_BONEINDICES, BoneIndexFormat);
-				if (!bChangeBoneIndicesFormat)
+				MUTABLE_CPUPROFILER_SCOPE(ResultFormatSetup);
+
+				// Start with the structure of the first mesh
+				Result->GetVertexBuffers().SetBufferCount(FirstBufferCount);
+				for (int32 BufferIndex = 0; BufferIndex < FirstBufferCount; ++BufferIndex)
 				{
-					pSecond->GetVertexBuffers().HasAnySemanticWithDifferentFormat(MBS_BONEINDICES, BoneIndexFormat);
+					Result->VertexBuffers.m_buffers[BufferIndex].m_channels = pFirst->VertexBuffers.m_buffers[BufferIndex].m_channels;
+					Result->VertexBuffers.m_buffers[BufferIndex].m_elementSize = pFirst->VertexBuffers.m_buffers[BufferIndex].m_elementSize;
 				}
 
-				bFastPath = bFastPath && !bChangeBoneIndicesFormat;
-			}
+				// See if we need to add additional buffers from the second mesh (like vertex colours or additional UV Channels)
+				// This is a bit ad-hoc: we only add buffers containing all new channels
+				for (const FMeshBuffer& SecondBuf : pSecond->GetVertexBuffers().m_buffers)
+				{
+					bool bSomeChannel = false;
+					bool bAllNewChannels = true;
+					for (const FMeshBufferChannel& SecondChan : SecondBuf.m_channels)
+					{
+						// Skip system buffers
+						if (SecondChan.m_semantic == MBS_VERTEXINDEX
+							||
+							SecondChan.m_semantic == MBS_LAYOUTBLOCK)
+						{
+							continue;
+						}
 
-			Ptr<const Mesh> pVFirst;
-			Ptr<const Mesh> pVSecond;
-			
-			if (!bFastPath)
-			{
-                MUTABLE_CPUPROFILER_SCOPE(SlowPath);
+						bSomeChannel = true;
 
-                // Very slow! Should we warn somehow?
+						int32 FoundBuffer = -1;
+						int32 FoundChannel = -1;
+						pFirst->GetVertexBuffers().FindChannel(SecondChan.m_semantic, SecondChan.m_semanticIndex, &FoundBuffer, &FoundChannel);
+						if (FoundBuffer >= 0)
+						{
+							// There's at least one channel that already exists in the first mesh. Don't add the buffer.
+							bAllNewChannels = false;
+							continue;
+						}
 
-				// Expand component counts in vertex channels of the format mesh
-				int32 vbcount = pFirst->GetVertexBuffers().m_buffers.Num();
-				Result->GetVertexBuffers().SetBufferCount( vbcount );
+						// If there are additional UV channels try to add them
+						if (!bAllNewChannels && SecondChan.m_semantic == MBS_TEXCOORDS
+							&&
+							SecondChan.m_semanticIndex > 0)
+						{
+							// Add additional UV channels if the previous one is found.
+							FMeshBufferSet& ResultVertexBuffers = Result->GetVertexBuffers();
+							ResultVertexBuffers.FindChannel(MBS_TEXCOORDS, SecondChan.m_semanticIndex - 1, &FoundBuffer, &FoundChannel);
 
-				for ( int32 vb = 0; vb<vbcount; ++vb )
+							if (FoundBuffer >= 0)
+							{
+								FMeshBuffer& Buffer = ResultVertexBuffers.m_buffers[FoundBuffer];
+								Buffer.m_channels.Insert(SecondChan, FoundChannel + 1);
+
+								ResultVertexBuffers.UpdateOffsets(FoundBuffer);
+							}
+						}
+					}
+
+					if (bSomeChannel && bAllNewChannels)
+					{
+						int32 NewBufferIndex = Result->GetVertexBuffers().m_buffers.Emplace();
+						Result->VertexBuffers.m_buffers[NewBufferIndex].m_channels = SecondBuf.m_channels;
+						Result->VertexBuffers.m_buffers[NewBufferIndex].m_elementSize = SecondBuf.m_elementSize;
+					}
+				}
+
+				// See if we need to enlarge the components of any of the result channels
+				int32 ResultBufferCount = Result->GetVertexBuffers().GetBufferCount();
+				for (int32 BufferIndex = 0; BufferIndex < FMath::Min(ResultBufferCount, FirstBufferCount); ++BufferIndex)
 				{
-					FMeshBuffer& result = Result->GetVertexBuffers().m_buffers[vb];
-					const FMeshBuffer& first = pFirst->GetVertexBuffers().m_buffers[vb];
+					// Expand component counts in vertex channels of the format mesh
+					FMeshBuffer& result = Result->GetVertexBuffers().m_buffers[BufferIndex];
+					const FMeshBuffer& first = pFirst->GetVertexBuffers().m_buffers[BufferIndex];
 
 					result.m_channels = first.m_channels;
 					result.m_elementSize = first.m_elementSize;
 
-					// See if we need to enlarge the components of any of the result channels
 					bool bResetOffsets = false;
-					for ( int32 c=0; c<result.m_channels.Num(); ++c )
+					for (int32 c = 0; c < result.m_channels.Num(); ++c)
 					{
 						int32 sb = -1;
 						int32 sc = -1;
 						pSecond->GetVertexBuffers().FindChannel
-								(
-									result.m_channels[c].m_semantic,
-									result.m_channels[c].m_semanticIndex,
-									&sb, &sc
-								);
-						if ( sb>=0 )
+						(
+							result.m_channels[c].m_semantic,
+							result.m_channels[c].m_semanticIndex,
+							&sb, &sc
+						);
+
+						if (sb >= 0)
 						{
 							const FMeshBuffer& second = pSecond->GetVertexBuffers().m_buffers[sb];
 
-							if ( second.m_channels[sc].m_componentCount
-								 >
-								 result.m_channels[c].m_componentCount )
+							if (second.m_channels[sc].m_componentCount>result.m_channels[c].m_componentCount)
 							{
 								result.m_channels[c].m_componentCount = second.m_channels[sc].m_componentCount;
 								bResetOffsets = true;
@@ -765,65 +811,13 @@ namespace mu
 					// Reset the channel offsets if necessary
 					if (bResetOffsets)
 					{
-						Result->GetVertexBuffers().UpdateOffsets(vb);
+						Result->GetVertexBuffers().UpdateOffsets(BufferIndex);
 					}
 				}
 
-                // See if we need to add additional buffers from the second mesh (like vertex colours or additional UV Channels)
-                // This is a bit ad-hoc: we only add buffers containing all new channels
-                for ( const FMeshBuffer& buf : pSecond->GetVertexBuffers().m_buffers )
-                {
-                    bool someChannel = false;
-                    bool allNewChannels = true;
-					for (const FMeshBufferChannel& chan : buf.m_channels)
-					{
-						// Skip system buffers
-						if (chan.m_semantic == MBS_VERTEXINDEX
-							||
-							chan.m_semantic == MBS_LAYOUTBLOCK)
-						{
-							continue;
-						}
-
-						someChannel = true;
-
-						int32 foundBuffer = -1;
-						int32 foundChannel = -1;
-						pFirst->GetVertexBuffers().FindChannel(chan.m_semantic, chan.m_semanticIndex, &foundBuffer, &foundChannel);
-						if (foundBuffer >= 0)
-						{
-							// There's at least one channel that already exists in the first mesh. Don't add the buffer.
-							allNewChannels = false;
-							continue;
-						}
-						
-						// If there are additional UV channels try to add them
-						if(!allNewChannels && chan.m_semantic == MBS_TEXCOORDS
-							&&
-							chan.m_semanticIndex > 0)
-						{
-							// Add additional UV channels if the previous one is found.
-							FMeshBufferSet& VertexBuffers = Result->GetVertexBuffers();
-							VertexBuffers.FindChannel(MBS_TEXCOORDS, chan.m_semanticIndex - 1, &foundBuffer, &foundChannel);
-
-							if (foundBuffer >= 0)
-							{
-								FMeshBuffer& Buffer = VertexBuffers.m_buffers[foundBuffer];
-								Buffer.m_channels.Insert(chan, foundChannel + 1);
-
-								VertexBuffers.UpdateOffsets(foundBuffer);
-							}
-						}
-					}
-
-                    if (someChannel && allNewChannels)
-                    {
-                        Result->GetVertexBuffers().m_buffers.Add(buf);
-                    }
-                }
 
 				// Change the format of the bone indices buffer
-				if(bChangeBoneIndicesFormat)
+				if (bChangeBoneIndicesFormat)
 				{
 					// Iterate all vertex buffers and update the format
 					FMeshBufferSet& VertexBuffers = Result->GetVertexBuffers();
@@ -853,208 +847,168 @@ namespace mu
 					}
 				}
 
+				// Manage vertex IDs
 				if (bNeedsExplicitIds)
 				{
 					// Make sure the result format is suitable for the explicit IDs
-
-					// Set to a dummy prefix, to ensure conversion happens
-					Result->MeshIDPrefix = 1;
 					Result->MakeIdsExplicit();
 				}
+			}
 
-				// Convert the source meshes to the new format
-                if (!bNeedsExplicitIds && pFirst->GetVertexBuffers().HasSameFormat(Result->GetVertexBuffers()))
-                {
-					Result->VertexBuffers = pVFirst->VertexBuffers;
-				}
-                else
-                {
-					check(ScratchMeshes.FirstReformat);
-
-					bool bOutSuccess = false;
-
-					if (bNeedsExplicitIds)
-					{
-						Ptr<Mesh> TempMesh = pFirst->Clone(EMeshCopyFlags::WithVertexBuffers);
-						TempMesh->MakeIdsExplicit();
-						MeshFormat(ScratchMeshes.FirstReformat.get(), TempMesh.get(), Result, false, true, false, false, bOutSuccess);
-					}
-					else
-					{
-						MeshFormat(ScratchMeshes.FirstReformat.get(), pFirst, Result, false, true, false, false, bOutSuccess);
-					}
-
-					if (bOutSuccess)
-					{
-						Result->VertexBuffers = ScratchMeshes.FirstReformat->VertexBuffers;
-					}
-					else
-					{
-						check(false);
-					}
 
-                }
+			// Step 2: Fill the result buffers
+			//-----------------------------------------------------------------------
+			{
+				MUTABLE_CPUPROFILER_SCOPE(ResultFill);
 
+				// We have the final result vertex buffers structure: allocate the memory for them.
+				Result->VertexBuffers.SetElementCount(FirstVertexCount + SecondVertexCount);
 
-                if (!bNeedsExplicitIds && pSecond->GetVertexBuffers().HasSameFormat(Result->GetVertexBuffers()))
-                {
-                    pVSecond = pSecond;
-                }
-                else
-                {
-					check(ScratchMeshes.SecondReformat);
+				int32 ResultBufferCount = Result->GetVertexBuffers().GetBufferCount();
+				for (int32 ResultBufferIndex = 0; ResultBufferIndex < ResultBufferCount; ++ResultBufferIndex)
+				{
+					FMeshBuffer& ResultBuffer = Result->VertexBuffers.m_buffers[ResultBufferIndex];
 
-					bool bOutSuccess = false;
+					// TODO: Don't assume the buffer order in First and Second matches Result, for more opportunities of fast-path
+					bool bFirstFastPath = (ResultBufferIndex < FirstBufferCount)
+						&&
+						pFirst->VertexBuffers.HasSameFormat(ResultBufferIndex, Result->VertexBuffers, ResultBufferIndex);
 
-					if (bNeedsExplicitIds)
+					int32 FirstResultBufferSize = ResultBuffer.m_elementSize * FirstVertexCount;
+					uint8* Destination = ResultBuffer.m_data.GetData();
+					if (bFirstFastPath)
 					{
-						Ptr<Mesh> TempMesh = pSecond->Clone(EMeshCopyFlags::WithVertexBuffers);
-						TempMesh->MakeIdsExplicit();
-						MeshFormat(ScratchMeshes.SecondReformat.get(), TempMesh.get(), Result, false, true, false, false, bOutSuccess);
-					}
-					else
-					{
-						MeshFormat(ScratchMeshes.SecondReformat.get(), pSecond, Result, false, true, false, false, bOutSuccess);
-					}
+						MUTABLE_CPUPROFILER_SCOPE(FirstFastPath);
 
-					if (bOutSuccess)
-					{
-						pVSecond = ScratchMeshes.SecondReformat;
+						const FMeshBuffer& FirstBuffer = pFirst->VertexBuffers.m_buffers[ResultBufferIndex];
+						check(FirstResultBufferSize == FirstBuffer.m_data.Num());
+						FMemory::Memcpy(Destination, FirstBuffer.m_data.GetData(), FirstResultBufferSize);
 					}
 					else
 					{
-						check(false);
-					}
-				}
+						MUTABLE_CPUPROFILER_SCOPE(FirstSlowPath);
 
-				check(Result->GetVertexBuffers().HasSameFormat(pVSecond->GetVertexBuffers()));
-			}
-			else
-			{
-                MUTABLE_CPUPROFILER_SCOPE(FastPath);
-
-				Result->VertexBuffers = pFirst->VertexBuffers;
-				pVSecond = pSecond;
-			}
+						int32 OffsetElements = 0;
+						MeshFormatBuffer(pFirst->VertexBuffers, ResultBuffer, OffsetElements, true, pFirst->MeshIDPrefix);
+					}
 
 
-			// Allocate additional vertices
-			Result->VertexBuffers.SetElementCount( FirstCount + SecondCount );
+					bool bSecondFastPath = (ResultBufferIndex < SecondBufferCount)
+						&&
+						pSecond->VertexBuffers.HasSameFormat(ResultBufferIndex, Result->VertexBuffers, ResultBufferIndex);
 
-            // Copy the second mesh vertex data
-			{
-				MUTABLE_CPUPROFILER_SCOPE(CopyVertexData);
-				for (int32 vb = 0; vb < Result->GetVertexBuffers().m_buffers.Num(); ++vb)
-				{
-					FMeshBuffer& result = Result->GetVertexBuffers().m_buffers[vb];
-					const FMeshBuffer& second = pVSecond->GetVertexBuffers().m_buffers[vb];
+					int32 SecondResultBufferSize = ResultBuffer.m_elementSize * SecondVertexCount;
+					Destination = ResultBuffer.m_data.GetData() + FirstResultBufferSize;
+					if (bSecondFastPath)
+					{
+						MUTABLE_CPUPROFILER_SCOPE(SecondFastPath);
 
-					if (SecondCount)
+						const FMeshBuffer& SecondBuffer = pSecond->VertexBuffers.m_buffers[ResultBufferIndex];
+						check(SecondResultBufferSize == SecondBuffer.m_data.Num());
+						FMemory::Memcpy(Destination, SecondBuffer.m_data.GetData(), SecondResultBufferSize);
+					}
+					else
 					{
-						int32 elemSize = Result->GetVertexBuffers().GetElementSize(vb);
-						int32 firstSize = FirstCount * elemSize;
-						int32 secondSize = SecondCount * elemSize;
+						MUTABLE_CPUPROFILER_SCOPE(SecondSlowPath);
 
-						FMemory::Memcpy( &result.m_data[firstSize], second.m_data.GetData(), secondSize);
+						int32 OffsetElements = FirstVertexCount;
+						MeshFormatBuffer(pSecond->VertexBuffers, ResultBuffer, OffsetElements, true, pSecond->MeshIDPrefix);
 					}
 				}
-			}
-
-
-            if (bRemapBoneIndices)
-            {
-                MUTABLE_CPUPROFILER_SCOPE(Remap);
 
-				// We need to remap the bone indices of the second mesh vertices that we already copied
-				// to result
-				check(!RemappedBoneMapIndices.IsEmpty())
-
-				// Iterate all vertex buffers and update the format
-				FMeshBufferSet& VertexBuffers = Result->GetVertexBuffers();
-				for (int32 VertexBufferIndex = 0; VertexBufferIndex < VertexBuffers.m_buffers.Num(); ++VertexBufferIndex)
+				// TODO: This still needs to be optimized
+				if (bRemapBoneIndices)
 				{
-					FMeshBuffer& ResultBuffer = VertexBuffers.m_buffers[VertexBufferIndex];
+					MUTABLE_CPUPROFILER_SCOPE(RemapBones);
 
-					const int32 ElemSize = VertexBuffers.GetElementSize(VertexBufferIndex);
-					const int32 FirstSize = FirstCount * ElemSize;
+					// We need to remap the bone indices of the second mesh vertices that we already copied to result
+					check(!RemappedBoneMapIndices.IsEmpty());
 
-					const int32 ChannelsCount = VertexBuffers.GetBufferChannelCount(VertexBufferIndex);
-					for (int32 ChannelIndex = 0; ChannelIndex < ChannelsCount; ++ChannelIndex)
+					// Iterate all vertex buffers and update the format
+					FMeshBufferSet& VertexBuffers = Result->GetVertexBuffers();
+					for (int32 VertexBufferIndex = 0; VertexBufferIndex < VertexBuffers.m_buffers.Num(); ++VertexBufferIndex)
 					{
-						if (ResultBuffer.m_channels[ChannelIndex].m_semantic != MBS_BONEINDICES)
-						{
-							continue;
-						}
-						
-						int32 ResultOffset = FirstSize + ResultBuffer.m_channels[ChannelIndex].m_offset;
+						FMeshBuffer& ResultBuffer = VertexBuffers.m_buffers[VertexBufferIndex];
 
-						const int32 NumComponents = ResultBuffer.m_channels[ChannelIndex].m_componentCount;
+						const int32 ElemSize = VertexBuffers.GetElementSize(VertexBufferIndex);
+						const int32 FirstSize = FirstVertexCount * ElemSize;
 
-						// Bone indices may need remapping
-						for (int32 VertexIndex = 0; VertexIndex < pVSecond->GetVertexCount(); ++VertexIndex)
+						const int32 ChannelsCount = VertexBuffers.GetBufferChannelCount(VertexBufferIndex);
+						for (int32 ChannelIndex = 0; ChannelIndex < ChannelsCount; ++ChannelIndex)
 						{
-							switch (BoneIndexFormat)
+							if (ResultBuffer.m_channels[ChannelIndex].m_semantic != MBS_BONEINDICES)
 							{
-							case MBF_INT8:
-							case MBF_UINT8:
-							{
-								uint8* pD = reinterpret_cast<uint8*>
-									(&ResultBuffer.m_data[ResultOffset]);
+								continue;
+							}
 
-								for (int32 ComponentIndex = 0; ComponentIndex < NumComponents; ++ComponentIndex)
+							int32 ResultOffset = FirstSize + ResultBuffer.m_channels[ChannelIndex].m_offset;
+
+							const int32 NumComponents = ResultBuffer.m_channels[ChannelIndex].m_componentCount;
+
+							// Bone indices may need remapping
+							for (int32 VertexIndex = 0; VertexIndex < pSecond->GetVertexCount(); ++VertexIndex)
+							{
+								switch (BoneIndexFormat)
+								{
+								case MBF_INT8:
+								case MBF_UINT8:
 								{
-									uint8 BoneMapIndex = pD[ComponentIndex];
+									uint8* pD = reinterpret_cast<uint8*>(&ResultBuffer.m_data[ResultOffset]);
 
-									// be defensive
-									if (BoneMapIndex < NumSecondBonesInBoneMap)
+									for (int32 ComponentIndex = 0; ComponentIndex < NumComponents; ++ComponentIndex)
 									{
-										pD[ComponentIndex] = (uint8)RemappedBoneMapIndices[BoneMapIndex];
+										uint8 BoneMapIndex = pD[ComponentIndex];
+
+										// be defensive
+										if (BoneMapIndex < NumSecondBonesInBoneMap)
+										{
+											pD[ComponentIndex] = (uint8)RemappedBoneMapIndices[BoneMapIndex];
+										}
+										else
+										{
+											pD[ComponentIndex] = 0;
+										}
 									}
-									else
-									{
-										pD[ComponentIndex] = 0;
-									}
-								}
 
-								ResultOffset += ElemSize;
-								break;
-							}
-
-							case MBF_INT16:
-							case MBF_UINT16:
-							{
-								uint16* pD = reinterpret_cast<uint16*>
-									(&ResultBuffer.m_data[ResultOffset]);
+									ResultOffset += ElemSize;
+									break;
+								}
 
-								for (int32 ComponentIndex = 0; ComponentIndex < NumComponents; ++ComponentIndex)
+								case MBF_INT16:
+								case MBF_UINT16:
 								{
-									uint16 BoneMapIndex = pD[ComponentIndex];
+									uint16* pD = reinterpret_cast<uint16*>(&ResultBuffer.m_data[ResultOffset]);
 
-									// be defensive
-									if (BoneMapIndex < NumSecondBonesInBoneMap)
-									{
-										pD[ComponentIndex] = (uint16)RemappedBoneMapIndices[BoneMapIndex];
-									}
-									else
+									for (int32 ComponentIndex = 0; ComponentIndex < NumComponents; ++ComponentIndex)
 									{
-										pD[ComponentIndex] = 0;
+										uint16 BoneMapIndex = pD[ComponentIndex];
+
+										// be defensive
+										if (BoneMapIndex < NumSecondBonesInBoneMap)
+										{
+											pD[ComponentIndex] = (uint16)RemappedBoneMapIndices[BoneMapIndex];
+										}
+										else
+										{
+											pD[ComponentIndex] = 0;
+										}
 									}
-								}
 
-								ResultOffset += ElemSize;
-								break;
-							}
+									ResultOffset += ElemSize;
+									break;
+								}
 
-							case MBF_INT32:
-							case MBF_UINT32:
-							{
-								// Unreal does not support 32 bit bone indices
-								checkf(false, TEXT("Format not supported."));
-								break;
-							}
+								case MBF_INT32:
+								case MBF_UINT32:
+								{
+									// Unreal does not support 32 bit bone indices
+									checkf(false, TEXT("Format not supported."));
+									break;
+								}
 
-							default:
-								checkf(false, TEXT("Format not supported."));
+								default:
+									checkf(false, TEXT("Format not supported."));
+								}
 							}
 						}
 					}
@@ -1062,6 +1016,7 @@ namespace mu
 			}
 		}
 
+
 		// Tags
 		Result->Tags = pFirst->Tags;
 
diff --git a/Engine/Plugins/Experimental/Mutable/Source/MutableRuntime/Private/MuR/Mesh.cpp b/Engine/Plugins/Experimental/Mutable/Source/MutableRuntime/Private/MuR/Mesh.cpp
--- a/Engine/Plugins/Experimental/Mutable/Source/MutableRuntime/Private/MuR/Mesh.cpp
+++ b/Engine/Plugins/Experimental/Mutable/Source/MutableRuntime/Private/MuR/Mesh.cpp
@@ -372,20 +372,13 @@ void Mesh::MakeVertexIdsRelative()
 
 void Mesh::MakeIdsExplicit()
 {
-	MUTABLE_CPUPROFILER_SCOPE(Mesh_MakeVertexIndicesExplicit);
-
-	if (!MeshIDPrefix)
-	{
-		// We already have explicit vertex IDs, or we don't have any.
-		return;
-	}
+	MUTABLE_CPUPROFILER_SCOPE(Mesh_MakeIdsExplicit);
 
 	int32 VertexCount = GetVertexCount();
+	check(VertexCount==0);
 
 	// Vertex IDs
 	{
-		MUTABLE_CPUPROFILER_SCOPE(VertexIDs);
-
 		bool bHasRelativeVertexIndices = false;
 
 		int32 OldBuf = -1;
@@ -396,62 +389,28 @@ void Mesh::MakeIdsExplicit()
 		{
 			check(OldChan == 0 && VertexBuffers.m_buffers[OldBuf].m_channels.Num() == 1);
 
-			FMeshBufferChannel& Channel = VertexBuffers.m_buffers[OldBuf].m_channels[0];
-
-			bool bHasExplicitVertexIndices = Channel.m_format == MBF_UINT64;
-			if (!bHasExplicitVertexIndices)
-			{
-				// The mesh has relative vertex IDs.
-				bHasRelativeVertexIndices = true;
-
-				check(Channel.m_format == MBF_UINT32);
-				const uint32* OldIdData = reinterpret_cast<const uint32*>(VertexBuffers.GetBufferData(OldBuf));
-
-				TMemoryTrackedArray<uint8> NewIds;
-				NewIds.SetNumUninitialized(VertexCount * sizeof(uint64));
-				uint64* NewIdData = reinterpret_cast<uint64*>(NewIds.GetData());
-
-				for (int32 i = 0; i < VertexCount; ++i)
-				{
-					uint32 OldId = *OldIdData++;
-					uint64 Id = (uint64(MeshIDPrefix) << 32) | uint64(OldId);
-					(*NewIdData++) = Id;
-				}
-
-				// 
-				FMeshBuffer& Buffer = VertexBuffers.m_buffers[OldBuf];
-				Swap(Buffer.m_data, NewIds);
-				Buffer.m_channels[0].m_format = MBF_UINT64;
-				Buffer.m_elementSize = sizeof(uint64);
-			}
+			FMeshBuffer& Buffer = VertexBuffers.m_buffers[OldBuf];
+			Buffer.m_channels[0].m_format = MBF_UINT64;
+			Buffer.m_elementSize = sizeof(uint64);
 		}
 
-		if (!bHasRelativeVertexIndices)
+		else
 		{
 			// The mesh has implicit Ids
 			// Create a new buffer with explicit ids
-			int32 NewBuffer = VertexBuffers.GetBufferCount();
-			VertexBuffers.SetBufferCount(NewBuffer + 1);
-			EMeshBufferSemantic Semantic = MBS_VERTEXINDEX;
-			int32 SemanticIndex = 0;
-			EMeshBufferFormat Format = MBF_UINT64;
-			int32 Components = 1;
-			int32 Offset = 0;
-			VertexBuffers.SetBuffer(NewBuffer, sizeof(uint64), 1, &Semantic, &SemanticIndex, &Format, &Components, &Offset);
-			uint64* IdData = reinterpret_cast<uint64*>(VertexBuffers.GetBufferData(NewBuffer));
-
-			for (int32 VertexIndex = 0; VertexIndex < VertexCount; ++VertexIndex)
-			{
-				uint64 Id = (uint64(MeshIDPrefix) << 32) | uint64(VertexIndex);
-				(*IdData++) = Id;
-			}
+			FMeshBuffer& Buffer = VertexBuffers.m_buffers.Emplace_GetRef();
+			Buffer.m_elementSize = sizeof(uint64);
+			FMeshBufferChannel& Channel = Buffer.m_channels.Emplace_GetRef();
+			Channel.m_semantic = MBS_VERTEXINDEX;
+			Channel.m_semanticIndex= 0;
+			Channel.m_format = MBF_UINT64;
+			Channel.m_componentCount = 1;
+			Channel.m_offset = 0;
 		}
 	}
 
 	// Layout block IDs
 	{
-		MUTABLE_CPUPROFILER_SCOPE(LayoutBlockIDs);
-
 		for (FMeshBuffer& Buffer: VertexBuffers.m_buffers)
 		{
 			for (FMeshBufferChannel& Channel : Buffer.m_channels)
@@ -461,31 +420,9 @@ void Mesh::MakeIdsExplicit()
 					continue;
 				}
 
-				if (Channel.m_format == MBF_UINT64)
-				{
-					continue;
-				}
-
 				check(Buffer.m_channels.Num() == 1);
 				check(Buffer.m_channels[0].m_offset == 0);
-				check(Buffer.m_elementSize == sizeof(uint16));
-
-				check(Channel.m_format == MBF_UINT16);
-				const uint16* OldIdData = reinterpret_cast<const uint16*>(Buffer.m_data.GetData());
-
-				TMemoryTrackedArray<uint8> NewIds;
-				NewIds.SetNumUninitialized(VertexCount *sizeof(uint64));
-				uint64* NewIdData = reinterpret_cast<uint64*>(NewIds.GetData());
-
-				for (int32 i = 0; i < VertexCount; ++i)
-				{
-					uint16 OldId = *OldIdData++;
-					uint64 Id = (uint64(MeshIDPrefix) << 32) | uint64(OldId);
-					(*NewIdData++) = Id;
-				}
 
-				// 
-				Swap(Buffer.m_data, NewIds);
 				Buffer.m_channels[0].m_format = MBF_UINT64;
 				Buffer.m_elementSize = sizeof(uint64);
 			}
@@ -612,7 +549,7 @@ void Mesh::GetSurface( int32 surfaceIndex,
 }
 
 
-uint32_t Mesh::GetSurfaceId( int surfaceIndex ) const
+uint32 Mesh::GetSurfaceId( int32 surfaceIndex ) const
 {
     if (surfaceIndex>=0 && surfaceIndex<Surfaces.Num())
     {
diff --git a/Engine/Plugins/Experimental/Mutable/Source/MutableRuntime/Private/MuR/MeshBufferSet.cpp b/Engine/Plugins/Experimental/Mutable/Source/MutableRuntime/Private/MuR/MeshBufferSet.cpp
--- a/Engine/Plugins/Experimental/Mutable/Source/MutableRuntime/Private/MuR/MeshBufferSet.cpp
+++ b/Engine/Plugins/Experimental/Mutable/Source/MutableRuntime/Private/MuR/MeshBufferSet.cpp
@@ -208,7 +208,7 @@ namespace mu
 			chan.m_semanticIndex = pSemanticIndices ? pSemanticIndices[c] : 0;
 			chan.m_format = pFormats ? pFormats[c] : MBF_NONE;
 			chan.m_componentCount = pComponentCount ? ((uint16)pComponentCount[c]) : 0;
-			chan.m_offset = pOffsets ? ((uint8_t)pOffsets[c]) : 0;
+			chan.m_offset = pOffsets ? ((uint8)pOffsets[c]) : 0;
 
 			int32 ThisChannelMinElemSize = chan.m_offset + chan.m_componentCount * GetMeshFormatData(chan.m_format).SizeInBytes;			
 			MinElemSize = FMath::Max(MinElemSize, ThisChannelMinElemSize);
@@ -262,7 +262,7 @@ namespace mu
 	uint8* FMeshBufferSet::GetBufferData(int32 buffer)
 	{
 		check(buffer >= 0 && buffer < m_buffers.Num());
-		uint8_t* pResult = m_buffers[buffer].m_data.GetData();
+		uint8* pResult = m_buffers[buffer].m_data.GetData();
 		return pResult;
 	}
 
@@ -271,7 +271,7 @@ namespace mu
 	const uint8* FMeshBufferSet::GetBufferData(int32 buffer) const
 	{
 		check(buffer >= 0 && buffer < m_buffers.Num());
-		const uint8_t* pResult = m_buffers[buffer].m_data.GetData();
+		const uint8* pResult = m_buffers[buffer].m_data.GetData();
 		return pResult;
 	}
 
@@ -303,15 +303,15 @@ namespace mu
 		*pBuffer = -1;
 		*pChannel = -1;
 
-		for (size_t b = 0; b < m_buffers.Num(); ++b)
+		for (int32 b = 0; b < m_buffers.Num(); ++b)
 		{
-			for (size_t c = 0; c < m_buffers[b].m_channels.Num(); ++c)
+			for (int32 c = 0; c < m_buffers[b].m_channels.Num(); ++c)
 			{
 				if (m_buffers[b].m_channels[c].m_semantic == semantic &&
 					m_buffers[b].m_channels[c].m_semanticIndex == semanticIndex)
 				{
-					*pBuffer = int(b);
-					*pChannel = int(c);
+					*pBuffer = b;
+					*pChannel = c;
 					return;
 				}
 			}
@@ -378,7 +378,7 @@ namespace mu
 
 		for (int32 b = 0; b < bc; ++b)
 		{
-			if (!HasSameFormat(b, Other))
+			if (!HasSameFormat(b, Other, b))
 			{
 				return false;
 			}
@@ -389,17 +389,9 @@ namespace mu
 
 
 	//---------------------------------------------------------------------------------------------
-	bool FMeshBufferSet::HasSameFormat(int32 buffer, const FMeshBufferSet& Other) const
+	bool FMeshBufferSet::HasSameFormat(int32 ThisBufferIndex, const FMeshBufferSet& Other, int32 OtherBufferIndex) const
 	{
-		if (m_buffers[buffer].m_channels == Other.m_buffers[buffer].m_channels
-			&&
-			m_buffers[buffer].m_elementSize == Other.m_buffers[buffer].m_elementSize
-			)
-		{
-			return true;
-		}
-
-		return false;
+		return m_buffers[ThisBufferIndex].HasSameFormat(Other.m_buffers[OtherBufferIndex]);
 	}
 
 
@@ -429,7 +421,7 @@ namespace mu
 	}
 
 	//-----------------------------------------------------------------------------------------
-	void FMeshBufferSet::CopyElement(uint32_t fromIndex, uint32_t toIndex)
+	void FMeshBufferSet::CopyElement(uint32 fromIndex, uint32 toIndex)
 	{
 		check(fromIndex < m_elementCount);
 		check(toIndex < m_elementCount);
@@ -659,7 +651,9 @@ namespace mu
 		}
 
 		if (m_buffers[b].m_elementSize < offset)
+		{
 			m_buffers[b].m_elementSize = offset;
+		}
 	}
 
 
diff --git a/Engine/Plugins/Experimental/Mutable/Source/MutableRuntime/Private/MuR/OpMeshFormat.cpp b/Engine/Plugins/Experimental/Mutable/Source/MutableRuntime/Private/MuR/OpMeshFormat.cpp
--- a/Engine/Plugins/Experimental/Mutable/Source/MutableRuntime/Private/MuR/OpMeshFormat.cpp
+++ b/Engine/Plugins/Experimental/Mutable/Source/MutableRuntime/Private/MuR/OpMeshFormat.cpp
@@ -24,261 +24,367 @@ namespace mu
 {
 
 	//-------------------------------------------------------------------------------------------------
-	void MeshFormatBuffer( const FMeshBufferSet& Source, FMeshBufferSet& Result, int32 bufferIndex )
+	void MeshFormatBuffer( const FMeshBufferSet& Source, FMeshBuffer& ResultBuffer, int32 ResultOffsetElements, bool bHasSpecialSemantics, uint32 IDPrefix )
 	{
-		int vCount = Source.GetElementCount();
+		int32 SourceElementCount = Source.GetElementCount();
+		if (!SourceElementCount)
+		{
+			return;
+		}
+
+		// This happens in the debugger
+		if (ResultBuffer.m_channels.IsEmpty() || ResultBuffer.m_elementSize==0)
+		{
+			return;
+		}
+
+		check(ResultBuffer.m_elementSize);
+		int32 ResultElementCount = ResultBuffer.m_data.Num() / ResultBuffer.m_elementSize;
+		check(SourceElementCount + ResultOffsetElements <= ResultElementCount);
 
-		int b = bufferIndex;
+		// For every channel in this buffer
+		for (int32 ResultChannelIndex = 0; ResultChannelIndex < ResultBuffer.m_channels.Num(); ++ResultChannelIndex)
 		{
-			// For every channel in this buffer
-			for (int c = 0; c < Result.GetBufferChannelCount(b); ++c)
+			FMeshBufferChannel& ResultChannel = ResultBuffer.m_channels[ResultChannelIndex];
+
+			// Find this channel in the source mesh
+			int32 SourceBufferIndex;
+			int32 SourceChannelIndex;
+			Source.FindChannel(ResultChannel.m_semantic, ResultChannel.m_semanticIndex, &SourceBufferIndex, &SourceChannelIndex);
+
+			int32 resultElemSize = ResultBuffer.m_elementSize;
+			int32 resultComponents = ResultChannel.m_componentCount;
+			int32 resultChannelSize = GetMeshFormatData(ResultChannel.m_format).SizeInBytes * resultComponents;
+			uint8* pResultBuf = ResultBuffer.m_data.GetData() + ResultOffsetElements*ResultBuffer.m_elementSize;
+			pResultBuf += ResultChannel.m_offset;
+
+			// Case 1: Special semantics that may be implicit or relative
+			//---------------------------------------------------------------------------------
+			if (bHasSpecialSemantics && ResultChannel.m_semantic == MBS_VERTEXINDEX)
 			{
-				// Find this channel in the source mesh
-				EMeshBufferSemantic resultSemantic;
-				int resultSemanticIndex;
-				EMeshBufferFormat resultFormat;
-				int resultComponents;
-				int resultOffset;
-				Result.GetChannel
-				(
-					b, c,
-					&resultSemantic, &resultSemanticIndex,
-					&resultFormat, &resultComponents,
-					&resultOffset
-				);
+				bool bHasVertexIndices = (SourceBufferIndex >=0 );
+				if (bHasVertexIndices)
+				{
+					check(SourceChannelIndex == 0 && Source.m_buffers[SourceBufferIndex].m_channels.Num() == 1);
+
+					const FMeshBuffer& SourceBuffer = Source.m_buffers[SourceBufferIndex];
+					const FMeshBufferChannel& SourceChannel = SourceBuffer.m_channels[SourceChannelIndex];
 
-				int sourceBuffer;
-				int sourceChannel;
-				Source.FindChannel
-				(resultSemantic, resultSemanticIndex, &sourceBuffer, &sourceChannel);
+					bool bHasSameFormat = SourceChannel.m_format == ResultChannel.m_format;
+					if (bHasSameFormat)
+					{
+						check(SourceChannel == ResultChannel);
+						FMemory::Memcpy(pResultBuf,SourceBuffer.m_data.GetData(),SourceBuffer.m_data.Num());
+					}
+					else
+					{
+						// Relative vertex IDs
+						check(IDPrefix);
+						check(SourceChannel.m_format == MBF_UINT32);
+						const uint32* SourceData = reinterpret_cast<const uint32*>(SourceBuffer.m_data.GetData());
 
-				int resultElemSize = Result.GetElementSize(b);
-				int resultChannelSize = GetMeshFormatData(resultFormat).SizeInBytes * resultComponents;
-				uint8_t* pResultBuf = Result.GetBufferData(b);
-				pResultBuf += resultOffset;
+						check(ResultChannel.m_format == MBF_UINT64);
+						uint64* ResultData = reinterpret_cast<uint64*>(pResultBuf);
 
-				if (sourceBuffer < 0)
+						for (int32 i = 0; i < SourceElementCount; ++i)
+						{
+							uint32 SourceId = *SourceData++;
+							uint64 Id = (uint64(IDPrefix) << 32) | uint64(SourceId);
+							*ResultData = Id;
+							++ResultData;
+						}
+					}
+				}
+				else
 				{
-					// Not found: fill with zeros.
+					// Implicit IDs
+					check(IDPrefix);
 
-					// Special case for derived channel data
-					bool generated = false;
+					check(ResultChannel.m_format == MBF_UINT64);
+					uint64* ResultData = reinterpret_cast<uint64*>(pResultBuf);
 
-					// If we have to add colour channels, we will add them as white, to be neutral.
-					// \todo: normal channels also should have special values.
-					if (resultSemantic == MBS_COLOUR)
+					for (int32 VertexIndex = 0; VertexIndex < SourceElementCount; ++VertexIndex)
 					{
-						generated = true;
+						uint64 Id = (uint64(IDPrefix) << 32) | uint64(VertexIndex);
+						*ResultData = Id;
+						++ResultData;
+					}
+				}
 
-						switch (resultFormat)
-						{
-						case MBF_FLOAT32:
+				continue;
+			}
+
+			else if (bHasSpecialSemantics && ResultChannel.m_semantic == MBS_LAYOUTBLOCK)
+			{
+				bool bHasBlockIds = (SourceBufferIndex >= 0);
+
+				if (bHasBlockIds)
+				{
+					const FMeshBuffer& SourceBuffer = Source.m_buffers[SourceBufferIndex];
+					const FMeshBufferChannel& SourceChannel = SourceBuffer.m_channels[SourceChannelIndex];
+
+					bool bHasSameFormat = SourceChannel.m_format == ResultChannel.m_format;
+					if (bHasSameFormat)
+					{
+						check(SourceChannel==ResultChannel);
+						FMemory::Memcpy(pResultBuf, SourceBuffer.m_data.GetData(), SourceBuffer.m_data.Num());
+					}
+					else
+					{
+						// Relative vertex IDs
+						check(SourceChannel.m_format == MBF_UINT16);
+						const uint16* SourceData = reinterpret_cast<const uint16*>(SourceBuffer.m_data.GetData());
+
+						check(ResultChannel.m_format == MBF_UINT64);
+						uint64* ResultData = reinterpret_cast<uint64*>(pResultBuf);
+
+						for (int32 i = 0; i < SourceElementCount; ++i)
 						{
-							for (int v = 0; v < vCount; ++v)
-							{
-								float* pTypedResultBuf = (float*)pResultBuf;
-								for (int i = 0; i < resultComponents; ++i)
-								{
-									pTypedResultBuf[i] = 1.0f;
-								}
-								pResultBuf += resultElemSize;
-							}
-							break;
+							uint16 SourceId = *SourceData++;
+							uint64 Id = (uint64(IDPrefix) << 32) | uint64(SourceId);
+							*ResultData = Id;
+							++ResultData;
 						}
+					}
+					continue;
+				}
+				else
+				{
+					// This seems to happen with objects that mix meshes with layouts with meshes without layouts.
+					// If we don't do anything here, it will be filled with zeros in the following code.
+					//ensure(false);
+				}
+
+			}
+
+			// Case 2: Not found in source: generate with default values, depending on semantic
+			//---------------------------------------------------------------------------------
+			if (SourceBufferIndex < 0)
+			{
+				// Not found: fill with zeros.
 
-						case MBF_NUINT8:
+				// Special case for derived channel data
+				bool generated = false;
+
+				// If we have to add colour channels, we will add them as white, to be neutral.
+				// \todo: normal channels also should have special values.
+				if (ResultChannel.m_semantic == MBS_COLOUR)
+				{
+					generated = true;
+
+					switch (ResultChannel.m_format)
+					{
+					case MBF_FLOAT32:
+					{
+						for (int32 v = 0; v < SourceElementCount; ++v)
 						{
-							for (int v = 0; v < vCount; ++v)
+							float* pTypedResultBuf = (float*)pResultBuf;
+							for (int32 i = 0; i < resultComponents; ++i)
 							{
-								uint8_t* pTypedResultBuf = (uint8_t*)pResultBuf;
-								for (int i = 0; i < resultComponents; ++i)
-								{
-									pTypedResultBuf[i] = 255;
-								}
-								pResultBuf += resultElemSize;
+								pTypedResultBuf[i] = 1.0f;
 							}
-							break;
+							pResultBuf += resultElemSize;
 						}
+						break;
+					}
 
-						case MBF_NUINT16:
+					case MBF_NUINT8:
+					{
+						for (int32 v = 0; v < SourceElementCount; ++v)
 						{
-							for (int v = 0; v < vCount; ++v)
+							uint8* pTypedResultBuf = (uint8*)pResultBuf;
+							for (int32 i = 0; i < resultComponents; ++i)
 							{
-								uint16* pTypedResultBuf = (uint16*)pResultBuf;
-								for (int i = 0; i < resultComponents; ++i)
-								{
-									pTypedResultBuf[i] = 65535;
-								}
-								pResultBuf += resultElemSize;
+								pTypedResultBuf[i] = 255;
 							}
-							break;
-						}
-
-						default:
-							// Format not implemented
-							check(false);
-							break;
+							pResultBuf += resultElemSize;
 						}
+						break;
 					}
 
-					if (!generated)
+					case MBF_NUINT16:
 					{
-						// TODO: and maybe raise a warning?
-						for (int v = 0; v < vCount; ++v)
+						for (int32 v = 0; v < SourceElementCount; ++v)
 						{
-							FMemory::Memzero(pResultBuf, resultChannelSize);
+							uint16* pTypedResultBuf = (uint16*)pResultBuf;
+							for (int32 i = 0; i < resultComponents; ++i)
+							{
+								pTypedResultBuf[i] = 65535;
+							}
 							pResultBuf += resultElemSize;
 						}
+						break;
+					}
+
+					default:
+						// Format not implemented
+						check(false);
+						break;
 					}
 				}
-				else
+
+				if (!generated)
 				{
-					// Get the data about the source format
-					EMeshBufferSemantic sourceSemantic;
-					int sourceSemanticIndex;
-					EMeshBufferFormat sourceFormat;
-					int sourceComponents;
-					int sourceOffset;
-					Source.GetChannel
-					(
-						sourceBuffer, sourceChannel,
-						&sourceSemantic, &sourceSemanticIndex,
-						&sourceFormat, &sourceComponents,
-						&sourceOffset
-					);
-					check(sourceSemantic == resultSemantic
-						&&
-						sourceSemanticIndex == resultSemanticIndex);
+					// TODO: and maybe raise a warning?
+					for (int32 v = 0; v < SourceElementCount; ++v)
+					{
+						FMemory::Memzero(pResultBuf, resultChannelSize);
+						pResultBuf += resultElemSize;
+					}
+				}
+
+				continue;
+			}
 
-					int32 sourceElemSize = Source.GetElementSize(sourceBuffer);
-					const uint8* pSourceBuf = Source.GetBufferData(sourceBuffer);
-					pSourceBuf += sourceOffset;
 
-					// Copy element by element
-					for (int v = 0; v < vCount; ++v)
+			// Case 3: Convert element by element
+			//---------------------------------------------------------------------------------
+			{
+				// Get the data about the source format
+				EMeshBufferSemantic sourceSemantic;
+				int32 sourceSemanticIndex;
+				EMeshBufferFormat sourceFormat;
+				int32 sourceComponents;
+				int32 sourceOffset;
+				Source.GetChannel
+				(
+					SourceBufferIndex, SourceChannelIndex,
+					&sourceSemantic, &sourceSemanticIndex,
+					&sourceFormat, &sourceComponents,
+					&sourceOffset
+				);
+				check(sourceSemantic == ResultChannel.m_semantic
+					&&
+					sourceSemanticIndex == ResultChannel.m_semanticIndex);
+
+				int32 sourceElemSize = Source.GetElementSize(SourceBufferIndex);
+				const uint8* pSourceBuf = Source.GetBufferData(SourceBufferIndex);
+				pSourceBuf += sourceOffset;
+
+				// Copy element by element
+				for (int32 v = 0; v < SourceElementCount; ++v)
+				{
+					if (ResultChannel.m_format == sourceFormat && resultComponents == sourceComponents)
 					{
-						if (resultFormat == sourceFormat && resultComponents == sourceComponents)
-						{
-							FMemory::Memcpy(pResultBuf, pSourceBuf, resultChannelSize);
-						}
-						else if (resultFormat == MBF_PACKEDDIR8_W_TANGENTSIGN
-							||
-							resultFormat == MBF_PACKEDDIRS8_W_TANGENTSIGN)
-						{
-							check(sourceComponents >= 3);
-							check(resultComponents == 4);
+						FMemory::Memcpy(pResultBuf, pSourceBuf, resultChannelSize);
+					}
+					else if (ResultChannel.m_format == MBF_PACKEDDIR8_W_TANGENTSIGN
+						||
+						ResultChannel.m_format == MBF_PACKEDDIRS8_W_TANGENTSIGN)
+					{
+						check(sourceComponents >= 3);
+						check(resultComponents == 4);
 
-							// convert the 3 first components
-							for (int i = 0; i < 3; ++i)
+						// convert the 3 first components
+						for (int32 i = 0; i < 3; ++i)
+						{
+							if (i < sourceComponents)
 							{
-								if (i < sourceComponents)
-								{
-									ConvertData
-									(
-										i,
-										pResultBuf, resultFormat,
-										pSourceBuf, sourceFormat
-									);
-								}
+								ConvertData
+								(
+									i,
+									pResultBuf, ResultChannel.m_format,
+									pSourceBuf, sourceFormat
+								);
 							}
+						}
 
 
-							// Add the tangent sign
-							uint8* pData = reinterpret_cast<uint8*>(pResultBuf);
+						// Add the tangent sign
+						uint8* pData = reinterpret_cast<uint8*>(pResultBuf);
 
-							// Look for the full tangent space
-							int32 tanXBuf, tanXChan, tanYBuf, tanYChan, tanZBuf, tanZChan;
-							Source.FindChannel(MBS_TANGENT, resultSemanticIndex, &tanXBuf, &tanXChan);
-							Source.FindChannel(MBS_BINORMAL, resultSemanticIndex, &tanYBuf, &tanYChan);
-							Source.FindChannel(MBS_NORMAL, resultSemanticIndex, &tanZBuf, &tanZChan);
+						// Look for the full tangent space
+						int32 tanXBuf, tanXChan, tanYBuf, tanYChan, tanZBuf, tanZChan;
+						Source.FindChannel(MBS_TANGENT, ResultChannel.m_semanticIndex, &tanXBuf, &tanXChan);
+						Source.FindChannel(MBS_BINORMAL, ResultChannel.m_semanticIndex, &tanYBuf, &tanYChan);
+						Source.FindChannel(MBS_NORMAL, ResultChannel.m_semanticIndex, &tanZBuf, &tanZChan);
 
-							if (tanXBuf >= 0 && tanYBuf >= 0 && tanZBuf >= 0)
-							{
-								UntypedMeshBufferIteratorConst xIt(Source, MBS_TANGENT, resultSemanticIndex);
-								UntypedMeshBufferIteratorConst yIt(Source, MBS_BINORMAL, resultSemanticIndex);
-								UntypedMeshBufferIteratorConst zIt(Source, MBS_NORMAL, resultSemanticIndex);
+						if (tanXBuf >= 0 && tanYBuf >= 0 && tanZBuf >= 0)
+						{
+							UntypedMeshBufferIteratorConst xIt(Source, MBS_TANGENT, ResultChannel.m_semanticIndex);
+							UntypedMeshBufferIteratorConst yIt(Source, MBS_BINORMAL, ResultChannel.m_semanticIndex);
+							UntypedMeshBufferIteratorConst zIt(Source, MBS_NORMAL, ResultChannel.m_semanticIndex);
 
-								xIt += v;
-								yIt += v;
-								zIt += v;
+							xIt += v;
+							yIt += v;
+							zIt += v;
 
-								FMatrix44f Mat(xIt.GetAsVec3f(), yIt.GetAsVec3f(), zIt.GetAsVec3f(), FVector3f(0, 0, 0));
+							FMatrix44f Mat(xIt.GetAsVec3f(), yIt.GetAsVec3f(), zIt.GetAsVec3f(), FVector3f(0, 0, 0));
 
-								uint8 sign = 0;
-								if (resultFormat == MBF_PACKEDDIR8_W_TANGENTSIGN)
-								{
-									sign = Mat.RotDeterminant() < 0 ? 0 : 255;
-								}
-								else if (resultFormat == MBF_PACKEDDIRS8_W_TANGENTSIGN)
-								{
-									sign = Mat.RotDeterminant() < 0 ? -128 : 127;
-								}
-								pData[3] = sign;
+							uint8 sign = 0;
+							if (ResultChannel.m_format == MBF_PACKEDDIR8_W_TANGENTSIGN)
+							{
+								sign = Mat.RotDeterminant() < 0 ? 0 : 255;
 							}
-							else
+							else if (ResultChannel.m_format == MBF_PACKEDDIRS8_W_TANGENTSIGN)
 							{
-								// At least initialize it to avoid garbage.
-								pData[3] = 0;
+								sign = Mat.RotDeterminant() < 0 ? -128 : 127;
 							}
+							pData[3] = sign;
 						}
 						else
 						{
-							// Convert formats
-							for (int i = 0; i < resultComponents; ++i)
+							// At least initialize it to avoid garbage.
+							pData[3] = 0;
+						}
+					}
+					else
+					{
+						// Convert formats
+						for (int32 i = 0; i < resultComponents; ++i)
+						{
+							if (i < sourceComponents)
 							{
-								if (i < sourceComponents)
-								{
-									ConvertData
-									(
-										i,
-										pResultBuf, resultFormat,
-										pSourceBuf, sourceFormat
-									);
-								}
-								else
-								{
-									// Add zeros. TODO: Warning?
-									FMemory::Memzero
-									(
-										pResultBuf + GetMeshFormatData(resultFormat).SizeInBytes * i,
-										GetMeshFormatData(resultFormat).SizeInBytes
-									);
-								}
+								ConvertData
+								(
+									i,
+									pResultBuf, ResultChannel.m_format,
+									pSourceBuf, sourceFormat
+								);
+							}
+							else
+							{
+								// Add zeros. TODO: Warning?
+								FMemory::Memzero
+								(
+									pResultBuf + GetMeshFormatData(ResultChannel.m_format).SizeInBytes * i,
+									GetMeshFormatData(ResultChannel.m_format).SizeInBytes
+								);
 							}
+						}
 
 
-							// Extra step to normalise some semantics in some formats
-							// TODO: Make it optional, and add different normalisation types n, n^2
-							// TODO: Optimise
-							if (sourceSemantic == MBS_BONEWEIGHTS)
+						// Extra step to normalise some semantics in some formats
+						// TODO: Make it optional, and add different normalisation types n, n^2
+						// TODO: Optimise
+						if (sourceSemantic == MBS_BONEWEIGHTS)
+						{
+							if (ResultChannel.m_format == MBF_NUINT8)
 							{
-								if (resultFormat == MBF_NUINT8)
+								uint8* pData = (uint8*)pResultBuf;
+								uint8 accum = 0;
+								for (int32 i = 0; i < resultComponents; ++i)
 								{
-									uint8_t* pData = (uint8_t*)pResultBuf;
-									uint8_t accum = 0;
-									for (int i = 0; i < resultComponents; ++i)
-									{
-										accum += pData[i];
-									}
-									pData[0] += 255 - accum;
+									accum += pData[i];
 								}
+								pData[0] += 255 - accum;
+							}
 
-								else if (resultFormat == MBF_NUINT16)
+							else if (ResultChannel.m_format == MBF_NUINT16)
+							{
+								uint16* pData = (uint16*)pResultBuf;
+								uint16 accum = 0;
+								for (int32 i = 0; i < resultComponents; ++i)
 								{
-									uint16* pData = (uint16*)pResultBuf;
-									uint16 accum = 0;
-									for (int i = 0; i < resultComponents; ++i)
-									{
-										accum += pData[i];
-									}
-									pData[0] += 65535 - accum;
+									accum += pData[i];
 								}
+								pData[0] += 65535 - accum;
 							}
 						}
-
-						pResultBuf += resultElemSize;
-						pSourceBuf += sourceElemSize;
 					}
+
+					pResultBuf += resultElemSize;
+					pSourceBuf += sourceElemSize;
 				}
 			}
 		}
@@ -292,7 +398,8 @@ namespace mu
 		FMeshBufferSet& Result,
 		bool bKeepSystemBuffers,
 		bool bIgnoreMissingChannels,
-		bool bIsVertexBuffer
+		bool bIsVertexBuffer, 
+		uint32 IDPrefix = 0
 	)
 	{
 		if (bIgnoreMissingChannels)
@@ -364,7 +471,7 @@ namespace mu
 		Result.SetElementCount(vCount);
 		for (int32 b = 0; b < Result.GetBufferCount(); ++b)
 		{
-			MeshFormatBuffer(Source, Result, b);
+			MeshFormatBuffer(Source, Result.m_buffers[b], 0, bIsVertexBuffer, IDPrefix);
 		}
 
 
@@ -473,14 +580,14 @@ namespace mu
 						if (resultBuf >= 0)
 						{
 							UntypedMeshBufferIteratorConst it(VertexBuffers, MBS_BONEINDICES, Channel.m_semanticIndex);
-							int32_t maxBoneIndex = 0;
-							for (int v = 0; v < VertexBuffers.GetElementCount(); ++v)
+							int32 maxBoneIndex = 0;
+							for (int32 v = 0; v < VertexBuffers.GetElementCount(); ++v)
 							{
 								// If MAX_TOTAL_INFLUENCES ever changed, the next line would no longer work or compile and 
 								// GetAsVec12i would need to be changed accordingly
 								int32 va[MAX_TOTAL_INFLUENCES];
 								it.GetAsInt32Vec(va, MAX_TOTAL_INFLUENCES);
-								for (int c = 0; c < it.GetComponents(); ++c)
+								for (int32 c = 0; c < it.GetComponents(); ++c)
 								{
 									maxBoneIndex = FMath::Max(maxBoneIndex, va[c]);
 								}
@@ -517,7 +624,7 @@ namespace mu
 		if (formatVertices)
 		{
 			FormatBufferSet(Source->GetVertexBuffers(), Result->GetVertexBuffers(),
-				keepSystemBuffers, ignoreMissingChannels, true);
+				keepSystemBuffers, ignoreMissingChannels, true, Source->MeshIDPrefix);
 		}
 		else
 		{
diff --git a/Engine/Plugins/Experimental/Mutable/Source/MutableRuntime/Public/MuR/Mesh.h b/Engine/Plugins/Experimental/Mutable/Source/MutableRuntime/Public/MuR/Mesh.h
--- a/Engine/Plugins/Experimental/Mutable/Source/MutableRuntime/Public/MuR/Mesh.h
+++ b/Engine/Plugins/Experimental/Mutable/Source/MutableRuntime/Public/MuR/Mesh.h
@@ -246,6 +246,8 @@ namespace mu
 
 		/** Create an explicit vertex buffer for vertex IDs if they are implicit. */
 		void MakeVertexIdsRelative();
+
+		/** Ensure the format of an empty mesh includes explicit IDs. The mesh cannot have any vertex data. */
 		void MakeIdsExplicit();
 
         //! \name Texture layouts
diff --git a/Engine/Plugins/Experimental/Mutable/Source/MutableRuntime/Public/MuR/MeshBufferSet.h b/Engine/Plugins/Experimental/Mutable/Source/MutableRuntime/Public/MuR/MeshBufferSet.h
--- a/Engine/Plugins/Experimental/Mutable/Source/MutableRuntime/Public/MuR/MeshBufferSet.h
+++ b/Engine/Plugins/Experimental/Mutable/Source/MutableRuntime/Public/MuR/MeshBufferSet.h
@@ -179,12 +179,6 @@ namespace mu
 		template<typename Type>
 		using TMemoryTrackedArray = TArray<Type, FDefaultMemoryTrackingAllocator<MemoryCounters::FMeshMemoryCounter>>;
 
-		//!
-		FMeshBuffer()
-		{
-			m_elementSize = 0;
-		}
-
 		//!
 		TArray<FMeshBufferChannel> m_channels;
 
@@ -192,7 +186,7 @@ namespace mu
 		TMemoryTrackedArray<uint8> m_data;
 
 		//!
-		uint32 m_elementSize;
+		uint32 m_elementSize = 0;
 
 		//!
 		void Serialise(mu::OutputArchive& arch) const;
@@ -209,6 +203,24 @@ namespace mu
 			return equal;
 		}
 
+		/** Return true if the buffer has any channel with the passed semantic. */
+		inline bool HasSemantic(EMeshBufferSemantic Semantic) const
+		{
+			for ( const FMeshBufferChannel& Channel : m_channels)
+			{
+				if (Channel.m_semantic == Semantic)
+				{
+					return true;
+				}
+			}
+			return false;
+		}
+
+		inline bool HasSameFormat(const FMeshBuffer& Other) const
+		{
+			return (m_channels == Other.m_channels && m_elementSize == Other.m_elementSize);
+		}
+
 	};
 
 
@@ -366,7 +378,7 @@ namespace mu
 		void CopyElement(uint32 fromIndex, uint32 toIndex);
 
 		//! Compare the format of the two buffers at index buffer and return true if they match.
-		bool HasSameFormat(int32 buffer, const FMeshBufferSet& pOther) const;
+		bool HasSameFormat(int32 ThisBufferIndex, const FMeshBufferSet& pOther, int32 OtherBufferIndex) const;
 
 		//! Get the total memory size of the buffers and this struct
 		int32 GetDataSize() const;
