[mutable] Make the core more resilient to missing streaming data.
- When loading streamed data, if mipmaps are missing from disk but there are smaller ones, return a smaller image.
- Otherwise, hard fail, but don't immediately crash.
Mutable will adjust the size of the expected image in many cases, but this means that for now it will usually scale it up to what was required, which may not be ideal.

#rnx

diff --git a/Engine/Plugins/Experimental/Mutable/Source/CustomizableObject/Internal/MuCO/UnrealMutableModelDiskStreamer.h b/Engine/Plugins/Experimental/Mutable/Source/CustomizableObject/Internal/MuCO/UnrealMutableModelDiskStreamer.h
--- a/Engine/Plugins/Experimental/Mutable/Source/CustomizableObject/Internal/MuCO/UnrealMutableModelDiskStreamer.h
+++ b/Engine/Plugins/Experimental/Mutable/Source/CustomizableObject/Internal/MuCO/UnrealMutableModelDiskStreamer.h
@@ -58,7 +58,7 @@ public:
 	// mu::ModelReader interface
 	OPERATION_ID BeginReadBlock(const mu::Model*, uint32 key0, void* pBuffer, uint64 size, TFunction<void(bool bSuccess)>* CompletionCallback) override;
 	bool IsReadCompleted(OPERATION_ID) override;
-	void EndRead(OPERATION_ID) override;
+	bool EndRead(OPERATION_ID) override;
 
 protected:
 
diff --git a/Engine/Plugins/Experimental/Mutable/Source/CustomizableObject/Private/MuCO/CustomizableObjectMipDataProvider.cpp b/Engine/Plugins/Experimental/Mutable/Source/CustomizableObject/Private/MuCO/CustomizableObjectMipDataProvider.cpp
--- a/Engine/Plugins/Experimental/Mutable/Source/CustomizableObject/Private/MuCO/CustomizableObjectMipDataProvider.cpp
+++ b/Engine/Plugins/Experimental/Mutable/Source/CustomizableObject/Private/MuCO/CustomizableObjectMipDataProvider.cpp
@@ -458,8 +458,18 @@ bool FMutableTextureMipDataProvider::PollMips(const FTextureUpdateSyncOptions& S
 		{
 			ImageLODCount = Image->GetLODCount();
 			// check(Image->GetLODCount() == OperationData->Levels.Num()); TODO PRP
-			check(Image->GetSizeX() == OperationData->Levels[0].SizeX);
-			check(Image->GetSizeY() == OperationData->Levels[0].SizeY);
+			
+			// No longer true, since missing data may mean we generate smaller images.
+			//check(Image->GetSizeX() == OperationData->Levels[0].SizeX);
+			//check(Image->GetSizeY() == OperationData->Levels[0].SizeY);
+			if (Image->GetSizeX() != OperationData->Levels[0].SizeX
+				||
+				Image->GetSizeY() != OperationData->Levels[0].SizeY)
+			{
+				OperationData = nullptr;
+				AdvanceTo(ETickState::CleanUp, ETickThread::Async);
+				return false;
+			}
 		}
 
 		int32 MipIndex = 0;
diff --git a/Engine/Plugins/Experimental/Mutable/Source/CustomizableObject/Private/MuCO/UnrealMutableModelDiskStreamer.cpp b/Engine/Plugins/Experimental/Mutable/Source/CustomizableObject/Private/MuCO/UnrealMutableModelDiskStreamer.cpp
--- a/Engine/Plugins/Experimental/Mutable/Source/CustomizableObject/Private/MuCO/UnrealMutableModelDiskStreamer.cpp
+++ b/Engine/Plugins/Experimental/Mutable/Source/CustomizableObject/Private/MuCO/UnrealMutableModelDiskStreamer.cpp
@@ -308,10 +308,11 @@ bool FUnrealMutableModelBulkReader::IsReadCompleted(mu::ModelReader::OPERATION_I
 }
 
 
-void FUnrealMutableModelBulkReader::EndRead(mu::ModelReader::OPERATION_ID OperationId)
+bool FUnrealMutableModelBulkReader::EndRead(mu::ModelReader::OPERATION_ID OperationId)
 {
 	MUTABLE_CPUPROFILER_SCOPE(FUnrealMutableModelBulkStreamer::EndRead);
 
+	bool bSuccess = true;
 	bool bFound = false;
 	for (FObjectData& o : Objects)
 	{
@@ -325,6 +326,12 @@ void FUnrealMutableModelBulkReader::EndRead(mu::ModelReader::OPERATION_ID Operat
 				{
 					UE_LOG(LogMutable, Error, TEXT("Operation failed to complete in EndRead."));
 					check(false);
+					bSuccess = false;
+				}
+				else if (ReadRequest->ReadRequest->GetReadResults()==nullptr)
+				{
+					// This means we have failed: file not found?
+					bSuccess = false;
 				}
 			}
 			o.CurrentReadRequests.Remove(OperationId);
@@ -338,6 +345,8 @@ void FUnrealMutableModelBulkReader::EndRead(mu::ModelReader::OPERATION_ID Operat
 		UE_LOG(LogMutable, Error, TEXT("Operation not found in EndRead."));
 		check(false);
 	}
+
+	return bSuccess;
 }
 
 
diff --git a/Engine/Plugins/Experimental/Mutable/Source/MutableRuntime/Internal/MuR/ModelPrivate.h b/Engine/Plugins/Experimental/Mutable/Source/MutableRuntime/Internal/MuR/ModelPrivate.h
--- a/Engine/Plugins/Experimental/Mutable/Source/MutableRuntime/Internal/MuR/ModelPrivate.h
+++ b/Engine/Plugins/Experimental/Mutable/Source/MutableRuntime/Internal/MuR/ModelPrivate.h
@@ -562,21 +562,38 @@ namespace mu
             return index;
         }
 
-        //! Get a constant image, assuming it is fully loaded. The image constant will be composed with lodaded mips if necessary.
+        //! Get a constant image, assuming at least some mips are loaded. The image constant will be composed with lodaded mips if necessary.
 		template <typename CreateImageFunc>
-        void GetConstant( int32 ConstantIndex, ImagePtrConst& res, int32 MipsToSkip, const CreateImageFunc& CreateImage) const
+        void GetConstant( int32 ConstantIndex, Ptr<const Image>& res, int32 MipsToSkip, const CreateImageFunc& CreateImage) const
         {
 			int32 ReallySkippedLODs = FMath::Min(ConstantImages[ConstantIndex].LODCount - 1, MipsToSkip);
 			int32 FirstLODIndexIndex = ConstantImages[ConstantIndex].FirstIndex;
-			int32 ResultLODIndexIndex = FirstLODIndexIndex + ReallySkippedLODs;
-			int32 FinalLODs = ConstantImages[ConstantIndex].LODCount - ReallySkippedLODs;
-			check(FinalLODs > 0);
 
 			// Get the first mip
+			int32 ResultLODIndexIndex = FirstLODIndexIndex + ReallySkippedLODs;
 			int32 ResultLODIndex = ConstantImageLODIndices[ResultLODIndexIndex];
 			Ptr<const Image> CurrentMip = ConstantImageLODs[ResultLODIndex].Value;
-			check(CurrentMip);
+
+			// We may need to skip more LODs if they are not loaded
+			while (!CurrentMip)
+			{
+				++ReallySkippedLODs;
 				
+				if (ReallySkippedLODs >= ConstantImages[ConstantIndex].LODCount)
+				{
+					// We don't have a single mip loaded for the image that was requested
+					ensure(false);
+					break;
+				}
+
+				++ResultLODIndexIndex;
+				ResultLODIndex = ConstantImageLODIndices[ResultLODIndexIndex];
+				CurrentMip = ConstantImageLODs[ResultLODIndex].Value;
+			}
+
+			int32 FinalLODs = ConstantImages[ConstantIndex].LODCount - ReallySkippedLODs;
+			check(FinalLODs > 0);
+
 			// Shortcut if we only want one mip
 			if (FinalLODs == 1)
 			{
diff --git a/Engine/Plugins/Experimental/Mutable/Source/MutableRuntime/Private/MuR/CodeRunner.cpp b/Engine/Plugins/Experimental/Mutable/Source/MutableRuntime/Private/MuR/CodeRunner.cpp
--- a/Engine/Plugins/Experimental/Mutable/Source/MutableRuntime/Private/MuR/CodeRunner.cpp
+++ b/Engine/Plugins/Experimental/Mutable/Source/MutableRuntime/Private/MuR/CodeRunner.cpp
@@ -924,7 +924,7 @@ namespace mu
 
 
     //---------------------------------------------------------------------------------------------
-    void CodeRunner::RunCode_ConstantResource(const FScheduledOp& item, const Model* pModel )
+    bool CodeRunner::RunCode_ConstantResource(const FScheduledOp& item, const Model* pModel )
     {
 		MUTABLE_CPUPROFILER_SCOPE(RunCode_Constant);
 
@@ -983,7 +983,10 @@ namespace mu
 				});
 
 			// Assume the ROM has been loaded previously in a task generated at IssueOp
-			check(Source);
+			if (!Source)
+			{
+				return false;
+			}
 
             StoreImage( item, Source );
 			//UE_LOG(LogMutableCore, Log, TEXT("Set image constant %d."), item.At);
@@ -1012,6 +1015,9 @@ namespace mu
             }
             break;
         }
+
+		// Success
+		return true;
     }
 
     //---------------------------------------------------------------------------------------------
diff --git a/Engine/Plugins/Experimental/Mutable/Source/MutableRuntime/Private/MuR/CodeRunner.h b/Engine/Plugins/Experimental/Mutable/Source/MutableRuntime/Private/MuR/CodeRunner.h
--- a/Engine/Plugins/Experimental/Mutable/Source/MutableRuntime/Private/MuR/CodeRunner.h
+++ b/Engine/Plugins/Experimental/Mutable/Source/MutableRuntime/Private/MuR/CodeRunner.h
@@ -121,7 +121,10 @@ namespace  mu
         void RunCode_Switch(const FScheduledOp&, const Model* );
         void RunCode_Instance(const FScheduledOp&, const Model*, uint32 LodMask );
         void RunCode_InstanceAddResource(const FScheduledOp&, const TSharedPtr<const Model>& Model, const Parameters* );
-        void RunCode_ConstantResource(const FScheduledOp&, const Model* );
+
+		/** Return false incase of failure. */
+        bool RunCode_ConstantResource(const FScheduledOp&, const Model* );
+
         void RunCode_Mesh(const FScheduledOp&, const Model* );
         void RunCode_Image(const FScheduledOp&, const Parameters*, const Model* );
         void RunCode_Layout(const FScheduledOp&, const Model* );
@@ -231,7 +234,11 @@ namespace  mu
 			/** */
 			virtual bool Prepare(CodeRunner*, bool& bOutFailed) { bOutFailed = false; return true; }
 			virtual void DoWork() {}
-			virtual void Complete(CodeRunner*) = 0;
+
+			/** Return true if succeeded. */
+			virtual bool Complete(CodeRunner*) = 0;
+
+			/** Return true if the task has been completed. */
 			virtual bool IsComplete(CodeRunner*)
 			{ 
 				return !Event.IsValid() || Event.IsCompleted(); 
@@ -258,8 +265,8 @@ namespace  mu
 			}
 
 			// FIssuedTask interface
-			bool Prepare(CodeRunner*, bool& bOutFailed) override;
-			void Complete(CodeRunner*) override;
+			virtual bool Prepare(CodeRunner*, bool& bOutFailed) override;
+			virtual bool Complete(CodeRunner*) override;
 
 		private:
 			int32 RomIndex = -1;
@@ -278,9 +285,9 @@ namespace  mu
 			}
 
 			// FIssuedTask interface
-			bool Prepare(CodeRunner*, bool& bOutFailed) override;
-			void Complete(CodeRunner*) override;
-			bool IsComplete(CodeRunner*) override;
+			virtual bool Prepare(CodeRunner*, bool& bOutFailed) override;
+			virtual bool Complete(CodeRunner*) override;
+			virtual bool IsComplete(CodeRunner*) override;
 
 		private:
 			ExtensionDataPtrConst Data;
@@ -301,8 +308,8 @@ namespace  mu
 			}
 
 			// FIssuedTask interface
-			bool Prepare(CodeRunner*, bool& bOutFailed) override;
-			void Complete(CodeRunner*) override;
+			virtual bool Prepare(CodeRunner*, bool& bOutFailed) override;
+			virtual bool Complete(CodeRunner*) override;
 
 		private:
  			int32 LODIndexIndex = -1;
diff --git a/Engine/Plugins/Experimental/Mutable/Source/MutableRuntime/Private/MuR/CodeRunnerTasked.cpp b/Engine/Plugins/Experimental/Mutable/Source/MutableRuntime/Private/MuR/CodeRunnerTasked.cpp
--- a/Engine/Plugins/Experimental/Mutable/Source/MutableRuntime/Private/MuR/CodeRunnerTasked.cpp
+++ b/Engine/Plugins/Experimental/Mutable/Source/MutableRuntime/Private/MuR/CodeRunnerTasked.cpp
@@ -302,6 +302,8 @@ namespace mu
 		const FTimespan MaxAllowedTime = FTimespan::FromMilliseconds(2.0); 
 		const FTimespan TimeOut = FTimespan::FromSeconds(FPlatformTime::Seconds()) + MaxAllowedTime;
 
+		bool bSuccess = true;
+
         while(!OpenTasks.IsEmpty() || !ClosedTasks.IsEmpty() || !IssuedTasks.IsEmpty())
         {
 			UpdateTraces();
@@ -323,7 +325,7 @@ namespace mu
 			//	UE_LOG(LogMutableCore, Log, TEXT("Tasks: %5d open, %5d issued, %5d closed, %d closed ready"), OpenTasks.Num(), IssuedTasks.Num(), ClosedTasks.Num(), ClosedReady);
 			//}
 
-			for (int32 Index = 0; Index < IssuedTasks.Num(); )
+			for (int32 Index = 0; bSuccess && Index < IssuedTasks.Num(); )
 			{
 				check(IssuedTasks[Index]);
 
@@ -331,7 +333,7 @@ namespace mu
 				if (bWorkDone)
 				{
 					const FScheduledOp& item = IssuedTasks[Index]->Op;
-					IssuedTasks[Index]->Complete(this);
+					bSuccess = IssuedTasks[Index]->Complete(this);
 
 					if (ScheduledStagePerOp[item] == item.Stage + 1)
 					{
@@ -349,6 +351,11 @@ namespace mu
 				}
 			}
 
+			if (!bSuccess)
+			{
+				return AbortRun();
+			}
+
 			while (!OpenTasks.IsEmpty())
 			{
 				// Get a new task to run
@@ -638,9 +645,9 @@ namespace mu
 		{}
 
 		// FIssuedTask interface
-		bool Prepare(CodeRunner*, bool& bOutFailed) override;
-		void DoWork() override;
-		void Complete(CodeRunner*) override;
+		virtual bool Prepare(CodeRunner*, bool& bOutFailed) override;
+		virtual void DoWork() override;
+		virtual bool Complete(CodeRunner*) override;
 
 	private:
 		int32 ImageCompressionQuality = 0;
@@ -745,12 +752,18 @@ namespace mu
 				int32 StartLevel = Mask->GetLODCount() - 1;
 				int32 LevelCount = Base->GetLODCount();
 
-				Ptr<Image> MaskFix = Runner->CloneOrTakeOver(Mask);
+				// Uncompress mask to avoid excessive RLE-compression and decompression in the following ops.
+				Ptr<Image> UncompressedMask = Runner->CreateImage(Mask->GetSizeX(), Mask->GetSizeY(), Mask->GetLODCount(), GetUncompressedFormat(Mask->GetFormat()), EInitializationType::NotInitialized);
+				bool bSuccess = false;
+				ImOp.ImagePixelFormat(bSuccess, ImageCompressionQuality, UncompressedMask.get(), Mask.get());
+
+				Ptr<Image> MaskFix = UncompressedMask;
 				MaskFix->DataStorage.SetNumLODs(LevelCount);
 
 				FMipmapGenerationSettings Settings{};
 				ImOp.ImageMipmap(ImageCompressionQuality, MaskFix.get(), MaskFix.get(), StartLevel, LevelCount, Settings);
 
+				Runner->Release(Mask);
 				Mask = MaskFix;
 			}
 		}
@@ -902,7 +915,7 @@ namespace mu
 	}
 
 
-	void FImageLayerTask::Complete( CodeRunner* Runner )
+	bool FImageLayerTask::Complete( CodeRunner* Runner )
 	{
 		// This runs in the Runner thread
 		Runner->Release(Blended);
@@ -926,6 +939,9 @@ namespace mu
 
 			Runner->StoreImage(Op, Result);
 		}
+
+		bool bSuccess = true;
+		return bSuccess;
 	}
 
 
@@ -940,9 +956,9 @@ namespace mu
 		{}
 
 		// FIssuedTask interface
-		bool Prepare(CodeRunner*, bool& bOutFailed) override;
-		void DoWork() override;
-		void Complete(CodeRunner*) override;
+		virtual bool Prepare(CodeRunner*, bool& bOutFailed) override;
+		virtual void DoWork() override;
+		virtual bool Complete(CodeRunner*) override;
 
 	private:
 		int32 ImageCompressionQuality = 0;
@@ -1152,7 +1168,7 @@ namespace mu
 	}
 
 
-	void FImageLayerColourTask::Complete(CodeRunner* Runner)
+	bool FImageLayerColourTask::Complete(CodeRunner* Runner)
 	{
 		// This runs in the Runner thread
 		Runner->Release(Mask);
@@ -1162,6 +1178,9 @@ namespace mu
 		{
 			Runner->StoreImage(Op, Result);
 		}
+
+		bool bSuccess = true;
+		return bSuccess;
 	}
 
 
@@ -1176,9 +1195,9 @@ namespace mu
 		{}
 
 		// FIssuedTask interface
-		bool Prepare(CodeRunner*, bool& bOutFailed) override;
-		void DoWork() override;
-		void Complete(CodeRunner*) override;
+		virtual bool Prepare(CodeRunner*, bool& bOutFailed) override;
+		virtual void DoWork() override;
+		virtual bool Complete(CodeRunner*) override;
 
 	private:
 		int ImageCompressionQuality = 0;
@@ -1249,7 +1268,7 @@ namespace mu
 	}
 
 
-	void FImagePixelFormatTask::Complete(CodeRunner* Runner)
+	bool FImagePixelFormatTask::Complete(CodeRunner* Runner)
 	{
 		// This runs in the Runner thread
 		Runner->Release(Base);
@@ -1259,6 +1278,9 @@ namespace mu
 		{
 			Runner->StoreImage(Op, Result);
 		}
+
+		bool bSuccess = true;
+		return bSuccess;
 	}
 
 	class FImageMipmapTask : public CodeRunner::FIssuedTask
@@ -1269,9 +1291,9 @@ namespace mu
 		{}
 
 		// FIssuedTask interface
-		bool Prepare(CodeRunner*, bool& bOutFailed) override;
-		void DoWork() override;
-		void Complete(CodeRunner*) override;
+		virtual bool Prepare(CodeRunner*, bool& bOutFailed) override;
+		virtual void DoWork() override;
+		virtual bool Complete(CodeRunner*) override;
 
 	private:
 		int32 ImageCompressionQuality = 0;
@@ -1360,7 +1382,7 @@ namespace mu
 	}
 
 
-	void FImageMipmapTask::Complete(CodeRunner* Runner)
+	bool FImageMipmapTask::Complete(CodeRunner* Runner)
 	{
 		FImageOperator ImOp = MakeImageOperator(Runner);
 		ImOp.ImageMipmap_ReleaseScratch(Scratch);
@@ -1376,6 +1398,9 @@ namespace mu
 		{
 			Runner->StoreImage(Op, Result);
 		}
+
+		bool bSuccess = true;
+		return bSuccess;
 	}
 
 
@@ -1390,9 +1415,9 @@ namespace mu
 		{}
 
 		// FIssuedTask interface
-		bool Prepare(CodeRunner*, bool& bOutFailed);
-		void DoWork() override;
-		void Complete(CodeRunner*) override;
+		virtual bool Prepare(CodeRunner*, bool& bOutFailed);
+		virtual void DoWork() override;
+		virtual bool Complete(CodeRunner*) override;
 
 	private:
 		OP::ImageSwizzleArgs Args;
@@ -1490,7 +1515,7 @@ namespace mu
 	}
 
 
-	void FImageSwizzleTask::Complete(CodeRunner* Runner)
+	bool FImageSwizzleTask::Complete(CodeRunner* Runner)
 	{
 		// This runs in the Runner thread
 		for (int i = 0; i < MUTABLE_OP_MAX_SWIZZLE_CHANNELS; ++i)
@@ -1505,6 +1530,9 @@ namespace mu
 		{
 			Runner->StoreImage(Op, Result);
 		}
+
+		bool bSuccess = true;
+		return bSuccess;
 	}
 
 
@@ -1517,9 +1545,9 @@ namespace mu
 		FImageSaturateTask(const FScheduledOp&, const OP::ImageSaturateArgs&);
 
 		// FIssuedTask interface
-		bool Prepare(CodeRunner*, bool& bOutFailed) override;
-		void DoWork() override;
-		void Complete(CodeRunner*) override;
+		virtual bool Prepare(CodeRunner*, bool& bOutFailed) override;
+		virtual void DoWork() override;
+		virtual bool Complete(CodeRunner*) override;
 
 	private:
 		const OP::ImageSaturateArgs Args;
@@ -1573,7 +1601,7 @@ namespace mu
 	}
 
 	//---------------------------------------------------------------------------------------------
-	void FImageSaturateTask::Complete(CodeRunner* Runner)
+	bool FImageSaturateTask::Complete(CodeRunner* Runner)
 	{
 		// This runs in the mutable Runner thread
 
@@ -1582,6 +1610,9 @@ namespace mu
 		{
 			Runner->StoreImage(Op, Result);
 		}
+
+		bool bSuccess = true;
+		return bSuccess;
 	}
 
 
@@ -1596,9 +1627,9 @@ namespace mu
 		{}
 
 		// FIssuedTask interface
-		bool Prepare(CodeRunner* Runner, bool& bFailed) override;
-		void DoWork() override;
-		void Complete(CodeRunner*) override;
+		virtual bool Prepare(CodeRunner* Runner, bool& bFailed) override;
+		virtual void DoWork() override;
+		virtual bool Complete(CodeRunner*) override;
 
 	private:
 		int32 ImageCompressionQuality = 0;
@@ -1689,7 +1720,7 @@ namespace mu
 
 
 	//---------------------------------------------------------------------------------------------
-	void FImageResizeTask::Complete(CodeRunner* Runner)
+	bool FImageResizeTask::Complete(CodeRunner* Runner)
 	{
 		// This runs in the Runner thread
 		Runner->Release(Base);
@@ -1699,6 +1730,9 @@ namespace mu
 		{
 			Runner->StoreImage(Op, Result);
 		}
+
+		bool bSuccess = true;
+		return bSuccess;
 	}
 
 
@@ -1713,9 +1747,9 @@ namespace mu
 		{}
 
 		// FIssuedTask interface
-		bool Prepare(CodeRunner* Runner, bool& bFailed) override;
-		void DoWork() override;
-		void Complete(CodeRunner*) override;
+		virtual bool Prepare(CodeRunner* Runner, bool& bFailed) override;
+		virtual void DoWork() override;
+		virtual bool Complete(CodeRunner*) override;
 
 	private:
 		OP::ImageResizeRelArgs Args;
@@ -1790,7 +1824,7 @@ namespace mu
 	}
 
 
-	void FImageResizeRelTask::Complete(CodeRunner* Runner)
+	bool FImageResizeRelTask::Complete(CodeRunner* Runner)
 	{
 		// This runs in the Runner thread
 		Runner->Release(Base);
@@ -1800,6 +1834,9 @@ namespace mu
 		{
 			Runner->StoreImage(Op, Result);
 		}
+
+		bool bSuccess = true;
+		return bSuccess;
 	}
 
 
@@ -1814,9 +1851,9 @@ namespace mu
 		{}
 
 		// FIssuedTask interface
-		bool Prepare(CodeRunner*, bool& bOutFailed) override;
-		void DoWork() override;
-		void Complete(CodeRunner* Runner) override;
+		virtual bool Prepare(CodeRunner*, bool& bOutFailed) override;
+		virtual void DoWork() override;
+		virtual bool Complete(CodeRunner* Runner) override;
 
 	private:
 		Ptr<Image> Result;
@@ -1847,13 +1884,16 @@ namespace mu
 	}
 
 
-	void FImageInvertTask::Complete(CodeRunner* Runner)
+	bool FImageInvertTask::Complete(CodeRunner* Runner)
 	{
 		// If we didn't take a shortcut
 		if (Result)
 		{
 			Runner->StoreImage(Op, Result);
 		}
+
+		bool bSuccess = true;
+		return bSuccess;
 	}
 
 
@@ -1868,9 +1908,9 @@ namespace mu
 		{}
 
 		// FIssuedTask interface
-		bool Prepare(CodeRunner*, bool& bOutFailed) override;
-		void DoWork() override;
-		void Complete(CodeRunner*) override;
+		virtual bool Prepare(CodeRunner*, bool& bOutFailed) override;
+		virtual void DoWork() override;
+		virtual bool Complete(CodeRunner*) override;
 
 	private:
 		int32 ImageCompressionQuality = 0;
@@ -2018,7 +2058,7 @@ namespace mu
 	}
 
 
-	void FImageComposeTask::Complete(CodeRunner* Runner)
+	bool FImageComposeTask::Complete(CodeRunner* Runner)
 	{
 		// This runs in the mutable Runner thread
 		Runner->Release(Block);
@@ -2029,6 +2069,9 @@ namespace mu
 		{
 			Runner->StoreImage(Op, Result);
 		}
+
+		bool bSuccess = true;
+		return bSuccess;
 	}
 
 
@@ -2106,14 +2149,14 @@ namespace mu
 	
 
 	//---------------------------------------------------------------------------------------------
-	void CodeRunner::FLoadMeshRomTask::Complete(CodeRunner* Runner)
+	bool CodeRunner::FLoadMeshRomTask::Complete(CodeRunner* Runner)
 	{
 		// This runs in the Runner thread
 		MUTABLE_CPUPROFILER_SCOPE(FLoadMeshRomTask_Complete);
 
 		if (!Runner || !Runner->m_pSystem)
 		{
-			return;
+			return false;
 		}
 
 		FProgram& Program = Runner->m_pModel->GetPrivate()->m_program;
@@ -2146,6 +2189,9 @@ namespace mu
 
 		Runner->m_pSystem->WorkingMemoryManager.MarkRomUsed(RomIndex, Runner->m_pModel);
 		--ModelCache->PendingOpsPerRom[RomIndex];
+
+		bool bSuccess = true;
+		return bSuccess;
 	}
 
 
@@ -2194,10 +2240,12 @@ namespace mu
 
 
 	//---------------------------------------------------------------------------------------------
-	void CodeRunner::FLoadExtensionDataTask::Complete(CodeRunner* Runner)
+	bool CodeRunner::FLoadExtensionDataTask::Complete(CodeRunner* Runner)
 	{
 		MUTABLE_CPUPROFILER_SCOPE(FLoadExtensionDataTask_Complete);
 
+		bool bSuccess = true;
+
 		// Complete should only be called if the load is finished
 		check(LoadHandle->LoadState != FExtensionDataLoadHandle::ELoadState::Pending);
 
@@ -2216,10 +2264,13 @@ namespace mu
 		{
 			check(LoadHandle->LoadState == FExtensionDataLoadHandle::ELoadState::FailedToLoad);
 			Constant.LoadState = FExtensionDataConstant::ELoadState::FailedToLoad;
+			bSuccess = false;
 		}
 
 		// Process the constant op normally, now that the data is loaded.
 		Runner->RunCode(Op, Runner->m_pParams, Runner->m_pModel, Runner->m_lodMask);
+
+		return bSuccess;
 	}
 
 
@@ -2332,49 +2383,64 @@ namespace mu
 	
 	
 	//---------------------------------------------------------------------------------------------
-	void CodeRunner::FLoadImageRomsTask::Complete(CodeRunner* Runner)
+	bool CodeRunner::FLoadImageRomsTask::Complete(CodeRunner* Runner)
 	{
 		// This runs in the Runner thread
 		MUTABLE_CPUPROFILER_SCOPE(FLoadImageRomsTask_Complete);
 
 		if (!Runner || !Runner->m_pSystem)
 		{
-			return;
+			return false;
 		}
 
 		FProgram& Program = Runner->m_pModel->GetPrivate()->m_program;
 		
 		FWorkingMemoryManager::FModelCacheEntry* ModelCache = Runner->m_pSystem->WorkingMemoryManager.FindModelCache(Runner->m_pModel.Get());
 
+		bool bSomeMissingData = false;
+
 		for (const int32 RomIndex : RomIndices)
 		{
 			// Since task could be reordered, we need to make sure we end the rom read before continuing
 			if (FRomLoadOp* RomLoadOp = Runner->RomLoadOps.Find(RomIndex))
 			{				
-				Runner->m_pSystem->StreamInterface->EndRead(RomLoadOp->m_streamID);
+				bool bSuccess = Runner->m_pSystem->StreamInterface->EndRead(RomLoadOp->m_streamID);
 
 				MUTABLE_CPUPROFILER_SCOPE(Unserialise);
 
-				InputMemoryStream Stream(RomLoadOp->m_streamBuffer.GetData(), RomLoadOp->m_streamBuffer.Num());
-				InputArchive Arch(&Stream);
+				if (bSuccess)
+				{
+					InputMemoryStream Stream(RomLoadOp->m_streamBuffer.GetData(), RomLoadOp->m_streamBuffer.Num());
+					InputArchive Arch(&Stream);
 
-				const int32 ResIndex = Program.m_roms[RomIndex].ResourceIndex;
+					const int32 ResIndex = Program.m_roms[RomIndex].ResourceIndex;
 
-				// TODO: Try to reuse buffer from PooledImages.
-				check(!Program.ConstantImageLODs[ResIndex].Value);
-				Ptr<Image> Value = Image::StaticUnserialise(Arch);
+					// TODO: Try to reuse buffer from PooledImages.
+					check(!Program.ConstantImageLODs[ResIndex].Value);
+					Ptr<Image> Value = Image::StaticUnserialise(Arch);
 
-				Program.SetImageRomValue(RomIndex, Value);
-				check(Program.ConstantImageLODs[ResIndex].Value);
+					Program.SetImageRomValue(RomIndex, Value);
+					check(Program.ConstantImageLODs[ResIndex].Value);
+				}
+				else
+				{
+					bSomeMissingData = true;
+				}
 				
 				Runner->RomLoadOps.Remove(*RomLoadOp);
 			}
 		}
-		
+
+		if (bSomeMissingData)
+		{
+			// Some data may be missing. We can try to go on if some mips are there. 
+			UE_LOG(LogMutableCore, Verbose, TEXT("FLoadImageRomsTask::Complete failed: missing data?"));
+		}
+
 		// Process the constant op normally, now that the rom is loaded.	
-		Runner->RunCode(Op, Runner->m_pParams, Runner->m_pModel, Runner->m_lodMask);
+		bool bSuccess = Runner->RunCode_ConstantResource(Op, Runner->m_pModel.Get());
 
-		for (int32 LODIndex = 0; LODIndex < LODIndexCount; ++LODIndex)
+		for (int32 LODIndex = 0; bSuccess && (LODIndex < LODIndexCount); ++LODIndex)
 		{
 			int32 CurrentIndexIndex = LODIndexIndex + LODIndex;
 			int32 CurrentIndex = Program.ConstantImageLODIndices[CurrentIndexIndex];
@@ -2396,6 +2462,8 @@ namespace mu
 				UE_LOG(LogMutableCore, Log, TEXT("FLoadImageRomsTask::Complete rom %d, now peding ops is %d."), RomIndex, ModelCache->PendingOpsPerRom[RomIndex]);
 			}
 		}
+
+		return bSuccess;
 	}
 
 
@@ -2409,7 +2477,7 @@ namespace mu
 
 		// FIssuedTask interface
 		virtual bool Prepare(CodeRunner*, bool& bOutFailed) override;
-		virtual void Complete(CodeRunner* Runner) override;
+		virtual bool Complete(CodeRunner* Runner) override;
 		
 	private:
 		uint8 MipmapsToSkip;
@@ -2453,7 +2521,7 @@ namespace mu
 	}
 
 
-	void FImageExternalLoadTask::Complete(CodeRunner* Runner)
+	bool FImageExternalLoadTask::Complete(CodeRunner* Runner)
 	{
 		if (ExternalCleanUpFunc)
 		{
@@ -2461,6 +2529,9 @@ namespace mu
 		}
 
 		Runner->StoreImage(Op, Result);
+
+		bool bSuccess = true;
+		return bSuccess;
 	}	
 
 
@@ -2474,7 +2545,7 @@ namespace mu
 
 		// FIssuedTask interface
 		virtual bool Prepare(CodeRunner*, bool& bOutFailed) override;
-		virtual void Complete(CodeRunner* Runner) override;
+		virtual bool Complete(CodeRunner* Runner) override;
 
 	private:
 		uint8 MipmapsToSkip;
@@ -2517,7 +2588,7 @@ namespace mu
 	}
 
 
-	void FMeshExternalLoadTask::Complete(CodeRunner* Runner)
+	bool FMeshExternalLoadTask::Complete(CodeRunner* Runner)
 	{
 		if (ExternalCleanUpFunc)
 		{
@@ -2525,6 +2596,9 @@ namespace mu
 		}
 
 		Runner->StoreMesh(Op, Result);
+
+		bool bSuccess = true;
+		return bSuccess;
 	}
 
 
diff --git a/Engine/Plugins/Experimental/Mutable/Source/MutableRuntime/Public/MuR/Serialisation.h b/Engine/Plugins/Experimental/Mutable/Source/MutableRuntime/Public/MuR/Serialisation.h
--- a/Engine/Plugins/Experimental/Mutable/Source/MutableRuntime/Public/MuR/Serialisation.h
+++ b/Engine/Plugins/Experimental/Mutable/Source/MutableRuntime/Public/MuR/Serialisation.h
@@ -97,14 +97,14 @@ namespace mu
 		//! \param Model.
         //! \param key0 key identifying the model data fragment that is requested.
         //!         This key interpretation depends on the implementation of the ModelStreamer,
-		//! \param pBuffer is an already-allocated buffer big enough to receive the expected data.
+		//! \param Buffer is an already-allocated buffer big enough to receive the expected data.
 		//! \param size is the size of the pBuffer buffer, which must match the size of the data
 		//! requested with the key identifiers.
 		//! \param CompletionCallback Optional callback. Copied inside the called function. Will always be called.
 		//! \return a previously unused identifier, now used for this operation, that can be used in
 		//! calls to the other methods of this interface. If the return value is negative it indicates
 		//! an unrecoverable error.
-		virtual OPERATION_ID BeginReadBlock(const mu::Model*, uint32 BlockKey, void* pBuffer, uint64 size, TFunction<void(bool bSuccess)>* CompletionCallback = nullptr) = 0;
+		virtual OPERATION_ID BeginReadBlock(const mu::Model*, uint32 BlockKey, void* Buffer, uint64 size, TFunction<void(bool bSuccess)>* CompletionCallback = nullptr) = 0;
 
         //! Check if a data request operation has been completed.
         //! This is a weak check than *may* return true if the given operation has completed, but
@@ -113,10 +113,12 @@ namespace mu
         //! return false.
         virtual bool IsReadCompleted( OPERATION_ID ) = 0;
 
-        //! Complete a data request operation. This method has to block until a data request issued
-        //! with OpenFile has been completed. After returning from this call, the ID cannot be used
-        //! any more to identify the same operation and becomes free.
-        virtual void EndRead( OPERATION_ID ) = 0;
+        /** Complete a data request operation.This method has to block until a data request issued
+        * with BeginReadBlock has been completed. After returning from this call, the ID cannot be used
+        * any more to identify the same operation and becomes free.
+		* \return true if the data was loaded successfully.
+		*/
+        virtual bool EndRead( OPERATION_ID ) = 0;
 
     };
 
