[Mutable] Fix not generating meshes in packaged builds.

Issue introduced in 34231165.

#rnx

diff --git a/Engine/Plugins/Experimental/Mutable/Source/CustomizableObject/Private/MuCO/CustomizableObjectInstance.cpp b/Engine/Plugins/Experimental/Mutable/Source/CustomizableObject/Private/MuCO/CustomizableObjectInstance.cpp
--- a/Engine/Plugins/Experimental/Mutable/Source/CustomizableObject/Private/MuCO/CustomizableObjectInstance.cpp
+++ b/Engine/Plugins/Experimental/Mutable/Source/CustomizableObject/Private/MuCO/CustomizableObjectInstance.cpp
@@ -2107,11 +2107,11 @@ bool UCustomizableInstancePrivate::UpdateSkeletalMesh_PostBeginUpdate0(UCustomiz
 
 		// We need the first valid mesh. get it from the component, considering that some LOSs may have been skipped.
 		mu::Ptr<const mu::Mesh> ComponentMesh;
-
-		int32 FirstValidLODIndex = Component.FirstLOD + OperationData->FirstLODAvailable;
-		check (OperationData->FirstLODAvailable<Component.LODCount)
-		ComponentMesh = OperationData->InstanceUpdateData.LODs[FirstValidLODIndex].Mesh;
-
+		for (int32 FirstValidLODIndex = Component.FirstLOD; FirstValidLODIndex < OperationData->InstanceUpdateData.LODs.Num() && !ComponentMesh; ++FirstValidLODIndex)
+		{
+			ComponentMesh = OperationData->InstanceUpdateData.LODs[FirstValidLODIndex].Mesh;
+		}
+		
 		if (ComponentMesh)
 		{
 			// Construct a new skeleton, fix up ActiveBones and Bonemap arrays and recompute the RefInvMatrices
