[Mutable] Fix NodeEditMaterial having two Pin Viewers.

Fixed by changing the Node Details hierarchy.

#rnx

diff --git a/Engine/Plugins/Experimental/Mutable/Source/CustomizableObjectEditor/Private/MuCOE/CustomizableObjectEditorModule.cpp b/Engine/Plugins/Experimental/Mutable/Source/CustomizableObjectEditor/Private/MuCOE/CustomizableObjectEditorModule.cpp
--- a/Engine/Plugins/Experimental/Mutable/Source/CustomizableObjectEditor/Private/MuCOE/CustomizableObjectEditorModule.cpp
+++ b/Engine/Plugins/Experimental/Mutable/Source/CustomizableObjectEditor/Private/MuCOE/CustomizableObjectEditorModule.cpp
@@ -70,6 +70,7 @@
 #include "MuCOE/GraphTraversal.h"
 #include "MuCOE/CustomizableObjectInstanceBaker.h"
 #include "Editor.h"
+#include "MuCOE/Nodes/CustomizableObjectNodeRemoveMesh.h"
 
 class AActor;
 class FString;
@@ -236,9 +237,9 @@ void FCustomizableObjectEditorModule::StartupModule()
 	FPropertyEditorModule& PropertyModule = FModuleManager::LoadModuleChecked<FPropertyEditorModule>("PropertyEditor");
 	RegisterCustomDetails(PropertyModule, UCustomizableObjectNodeLayoutBlocks::StaticClass(), FOnGetDetailCustomizationInstance::CreateStatic(&FCustomizableObjectNodeLayoutBlocksDetails::MakeInstance));
 	RegisterCustomDetails(PropertyModule, UCustomizableObjectNodeEditMaterial::StaticClass(), FOnGetDetailCustomizationInstance::CreateStatic(&FCustomizableObjectNodeEditMaterialDetails::MakeInstance));
+	RegisterCustomDetails(PropertyModule, UCustomizableObjectNodeRemoveMesh::StaticClass(), FOnGetDetailCustomizationInstance::CreateStatic(&FCustomizableObjectNodeEditMaterialBaseDetails::MakeInstance));
 	RegisterCustomDetails(PropertyModule, UCustomizableObjectNodeRemoveMeshBlocks::StaticClass(), FOnGetDetailCustomizationInstance::CreateStatic(&FCustomizableObjectNodeRemoveMeshBlocksDetails::MakeInstance));
 	RegisterCustomDetails(PropertyModule, UCustomizableObjectNodeExtendMaterial::StaticClass(), FOnGetDetailCustomizationInstance::CreateStatic(&FCustomizableObjectNodeParentedMaterialDetails::MakeInstance));
-	RegisterCustomDetails(PropertyModule, UCustomizableObjectNodeEditMaterialBase::StaticClass(), FOnGetDetailCustomizationInstance::CreateStatic(&FCustomizableObjectNodeEditMaterialBaseDetails::MakeInstance));
 	RegisterCustomDetails(PropertyModule, UCustomizableObjectNodeMorphMaterial::StaticClass(), FOnGetDetailCustomizationInstance::CreateStatic(&FCustomizableObjectNodeMorphMaterialDetails::MakeInstance));
 	RegisterCustomDetails(PropertyModule, UCustomizableObjectNodeObject::StaticClass(), FOnGetDetailCustomizationInstance::CreateStatic(&FCustomizableObjectNodeObjectDetails::MakeInstance));
 	RegisterCustomDetails(PropertyModule, UCustomizableObjectNodeObjectGroup::StaticClass(), FOnGetDetailCustomizationInstance::CreateStatic(&FCustomizableObjectNodeObjectGroupDetails::MakeInstance));
diff --git a/Engine/Plugins/Experimental/Mutable/Source/CustomizableObjectEditor/Private/MuCOE/Nodes/CustomizableObjectNodeEditMaterialBaseDetails.cpp b/Engine/Plugins/Experimental/Mutable/Source/CustomizableObjectEditor/Private/MuCOE/Nodes/CustomizableObjectNodeEditMaterialBaseDetails.cpp
--- a/Engine/Plugins/Experimental/Mutable/Source/CustomizableObjectEditor/Private/MuCOE/Nodes/CustomizableObjectNodeEditMaterialBaseDetails.cpp
+++ b/Engine/Plugins/Experimental/Mutable/Source/CustomizableObjectEditor/Private/MuCOE/Nodes/CustomizableObjectNodeEditMaterialBaseDetails.cpp
@@ -43,7 +43,7 @@ void FCustomizableObjectNodeEditMaterialBaseDetails::CustomizeDetails(IDetailLay
 		// Add all the materials to the combobox
 		TSharedPtr<FString> LayoutToSelect = GenerateLayoutComboboxOptions();
 
-		TSharedRef<IPropertyHandle> LayoutProperty = DetailBuilder.GetProperty(GET_MEMBER_NAME_CHECKED(UCustomizableObjectNodeEditMaterialBase, ParentLayoutIndex));
+		TSharedRef<IPropertyHandle> LayoutProperty = DetailBuilder.GetProperty(GET_MEMBER_NAME_CHECKED(UCustomizableObjectNodeEditMaterialBase, ParentLayoutIndex), UCustomizableObjectNodeEditMaterialBase::StaticClass());
 
 		BlocksCategory.AddCustomRow(LOCTEXT("BlocksRow", "Blocks"))
 		[
diff --git a/Engine/Plugins/Experimental/Mutable/Source/CustomizableObjectEditor/Private/MuCOE/Nodes/CustomizableObjectNodeEditMaterialDetails.cpp b/Engine/Plugins/Experimental/Mutable/Source/CustomizableObjectEditor/Private/MuCOE/Nodes/CustomizableObjectNodeEditMaterialDetails.cpp
--- a/Engine/Plugins/Experimental/Mutable/Source/CustomizableObjectEditor/Private/MuCOE/Nodes/CustomizableObjectNodeEditMaterialDetails.cpp
+++ b/Engine/Plugins/Experimental/Mutable/Source/CustomizableObjectEditor/Private/MuCOE/Nodes/CustomizableObjectNodeEditMaterialDetails.cpp
@@ -24,7 +24,7 @@ TSharedRef<IDetailCustomization> FCustomizableObjectNodeEditMaterialDetails::Mak
 
 void FCustomizableObjectNodeEditMaterialDetails::CustomizeDetails( IDetailLayoutBuilder& DetailBuilder )
 {
-	FCustomizableObjectNodeDetails::CustomizeDetails(DetailBuilder);
+	FCustomizableObjectNodeEditMaterialBaseDetails::CustomizeDetails(DetailBuilder);
 
 	Node = nullptr;
 	const IDetailsView* DetailsView = DetailBuilder.GetDetailsView();
diff --git a/Engine/Plugins/Experimental/Mutable/Source/CustomizableObjectEditor/Private/MuCOE/Nodes/CustomizableObjectNodeEditMaterialDetails.h b/Engine/Plugins/Experimental/Mutable/Source/CustomizableObjectEditor/Private/MuCOE/Nodes/CustomizableObjectNodeEditMaterialDetails.h
--- a/Engine/Plugins/Experimental/Mutable/Source/CustomizableObjectEditor/Private/MuCOE/Nodes/CustomizableObjectNodeEditMaterialDetails.h
+++ b/Engine/Plugins/Experimental/Mutable/Source/CustomizableObjectEditor/Private/MuCOE/Nodes/CustomizableObjectNodeEditMaterialDetails.h
@@ -3,7 +3,7 @@
 #pragma once
 
 
-#include "CustomizableObjectNodeDetails.h"
+#include "CustomizableObjectNodeEditMaterialBaseDetails.h"
 #include "IDetailCustomization.h"
 
 class FString;
@@ -12,21 +12,18 @@ class UCustomizableObjectNodeEditMaterial;
 class SCustomizableObjectNodeLayoutBlocksSelector;
 
 
-class FCustomizableObjectNodeEditMaterialDetails : public FCustomizableObjectNodeDetails
+class FCustomizableObjectNodeEditMaterialDetails : public FCustomizableObjectNodeEditMaterialBaseDetails
 {
 public:
 	/** Makes a new instance of this detail layout class for a specific detail view requesting it */
 	static TSharedRef<IDetailCustomization> MakeInstance();
 
 	/** ILayoutDetails interface */
-	virtual void CustomizeDetails( IDetailLayoutBuilder& DetailBuilder ) override;
-
+	virtual void CustomizeDetails(IDetailLayoutBuilder& DetailBuilder) override;
 
 private:
-
 	UCustomizableObjectNodeEditMaterial* Node;
 
 	// Layout block editor widget
 	TSharedPtr<SCustomizableObjectNodeLayoutBlocksSelector> LayoutBlocksSelector;
-
 };
diff --git a/Engine/Plugins/Experimental/Mutable/Source/CustomizableObjectEditor/Private/MuCOE/Nodes/CustomizableObjectNodeMorphMaterialDetails.cpp b/Engine/Plugins/Experimental/Mutable/Source/CustomizableObjectEditor/Private/MuCOE/Nodes/CustomizableObjectNodeMorphMaterialDetails.cpp
--- a/Engine/Plugins/Experimental/Mutable/Source/CustomizableObjectEditor/Private/MuCOE/Nodes/CustomizableObjectNodeMorphMaterialDetails.cpp
+++ b/Engine/Plugins/Experimental/Mutable/Source/CustomizableObjectEditor/Private/MuCOE/Nodes/CustomizableObjectNodeMorphMaterialDetails.cpp
@@ -26,7 +26,7 @@ TSharedRef<IDetailCustomization> FCustomizableObjectNodeMorphMaterialDetails::Ma
 
 void FCustomizableObjectNodeMorphMaterialDetails::CustomizeDetails( IDetailLayoutBuilder& DetailBuilder )
 {
-	FCustomizableObjectNodeDetails::CustomizeDetails(DetailBuilder);
+	FCustomizableObjectNodeEditMaterialBaseDetails::CustomizeDetails(DetailBuilder);
 
 	const IDetailsView* DetailsView = DetailBuilder.GetDetailsView();
 	if ( DetailsView->GetSelectedObjects().Num() )
diff --git a/Engine/Plugins/Experimental/Mutable/Source/CustomizableObjectEditor/Private/MuCOE/Nodes/CustomizableObjectNodeMorphMaterialDetails.h b/Engine/Plugins/Experimental/Mutable/Source/CustomizableObjectEditor/Private/MuCOE/Nodes/CustomizableObjectNodeMorphMaterialDetails.h
--- a/Engine/Plugins/Experimental/Mutable/Source/CustomizableObjectEditor/Private/MuCOE/Nodes/CustomizableObjectNodeMorphMaterialDetails.h
+++ b/Engine/Plugins/Experimental/Mutable/Source/CustomizableObjectEditor/Private/MuCOE/Nodes/CustomizableObjectNodeMorphMaterialDetails.h
@@ -3,6 +3,7 @@
 #pragma once
 
 #include "CustomizableObjectNodeDetails.h"
+#include "CustomizableObjectNodeEditMaterialBaseDetails.h"
 #include "IDetailCustomization.h"
 
 namespace ESelectInfo { enum Type : int; }
@@ -12,7 +13,7 @@ class IDetailLayoutBuilder;
 class IPropertyHandle;
 class UCustomizableObjectNodeObject;
 
-class FCustomizableObjectNodeMorphMaterialDetails : public FCustomizableObjectNodeDetails
+class FCustomizableObjectNodeMorphMaterialDetails : public FCustomizableObjectNodeEditMaterialBaseDetails
 {
 public:
 	// Makes a new instance of this detail layout class for a specific detail view requesting it 
@@ -21,16 +22,10 @@ public:
 	// ILayoutDetails interface
 	void CustomizeDetails( IDetailLayoutBuilder& DetailBuilder ) override;
 
-
 private:
-
 	class UCustomizableObjectNodeMorphMaterial* Node = nullptr;
 
 	TArray< TSharedPtr<FString> > MorphTargetComboOptions;
 
-	void OnParentComboBoxSelectionChanged(TSharedPtr<FString> Selection, ESelectInfo::Type SelectInfo, TSharedRef<IPropertyHandle> ParentProperty);
 	void OnMorphTargetComboBoxSelectionChanged(TSharedPtr<FString> Selection, ESelectInfo::Type SelectInfo, TSharedRef<IPropertyHandle> Property);
-
-	TSharedPtr<FString> PrepareComboboxSelection(const int LODIndex, TArray<UCustomizableObjectNodeObject*>& ParentObjectNodes);
-
 };
diff --git a/Engine/Plugins/Experimental/Mutable/Source/CustomizableObjectEditor/Private/MuCOE/Nodes/CustomizableObjectNodeParentedMaterialDetails.cpp b/Engine/Plugins/Experimental/Mutable/Source/CustomizableObjectEditor/Private/MuCOE/Nodes/CustomizableObjectNodeParentedMaterialDetails.cpp
--- a/Engine/Plugins/Experimental/Mutable/Source/CustomizableObjectEditor/Private/MuCOE/Nodes/CustomizableObjectNodeParentedMaterialDetails.cpp
+++ b/Engine/Plugins/Experimental/Mutable/Source/CustomizableObjectEditor/Private/MuCOE/Nodes/CustomizableObjectNodeParentedMaterialDetails.cpp
@@ -35,6 +35,8 @@ void FCustomizableObjectNodeParentedMaterialDetails::CustomizeDetails(IDetailLay
 	Node = nullptr;
 	NodeParentedMaterial = nullptr;
 
+	TSharedPtr<IPropertyHandle> ParentProperty;
+
 	const TArray<TWeakObjectPtr<UObject>>& SelectedObjects = DetailBuilder.GetDetailsView()->GetSelectedObjects();
 	if (SelectedObjects.Num())
 	{
@@ -43,15 +45,16 @@ void FCustomizableObjectNodeParentedMaterialDetails::CustomizeDetails(IDetailLay
 		if (UCustomizableObjectNodeEditMaterialBase* NodeEditMaterial = Cast<UCustomizableObjectNodeEditMaterialBase>(SelectedObjects[0].Get()))
 		{
 			NodeParentedMaterial = NodeEditMaterial;
+			ParentProperty = DetailBuilder.GetProperty("ParentMaterialObject", UCustomizableObjectNodeEditMaterialBase::StaticClass());
 		}
 		else if (UCustomizableObjectNodeExtendMaterial* NodeExtendMaterial = Cast<UCustomizableObjectNodeExtendMaterial>(SelectedObjects[0].Get()))
 		{
 			NodeParentedMaterial = NodeExtendMaterial;
+			ParentProperty = DetailBuilder.GetProperty("ParentMaterialObject", UCustomizableObjectNodeExtendMaterial::StaticClass());
 		}
 	}
 
-	IDetailCategoryBuilder& BlocksCategory = DetailBuilder.EditCategory("Parent");
-
+	IDetailCategoryBuilder& ParentCategory = DetailBuilder.EditCategory("Parent");
 
 	if (NodeParentedMaterial)
 	{
@@ -98,10 +101,8 @@ void FCustomizableObjectNodeParentedMaterialDetails::CustomizeDetails(IDetailLay
 				break;
 			}
 		}
-
-		TSharedRef<IPropertyHandle> ParentProperty = DetailBuilder.GetProperty("ParentMaterialObject");
-
-		BlocksCategory.AddCustomRow(LOCTEXT("FCustomizableObjectNodeParentedMaterialDetails", "Blocks"))
+		
+		ParentCategory.AddCustomRow(LOCTEXT("FCustomizableObjectNodeParentedMaterialDetails", "Blocks"))
 		[
 			SNew(SHorizontalBox)
 			+ SHorizontalBox::Slot()
@@ -124,7 +125,7 @@ void FCustomizableObjectNodeParentedMaterialDetails::CustomizeDetails(IDetailLay
 						SNew(STextComboBox)
 						.OptionsSource(&ParentMaterialOptionNames)
 						.InitiallySelectedItem(SelectedItem)
-						.OnSelectionChanged(this, &FCustomizableObjectNodeParentedMaterialDetails::OnParentComboBoxSelectionChanged, ParentProperty)
+						.OnSelectionChanged(this, &FCustomizableObjectNodeParentedMaterialDetails::OnParentComboBoxSelectionChanged, ParentProperty.ToSharedRef())
 					]
 				]
 			]
@@ -132,7 +133,7 @@ void FCustomizableObjectNodeParentedMaterialDetails::CustomizeDetails(IDetailLay
 	}
 	else
 	{
-		BlocksCategory.AddCustomRow(LOCTEXT("FCustomizableObjectNodeEditMaterialBaseDetails", "Node"))
+		ParentCategory.AddCustomRow(LOCTEXT("FCustomizableObjectNodeEditMaterialBaseDetails", "Node"))
 		[
 			SNew(STextBlock)
 			.Text(LOCTEXT("Node not found", "Node not found"))
diff --git a/Engine/Plugins/Experimental/Mutable/Source/CustomizableObjectEditor/Private/MuCOE/Nodes/CustomizableObjectNodeRemoveMeshBlocksDetails.cpp b/Engine/Plugins/Experimental/Mutable/Source/CustomizableObjectEditor/Private/MuCOE/Nodes/CustomizableObjectNodeRemoveMeshBlocksDetails.cpp
--- a/Engine/Plugins/Experimental/Mutable/Source/CustomizableObjectEditor/Private/MuCOE/Nodes/CustomizableObjectNodeRemoveMeshBlocksDetails.cpp
+++ b/Engine/Plugins/Experimental/Mutable/Source/CustomizableObjectEditor/Private/MuCOE/Nodes/CustomizableObjectNodeRemoveMeshBlocksDetails.cpp
@@ -24,7 +24,7 @@ TSharedRef<IDetailCustomization> FCustomizableObjectNodeRemoveMeshBlocksDetails:
 
 void FCustomizableObjectNodeRemoveMeshBlocksDetails::CustomizeDetails( IDetailLayoutBuilder& DetailBuilder )
 {
-	FCustomizableObjectNodeDetails::CustomizeDetails(DetailBuilder);
+	FCustomizableObjectNodeEditMaterialBaseDetails::CustomizeDetails(DetailBuilder);
 
 	Node = nullptr;
 	const IDetailsView* DetailsView = DetailBuilder.GetDetailsView();
diff --git a/Engine/Plugins/Experimental/Mutable/Source/CustomizableObjectEditor/Private/MuCOE/Nodes/CustomizableObjectNodeRemoveMeshBlocksDetails.h b/Engine/Plugins/Experimental/Mutable/Source/CustomizableObjectEditor/Private/MuCOE/Nodes/CustomizableObjectNodeRemoveMeshBlocksDetails.h
--- a/Engine/Plugins/Experimental/Mutable/Source/CustomizableObjectEditor/Private/MuCOE/Nodes/CustomizableObjectNodeRemoveMeshBlocksDetails.h
+++ b/Engine/Plugins/Experimental/Mutable/Source/CustomizableObjectEditor/Private/MuCOE/Nodes/CustomizableObjectNodeRemoveMeshBlocksDetails.h
@@ -4,6 +4,7 @@
 
 
 #include "CustomizableObjectNodeDetails.h"
+#include "CustomizableObjectNodeEditMaterialBaseDetails.h"
 #include "IDetailCustomization.h"
 
 class FString;
@@ -12,21 +13,18 @@ class UCustomizableObjectNodeRemoveMeshBlocks;
 class SCustomizableObjectNodeLayoutBlocksSelector;
 
 
-class FCustomizableObjectNodeRemoveMeshBlocksDetails : public FCustomizableObjectNodeDetails
+class FCustomizableObjectNodeRemoveMeshBlocksDetails : public FCustomizableObjectNodeEditMaterialBaseDetails
 {
 public:
 	/** Makes a new instance of this detail layout class for a specific detail view requesting it */
 	static TSharedRef<IDetailCustomization> MakeInstance();
 
 	/** ILayoutDetails interface */
-	virtual void CustomizeDetails( IDetailLayoutBuilder& DetailBuilder ) override;
-
+	virtual void CustomizeDetails(IDetailLayoutBuilder& DetailBuilder) override;
 
 private:
-
 	UCustomizableObjectNodeRemoveMeshBlocks* Node;
 
 	// Layout block editor widget
 	TSharedPtr<SCustomizableObjectNodeLayoutBlocksSelector> LayoutBlocksSelector;
-
 };
