[Mutable] Fix crash when generating components.

#rnx

diff --git a/Engine/Plugins/Experimental/Mutable/Source/CustomizableObjectEditor/Private/MuCOE/GenerateMutableSource/GenerateMutableSourceComponent.cpp b/Engine/Plugins/Experimental/Mutable/Source/CustomizableObjectEditor/Private/MuCOE/GenerateMutableSource/GenerateMutableSourceComponent.cpp
--- a/Engine/Plugins/Experimental/Mutable/Source/CustomizableObjectEditor/Private/MuCOE/GenerateMutableSource/GenerateMutableSourceComponent.cpp
+++ b/Engine/Plugins/Experimental/Mutable/Source/CustomizableObjectEditor/Private/MuCOE/GenerateMutableSource/GenerateMutableSourceComponent.cpp
@@ -86,17 +86,17 @@ mu::Ptr<mu::NodeComponent> GenerateMutableSourceComponent(const UEdGraphPin * Pi
 			for (int32 SectionIndex = 0; SectionIndex < SectionCount; ++SectionIndex)
 			{
 				// Is there a pin in the unreal node for this section?
-				UEdGraphPin* InMaterialPin = TypedComponentMesh->GetMaterialPin(LODIndex,SectionIndex);
-				UEdGraphPin* ConnectedMaterialPin = FollowInputPin(*InMaterialPin);
-
-				if (ConnectedMaterialPin)
+				if (UEdGraphPin* InMaterialPin = TypedComponentMesh->GetMaterialPin(LODIndex,SectionIndex))
 				{
-					GenerationContext.ComponentMeshOverride = MeshNode;
+					if (UEdGraphPin* ConnectedMaterialPin = FollowInputPin(*InMaterialPin))
+					{
+						GenerationContext.ComponentMeshOverride = MeshNode;
 					
-					mu::Ptr<mu::NodeSurface> SurfaceNode = GenerateMutableSourceSurface(ConnectedMaterialPin, GenerationContext);
-					LODNode->Surfaces.Add(SurfaceNode);
+						mu::Ptr<mu::NodeSurface> SurfaceNode = GenerateMutableSourceSurface(ConnectedMaterialPin, GenerationContext);
+						LODNode->Surfaces.Add(SurfaceNode);
 
-					GenerationContext.ComponentMeshOverride = nullptr;
+						GenerationContext.ComponentMeshOverride = nullptr;
+					}
 				}
 			}
 		}
