[mutable] Layout node cleanup. Remove boilerplate and unnecessary hierarchy. Move compiler-only attribute out of runtime data.
#rnx

diff --git a/Engine/Plugins/Experimental/Mutable/Source/CustomizableObject/Internal/MuCO/CustomizableObjectPrivate.h b/Engine/Plugins/Experimental/Mutable/Source/CustomizableObject/Internal/MuCO/CustomizableObjectPrivate.h
--- a/Engine/Plugins/Experimental/Mutable/Source/CustomizableObject/Internal/MuCO/CustomizableObjectPrivate.h
+++ b/Engine/Plugins/Experimental/Mutable/Source/CustomizableObject/Internal/MuCO/CustomizableObjectPrivate.h
@@ -1089,6 +1089,8 @@ public:
 
 		AddedRomFlags,
 
+		LayoutNodeCleanup,
+
 		// -----<new versions can be added above this line>--------
 		LastCustomizableObjectVersion
 	};
diff --git a/Engine/Plugins/Experimental/Mutable/Source/CustomizableObjectEditor/Private/MuCOE/CustomizableObjectLayout.cpp b/Engine/Plugins/Experimental/Mutable/Source/CustomizableObjectEditor/Private/MuCOE/CustomizableObjectLayout.cpp
--- a/Engine/Plugins/Experimental/Mutable/Source/CustomizableObjectEditor/Private/MuCOE/CustomizableObjectLayout.cpp
+++ b/Engine/Plugins/Experimental/Mutable/Source/CustomizableObjectEditor/Private/MuCOE/CustomizableObjectLayout.cpp
@@ -88,7 +88,7 @@ void UCustomizableObjectLayout::SetLayoutName(FString Name)
 
 void UCustomizableObjectLayout::GenerateBlocksFromUVs()
 {
-	mu::NodeLayoutBlocksPtr Layout = nullptr;
+	mu::Ptr<mu::NodeLayout> LayoutNode;
 
 	UCustomizableObjectNode* Node = Cast<UCustomizableObjectNode>(GetOuter());
 
@@ -124,18 +124,18 @@ void UCustomizableObjectLayout::GenerateBlocksFromUVs()
 		if (MutableMesh)
 		{
 			// Generating blocks with the mutable mesh
-			Layout = mu::NodeLayoutBlocks::GenerateLayoutBlocks(MutableMesh, UVChannel, GridSize.X, GridSize.Y);
+			LayoutNode = mu::NodeLayout::GenerateLayoutBlocks(MutableMesh, UVChannel, GridSize.X, GridSize.Y);
 		}
 	
-		if (Layout)
+		if (LayoutNode)
 		{
 			Blocks.Empty();
 		
 			// Generating the layout blocks with the mutable layout
-			for (int i = 0; i < Layout->GetBlockCount(); ++i)
+			for (int i = 0; i < LayoutNode->Layout->GetBlockCount(); ++i)
 			{
-				mu::FImageSize Min = Layout->GetLayout()->Blocks[i].Min;
-				mu::FImageSize Size = Layout->GetLayout()->Blocks[i].Size;
+				mu::FImageSize Min = LayoutNode->Layout->Blocks[i].Min;
+				mu::FImageSize Size = LayoutNode->Layout->Blocks[i].Size;
 		
 				FCustomizableObjectLayoutBlock block(FIntPoint(Min.X, Min.Y), FIntPoint(Min.X + Size.X, Min.Y + Size.Y));
 				Blocks.Add(block);
diff --git a/Engine/Plugins/Experimental/Mutable/Source/CustomizableObjectEditor/Private/MuCOE/GenerateMutableSource/GenerateMutableSourceLayout.cpp b/Engine/Plugins/Experimental/Mutable/Source/CustomizableObjectEditor/Private/MuCOE/GenerateMutableSource/GenerateMutableSourceLayout.cpp
--- a/Engine/Plugins/Experimental/Mutable/Source/CustomizableObjectEditor/Private/MuCOE/GenerateMutableSource/GenerateMutableSourceLayout.cpp
+++ b/Engine/Plugins/Experimental/Mutable/Source/CustomizableObjectEditor/Private/MuCOE/GenerateMutableSource/GenerateMutableSourceLayout.cpp
@@ -11,7 +11,7 @@
 #define LOCTEXT_NAMESPACE "CustomizableObjectEditor"
 
 
-mu::NodeLayoutPtr GenerateMutableSourceLayout(const UEdGraphPin * Pin, FMutableGraphGenerationContext & GenerationContext, bool bLinkedToExtendMaterial)
+mu::Ptr<mu::NodeLayout> GenerateMutableSourceLayout(const UEdGraphPin * Pin, FMutableGraphGenerationContext & GenerationContext, bool bLinkedToExtendMaterial)
 {
 	check(Pin)
 	RETURN_ON_CYCLE(*Pin, GenerationContext)
@@ -26,7 +26,7 @@ mu::NodeLayoutPtr GenerateMutableSourceLayout(const UEdGraphPin * Pin, FMutableG
 		return static_cast<mu::NodeLayout*>(Generated->Node.get());
 	}
 
-	mu::NodeLayoutPtr Result;
+	mu::Ptr<mu::NodeLayout> Result;
 	
 	if (const UCustomizableObjectNodeLayoutBlocks* TypedNodeBlocks = Cast<UCustomizableObjectNodeLayoutBlocks>(Node))
 	{
@@ -52,39 +52,33 @@ mu::NodeLayoutPtr GenerateMutableSourceLayout(const UEdGraphPin * Pin, FMutableG
 			}
 		}
 
-		mu::NodeLayoutBlocksPtr LayoutNode = new mu::NodeLayoutBlocks;
+		mu::Ptr<mu::NodeLayout> LayoutNode = new mu::NodeLayout;
 		Result = LayoutNode;
 
-		LayoutNode->SetGridSize(TypedNodeBlocks->Layout->GetGridSize().X, TypedNodeBlocks->Layout->GetGridSize().Y);
-		LayoutNode->SetMaxGridSize(TypedNodeBlocks->Layout->GetMaxGridSize().X, TypedNodeBlocks->Layout->GetMaxGridSize().Y);
-		LayoutNode->SetBlockCount(TypedNodeBlocks->Layout->Blocks.Num() ? TypedNodeBlocks->Layout->Blocks.Num() : 1);
+		LayoutNode->Layout->SetGridSize(TypedNodeBlocks->Layout->GetGridSize().X, TypedNodeBlocks->Layout->GetGridSize().Y);
+		LayoutNode->Layout->SetMaxGridSize(TypedNodeBlocks->Layout->GetMaxGridSize().X, TypedNodeBlocks->Layout->GetMaxGridSize().Y);
+		LayoutNode->Layout->SetBlockCount(TypedNodeBlocks->Layout->Blocks.Num() ? TypedNodeBlocks->Layout->Blocks.Num() : 1);
 
 		mu::EPackStrategy PackStrategy = ConvertLayoutStrategy(TypedNodeBlocks->Layout->GetPackingStrategy());
-		LayoutNode->SetLayoutPackingStrategy(PackStrategy);
+		LayoutNode->Layout->SetLayoutPackingStrategy(PackStrategy);
 
-		LayoutNode->SetBlockReductionMethod(TypedNodeBlocks->Layout->GetBlockReductionMethod() == ECustomizableObjectLayoutBlockReductionMethod::Halve ? mu::EReductionMethod::HALVE_REDUCTION : mu::EReductionMethod::UNITARY_REDUCTION);
+		LayoutNode->Layout->ReductionMethod = (TypedNodeBlocks->Layout->GetBlockReductionMethod() == ECustomizableObjectLayoutBlockReductionMethod::Halve ? mu::EReductionMethod::HALVE_REDUCTION : mu::EReductionMethod::UNITARY_REDUCTION);
 
 		if (bLinkedToExtendMaterial)
 		{
 			// Layout warnings can be safely ignored in this case. Vertices that do not belong to any layout block will be removed (Extend Materials only)
-			LayoutNode->SetIgnoreWarningsLOD(0);
+			LayoutNode->FirstLODToIgnoreWarnings = 0;
 		}
 		else
 		{
-			LayoutNode->SetIgnoreWarningsLOD(TypedNodeBlocks->Layout->GetIgnoreVertexLayoutWarnings() ? TypedNodeBlocks->Layout->GetFirstLODToIgnoreWarnings() : -1);
+			LayoutNode->FirstLODToIgnoreWarnings = TypedNodeBlocks->Layout->GetIgnoreVertexLayoutWarnings() ? TypedNodeBlocks->Layout->GetFirstLODToIgnoreWarnings() : -1;
 		}
 
 		if (TypedNodeBlocks->Layout->Blocks.Num())
 		{
 			for (int BlockIndex = 0; BlockIndex < TypedNodeBlocks->Layout->Blocks.Num(); ++BlockIndex)
 			{
-				LayoutNode->SetBlock(BlockIndex,
-					TypedNodeBlocks->Layout->Blocks[BlockIndex].Min.X,
-					TypedNodeBlocks->Layout->Blocks[BlockIndex].Min.Y,
-					TypedNodeBlocks->Layout->Blocks[BlockIndex].Max.X - TypedNodeBlocks->Layout->Blocks[BlockIndex].Min.X,
-					TypedNodeBlocks->Layout->Blocks[BlockIndex].Max.Y - TypedNodeBlocks->Layout->Blocks[BlockIndex].Min.Y);
-
-				LayoutNode->SetBlockOptions(BlockIndex, TypedNodeBlocks->Layout->Blocks[BlockIndex].Priority, TypedNodeBlocks->Layout->Blocks[BlockIndex].bReduceBothAxes, TypedNodeBlocks->Layout->Blocks[BlockIndex].bReduceByTwo);
+				LayoutNode->Layout->Blocks[BlockIndex]= ToMutable(TypedNodeBlocks->Layout->Blocks[BlockIndex]);
 			}
 		}
 		else
@@ -92,8 +86,11 @@ mu::NodeLayoutPtr GenerateMutableSourceLayout(const UEdGraphPin * Pin, FMutableG
 			FString msg = "Layout without any block found. A grid sized block will be used instead.";
 			GenerationContext.Compiler->CompilerLog(FText::FromString(msg), Node, EMessageSeverity::Warning);
 
-			LayoutNode->SetBlock(0, 0, 0, TypedNodeBlocks->Layout->GetGridSize().X, TypedNodeBlocks->Layout->GetGridSize().Y);
-			LayoutNode->SetBlockOptions(0, 0, false, false);
+			LayoutNode->Layout->Blocks[0].Min = { 0,0 };
+			LayoutNode->Layout->Blocks[0].Size = { uint16(TypedNodeBlocks->Layout->GetGridSize().X), uint16(TypedNodeBlocks->Layout->GetGridSize().Y) };
+			LayoutNode->Layout->Blocks[0].Priority = 0;
+			LayoutNode->Layout->Blocks[0].bReduceBothAxes = false;
+			LayoutNode->Layout->Blocks[0].bReduceByTwo = false;
 		}
 	}
 	
diff --git a/Engine/Plugins/Experimental/Mutable/Source/CustomizableObjectEditor/Private/MuCOE/GenerateMutableSource/GenerateMutableSourceLayout.h b/Engine/Plugins/Experimental/Mutable/Source/CustomizableObjectEditor/Private/MuCOE/GenerateMutableSource/GenerateMutableSourceLayout.h
--- a/Engine/Plugins/Experimental/Mutable/Source/CustomizableObjectEditor/Private/MuCOE/GenerateMutableSource/GenerateMutableSourceLayout.h
+++ b/Engine/Plugins/Experimental/Mutable/Source/CustomizableObjectEditor/Private/MuCOE/GenerateMutableSource/GenerateMutableSourceLayout.h
@@ -3,9 +3,27 @@
 #pragma once
 
 #include "MuT/NodeLayout.h"
+#include "MuCOE/CustomizableObjectLayout.h"
 
 class UEdGraphPin;
 struct FMutableGraphGenerationContext;
 
 
-mu::NodeLayoutPtr GenerateMutableSourceLayout(const UEdGraphPin* Pin, FMutableGraphGenerationContext& GenerationContext, bool bLinkedToExtendMaterial = false);
+mu::Ptr<mu::NodeLayout> GenerateMutableSourceLayout(const UEdGraphPin* Pin, FMutableGraphGenerationContext& GenerationContext, bool bLinkedToExtendMaterial = false);
+
+
+/** */
+inline mu::Layout::FBlock ToMutable(const FCustomizableObjectLayoutBlock& UnrealBlock)
+{
+	mu::Layout::FBlock MutableBlock;
+
+	MutableBlock.Min = { uint16(UnrealBlock.Min.X), uint16(UnrealBlock.Min.Y) };
+	FIntPoint Size = UnrealBlock.Max - UnrealBlock.Min;
+	MutableBlock.Size = { uint16(Size.X), uint16(Size.Y) };
+
+	MutableBlock.Priority = UnrealBlock.Priority;
+	MutableBlock.bReduceBothAxes = UnrealBlock.bReduceBothAxes;
+	MutableBlock.bReduceByTwo = UnrealBlock.bReduceByTwo;
+
+	return MutableBlock;
+}
diff --git a/Engine/Plugins/Experimental/Mutable/Source/CustomizableObjectEditor/Private/MuCOE/GenerateMutableSource/GenerateMutableSourceMesh.cpp b/Engine/Plugins/Experimental/Mutable/Source/CustomizableObjectEditor/Private/MuCOE/GenerateMutableSource/GenerateMutableSourceMesh.cpp
--- a/Engine/Plugins/Experimental/Mutable/Source/CustomizableObjectEditor/Private/MuCOE/GenerateMutableSource/GenerateMutableSourceMesh.cpp
+++ b/Engine/Plugins/Experimental/Mutable/Source/CustomizableObjectEditor/Private/MuCOE/GenerateMutableSource/GenerateMutableSourceMesh.cpp
@@ -3133,18 +3133,21 @@ mu::NodeMeshPtr GenerateMorphMesh(const UEdGraphPin* Pin,
 }
 
 /** Create a default layout. Used when no layout is found. */
-mu::NodeLayoutBlocksPtr CreateDefaultLayout()
+mu::Ptr<mu::NodeLayout> CreateDefaultLayout()
 {
 	constexpr int32 GridSize = 4;
 	
-	mu::NodeLayoutBlocksPtr LayoutNode = new mu::NodeLayoutBlocks();
-	LayoutNode->SetGridSize(GridSize, GridSize);
-	LayoutNode->SetMaxGridSize(GridSize, GridSize);
-	LayoutNode->SetLayoutPackingStrategy(mu::EPackStrategy::RESIZABLE_LAYOUT);
-	LayoutNode->SetBlockReductionMethod(mu::EReductionMethod::HALVE_REDUCTION);
-	LayoutNode->SetBlockCount(1);
-	LayoutNode->SetBlock(0, 0, 0, GridSize, GridSize);
-	LayoutNode->SetBlockOptions(0, 0, false, false);
+	mu::Ptr<mu::NodeLayout> LayoutNode = new mu::NodeLayout();
+	LayoutNode->Layout->SetGridSize(GridSize, GridSize);
+	LayoutNode->Layout->SetMaxGridSize(GridSize, GridSize);
+	LayoutNode->Layout->SetLayoutPackingStrategy(mu::EPackStrategy::RESIZABLE_LAYOUT);
+	LayoutNode->Layout->ReductionMethod = mu::EReductionMethod::HALVE_REDUCTION;
+	LayoutNode->Layout->SetBlockCount(1);
+	LayoutNode->Layout->Blocks[0].Min = { 0, 0 };
+	LayoutNode->Layout->Blocks[0].Size = { GridSize, GridSize };
+	LayoutNode->Layout->Blocks[0].Priority = 0;
+	LayoutNode->Layout->Blocks[0].bReduceBothAxes = false;
+	LayoutNode->Layout->Blocks[0].bReduceByTwo = false;
 
 	return LayoutNode;
 }
@@ -3376,7 +3379,7 @@ mu::NodeMeshPtr GenerateMutableSourceMesh(const UEdGraphPin* Pin,
 					{
 						if (const UEdGraphPin* ConnectedPin = FollowInputPin(*LayoutPin))
 						{
-							mu::NodeLayoutPtr LayoutNode = GenerateMutableSourceLayout(ConnectedPin, GenerationContext, bLinkedToExtendMaterial);
+							mu::Ptr<mu::NodeLayout> LayoutNode = GenerateMutableSourceLayout(ConnectedPin, GenerationContext, bLinkedToExtendMaterial);
 							MeshNode->SetLayout(LayoutIndex, LayoutNode);
 							bAtLeastOneLayout = true;
 						}
@@ -3387,7 +3390,7 @@ mu::NodeMeshPtr GenerateMutableSourceMesh(const UEdGraphPin* Pin,
 				{
 					MeshNode->SetLayoutCount(1);
 
-					mu::NodeLayoutBlocksPtr LayoutNode = CreateDefaultLayout();
+					mu::Ptr<mu::NodeLayout> LayoutNode = CreateDefaultLayout();
 					MeshNode->SetLayout(0, LayoutNode);
 					LayoutNode->SetMessageContext(Node); // We need it here because we create multiple nodes.
 				}
@@ -3460,12 +3463,12 @@ mu::NodeMeshPtr GenerateMutableSourceMesh(const UEdGraphPin* Pin,
 
 				if (const UEdGraphPin* ConnectedPin = FollowInputPin(*TypedNodeStatic->LODs[LODIndex].Materials[SectionIndex].LayoutPinRef.Get()))
 				{
-					mu::NodeLayoutPtr LayoutNode = GenerateMutableSourceLayout(ConnectedPin, GenerationContext);
+					mu::Ptr<mu::NodeLayout> LayoutNode = GenerateMutableSourceLayout(ConnectedPin, GenerationContext);
 					MeshNode->SetLayout(0, LayoutNode);
 				}
 				else
 				{
-					mu::NodeLayoutBlocksPtr LayoutNode = CreateDefaultLayout();
+					mu::Ptr<mu::NodeLayout> LayoutNode = CreateDefaultLayout();
 					MeshNode->SetLayout(0, LayoutNode);
 					LayoutNode->SetMessageContext(Node);  // We need it here because we create multiple nodes.
 				}
@@ -4080,34 +4083,28 @@ mu::NodeMeshPtr GenerateMutableSourceMesh(const UEdGraphPin* Pin,
 							// Generating node Layouts
 							for (int32 i = 0; i < Layouts.Num(); ++i)
 							{
-								mu::NodeLayoutBlocksPtr LayoutNode = new mu::NodeLayoutBlocks;
+								mu::Ptr<mu::NodeLayout> LayoutNode = new mu::NodeLayout;
 
-								LayoutNode->SetGridSize(Layouts[i]->GetGridSize().X, Layouts[i]->GetGridSize().Y);
-								LayoutNode->SetMaxGridSize(Layouts[i]->GetMaxGridSize().X, Layouts[i]->GetMaxGridSize().Y);
-								LayoutNode->SetBlockCount(Layouts[i]->Blocks.Num() ? Layouts[i]->Blocks.Num() : 1);
+								LayoutNode->Layout->SetGridSize(Layouts[i]->GetGridSize().X, Layouts[i]->GetGridSize().Y);
+								LayoutNode->Layout->SetMaxGridSize(Layouts[i]->GetMaxGridSize().X, Layouts[i]->GetMaxGridSize().Y);
+								LayoutNode->Layout->SetBlockCount(Layouts[i]->Blocks.Num() ? Layouts[i]->Blocks.Num() : 1);
 
 								mu::EPackStrategy PackStrategy = ConvertLayoutStrategy(Layouts[i]->GetPackingStrategy());
-								LayoutNode->SetLayoutPackingStrategy(PackStrategy);
+								LayoutNode->Layout->SetLayoutPackingStrategy(PackStrategy);
 								
-								LayoutNode->SetBlockReductionMethod(Layouts[i]->GetBlockReductionMethod() == ECustomizableObjectLayoutBlockReductionMethod::Halve ? mu::EReductionMethod::HALVE_REDUCTION : mu::EReductionMethod::UNITARY_REDUCTION);
+								LayoutNode->Layout->ReductionMethod = (Layouts[i]->GetBlockReductionMethod() == ECustomizableObjectLayoutBlockReductionMethod::Halve ? mu::EReductionMethod::HALVE_REDUCTION : mu::EReductionMethod::UNITARY_REDUCTION);
 
 								if (bLinkedToExtendMaterial)
 								{
 									// Layout warnings can be safely ignored in this case. Vertices that do not belong to any layout block will be removed (Extend Materials only)
-									LayoutNode->SetIgnoreWarningsLOD(0);
+									LayoutNode->FirstLODToIgnoreWarnings = 0;
 								}
 
 								if (Layouts[i]->Blocks.Num())
 								{
 									for (int BlockIndex = 0; BlockIndex < Layouts[i]->Blocks.Num(); ++BlockIndex)
 									{
-										LayoutNode->SetBlock(BlockIndex,
-											Layouts[i]->Blocks[BlockIndex].Min.X,
-											Layouts[i]->Blocks[BlockIndex].Min.Y,
-											Layouts[i]->Blocks[BlockIndex].Max.X - Layouts[i]->Blocks[BlockIndex].Min.X,
-											Layouts[i]->Blocks[BlockIndex].Max.Y - Layouts[i]->Blocks[BlockIndex].Min.Y);
-
-										LayoutNode->SetBlockOptions(BlockIndex, Layouts[i]->Blocks[BlockIndex].Priority, Layouts[i]->Blocks[BlockIndex].bReduceBothAxes, Layouts[i]->Blocks[BlockIndex].bReduceByTwo);
+										LayoutNode->Layout->Blocks[BlockIndex] = ToMutable( Layouts[i]->Blocks[BlockIndex] );
 									}
 								}
 								else
@@ -4115,8 +4112,11 @@ mu::NodeMeshPtr GenerateMutableSourceMesh(const UEdGraphPin* Pin,
 									FString msg = "Mesh Column [" + MutableColumnName + "] Layout doesn't has any block. A grid sized block will be used instead.";
 									GenerationContext.Compiler->CompilerLog(FText::FromString(msg), Node, EMessageSeverity::Warning);
 
-									LayoutNode->SetBlock(0, 0, 0, Layouts[i]->GetGridSize().X, Layouts[i]->GetGridSize().Y);
-									LayoutNode->SetBlockOptions(0, 0, false, false);
+									LayoutNode->Layout->Blocks[0].Min = { 0,0 };
+									LayoutNode->Layout->Blocks[0].Size = { uint16(Layouts[i]->GetGridSize().X), uint16(Layouts[i]->GetGridSize().Y)};
+									LayoutNode->Layout->Blocks[0].Priority = 0;
+									LayoutNode->Layout->Blocks[0].bReduceBothAxes = false;
+									LayoutNode->Layout->Blocks[0].bReduceByTwo = false;
 								}
 
 								MeshTableNode->SetLayout(i, LayoutNode);
diff --git a/Engine/Plugins/Experimental/Mutable/Source/CustomizableObjectEditor/Private/MuCOE/SMutableGraphViewer.cpp b/Engine/Plugins/Experimental/Mutable/Source/CustomizableObjectEditor/Private/MuCOE/SMutableGraphViewer.cpp
--- a/Engine/Plugins/Experimental/Mutable/Source/CustomizableObjectEditor/Private/MuCOE/SMutableGraphViewer.cpp
+++ b/Engine/Plugins/Experimental/Mutable/Source/CustomizableObjectEditor/Private/MuCOE/SMutableGraphViewer.cpp
@@ -595,7 +595,7 @@ void SMutableGraphViewer::GetChildrenForInfo(TSharedPtr<FMutableGraphTreeElement
 		AddChildFunc(Private->m_pColour.get(), FString::Printf(TEXT("COLOR")));
 	}
 
-	else if (ParentNode->GetType() == mu::NodeLayoutBlocks::GetStaticType())
+	else if (ParentNode->GetType() == mu::NodeLayout::GetStaticType())
 	{
 		// Nothing to show
 	}
diff --git a/Engine/Plugins/Experimental/Mutable/Source/MutableRuntime/Private/MuR/Layout.cpp b/Engine/Plugins/Experimental/Mutable/Source/MutableRuntime/Private/MuR/Layout.cpp
--- a/Engine/Plugins/Experimental/Mutable/Source/MutableRuntime/Private/MuR/Layout.cpp
+++ b/Engine/Plugins/Experimental/Mutable/Source/MutableRuntime/Private/MuR/Layout.cpp
@@ -46,7 +46,6 @@ namespace mu {
 		pResult->MaxSize = MaxSize;
 		pResult->Blocks = Blocks;
 		pResult->Strategy = Strategy;
-		pResult->FirstLODToIgnoreWarnings = FirstLODToIgnoreWarnings;
 		pResult->ReductionMethod = ReductionMethod;
 		return pResult;
 	}
@@ -60,7 +59,6 @@ namespace mu {
 			(Blocks == o.Blocks) &&
 			(Strategy == o.Strategy) &&
 			// maybe this is not needed
-			(FirstLODToIgnoreWarnings == o.FirstLODToIgnoreWarnings) &&
 			(ReductionMethod==o.ReductionMethod);
 	}
 
@@ -148,7 +146,6 @@ namespace mu {
 
 		arch << MaxSize;
 		arch << uint32(Strategy);
-		arch << FirstLODToIgnoreWarnings;
 		arch << uint32(ReductionMethod);
 	}
 
@@ -164,8 +161,6 @@ namespace mu {
 		arch >> Temp;
 		Strategy = EPackStrategy(Temp);
 
-		arch >> FirstLODToIgnoreWarnings;
-
 		arch >> Temp;
 		ReductionMethod = EReductionMethod(Temp);
 	}
@@ -216,32 +211,5 @@ namespace mu {
 	}
 
 
-	//---------------------------------------------------------------------------------------------
-	void Layout::SetIgnoreLODWarnings(int32 LOD)
-	{
-		FirstLODToIgnoreWarnings = LOD;
-	}
-
-
-	//---------------------------------------------------------------------------------------------
-	int32 Layout::GetIgnoreLODWarnings()
-	{
-		return FirstLODToIgnoreWarnings;
-	}
-
-
-	//---------------------------------------------------------------------------------------------
-	void Layout::SetBlockReductionMethod(EReductionMethod Method)
-	{
-		ReductionMethod = Method;
-	}
-
-
-	//---------------------------------------------------------------------------------------------
-	EReductionMethod Layout::GetBlockReductionMethod() const
-	{
-		return ReductionMethod;
-	}
-
 }
 
diff --git a/Engine/Plugins/Experimental/Mutable/Source/MutableRuntime/Private/MuR/OpLayoutPack.h b/Engine/Plugins/Experimental/Mutable/Source/MutableRuntime/Private/MuR/OpLayoutPack.h
--- a/Engine/Plugins/Experimental/Mutable/Source/MutableRuntime/Private/MuR/OpLayoutPack.h
+++ b/Engine/Plugins/Experimental/Mutable/Source/MutableRuntime/Private/MuR/OpLayoutPack.h
@@ -410,7 +410,7 @@ namespace mu
 		bool usePriority = false;
 
 		EPackStrategy LayoutStrategy = pSourceLayout->GetLayoutPackingStrategy();
-		EReductionMethod ReductionMethod = pSourceLayout->GetBlockReductionMethod();
+		EReductionMethod ReductionMethod = pSourceLayout->ReductionMethod;
 
         // Look for the maximum block sizes on the layout and the total area
 		uint16 maxX = 0;
@@ -636,7 +636,7 @@ namespace mu
         pResult->SetGridSize( maxX, maxY );
         pResult->SetMaxGridSize(layoutSizeX, layoutSizeY);
 		pResult->SetLayoutPackingStrategy(LayoutStrategy);
-		pResult->SetBlockReductionMethod(ReductionMethod);
+		pResult->ReductionMethod = ReductionMethod;
 
         for ( int32 index=0; index<blockCount; ++index )
         {
diff --git a/Engine/Plugins/Experimental/Mutable/Source/MutableRuntime/Public/MuR/Layout.h b/Engine/Plugins/Experimental/Mutable/Source/MutableRuntime/Public/MuR/Layout.h
--- a/Engine/Plugins/Experimental/Mutable/Source/MutableRuntime/Public/MuR/Layout.h
+++ b/Engine/Plugins/Experimental/Mutable/Source/MutableRuntime/Public/MuR/Layout.h
@@ -106,18 +106,6 @@ namespace mu
 		//! Set the texture layout packing strategy
 		EPackStrategy GetLayoutPackingStrategy() const;
 
-		//! Set at which LOD the unassigned vertices warnings will star to be ignored
-		void SetIgnoreLODWarnings(int32 LOD);
-
-		//! Get the LOD where the unassigned vertices warnings starts to be ignored
-		int32 GetIgnoreLODWarnings();
-
-		//! Set the block reduction method a the Fixed_Layout strategy
-		void SetBlockReductionMethod(EReductionMethod Method);
-
-		//! Returns the block reduction method
-		EReductionMethod GetBlockReductionMethod() const;
-
 	protected:
 
 		//! Forbidden. Manage with the Ptr<> template.
@@ -191,8 +179,6 @@ namespace mu
 		//! Packing strategy
 		EPackStrategy Strategy = EPackStrategy::RESIZABLE_LAYOUT;
 		 
-		int32 FirstLODToIgnoreWarnings = 0;
-
 		EReductionMethod ReductionMethod = EReductionMethod::HALVE_REDUCTION;
 
 
diff --git a/Engine/Plugins/Experimental/Mutable/Source/MutableRuntime/Public/MuR/Model.h b/Engine/Plugins/Experimental/Mutable/Source/MutableRuntime/Public/MuR/Model.h
--- a/Engine/Plugins/Experimental/Mutable/Source/MutableRuntime/Public/MuR/Model.h
+++ b/Engine/Plugins/Experimental/Mutable/Source/MutableRuntime/Public/MuR/Model.h
@@ -9,14 +9,6 @@
 #include "Templates/SharedPointer.h"
 
 
-//! This version number changes whenever there is a compatibility-breaking change in the Model
-//! data structures. Compiled models are not necessarily compatible when the runtime is updated,
-//! so this version number can be used externally to verify this. It is not used internally, and
-//! serializing models from different versions than this runtime will probably result in a crash.
-#define MUTABLE_COMPILED_MODEL_CODE_VERSION		uint32( 85 )
-#define MUTABLE_PARAMETERS_VERSION              uint32( 3 )
-
-
 namespace mu
 {
 	class InputArchive;
diff --git a/Engine/Plugins/Experimental/Mutable/Source/MutableTools/Internal/MuT/NodeSurfaceEditPrivate.h b/Engine/Plugins/Experimental/Mutable/Source/MutableTools/Internal/MuT/NodeSurfaceEditPrivate.h
--- a/Engine/Plugins/Experimental/Mutable/Source/MutableTools/Internal/MuT/NodeSurfaceEditPrivate.h
+++ b/Engine/Plugins/Experimental/Mutable/Source/MutableTools/Internal/MuT/NodeSurfaceEditPrivate.h
@@ -28,10 +28,6 @@ namespace mu
 		NodePatchMeshPtr m_pMesh;
 		NodeMeshPtr m_pMorph;
 
-        //! This flag indicates that the mesh in the m_pMorph connection is a target mesh, and not
-        //! morph information. This means that the morph information needs to be generated.
-        bool m_morphIsTarget = true;
-
 		struct FTexture
 		{
 			NodeImagePtr m_pExtend;
diff --git a/Engine/Plugins/Experimental/Mutable/Source/MutableTools/Internal/MuT/NodeSurfaceNewPrivate.h b/Engine/Plugins/Experimental/Mutable/Source/MutableTools/Internal/MuT/NodeSurfaceNewPrivate.h
--- a/Engine/Plugins/Experimental/Mutable/Source/MutableTools/Internal/MuT/NodeSurfaceNewPrivate.h
+++ b/Engine/Plugins/Experimental/Mutable/Source/MutableTools/Internal/MuT/NodeSurfaceNewPrivate.h
@@ -6,7 +6,6 @@
 #include "MuT/NodeSurfaceNew.h"
 
 #include "MuT/NodeMesh.h"
-#include "MuT/NodePatchMesh.h"
 #include "MuT/NodeImage.h"
 #include "MuT/NodeScalar.h"
 #include "MuT/NodeString.h"
diff --git a/Engine/Plugins/Experimental/Mutable/Source/MutableTools/Private/MuT/CodeGenerator.cpp b/Engine/Plugins/Experimental/Mutable/Source/MutableTools/Private/MuT/CodeGenerator.cpp
--- a/Engine/Plugins/Experimental/Mutable/Source/MutableTools/Private/MuT/CodeGenerator.cpp
+++ b/Engine/Plugins/Experimental/Mutable/Source/MutableTools/Private/MuT/CodeGenerator.cpp
@@ -946,7 +946,7 @@ namespace mu
             {
                 Ptr<ASTOp> layoutOp;
 
-				Ptr<const Layout> pLayout = meshResults.GeneratedLayouts[LayoutIndex].get();
+				Ptr<const Layout> pLayout = meshResults.GeneratedLayouts[LayoutIndex].Layout;
                 if ( pLayout )
                 {
 					if (SharedMeshResults)
@@ -989,7 +989,7 @@ namespace mu
 									Ptr<ASTOpConstantResource> op = new ASTOpConstantResource();
 									op->Type = OP_TYPE::LA_CONSTANT;
 
-									Ptr<const Layout> pCloned = data.GeneratedLayouts[LayoutIndex];
+									Ptr<const Layout> pCloned = data.GeneratedLayouts[LayoutIndex].Layout;
 									op->SetValue(pCloned, m_compilerOptions->OptimisationOptions.DiskCacheContext);
 
 									layoutFragmentAd = op;
@@ -1210,7 +1210,7 @@ namespace mu
 						}
 						else
 						{
-							const Layout* pLayout = meshResults.GeneratedLayouts[LayoutIndex].get();
+							const Layout* pLayout = meshResults.GeneratedLayouts[LayoutIndex].Layout.get();
 							check(pLayout);
 
 							Ptr<ASTOpInstanceAdd> op = new ASTOpInstanceAdd();
@@ -1374,7 +1374,7 @@ namespace mu
 									{
 										if (LayoutIndex >= meshResults.extraMeshLayouts[editIndex].GeneratedLayouts.Num()
 											||
-											!meshResults.extraMeshLayouts[editIndex].GeneratedLayouts[LayoutIndex])
+											!meshResults.extraMeshLayouts[editIndex].GeneratedLayouts[LayoutIndex].Layout)
 										{
 											FString Msg = FString::Printf(TEXT("Trying to extend a layout that doesn't exist in object [%s]."),
 												*m_currentParents.Last().m_pObject->Name
@@ -1384,7 +1384,7 @@ namespace mu
 										}
 										else
 										{
-											Ptr<const Layout> pExtendLayout = meshResults.extraMeshLayouts[editIndex].GeneratedLayouts[LayoutIndex];
+											Ptr<const Layout> pExtendLayout = meshResults.extraMeshLayouts[editIndex].GeneratedLayouts[LayoutIndex].Layout;
 
 											// Size of a layout block in pixels
 											FIntPoint extlayout = pExtendLayout->GetGridSize();
diff --git a/Engine/Plugins/Experimental/Mutable/Source/MutableTools/Private/MuT/CodeGenerator.h b/Engine/Plugins/Experimental/Mutable/Source/MutableTools/Private/MuT/CodeGenerator.h
--- a/Engine/Plugins/Experimental/Mutable/Source/MutableTools/Private/MuT/CodeGenerator.h
+++ b/Engine/Plugins/Experimental/Mutable/Source/MutableTools/Private/MuT/CodeGenerator.h
@@ -497,6 +497,19 @@ namespace mu
         //-----------------------------------------------------------------------------------------
         // Meshes
 
+		/** */
+		struct FGeneratedLayout
+		{
+			Ptr<const Layout> Layout;
+			int32 FirstLODToIgnoreWarnings=0;
+
+			FORCEINLINE bool operator==(const FGeneratedLayout& Other) const
+			{
+				return Layout == Other.Layout
+					&& FirstLODToIgnoreWarnings == Other.FirstLODToIgnoreWarnings;
+			}
+		};
+
 		/** Options that affect the generation of meshes. It is like list of what required data we want
 		* while parsing down the mesh node graph.
 		*/
@@ -517,7 +530,7 @@ namespace mu
 			* they are supposed to match some other set of layouts. If the array is empty, layouts
 			* are generated normally.
 			*/
-			TArray<Ptr<const Layout>> OverrideLayouts;			
+			TArray<FGeneratedLayout> OverrideLayouts;
 
 			/** Optional context to use instead of the node error context.
 			 * Be careful since it is not used everywhere. Check usages before assigning a value to it. */
@@ -553,14 +566,14 @@ namespace mu
 			Ptr<ASTOp> baseMeshOp;
 
 			/** Generated node layouts with their own block ids. */
-			TArray<Ptr<const Layout>> GeneratedLayouts;
+			TArray<FGeneratedLayout> GeneratedLayouts;
 
 			TArray<Ptr<ASTOp>> layoutOps;
 
 			struct FExtraLayouts
 			{
 				/** Source node layouts to use with these extra mesh. They don't have block ids. */
-				TArray<Ptr<const Layout>> GeneratedLayouts;
+				TArray<FGeneratedLayout> GeneratedLayouts;
 				Ptr<ASTOp> condition;
 				Ptr<ASTOp> meshFragment;
 			};
@@ -623,7 +636,8 @@ namespace mu
 			int32 currentLayoutChannel,
 			const void* errorContext,
 			const FMeshGenerationOptions& MeshOptions,
-			bool bUseAbsoluteBlockIds);
+			bool bUseAbsoluteBlockIds, 
+			int32 FirstLODToIgnoreWarnings);
 
 		//!
 		Ptr<const Layout> AddLayout(Ptr<const Layout> SourceLayout, uint32 MeshIDPrefix);
diff --git a/Engine/Plugins/Experimental/Mutable/Source/MutableTools/Private/MuT/CodeGenerator_Image.cpp b/Engine/Plugins/Experimental/Mutable/Source/MutableTools/Private/MuT/CodeGenerator_Image.cpp
--- a/Engine/Plugins/Experimental/Mutable/Source/MutableTools/Private/MuT/CodeGenerator_Image.cpp
+++ b/Engine/Plugins/Experimental/Mutable/Source/MutableTools/Private/MuT/CodeGenerator_Image.cpp
@@ -1405,7 +1405,7 @@ namespace mu
 			GenerateMesh( MeshOptions, MeshResult, node.m_pMesh );
 
 			// Match the block id of the block we are generating with the id that resulted in the generated mesh			
-			Ptr<const Layout> Layout = MeshResult.GeneratedLayouts.IsValidIndex(node.m_layout) ? MeshResult.GeneratedLayouts[node.m_layout] : nullptr;
+			Ptr<const Layout> Layout = MeshResult.GeneratedLayouts.IsValidIndex(node.m_layout) ? MeshResult.GeneratedLayouts[node.m_layout].Layout : nullptr;
 			if (Layout && Layout->Blocks.IsValidIndex(LayoutBlockIndex))
 			{
 				GeneratedLayoutBlockId = Layout->Blocks[LayoutBlockIndex].Id;
diff --git a/Engine/Plugins/Experimental/Mutable/Source/MutableTools/Private/MuT/CodeGenerator_Mesh.cpp b/Engine/Plugins/Experimental/Mutable/Source/MutableTools/Private/MuT/CodeGenerator_Mesh.cpp
--- a/Engine/Plugins/Experimental/Mutable/Source/MutableTools/Private/MuT/CodeGenerator_Mesh.cpp
+++ b/Engine/Plugins/Experimental/Mutable/Source/MutableTools/Private/MuT/CodeGenerator_Mesh.cpp
@@ -40,7 +40,6 @@
 #include "MuT/ErrorLogPrivate.h"
 #include "MuT/NodeImageProject.h"
 #include "MuT/NodeLayout.h"
-#include "MuT/NodeLayoutPrivate.h"
 #include "MuT/NodeMesh.h"
 #include "MuT/NodeMeshApplyPose.h"
 #include "MuT/NodeMeshApplyPosePrivate.h"
@@ -270,7 +269,8 @@ namespace mu
 		int32 currentLayoutChannel,
 		const void* errorContext,
 		const FMeshGenerationOptions& MeshOptions,
-		bool bUseAbsoluteBlockIds
+		bool bUseAbsoluteBlockIds,
+		int32 FirstLODToIgnoreWarnings
 		)
 	{
 		MUTABLE_CPUPROFILER_SCOPE(LayoutUV_PrepareForLayout);
@@ -593,7 +593,7 @@ namespace mu
 		}
 
 		// Warn about vertices without a block id
-		if (Layout->FirstLODToIgnoreWarnings == -1 || m_currentParents.Last().m_lod < Layout->FirstLODToIgnoreWarnings)
+		if (FirstLODToIgnoreWarnings == -1 || m_currentParents.Last().m_lod < FirstLODToIgnoreWarnings)
 		{
 			TArray<float> UnassignedUVs;
 			UnassignedUVs.Reserve(NumVertices / 100);
@@ -941,7 +941,7 @@ namespace mu
 
                 if (BaseResult.GeneratedLayouts.Num()>node.LayoutOrGroup )
                 {
-                    const Layout* pLayout = BaseResult.GeneratedLayouts[node.LayoutOrGroup].get();
+                    const Layout* pLayout = BaseResult.GeneratedLayouts[node.LayoutOrGroup].Layout.get();
                     op->Layout = (uint16)node.LayoutOrGroup;
 
                     for ( int32 i=0; i<node.Blocks.Num(); ++i )
@@ -1343,7 +1343,7 @@ namespace mu
 					bool bLayoutsAreEqual = true;
 					for (int32 l = 0; l < Candidate.Mesh->GetLayoutCount(); ++l)
 					{
-						bLayoutsAreEqual = (*Candidate.Mesh->GetLayout(l) == *InOptions.OverrideLayouts[l]);
+						bLayoutsAreEqual = (*Candidate.Mesh->GetLayout(l) == *InOptions.OverrideLayouts[l].Layout);
 						if ( !bLayoutsAreEqual )
 						{
 							break;
@@ -1372,7 +1372,7 @@ namespace mu
 				{
 					for (int32 l = 0; l < DuplicateOf.Mesh->GetLayoutCount(); ++l)
 					{
-						const Layout* OverridingLayout = InOptions.OverrideLayouts[l].get();
+						FGeneratedLayout OverridingLayout = InOptions.OverrideLayouts[l];
 						OutResult.GeneratedLayouts.Add(OverridingLayout);
 					}
 				}
@@ -1381,7 +1381,7 @@ namespace mu
 					for (int32 l = 0; l < DuplicateOf.Mesh->GetLayoutCount(); ++l)
 					{
 						const Layout* DuplicatedLayout = DuplicateOf.Mesh->GetLayout(l);
-						OutResult.GeneratedLayouts.Add(DuplicatedLayout);
+						OutResult.GeneratedLayouts.Add({ DuplicatedLayout,0 });
 					}
 				}
 			}
@@ -1435,18 +1435,15 @@ namespace mu
 							continue;
 						}
 
-						// TODO: In a cleanup of the design of the layouts, we should remove this cast.
-						check(pLayoutNode->GetType()==NodeLayoutBlocks::GetStaticType() );
-						const NodeLayoutBlocks* TypedNode = static_cast<const NodeLayoutBlocks*>(pLayoutNode.get());
-
-						Ptr<const Layout> SourceLayout = TypedNode->GetPrivate()->m_pLayout;
+						Ptr<const Layout> SourceLayout = pLayoutNode->Layout;
 						Ptr<const Layout> GeneratedLayout = AddLayout(SourceLayout, MeshIDPrefix);
+						int32 FirstLODToIgnoreWarnings = pLayoutNode->FirstLODToIgnoreWarnings;
 						const void* Context = InOptions.OverrideContext.Get(InNode->GetMessageContext());
 
 						bool bUseAbsoluteBlockIds = false;
-						PrepareForLayout(GeneratedLayout, Cloned, LayoutIndex, Context, InOptions, bUseAbsoluteBlockIds);
+						PrepareForLayout(GeneratedLayout, Cloned, LayoutIndex, Context, InOptions, bUseAbsoluteBlockIds, FirstLODToIgnoreWarnings);
 
-						OutResult.GeneratedLayouts.Add(GeneratedLayout);
+						OutResult.GeneratedLayouts.Add({ GeneratedLayout, FirstLODToIgnoreWarnings });
 					}
 				}
 				else
@@ -1454,15 +1451,16 @@ namespace mu
 					// We need to apply the transform of the layouts used to override
 					for (int32 LayoutIndex = 0; LayoutIndex < InOptions.OverrideLayouts.Num(); ++LayoutIndex)
 					{
-						Ptr<const Layout> GeneratedLayout = InOptions.OverrideLayouts[LayoutIndex];
+						Ptr<const Layout> GeneratedLayout = InOptions.OverrideLayouts[LayoutIndex].Layout;
+						int32 FirstLODToIgnoreWarnings = InOptions.OverrideLayouts[LayoutIndex].FirstLODToIgnoreWarnings;
 						const void* Context = InOptions.OverrideContext.Get(InNode->GetMessageContext());
 
 						// In this case we need the layout block ids to use the ids in the parent layout, and not be prefixed with
 						// the current mesh id prefix. For this reason we need them to be absolute.
 						bool bUseAbsoluteBlockIds = true;
-						PrepareForLayout(GeneratedLayout, Cloned, LayoutIndex, Context, InOptions, bUseAbsoluteBlockIds);
+						PrepareForLayout(GeneratedLayout, Cloned, LayoutIndex, Context, InOptions, bUseAbsoluteBlockIds, FirstLODToIgnoreWarnings);
 
-						OutResult.GeneratedLayouts.Add(GeneratedLayout);
+						OutResult.GeneratedLayouts.Add({ GeneratedLayout,FirstLODToIgnoreWarnings });
 					}
 				}
 			}
diff --git a/Engine/Plugins/Experimental/Mutable/Source/MutableTools/Private/MuT/NodeLayout.cpp b/Engine/Plugins/Experimental/Mutable/Source/MutableTools/Private/MuT/NodeLayout.cpp
--- a/Engine/Plugins/Experimental/Mutable/Source/MutableTools/Private/MuT/NodeLayout.cpp
+++ b/Engine/Plugins/Experimental/Mutable/Source/MutableTools/Private/MuT/NodeLayout.cpp
@@ -3,127 +3,28 @@
 #include "MuT/NodeLayout.h"
 
 #include "Math/IntPoint.h"
-#include "Math/UnrealMathSSE.h"
 #include "Misc/AssertionMacros.h"
 #include "MuR/ConvertData.h"
 #include "MuR/MeshBufferSet.h"
 #include "MuR/MeshPrivate.h"
 #include "MuR/MutableMath.h"
-#include "MuT/NodeLayoutPrivate.h"
 
 
 namespace mu
 {
 
-	static FNodeType s_nodeLayoutType = FNodeType(Node::EType::Layout, Node::GetStaticType() );
+	FNodeType NodeLayout::StaticType = FNodeType(Node::EType::Layout, Node::GetStaticType() );
 
 
-	//---------------------------------------------------------------------------------------------
-	const FNodeType* NodeLayout::GetType() const
+	NodeLayout::NodeLayout()
 	{
-		return GetStaticType();
+		Layout = new mu::Layout();
 	}
 
 
-	//---------------------------------------------------------------------------------------------
-	const FNodeType* NodeLayout::GetStaticType()
+	Ptr<NodeLayout> NodeLayout::GenerateLayoutBlocks(const Ptr<Mesh> pMesh, int32 layoutIndex, int32 gridSizeX, int32 gridSizeY)
 	{
-		return &s_nodeLayoutType;
-	}
-
-
-	//---------------------------------------------------------------------------------------------
-	FNodeType NodeLayoutBlocks::Private::s_type = FNodeType(Node::EType::LayoutBlocks, NodeLayout::GetStaticType() );
-
-
-	//---------------------------------------------------------------------------------------------
-	//!
-	//---------------------------------------------------------------------------------------------
-
-	MUTABLE_IMPLEMENT_NODE( NodeLayoutBlocks )
-
-
-	//---------------------------------------------------------------------------------------------
-	// Own Interface
-	//---------------------------------------------------------------------------------------------
-	void NodeLayoutBlocks::SetGridSize( int32 x, int32 y )
-	{
-		m_pD->m_pLayout->SetGridSize( x, y );
-	}
-
-
-	//---------------------------------------------------------------------------------------------
-	void NodeLayoutBlocks::SetMaxGridSize(int32 x, int32 y)
-	{
-		m_pD->m_pLayout->SetMaxGridSize(x, y);
-	}
-
-
-	//---------------------------------------------------------------------------------------------
-	void NodeLayoutBlocks::GetGridSize( int32* pX, int32* pY ) const
-	{
-		FIntPoint grid = m_pD->m_pLayout->GetGridSize();
-		if (pX) *pX = grid[0];
-		if (pY) *pY = grid[1];
-	}
-
-
-	//---------------------------------------------------------------------------------------------
-	void NodeLayoutBlocks::GetMaxGridSize(int32* pX, int32* pY) const
-	{
-		m_pD->m_pLayout->GetMaxGridSize(pX, pY);
-	}
-
-
-	//---------------------------------------------------------------------------------------------
-	void NodeLayoutBlocks::SetBlockCount( int32 n )
-	{
-		m_pD->m_pLayout->SetBlockCount( n );
-	}
-
-
-	//---------------------------------------------------------------------------------------------
-	int32 NodeLayoutBlocks::GetBlockCount()
-	{
-		return m_pD->m_pLayout->GetBlockCount();
-	}
-
-
-	//---------------------------------------------------------------------------------------------
-	Ptr<const Layout> NodeLayoutBlocks::GetLayout() const
-	{
-		return m_pD->m_pLayout;
-	}
-
-
-	//---------------------------------------------------------------------------------------------
-    void NodeLayoutBlocks::SetBlock( int32 index, int32 minx, int32 miny, int32 sizex, int32 sizey )
-	{
-		m_pD->m_pLayout->Blocks[index].Min = FImageSize( minx, miny );
-		m_pD->m_pLayout->Blocks[index].Size = FImageSize( sizex, sizey );
-	}
-
-
-	//---------------------------------------------------------------------------------------------
-	void NodeLayoutBlocks::SetBlockOptions(int32 index, int32 priority, bool bReduceBothAxes, bool bReduceByTwo)
-	{
-		m_pD->m_pLayout->Blocks[index].Priority = priority;
-		m_pD->m_pLayout->Blocks[index].bReduceBothAxes = bReduceBothAxes;
-		m_pD->m_pLayout->Blocks[index].bReduceByTwo = bReduceByTwo;
-	}
-
-
-	//---------------------------------------------------------------------------------------------
-	void NodeLayoutBlocks::SetLayoutPackingStrategy(mu::EPackStrategy strategy)
-	{
-		m_pD->m_pLayout->SetLayoutPackingStrategy(strategy);
-	}
-
-
-	//---------------------------------------------------------------------------------------------
-	NodeLayoutBlocksPtr NodeLayoutBlocks::GenerateLayoutBlocks(const MeshPtr pMesh, int32 layoutIndex, int32 gridSizeX, int32 gridSizeY)
-	{
-		NodeLayoutBlocksPtr layout = nullptr;
+		Ptr<NodeLayout> LayoutNode = nullptr;
 
 		if (pMesh && layoutIndex >=0 && gridSizeX+gridSizeY>0)
 		{
@@ -173,10 +74,10 @@ namespace mu
 				UVs[v * 6 + 5][1] = uv_1[1];
 			}
 
-			layout = new NodeLayoutBlocks;
-			layout->SetGridSize(gridSizeX, gridSizeY);
-			layout->SetMaxGridSize(gridSizeX, gridSizeY);
-			layout->SetLayoutPackingStrategy(EPackStrategy::RESIZABLE_LAYOUT);
+			LayoutNode = new NodeLayout;
+			LayoutNode->Layout->SetGridSize(gridSizeX, gridSizeY);
+			LayoutNode->Layout->SetMaxGridSize(gridSizeX, gridSizeY);
+			LayoutNode->Layout->SetLayoutPackingStrategy(EPackStrategy::RESIZABLE_LAYOUT);
 			
 			TArray<box<FIntVector2>> blocks;
 			
@@ -283,7 +184,7 @@ namespace mu
 			
 			bool intersections = true;
 			
-			//Cheking if the blocks intersect with each other
+			// Check if the blocks intersect with each other
 			while (intersections)
 			{
 				intersections = false;
@@ -303,55 +204,24 @@ namespace mu
 				}
 			}
 			
-			int32 numBlocks = blocks.Num();
+			int32 NumBlocks = blocks.Num();
 			
 			//Generating layout blocks
-			if (numBlocks > 0)
+			if (NumBlocks > 0)
 			{
-				layout->SetBlockCount(numBlocks);
+				LayoutNode->Layout->SetBlockCount(NumBlocks);
 			
-				for (int32 i = 0; i < numBlocks; ++i)
+				for (int32 BlockIndex = 0; BlockIndex < NumBlocks; ++BlockIndex)
 				{
-					int32 blockIndex = i;
-					int32 minX = blocks[i].min[0];
-					int32 minY = blocks[i].min[1];
-					int32 sizeX = blocks[i].size[0];
-					int32 sizeY = blocks[i].size[1];
-			
-					layout->SetBlock(blockIndex, minX, minY, sizeX, sizeY);
+					LayoutNode->Layout->Blocks[BlockIndex].Min = UE::Math::TIntVector2<uint16>(blocks[BlockIndex].min);
+					LayoutNode->Layout->Blocks[BlockIndex].Size = UE::Math::TIntVector2<uint16>(blocks[BlockIndex].size);
 				}
 			}
 		}
 
-		return layout;
-	}
-
-	//---------------------------------------------------------------------------------------------
-	void NodeLayoutBlocks::SetIgnoreWarningsLOD(int32 LOD)
-	{
-		m_pD->m_pLayout->SetIgnoreLODWarnings(LOD);
-	}
-
-
-	//---------------------------------------------------------------------------------------------
-	int32 NodeLayoutBlocks::GetIgnoreWarningsLOD()
-	{
-		return m_pD->m_pLayout->GetIgnoreLODWarnings();
-	}
-
-
-	//---------------------------------------------------------------------------------------------
-	void NodeLayoutBlocks::SetBlockReductionMethod(EReductionMethod Method)
-	{
-		m_pD->m_pLayout->SetBlockReductionMethod(Method);
+		return LayoutNode;
 	}
 
-
-	//---------------------------------------------------------------------------------------------
-	EReductionMethod NodeLayoutBlocks::GetBlockReductionMethod()
-	{
-		return m_pD->m_pLayout->GetBlockReductionMethod();
-	}
 }
 
 
diff --git a/Engine/Plugins/Experimental/Mutable/Source/MutableTools/Private/MuT/NodeLayoutPrivate.h b/Engine/Plugins/Experimental/Mutable/Source/MutableTools/Private/MuT/NodeLayoutPrivate.h
deleted file mode 100644
diff --git a/Engine/Plugins/Experimental/Mutable/Source/MutableTools/Private/MuT/NodeMeshConstant.cpp b/Engine/Plugins/Experimental/Mutable/Source/MutableTools/Private/MuT/NodeMeshConstant.cpp
--- a/Engine/Plugins/Experimental/Mutable/Source/MutableTools/Private/MuT/NodeMeshConstant.cpp
+++ b/Engine/Plugins/Experimental/Mutable/Source/MutableTools/Private/MuT/NodeMeshConstant.cpp
@@ -52,7 +52,7 @@ namespace mu
 	{
 		check( index >=0 && index < m_pD->Layouts.Num() );
 
-		NodeLayoutPtr pResult;
+		Ptr<NodeLayout> pResult;
 
 		if (index >= 0 && index < m_pD->Layouts.Num())
 		{
diff --git a/Engine/Plugins/Experimental/Mutable/Source/MutableTools/Private/MuT/NodeMeshTable.cpp b/Engine/Plugins/Experimental/Mutable/Source/MutableTools/Private/MuT/NodeMeshTable.cpp
--- a/Engine/Plugins/Experimental/Mutable/Source/MutableTools/Private/MuT/NodeMeshTable.cpp
+++ b/Engine/Plugins/Experimental/Mutable/Source/MutableTools/Private/MuT/NodeMeshTable.cpp
@@ -49,9 +49,9 @@ namespace mu
 
 
 	//---------------------------------------------------------------------------------------------
-	NodeLayoutPtr NodeMeshTable::GetLayout( int i ) const
+	Ptr<NodeLayout> NodeMeshTable::GetLayout( int i ) const
 	{
-		NodeLayoutPtr pResult;
+		Ptr<NodeLayout> pResult;
 
 		if (i >= 0 && i < Layouts.Num())
 		{
@@ -63,7 +63,7 @@ namespace mu
 
 
 	//---------------------------------------------------------------------------------------------
-	void NodeMeshTable::SetLayout( int i, NodeLayoutPtr pLayout )
+	void NodeMeshTable::SetLayout( int i, Ptr<NodeLayout> pLayout )
 	{
 		check( i>=0 && i<GetLayoutCount() );
 		Layouts[i] = pLayout;
diff --git a/Engine/Plugins/Experimental/Mutable/Source/MutableTools/Private/MuT/NodeSurfaceNew.cpp b/Engine/Plugins/Experimental/Mutable/Source/MutableTools/Private/MuT/NodeSurfaceNew.cpp
--- a/Engine/Plugins/Experimental/Mutable/Source/MutableTools/Private/MuT/NodeSurfaceNew.cpp
+++ b/Engine/Plugins/Experimental/Mutable/Source/MutableTools/Private/MuT/NodeSurfaceNew.cpp
@@ -8,7 +8,6 @@
 #include "MuT/NodeColour.h"
 #include "MuT/NodeImage.h"
 #include "MuT/NodeMesh.h"
-#include "MuT/NodePatchImage.h"
 #include "MuT/NodePrivate.h"
 #include "MuT/NodeScalar.h"
 #include "MuT/NodeString.h"
diff --git a/Engine/Plugins/Experimental/Mutable/Source/MutableTools/Public/MuT/Node.h b/Engine/Plugins/Experimental/Mutable/Source/MutableTools/Public/MuT/Node.h
--- a/Engine/Plugins/Experimental/Mutable/Source/MutableTools/Public/MuT/Node.h
+++ b/Engine/Plugins/Experimental/Mutable/Source/MutableTools/Public/MuT/Node.h
@@ -114,7 +114,6 @@ namespace mu
 			RangeFromScalar,
 
 			Layout,
-			LayoutBlocks,
 
 			PatchImage,
 			PatchMesh,
diff --git a/Engine/Plugins/Experimental/Mutable/Source/MutableTools/Public/MuT/NodeLayout.h b/Engine/Plugins/Experimental/Mutable/Source/MutableTools/Public/MuT/NodeLayout.h
--- a/Engine/Plugins/Experimental/Mutable/Source/MutableTools/Public/MuT/NodeLayout.h
+++ b/Engine/Plugins/Experimental/Mutable/Source/MutableTools/Public/MuT/NodeLayout.h
@@ -13,129 +13,38 @@
 namespace mu
 {
 
-	//
-	class NodeLayout;
-	typedef Ptr<NodeLayout> NodeLayoutPtr;
-	typedef Ptr<const NodeLayout> NodeLayoutPtrConst;
-
-	//
-	class NodeLayoutBlocks;
-	typedef Ptr<NodeLayoutBlocks> NodeLayoutBlocksPtr;
-	typedef Ptr<const NodeLayoutBlocks> NodeLayoutBlocksPtrConst;
-
-
-	//! This node is used to define the texture layout for a texture coordinates channel of a
-	//! constant mesh.
+	/** This node is used to define the texture layout for a texture coordinates channel of a mesh. */
 	class MUTABLETOOLS_API NodeLayout : public Node
 	{
 	public:
 
-		//-----------------------------------------------------------------------------------------
-		// Node Interface
-		//-----------------------------------------------------------------------------------------
-
-		const FNodeType* GetType() const override;
-		static const FNodeType* GetStaticType();
+		Ptr<Layout> Layout;
 
+		int32 FirstLODToIgnoreWarnings = 0;
 
-		//-----------------------------------------------------------------------------------------
-		// Interface pattern
-		//-----------------------------------------------------------------------------------------
-		class Private;
-
-	protected:
-
-		//! Forbidden. Manage with the Ptr<> template.
-		~NodeLayout() {}
-
-	};
-
-
-	//! This node is used to define the texture layout for a texture coordinates channel of a
-	//! constant mesh.
-	//! The blocks defined here will be set to the channel of the mesh were this node is connected.
-	class MUTABLETOOLS_API NodeLayoutBlocks : public NodeLayout
-	{
 	public:
 
-		NodeLayoutBlocks();
+		NodeLayout();
 
-		//-----------------------------------------------------------------------------------------
 		// Node Interface
-		//-----------------------------------------------------------------------------------------
-
-		const FNodeType* GetType() const override;
-		static const FNodeType* GetStaticType();
-
-		//-----------------------------------------------------------------------------------------
-		// Own Interface
-		//-----------------------------------------------------------------------------------------
+		virtual const FNodeType* GetType() const override { return GetStaticType(); }
+		static const FNodeType* GetStaticType() { return &StaticType; }
 
-		//! Set the size of the grid where the blocks will be defined
-		void SetGridSize( int32 x, int32 y );
-
-		//! Set the maximum size of the grid where the blocks can be defined
-		void SetMaxGridSize(int32 x, int32 y);
-
-		//! Get the size of the grid where the blocks will be defined
-		void GetGridSize( int32* pX, int32* pY ) const;
-
-		//! Get the maximum size of the grid where the blocks can be defined
-		void GetMaxGridSize(int32* pX, int32* pY) const;
-
-		//! Set the number of blocks in the layout.
-		//! It keeps the current data as much as possible, and the new data is undefined.
-		void SetBlockCount( int32 );
-
-		//! Get the number of blocks in the layout
-		int32 GetBlockCount();
-
-		/** */
-		Ptr<const Layout> GetLayout() const;
-
-		//! Set a block of the layout.
-		//! minx and miny refer to the lowest left corner of the block.
-        void SetBlock( int32 index, int32 minx, int32 miny, int32 sizex, int32 sizey );
-
-		//! Set reduction block options like priority or if the block has to be reduced symmetrically.
-		void SetBlockOptions(int32 index, int32 priority, bool bReduceBothAxes, bool bReduceByTwo);
-
-		//! Set the texture layout packing strategy 
-		void SetLayoutPackingStrategy(EPackStrategy strategy);
+		// Own interface
 
 		//! Generate the blocks of a layout using the UV of the meshes
-		static NodeLayoutBlocksPtr GenerateLayoutBlocks(const MeshPtr pMesh, int32 layoutIndex, int32 gridSizeX, int32 gridSizeY);
-
-		//! Set at which LOD the unassigned vertices warnings will star to be ignored
-		void SetIgnoreWarningsLOD(int32 LOD);
-
-		//! Get the LOD where the unassigned vertices warnings starts to be ignored
-		int32 GetIgnoreWarningsLOD();
-
-		//! Set the block reduction method a the Fixed_Layout strategy
-		void SetBlockReductionMethod(EReductionMethod strategy);
-
-		//! Returns the block reduction method
-		EReductionMethod GetBlockReductionMethod();
-
-
-		//-----------------------------------------------------------------------------------------
-		// Interface pattern
-		//-----------------------------------------------------------------------------------------
-		class Private;
-		Private* GetPrivate() const;
+		static Ptr<NodeLayout> GenerateLayoutBlocks(const Ptr<Mesh> pMesh, int32 layoutIndex, int32 gridSizeX, int32 gridSizeY);
 
 	protected:
 
 		//! Forbidden. Manage with the Ptr<> template.
-		~NodeLayoutBlocks();
+		~NodeLayout() {}
 
 	private:
 
-		Private* m_pD;
+		static FNodeType StaticType;
 
 	};
 
 
-
 }
diff --git a/Engine/Plugins/Experimental/Mutable/Source/MutableTools/Public/MuT/NodeMeshTable.h b/Engine/Plugins/Experimental/Mutable/Source/MutableTools/Public/MuT/NodeMeshTable.h
--- a/Engine/Plugins/Experimental/Mutable/Source/MutableTools/Public/MuT/NodeMeshTable.h
+++ b/Engine/Plugins/Experimental/Mutable/Source/MutableTools/Public/MuT/NodeMeshTable.h
@@ -30,7 +30,7 @@ namespace mu
 		bool bNoneOption = false;
 		FString DefaultRowName;
 
-		TArray<NodeLayoutPtr> Layouts;
+		TArray<Ptr<NodeLayout>> Layouts;
 
 	public:
 
@@ -53,8 +53,8 @@ namespace mu
 		void SetLayoutCount( int );
 
 		//! Get the node defining a layout of the meshes on this column.
-		NodeLayoutPtr GetLayout( int index ) const;
-		void SetLayout( int index, NodeLayoutPtr );
+		Ptr<NodeLayout> GetLayout( int index ) const;
+		void SetLayout( int index, Ptr<NodeLayout>);
 
 		//! Adds the "None" option to the parameter that represents this table column
 		void SetNoneOption(bool bAddOption);
