[mutable] Ensure table row order is deterministic.
#rnx

diff --git a/Engine/Plugins/Experimental/Mutable/Source/CustomizableObjectEditor/Private/MuCOE/GenerateMutableSource/GenerateMutableSourceTable.cpp b/Engine/Plugins/Experimental/Mutable/Source/CustomizableObjectEditor/Private/MuCOE/GenerateMutableSource/GenerateMutableSourceTable.cpp
--- a/Engine/Plugins/Experimental/Mutable/Source/CustomizableObjectEditor/Private/MuCOE/GenerateMutableSource/GenerateMutableSourceTable.cpp
+++ b/Engine/Plugins/Experimental/Mutable/Source/CustomizableObjectEditor/Private/MuCOE/GenerateMutableSource/GenerateMutableSourceTable.cpp
@@ -699,6 +699,10 @@ TArray<FName> GetEnabledRows(const UDataTable& DataTable, const UCustomizableObj
 	TArray<FName> TableRowNames = DataTable.GetRowNames();
 	FBoolProperty* BoolProperty = nullptr;
 
+	// Sort them to avoid cooked data indeterminism problems. Rows may come from different tables and their loading order
+	// is not defined.
+	TableRowNames.Sort([](const FName& A, const FName& B) { return A.ToString() > B.ToString(); });
+
 	for (TFieldIterator<FProperty> PropertyIt(TableStruct); PropertyIt && TableNode.bDisableCheckedRows; ++PropertyIt)
 	{
 		BoolProperty = CastField<FBoolProperty>(*PropertyIt);
