[mutable] Now table nodes support pose assets.

#rnx

diff --git a/Engine/Plugins/Experimental/Mutable/Source/CustomizableObjectEditor/Private/MuCOE/EdGraphSchema_CustomizableObject.cpp b/Engine/Plugins/Experimental/Mutable/Source/CustomizableObjectEditor/Private/MuCOE/EdGraphSchema_CustomizableObject.cpp
--- a/Engine/Plugins/Experimental/Mutable/Source/CustomizableObjectEditor/Private/MuCOE/EdGraphSchema_CustomizableObject.cpp
+++ b/Engine/Plugins/Experimental/Mutable/Source/CustomizableObjectEditor/Private/MuCOE/EdGraphSchema_CustomizableObject.cpp
@@ -237,6 +237,7 @@ const FName UEdGraphSchema_CustomizableObject::PC_Enum("enum");
 const FName UEdGraphSchema_CustomizableObject::PC_Stack("stack");
 const FName UEdGraphSchema_CustomizableObject::PC_MaterialAsset("materialAsset");
 const FName UEdGraphSchema_CustomizableObject::PC_Wildcard("wildcard");
+const FName UEdGraphSchema_CustomizableObject::PC_PoseAsset("poseAsset");
 
 
 TSharedPtr<FCustomizableObjectSchemaAction_NewNode> UEdGraphSchema_CustomizableObject::AddNewNodeAction(FGraphActionListBuilderBase& ContextMenuBuilder, const FString& Category, const FText& MenuDesc, const FText& Tooltip, const int32 Grouping, const FString& Keywords)
@@ -510,7 +511,7 @@ void UEdGraphSchema_CustomizableObject::GetGraphContextActions(FGraphContextMenu
 
 	{
 		// External Pin Nodes
-		TArray<FName> PinTypes({ PC_Material, PC_Mesh, PC_Image, PC_Projector, PC_GroupProjector, PC_Color, PC_Float, PC_Bool, PC_Enum, PC_Stack, PC_PassThroughImage, PC_MaterialAsset });
+		TArray<FName> PinTypes({ PC_Material, PC_Mesh, PC_Image, PC_Projector, PC_GroupProjector, PC_Color, PC_Float, PC_Bool, PC_Enum, PC_Stack, PC_PassThroughImage, PC_MaterialAsset, PC_PoseAsset });
 
 		// Add pin types from extensions
 		for (const FRegisteredCustomizableObjectPinType& PinType : ICustomizableObjectModule::Get().GetExtendedPinTypes())
@@ -717,6 +718,10 @@ FLinearColor UEdGraphSchema_CustomizableObject::GetPinTypeColor(const FName& Typ
 	{
 		return FLinearColor(1.000000f, 1.000000f, 1.000000f, 1.000000f); // White
 	}
+	else if (TypeString == PC_PoseAsset)
+	{
+		return FLinearColor(0.700000f, 0.000000f, 0.000000f, 1.000000f); // Dark Red
+	}
 
 	for (const FRegisteredCustomizableObjectPinType& PinType : ICustomizableObjectModule::Get().GetExtendedPinTypes())
 	{
@@ -1356,6 +1361,10 @@ FText UEdGraphSchema_CustomizableObject::GetPinCategoryName(const FName& PinCate
 	{
 		return LOCTEXT("Wildcard_Pin_Category", "Wildcard");
 	}
+	else if (PinCategory == UEdGraphSchema_CustomizableObject::PC_PoseAsset)
+	{
+		return LOCTEXT("Pose_Pin_Category", "PoseAsset");
+	}
 	else
 	{
 		for (const FRegisteredCustomizableObjectPinType& PinType : ICustomizableObjectModule::Get().GetExtendedPinTypes())
diff --git a/Engine/Plugins/Experimental/Mutable/Source/CustomizableObjectEditor/Private/MuCOE/GenerateMutableSource/GenerateMutableSourceMesh.cpp b/Engine/Plugins/Experimental/Mutable/Source/CustomizableObjectEditor/Private/MuCOE/GenerateMutableSource/GenerateMutableSourceMesh.cpp
--- a/Engine/Plugins/Experimental/Mutable/Source/CustomizableObjectEditor/Private/MuCOE/GenerateMutableSource/GenerateMutableSourceMesh.cpp
+++ b/Engine/Plugins/Experimental/Mutable/Source/CustomizableObjectEditor/Private/MuCOE/GenerateMutableSource/GenerateMutableSourceMesh.cpp
@@ -3934,39 +3934,61 @@ mu::NodeMeshPtr GenerateMutableSourceMesh(const UEdGraphPin* Pin,
 
 	else if (const UCustomizableObjectNodeAnimationPose* TypedNode = Cast<UCustomizableObjectNodeAnimationPose>(Node))
 	{
-		if (const UEdGraphPin* ConnectedPin = FollowInputPin(*TypedNode->GetInputMeshPin()))
+		if (const UEdGraphPin* InputMeshPin = FollowInputPin(*TypedNode->GetInputMeshPin()))
 		{
-			mu::Ptr<mu::NodeMesh> InputMeshNode = GenerateMutableSourceMesh(ConnectedPin, GenerationContext, MeshData, false, bOnlyConnectedLOD);
+			mu::Ptr<mu::NodeMesh> InputMeshNode = GenerateMutableSourceMesh(InputMeshPin, GenerationContext, MeshData, false, bOnlyConnectedLOD);
 
-			if (TypedNode->PoseAsset && GenerationContext.GetCurrentComponentInfo().RefSkeletalMesh)
+			if (GenerationContext.GetCurrentComponentInfo().RefSkeletalMesh)
 			{
-				TArray<FName> ArrayBoneName;
-				TArray<FTransform> ArrayTransform;
-				UCustomizableObjectNodeAnimationPose::StaticRetrievePoseInformation(TypedNode->PoseAsset, GenerationContext.GetCurrentComponentInfo().RefSkeletalMesh, ArrayBoneName, ArrayTransform);
-				mu::NodeMeshApplyPosePtr NodeMeshApplyPose = CreateNodeMeshApplyPose(GenerationContext, InputMeshNode, ArrayBoneName, ArrayTransform);
+				if (TypedNode->PoseAsset)
+				{
+					TArray<FName> ArrayBoneName;
+					TArray<FTransform> ArrayTransform;
+					UCustomizableObjectNodeAnimationPose::StaticRetrievePoseInformation(TypedNode->PoseAsset, GenerationContext.GetCurrentComponentInfo().RefSkeletalMesh, ArrayBoneName, ArrayTransform);
+					mu::NodeMeshApplyPosePtr NodeMeshApplyPose = CreateNodeMeshApplyPose(GenerationContext, InputMeshNode, ArrayBoneName, ArrayTransform);
+
+					if (NodeMeshApplyPose)
+					{
+						Result = NodeMeshApplyPose;
+					}
+					else
+					{
+						FString msg = FString::Printf(TEXT("Couldn't get bone transform information from a Pose Asset."));
+						GenerationContext.Compiler->CompilerLog(FText::FromString(msg), Node);
 
-				if (NodeMeshApplyPose)
+						Result = nullptr;
+					}
+				}
+				else if (const UEdGraphPin* TablePosePin = FollowInputPin(*TypedNode->GetTablePosePin()))
 				{
-					Result = NodeMeshApplyPose;
+					if (const UCustomizableObjectNodeTable* TypedNodeTable = Cast<UCustomizableObjectNodeTable>(TablePosePin->GetOwningNode()))
+					{
+						mu::NodeMeshApplyPosePtr NodeMeshApplyPose = new mu::NodeMeshApplyPose();
+						mu::Ptr<mu::NodeMesh> MeshTableNode = GenerateMutableSourceMesh(TablePosePin, GenerationContext, MeshData, false, bOnlyConnectedLOD);
+
+						NodeMeshApplyPose->SetBase(InputMeshNode);
+						NodeMeshApplyPose->SetPose(MeshTableNode);
+
+						Result = NodeMeshApplyPose;
+					}
 				}
 				else
 				{
-					FString msg = FString::Printf(TEXT("Couldn't get bone transform information from a Pose Asset."));
-					GenerationContext.Compiler->CompilerLog(FText::FromString(msg), Node);
+					if (!TypedNode->PoseAsset) // Check if the slot has a selected pose. Could be left empty by the user
+					{
+						FString msg = FString::Printf(TEXT("Found pose mesh node without a pose asset assigned."));
+						GenerationContext.Compiler->CompilerLog(FText::FromString(msg), TypedNode);
+					}
 
-					Result = nullptr;
+					Result = InputMeshNode;
 				}
 			}
-			else
-			{
-				Result = InputMeshNode;
-			}
+			
 		}
 	}
 
 	else if (const UCustomizableObjectNodeTable* TypedNodeTable = Cast<UCustomizableObjectNodeTable>(Node))
 	{
-		//This node will add a checker texture in case of error
 		mu::NodeMeshConstantPtr EmptyNode = new mu::NodeMeshConstant();
 		Result = EmptyNode;
 		bool bSuccess = true;
@@ -3975,6 +3997,8 @@ mu::NodeMeshPtr GenerateMutableSourceMesh(const UEdGraphPin* Pin,
 
 		if (DataTable)
 		{
+			const UEdGraphSchema_CustomizableObject* Schema = GetDefault<UEdGraphSchema_CustomizableObject>();
+
 			// Getting the real name of the data table column
 			FString DataTableColumnName = TypedNodeTable->GetColumnNameByPin(Pin);
 			FProperty* Property = DataTable->FindTableProperty(FName(*DataTableColumnName));
@@ -3989,8 +4013,9 @@ mu::NodeMeshPtr GenerateMutableSourceMesh(const UEdGraphPin* Pin,
 
 			USkeletalMesh* DefaultSkeletalMesh = TypedNodeTable->GetColumnDefaultAssetByType<USkeletalMesh>(Pin);
 			UStaticMesh* DefaultStaticMesh = TypedNodeTable->GetColumnDefaultAssetByType<UStaticMesh>(Pin);
+			UPoseAsset* DefaultPoseAsset = TypedNodeTable->GetColumnDefaultAssetByType<UPoseAsset>(Pin);
 
-			if (bSuccess && !DefaultSkeletalMesh && !DefaultStaticMesh)
+			if (bSuccess && !DefaultSkeletalMesh && !DefaultStaticMesh && !DefaultPoseAsset)
 			{
 				FString Msg = FString::Printf(TEXT("Couldn't find a default value in the data table's struct for the column [%s]."), *DataTableColumnName);
 				GenerationContext.Compiler->CompilerLog(FText::FromString(Msg), Node);
@@ -4010,24 +4035,28 @@ mu::NodeMeshPtr GenerateMutableSourceMesh(const UEdGraphPin* Pin,
 					
 					int32 LODIndexConnected = -1; // LOD which the pin is connected to
 					int32 SectionIndexConnected = -1;
-					
 					int32 LODIndex = 0;
 					int32 SectionIndex = 0;
 
-					TypedNodeTable->GetPinLODAndSection(Pin, LODIndexConnected, SectionIndexConnected);
-
 					// Getting the mutable table mesh column name
-					FString MutableColumnName;
-					if (DefaultSkeletalMesh)
-					{
-						GetLODAndSectionForAutomaticLODs(GenerationContext, *Node, *DefaultSkeletalMesh, LODIndexConnected, SectionIndexConnected, LODIndex, SectionIndex, bOnlyConnectedLOD);
-						MutableColumnName = TypedNodeTable->GenerateSkeletalMeshMutableColumName(DataTableColumnName, LODIndex, SectionIndex);
-					}
-					else
+					FString MutableColumnName = DataTableColumnName;
+
+					if (Pin->PinType.PinCategory == Schema->PC_Mesh)
 					{
-						MutableColumnName = TypedNodeTable->GenerateStaticMeshMutableColumName(DataTableColumnName, SectionIndexConnected);
+						// LOD and sections are relevant for Skeletal and Static meshes but not for Pose Assets
+						TypedNodeTable->GetPinLODAndSection(Pin, LODIndexConnected, SectionIndexConnected);
+
+						if (DefaultSkeletalMesh)
+						{
+							GetLODAndSectionForAutomaticLODs(GenerationContext, *Node, *DefaultSkeletalMesh, LODIndexConnected, SectionIndexConnected, LODIndex, SectionIndex, bOnlyConnectedLOD);
+							MutableColumnName = TypedNodeTable->GenerateSkeletalMeshMutableColumName(DataTableColumnName, LODIndex, SectionIndex);
+						}
+						else
+						{
+							MutableColumnName = TypedNodeTable->GenerateStaticMeshMutableColumName(DataTableColumnName, SectionIndexConnected);
+						}
 					}
-					
+
 					// Generating a new Mesh column if not exists
 					if (Table->FindColumn(MutableColumnName) == INDEX_NONE)
 					{
@@ -4050,83 +4079,87 @@ mu::NodeMeshPtr GenerateMutableSourceMesh(const UEdGraphPin* Pin,
 						MeshTableNode->SetNoneOption(TypedNodeTable->bAddNoneOption);
 						MeshTableNode->SetDefaultRowName(TypedNodeTable->DefaultRowName.ToString());
 
-						if (DefaultSkeletalMesh)
+						// Pose Assets do not need this part of the code
+						if (Pin->PinType.PinCategory == Schema->PC_Mesh)
 						{
-							FSkeletalMeshModel* ImportedModel = DefaultSkeletalMesh->GetImportedModel();
-
-							if (ImportedModel->LODModels.IsValidIndex(LODIndex) &&
-								ImportedModel->LODModels[LODIndex].Sections.IsValidIndex(SectionIndex))
+							if (DefaultSkeletalMesh)
 							{
-								// TODO: this should be made for all the meshes of the Column to support meshes with different values
-								// Filling Mesh Data
-								MeshData.bHasVertexColors = DefaultSkeletalMesh->GetHasVertexColors();
-								MeshData.NumTexCoordChannels = ImportedModel->LODModels[LODIndex].NumTexCoords;
-								MeshData.MaxBoneIndexTypeSizeBytes = ImportedModel->LODModels[LODIndex].RequiredBones.Num() > 256 ? 2 : 1;
-								MeshData.MaxNumBonesPerVertex = ImportedModel->LODModels[LODIndex].GetMaxBoneInfluences();
-
-								// When mesh data is combined we will get an upper and lower bound of the number of triangles.
-								MeshData.MaxNumTriangles = ImportedModel->LODModels[LODIndex].Sections[SectionIndex].NumTriangles;
-								MeshData.MinNumTriangles = ImportedModel->LODModels[LODIndex].Sections[SectionIndex].NumTriangles;
-
-								// With tables we have to use the default skeletal mesh to tell if the surface mesh needs morphs or clothing.
-								// For now tables will not have morphs as there isn't a way to select them.
-								MeshData.bHasRealTimeMorphs = false;
-								MeshData.bHasClothing = GenerationContext.Options.bClothingEnabled && ImportedModel->LODModels[LODIndex].HasClothData();
+								FSkeletalMeshModel* ImportedModel = DefaultSkeletalMesh->GetImportedModel();
+
+								if (ImportedModel->LODModels.IsValidIndex(LODIndex) &&
+									ImportedModel->LODModels[LODIndex].Sections.IsValidIndex(SectionIndex))
+								{
+									// TODO: this should be made for all the meshes of the Column to support meshes with different values
+									// Filling Mesh Data
+									MeshData.bHasVertexColors = DefaultSkeletalMesh->GetHasVertexColors();
+									MeshData.NumTexCoordChannels = ImportedModel->LODModels[LODIndex].NumTexCoords;
+									MeshData.MaxBoneIndexTypeSizeBytes = ImportedModel->LODModels[LODIndex].RequiredBones.Num() > 256 ? 2 : 1;
+									MeshData.MaxNumBonesPerVertex = ImportedModel->LODModels[LODIndex].GetMaxBoneInfluences();
+
+									// When mesh data is combined we will get an upper and lower bound of the number of triangles.
+									MeshData.MaxNumTriangles = ImportedModel->LODModels[LODIndex].Sections[SectionIndex].NumTriangles;
+									MeshData.MinNumTriangles = ImportedModel->LODModels[LODIndex].Sections[SectionIndex].NumTriangles;
+
+									// With tables we have to use the default skeletal mesh to tell if the surface mesh needs morphs or clothing.
+									// For now tables will not have morphs as there isn't a way to select them.
+									MeshData.bHasRealTimeMorphs = false;
+									MeshData.bHasClothing = GenerationContext.Options.bClothingEnabled && ImportedModel->LODModels[LODIndex].HasClothData();
+								}
 							}
-						}
 
-						TArray<UCustomizableObjectLayout*> Layouts = TypedNodeTable->GetLayouts(Pin);
-						MeshTableNode->SetLayoutCount(Layouts.Num());
+							TArray<UCustomizableObjectLayout*> Layouts = TypedNodeTable->GetLayouts(Pin);
+							MeshTableNode->SetLayoutCount(Layouts.Num());
 
-						if (Layouts.Num())
-						{
-							// Generating node Layouts
-							for (int32 i = 0; i < Layouts.Num(); ++i)
+							if (Layouts.Num())
 							{
-								mu::Ptr<mu::NodeLayout> LayoutNode = new mu::NodeLayout;
+								// Generating node Layouts
+								for (int32 i = 0; i < Layouts.Num(); ++i)
+								{
+									mu::Ptr<mu::NodeLayout> LayoutNode = new mu::NodeLayout;
 
-								LayoutNode->Layout->SetGridSize(Layouts[i]->GetGridSize().X, Layouts[i]->GetGridSize().Y);
-								LayoutNode->Layout->SetMaxGridSize(Layouts[i]->GetMaxGridSize().X, Layouts[i]->GetMaxGridSize().Y);
-								LayoutNode->Layout->SetBlockCount(Layouts[i]->Blocks.Num() ? Layouts[i]->Blocks.Num() : 1);
+									LayoutNode->Layout->SetGridSize(Layouts[i]->GetGridSize().X, Layouts[i]->GetGridSize().Y);
+									LayoutNode->Layout->SetMaxGridSize(Layouts[i]->GetMaxGridSize().X, Layouts[i]->GetMaxGridSize().Y);
+									LayoutNode->Layout->SetBlockCount(Layouts[i]->Blocks.Num() ? Layouts[i]->Blocks.Num() : 1);
 
-								mu::EPackStrategy PackStrategy = ConvertLayoutStrategy(Layouts[i]->GetPackingStrategy());
-								LayoutNode->Layout->SetLayoutPackingStrategy(PackStrategy);
-								
-								LayoutNode->Layout->ReductionMethod = (Layouts[i]->GetBlockReductionMethod() == ECustomizableObjectLayoutBlockReductionMethod::Halve ? mu::EReductionMethod::HALVE_REDUCTION : mu::EReductionMethod::UNITARY_REDUCTION);
+									mu::EPackStrategy PackStrategy = ConvertLayoutStrategy(Layouts[i]->GetPackingStrategy());
+									LayoutNode->Layout->SetLayoutPackingStrategy(PackStrategy);
 
-								if (bLinkedToExtendMaterial)
-								{
-									// Layout warnings can be safely ignored in this case. Vertices that do not belong to any layout block will be removed (Extend Materials only)
-									LayoutNode->FirstLODToIgnoreWarnings = 0;
-								}
+									LayoutNode->Layout->ReductionMethod = (Layouts[i]->GetBlockReductionMethod() == ECustomizableObjectLayoutBlockReductionMethod::Halve ? mu::EReductionMethod::HALVE_REDUCTION : mu::EReductionMethod::UNITARY_REDUCTION);
 
-								if (Layouts[i]->Blocks.Num())
-								{
-									for (int BlockIndex = 0; BlockIndex < Layouts[i]->Blocks.Num(); ++BlockIndex)
+									if (bLinkedToExtendMaterial)
 									{
-										LayoutNode->Layout->Blocks[BlockIndex] = ToMutable( Layouts[i]->Blocks[BlockIndex] );
+										// Layout warnings can be safely ignored in this case. Vertices that do not belong to any layout block will be removed (Extend Materials only)
+										LayoutNode->FirstLODToIgnoreWarnings = 0;
+									}
+
+									if (Layouts[i]->Blocks.Num())
+									{
+										for (int BlockIndex = 0; BlockIndex < Layouts[i]->Blocks.Num(); ++BlockIndex)
+										{
+											LayoutNode->Layout->Blocks[BlockIndex] = ToMutable(Layouts[i]->Blocks[BlockIndex]);
+										}
+									}
+									else
+									{
+										FString msg = "Mesh Column [" + MutableColumnName + "] Layout doesn't has any block. A grid sized block will be used instead.";
+										GenerationContext.Compiler->CompilerLog(FText::FromString(msg), Node, EMessageSeverity::Warning);
+
+										LayoutNode->Layout->Blocks[0].Min = { 0,0 };
+										LayoutNode->Layout->Blocks[0].Size = { uint16(Layouts[i]->GetGridSize().X), uint16(Layouts[i]->GetGridSize().Y) };
+										LayoutNode->Layout->Blocks[0].Priority = 0;
+										LayoutNode->Layout->Blocks[0].bReduceBothAxes = false;
+										LayoutNode->Layout->Blocks[0].bReduceByTwo = false;
 									}
-								}
-								else
-								{
-									FString msg = "Mesh Column [" + MutableColumnName + "] Layout doesn't has any block. A grid sized block will be used instead.";
-									GenerationContext.Compiler->CompilerLog(FText::FromString(msg), Node, EMessageSeverity::Warning);
-
-									LayoutNode->Layout->Blocks[0].Min = { 0,0 };
-									LayoutNode->Layout->Blocks[0].Size = { uint16(Layouts[i]->GetGridSize().X), uint16(Layouts[i]->GetGridSize().Y)};
-									LayoutNode->Layout->Blocks[0].Priority = 0;
-									LayoutNode->Layout->Blocks[0].bReduceBothAxes = false;
-									LayoutNode->Layout->Blocks[0].bReduceByTwo = false;
-								}
 
-								MeshTableNode->SetLayout(i, LayoutNode);
+									MeshTableNode->SetLayout(i, LayoutNode);
+								}
 							}
-						}
 
-						// Applying Mesh Morph Nodes
-						if (DefaultSkeletalMesh && GenerationContext.MeshMorphStack.Num())
-						{
-							MorphResult = GenerateMorphMesh(Pin, GenerationContext.MeshMorphStack, 0, Result, GenerationContext, MeshData, bOnlyConnectedLOD, MutableColumnName);
+							// Applying Mesh Morph Nodes
+							if (DefaultSkeletalMesh && GenerationContext.MeshMorphStack.Num())
+							{
+								MorphResult = GenerateMorphMesh(Pin, GenerationContext.MeshMorphStack, 0, Result, GenerationContext, MeshData, bOnlyConnectedLOD, MutableColumnName);
+							}
 						}
 					}
 				}
diff --git a/Engine/Plugins/Experimental/Mutable/Source/CustomizableObjectEditor/Private/MuCOE/GenerateMutableSource/GenerateMutableSourceTable.cpp b/Engine/Plugins/Experimental/Mutable/Source/CustomizableObjectEditor/Private/MuCOE/GenerateMutableSource/GenerateMutableSourceTable.cpp
--- a/Engine/Plugins/Experimental/Mutable/Source/CustomizableObjectEditor/Private/MuCOE/GenerateMutableSource/GenerateMutableSourceTable.cpp
+++ b/Engine/Plugins/Experimental/Mutable/Source/CustomizableObjectEditor/Private/MuCOE/GenerateMutableSource/GenerateMutableSourceTable.cpp
@@ -3,6 +3,7 @@
 #include "MuCOE/GenerateMutableSource/GenerateMutableSourceTable.h"
 
 #include "Animation/AnimInstance.h"
+#include "Animation/PoseAsset.h"
 #include "AssetRegistry/ARFilter.h"
 #include "AssetRegistry/AssetRegistryModule.h"
 #include "Engine/CompositeDataTable.h"
@@ -18,6 +19,7 @@
 #include "MuCOE/CustomizableObjectCompiler.h"
 #include "MuCOE/GenerateMutableSource/GenerateMutableSourceMesh.h"
 #include "MuCOE/Nodes/CustomizableObjectNodeTable.h"
+#include "MuCOE/Nodes/CustomizableObjectNodeAnimationPose.h"
 #include "MuR/Mesh.h"
 #include "PhysicsEngine/PhysicsAsset.h"
 #include "Rendering/SkeletalMeshLODModel.h"
@@ -28,7 +30,7 @@
 #define LOCTEXT_NAMESPACE "CustomizableObjectEditor"
 
 
-bool FillTableColumn(const UCustomizableObjectNodeTable* TableNode, mu::TablePtr MutableTable,	const FString& ColumnName,	const FString& RowName,	const int32 RowIdx,	uint8* CellData, const FProperty* ColumnProperty,
+bool FillTableColumn(const UCustomizableObjectNodeTable* TableNode, mu::TablePtr MutableTable, const FString& ColumnName, const FString& RowName, const int32 RowIdx, uint8* CellData, const FProperty* ColumnProperty,
 	const int LODIndexConnected, const int32 SectionIndexConnected, int32 LODIndex, int32 SectionIndex, const bool bOnlyConnectedLOD, FMutableGraphGenerationContext& GenerationContext)
 {
 	int32 CurrentColumn;
@@ -587,6 +589,42 @@ bool FillTableColumn(const UCustomizableObjectNodeTable* TableNode, mu::TablePtr
 			}
 		}
 
+		else if (SoftObjectProperty->PropertyClass->IsChildOf(UPoseAsset::StaticClass()))
+		{
+			if (UPoseAsset* PoseAsset = Cast<UPoseAsset>(Object))
+			{
+				CurrentColumn = MutableTable.get()->FindColumn(ColumnName);
+
+				if (CurrentColumn == -1)
+				{
+					CurrentColumn = MutableTable->AddColumn(ColumnName, mu::ETableColumnType::Mesh);
+				}
+
+				TArray<FName> ArrayBoneName;
+				TArray<FTransform> ArrayTransform;
+				UCustomizableObjectNodeAnimationPose::StaticRetrievePoseInformation(PoseAsset, GenerationContext.GetCurrentComponentInfo().RefSkeletalMesh, ArrayBoneName, ArrayTransform);
+
+				mu::Ptr<mu::Mesh> MutableMesh = new mu::Mesh();
+				mu::Ptr<mu::Skeleton> MutableSkeleton = new mu::Skeleton;
+
+				MutableMesh->SetSkeleton(MutableSkeleton);
+				MutableMesh->SetBonePoseCount(ArrayBoneName.Num());
+				MutableSkeleton->SetBoneCount(ArrayBoneName.Num());
+
+				for (int32 i = 0; i < ArrayBoneName.Num(); ++i)
+				{
+					const FName BoneName = ArrayBoneName[i];
+					const mu::FBoneName& BoneId = GenerationContext.GetBoneUnique(BoneName);
+
+					MutableSkeleton->SetDebugName(i, BoneName);
+					MutableSkeleton->SetBoneName(i, { BoneId });
+					MutableMesh->SetBonePose(i, BoneId, (FTransform3f)ArrayTransform[i], mu::EBoneUsageFlags::Skinning);
+				}
+
+				MutableTable->SetCell(CurrentColumn, RowIdx, MutableMesh.get());
+			}
+		}
+
 		else
 		{
 			// Unsuported Variable Type
diff --git a/Engine/Plugins/Experimental/Mutable/Source/CustomizableObjectEditor/Private/MuCOE/Nodes/CustomizableObjectNodeAnimationPose.cpp b/Engine/Plugins/Experimental/Mutable/Source/CustomizableObjectEditor/Private/MuCOE/Nodes/CustomizableObjectNodeAnimationPose.cpp
--- a/Engine/Plugins/Experimental/Mutable/Source/CustomizableObjectEditor/Private/MuCOE/Nodes/CustomizableObjectNodeAnimationPose.cpp
+++ b/Engine/Plugins/Experimental/Mutable/Source/CustomizableObjectEditor/Private/MuCOE/Nodes/CustomizableObjectNodeAnimationPose.cpp
@@ -31,12 +31,16 @@ void UCustomizableObjectNodeAnimationPose::AllocateDefaultPins(UCustomizableObje
 	const UEdGraphSchema_CustomizableObject* Schema = GetDefault<UEdGraphSchema_CustomizableObject>();
 
 	FString PinName = TEXT("Input Mesh");
-	UEdGraphPin* PinImagePin = CustomCreatePin(EGPD_Input, Schema->PC_Mesh, FName(*PinName));
-	PinImagePin->bDefaultValueIsIgnored = true;
+	UEdGraphPin* InMeshPin = CustomCreatePin(EGPD_Input, Schema->PC_Mesh, FName(*PinName));
+	InMeshPin->bDefaultValueIsIgnored = true;
+
+	PinName = TEXT("Table Pose");
+	UEdGraphPin* TablePosePin = CustomCreatePin(EGPD_Input, Schema->PC_PoseAsset, FName(*PinName));
+	TablePosePin->bDefaultValueIsIgnored = true;
 
 	PinName = TEXT("Output Mesh");
-	PinImagePin = CustomCreatePin(EGPD_Output, Schema->PC_Mesh, FName(*PinName));
-	PinImagePin->bDefaultValueIsIgnored = true;
+	UEdGraphPin* OutMeshPin = CustomCreatePin(EGPD_Output, Schema->PC_Mesh, FName(*PinName));
+	OutMeshPin->bDefaultValueIsIgnored = true;
 }
 
 
@@ -50,6 +54,16 @@ UEdGraphPin* UCustomizableObjectNodeAnimationPose::GetInputMeshPin() const
 }
 
 
+UEdGraphPin* UCustomizableObjectNodeAnimationPose::GetTablePosePin() const
+{
+	FString PinName = FString(TEXT("Table Pose"));
+
+	UEdGraphPin* Pin = FindPin(PinName);
+
+	return Pin;
+}
+
+
 FText UCustomizableObjectNodeAnimationPose::GetNodeTitle(ENodeTitleType::Type TitleType) const
 {
 	if (PoseAsset)
diff --git a/Engine/Plugins/Experimental/Mutable/Source/CustomizableObjectEditor/Private/MuCOE/Nodes/CustomizableObjectNodeTable.cpp b/Engine/Plugins/Experimental/Mutable/Source/CustomizableObjectEditor/Private/MuCOE/Nodes/CustomizableObjectNodeTable.cpp
--- a/Engine/Plugins/Experimental/Mutable/Source/CustomizableObjectEditor/Private/MuCOE/Nodes/CustomizableObjectNodeTable.cpp
+++ b/Engine/Plugins/Experimental/Mutable/Source/CustomizableObjectEditor/Private/MuCOE/Nodes/CustomizableObjectNodeTable.cpp
@@ -2,6 +2,7 @@
 
 #include "MuCOE/Nodes/CustomizableObjectNodeTable.h"
 
+#include "Animation/PoseAsset.h"
 #include "Engine/StaticMesh.h"
 #include "Engine/SkeletalMesh.h"
 #include "Engine/Texture2DArray.h"
@@ -392,7 +393,8 @@ void UCustomizableObjectNodeTable::AllocateDefaultPins(UCustomizableObjectNodeRe
 		UStaticMesh::StaticClass(),
 		UTexture2D::StaticClass(),
 		UTexture::StaticClass(),
-		UMaterialInstance::StaticClass()
+		UMaterialInstance::StaticClass(),
+		UPoseAsset::StaticClass()
 	};
 
 	TArray<UEdGraphPin*> OldPins(Pins);
@@ -492,6 +494,15 @@ void UCustomizableObjectNodeTable::AllocateDefaultPins(UCustomizableObjectNodeRe
 
 						OutPin = CustomCreatePin(EGPD_Output, Schema->PC_MaterialAsset, FName(*PinName), PinData);
 					}
+
+					else if (Object->IsA(UPoseAsset::StaticClass()))
+					{
+						UCustomizableObjectNodeTableObjectPinData* PinData = NewObject<UCustomizableObjectNodeTableObjectPinData>(this);
+						PinData->ColumnName = ColumnName;
+						PinData->StructColumnId = ColumnPropertyId;
+
+						OutPin = CustomCreatePin(EGPD_Output, Schema->PC_PoseAsset, FName(*PinName), PinData);
+					}
 				}
 				else
 				{
@@ -774,6 +785,17 @@ bool UCustomizableObjectNodeTable::IsNodeOutDatedAndNeedsRefresh()
 						bNeedsUpdate = true;
 					}
 
+					NumPins++;
+				}
+				else if (Object->IsA(UPoseAsset::StaticClass()))
+				{
+					FString PinName = DataTableUtils::GetPropertyExportName(ColumnProperty);
+					
+					if (CheckPinUpdated(PinName, Schema->PC_PoseAsset))
+					{
+						bNeedsUpdate = true;
+					}
+
 					NumPins++;
 				}
 			}
@@ -1038,7 +1060,7 @@ FString UCustomizableObjectNodeTable::GetColumnNameByPin(const UEdGraphPin* Pin)
 		return PinData->ColumnName;
 	}
 
-	return FString();
+	return Pin->PinFriendlyName.ToString();
 }
 
 
@@ -1477,6 +1499,7 @@ FSkeletalMaterial* UCustomizableObjectNodeTable::GetDefaultSkeletalMaterialFor(c
 	return nullptr;
 }
 
+
 int32 UCustomizableObjectNodeTable::GetDefaultSkeletalMaterialIndexFor(const UEdGraphPin& MeshPin) const
 {
 	USkeletalMesh* SkeletalMesh = GetColumnDefaultAssetByType<USkeletalMesh>(&MeshPin);
diff --git a/Engine/Plugins/Experimental/Mutable/Source/CustomizableObjectEditor/Public/MuCOE/EdGraphSchema_CustomizableObject.h b/Engine/Plugins/Experimental/Mutable/Source/CustomizableObjectEditor/Public/MuCOE/EdGraphSchema_CustomizableObject.h
--- a/Engine/Plugins/Experimental/Mutable/Source/CustomizableObjectEditor/Public/MuCOE/EdGraphSchema_CustomizableObject.h
+++ b/Engine/Plugins/Experimental/Mutable/Source/CustomizableObjectEditor/Public/MuCOE/EdGraphSchema_CustomizableObject.h
@@ -105,6 +105,7 @@ public:
 	static const FName PC_Stack;
 	static const FName PC_MaterialAsset;
 	static const FName PC_Wildcard;
+	static const FName PC_PoseAsset;
 
 	// EdGraphSchema interface
 	virtual void GetAssetsGraphHoverMessage(const TArray<FAssetData>& Assets, const UEdGraph* HoverGraph, FString& OutTooltipText, bool& OutOkIcon) const override;
diff --git a/Engine/Plugins/Experimental/Mutable/Source/CustomizableObjectEditor/Public/MuCOE/Nodes/CustomizableObjectNodeAnimationPose.h b/Engine/Plugins/Experimental/Mutable/Source/CustomizableObjectEditor/Public/MuCOE/Nodes/CustomizableObjectNodeAnimationPose.h
--- a/Engine/Plugins/Experimental/Mutable/Source/CustomizableObjectEditor/Public/MuCOE/Nodes/CustomizableObjectNodeAnimationPose.h
+++ b/Engine/Plugins/Experimental/Mutable/Source/CustomizableObjectEditor/Public/MuCOE/Nodes/CustomizableObjectNodeAnimationPose.h
@@ -32,6 +32,7 @@ public:
 
 	// Begin EdGraphNode interface
 	UEdGraphPin* GetInputMeshPin() const;
+	UEdGraphPin* GetTablePosePin() const;
 	virtual FText GetNodeTitle(ENodeTitleType::Type TitleType) const override;
 	virtual FLinearColor GetNodeTitleColor() const override;
 	// End EdGraphNode interface
