[mutable] Added a new array variable to customizable object to store all the mutable mesh components info (name and reference skeletal mesh).

#rnx

diff --git a/Engine/Plugins/Experimental/Mutable/Source/CustomizableObject/Internal/MuCO/CustomizableObjectCustomVersion.h b/Engine/Plugins/Experimental/Mutable/Source/CustomizableObject/Internal/MuCO/CustomizableObjectCustomVersion.h
--- a/Engine/Plugins/Experimental/Mutable/Source/CustomizableObject/Internal/MuCO/CustomizableObjectCustomVersion.h
+++ b/Engine/Plugins/Experimental/Mutable/Source/CustomizableObject/Internal/MuCO/CustomizableObjectCustomVersion.h
@@ -122,6 +122,8 @@ struct CUSTOMIZABLEOBJECT_API FCustomizableObjectCustomVersion
 
 		NodeTablePinViewer,
 
+		NewComponentOptions,
+
 		// -----<new versions can be added above this line>-------------------------------------------------
 		VersionPlusOne,
 		LatestVersion = VersionPlusOne - 1
diff --git a/Engine/Plugins/Experimental/Mutable/Source/CustomizableObject/Internal/MuCO/CustomizableObjectPrivate.h b/Engine/Plugins/Experimental/Mutable/Source/CustomizableObject/Internal/MuCO/CustomizableObjectPrivate.h
--- a/Engine/Plugins/Experimental/Mutable/Source/CustomizableObject/Internal/MuCO/CustomizableObjectPrivate.h
+++ b/Engine/Plugins/Experimental/Mutable/Source/CustomizableObject/Internal/MuCO/CustomizableObjectPrivate.h
@@ -816,6 +816,28 @@ struct CUSTOMIZABLEOBJECT_API FMutableCachedPlatformData
 };
 
 
+USTRUCT()
+struct CUSTOMIZABLEOBJECT_API FMutableMeshComponentData
+{
+	GENERATED_USTRUCT_BODY()
+
+	/** Name to identify this component. */
+	UPROPERTY(EditAnywhere, Category = CustomizableObject)
+	FName Name;
+
+	/** All the SkeletalMeshes generated for this CustomizableObject instances will use the Reference Skeletal Mesh
+	* properties for everything that Mutable doesn't create or modify. This includes data like LOD distances, Physics
+	* properties, Bounding Volumes, Skeleton, etc.
+	*
+	* While a CustomizableObject instance is being created for the first time, and in some situation with lots of
+	* objects this may require some seconds, the Reference Skeletal Mesh is used for the actor. This works as a better
+	* solution than the alternative of not showing anything, although this can be disabled with the function
+	* "SetReplaceDiscardedWithReferenceMeshEnabled" (See the c++ section). */
+	UPROPERTY(EditAnywhere, Category = CustomizableObject)
+	TObjectPtr<USkeletalMesh> ReferenceSkeletalMesh;
+};
+
+
 UCLASS()
 class CUSTOMIZABLEOBJECT_API UCustomizableObjectPrivate : public UObject
 {
@@ -971,6 +993,10 @@ public:
 
 #if WITH_EDITORONLY_DATA
 
+	/** Necessary information to generate each component of the Customizable Object. */
+	UPROPERTY(EditAnywhere, Category = CustomizableObject)
+	TArray<FMutableMeshComponentData> MutableMeshComponents;
+
 	/** Unique Identifier - Deterministic. Used to locate Model and Streamable data on disk. Should not be modified. */
 	FGuid Identifier;
 
diff --git a/Engine/Plugins/Experimental/Mutable/Source/CustomizableObject/Private/MuCO/CustomizableObject.cpp b/Engine/Plugins/Experimental/Mutable/Source/CustomizableObject/Private/MuCO/CustomizableObject.cpp
--- a/Engine/Plugins/Experimental/Mutable/Source/CustomizableObject/Private/MuCO/CustomizableObject.cpp
+++ b/Engine/Plugins/Experimental/Mutable/Source/CustomizableObject/Private/MuCO/CustomizableObject.cpp
@@ -253,7 +253,7 @@ void UCustomizableObject::PostLoad()
 #if WITH_EDITOR
 	if (ReferenceSkeletalMesh_DEPRECATED)
 	{
-		ReferenceSkeletalMeshes.Add(ReferenceSkeletalMesh_DEPRECATED);
+		ReferenceSkeletalMeshes_DEPRECATED.Add(ReferenceSkeletalMesh_DEPRECATED);
 		ReferenceSkeletalMesh_DEPRECATED = nullptr;
 	}
 
@@ -266,6 +266,23 @@ void UCustomizableObject::PostLoad()
 		GetPrivate()->EmbeddedDataBytesLimit = CompileOptions_DEPRECATED.EmbeddedDataBytesLimit;
 		GetPrivate()->PackagedDataBytesLimit = CompileOptions_DEPRECATED.PackagedDataBytesLimit;
 	}
+
+	if (CustomizableObjectCustomVersion < FCustomizableObjectCustomVersion::NewComponentOptions)
+	{
+		if (GetPrivate()->MutableMeshComponents.IsEmpty())
+		{
+			for (int32 SkeletalMeshIndex = 0; SkeletalMeshIndex < ReferenceSkeletalMeshes_DEPRECATED.Num(); ++SkeletalMeshIndex)
+			{
+				FMutableMeshComponentData NewComponent;
+				NewComponent.Name = FName(FString::FromInt(SkeletalMeshIndex));
+				NewComponent.ReferenceSkeletalMesh = ReferenceSkeletalMeshes_DEPRECATED[SkeletalMeshIndex];
+
+				GetPrivate()->MutableMeshComponents.Add(NewComponent);
+			}
+
+			ReferenceSkeletalMeshes_DEPRECATED.Empty();
+		}
+	}
 #endif
 
 	// Register to dirty delegate so we update derived data version ID each time that the package is marked as dirty.
@@ -297,6 +314,7 @@ void UCustomizableObject::PostLoad()
 			GetPrivate()->LoadCompiledDataFromDisk();
 		}
 	}
+
 #endif
 }
 
@@ -1115,9 +1133,9 @@ void UCustomizableObjectPrivate::AddUncompiledCOWarning(const FString& Additiona
 USkeletalMesh* UCustomizableObject::GetRefSkeletalMesh(int32 ComponentIndex) const
 {
 #if WITH_EDITORONLY_DATA
-	if (ReferenceSkeletalMeshes.IsValidIndex(ComponentIndex))
+	if (GetPrivate()->MutableMeshComponents.IsValidIndex(ComponentIndex))
 	{
-		return ReferenceSkeletalMeshes[ComponentIndex];
+		return GetPrivate()->MutableMeshComponents[ComponentIndex].ReferenceSkeletalMesh;
 	}
 #else
 	const FModelResources& ModelResources = Private->GetModelResources();
diff --git a/Engine/Plugins/Experimental/Mutable/Source/CustomizableObject/Public/MuCO/CustomizableObject.h b/Engine/Plugins/Experimental/Mutable/Source/CustomizableObject/Public/MuCO/CustomizableObject.h
--- a/Engine/Plugins/Experimental/Mutable/Source/CustomizableObject/Public/MuCO/CustomizableObject.h
+++ b/Engine/Plugins/Experimental/Mutable/Source/CustomizableObject/Public/MuCO/CustomizableObject.h
@@ -361,20 +361,14 @@ public:
 
 #if WITH_EDITORONLY_DATA
 private:
+
 	UPROPERTY()
 	TObjectPtr<USkeletalMesh> ReferenceSkeletalMesh_DEPRECATED;
 
+	UPROPERTY()
+	TArray<TObjectPtr<USkeletalMesh>> ReferenceSkeletalMeshes_DEPRECATED;
+
 public:
-	/** All the SkeletalMeshes generated for this CustomizableObject instances will use the Reference Skeletal Mesh
-	* properties for everything that Mutable doesn't create or modify. This includes data like LOD distances, Physics
-	* properties, Bounding Volumes, Skeleton, etc.
-	*
-	* While a CustomizableObject instance is being created for the first time, and in some situation with lots of
-	* objects this may require some seconds, the Reference Skeletal Mesh is used for the actor. This works as a better
-	* solution than the alternative of not showing anything, although this can be disabled with the function
-	* "SetReplaceDiscardedWithReferenceMeshEnabled" (See the c++ section). */
-	UPROPERTY(EditAnywhere, Category = CustomizableObject)
-	TArray<TObjectPtr<USkeletalMesh>> ReferenceSkeletalMeshes;
 
 	/**
 	  * The optional VersionBridge asset, which must implement the ICustomizableObjectVersionBridgeInterface, will be used to
diff --git a/Engine/Plugins/Experimental/Mutable/Source/CustomizableObjectEditor/Private/MuCOE/CustomizableObjectCompiler.cpp b/Engine/Plugins/Experimental/Mutable/Source/CustomizableObjectEditor/Private/MuCOE/CustomizableObjectCompiler.cpp
--- a/Engine/Plugins/Experimental/Mutable/Source/CustomizableObjectEditor/Private/MuCOE/CustomizableObjectCompiler.cpp
+++ b/Engine/Plugins/Experimental/Mutable/Source/CustomizableObjectEditor/Private/MuCOE/CustomizableObjectCompiler.cpp
@@ -648,23 +648,46 @@ mu::NodeObjectPtr FCustomizableObjectCompiler::GenerateMutableRoot(
 
 		VisitedObjects.Empty();
 		ActualRootObject = Root->ParentObject ? GetFullGraphRootObject(Root, VisitedObjects) : Object;
-
 	}
 
 	// Ensure that the CO has a valid AutoLODStrategy on the ActualRoot.
 	if (ActualRoot->AutoLODStrategy == ECustomizableObjectAutomaticLODStrategy::Inherited)
 	{
-		CompilerLog(LOCTEXT("RootInheritsFromParent", "Error! Base CustomizableObject's LOD Strategy can't be set to 'Inherit from parent object'"), ActualRoot);
+		CompilerLog(LOCTEXT("RootInheritsFromParent", "Error! Base CustomizableObject's LOD Strategy can't be set to 'Inherit from parent object'"), ActualRoot, EMessageSeverity::Error);
 		return nullptr;
 	}
 
-	// Make sure we have a valid Reference SkeletalMesh and Skeleton for each component
-	for (int32 ComponentIndex = 0; ComponentIndex < ActualRoot->NumMeshComponents; ++ComponentIndex)
+	if (ActualRootObject->GetPrivate()->MutableMeshComponents.IsEmpty())
 	{
+		CompilerLog(LOCTEXT("EmptyComponentName", "Error! There are no components defined in the Object Properties Tab."), ActualRoot, EMessageSeverity::Error);
+		return nullptr;
+	}
+
+	TArray<const FName> ComponentNames;
+
+	// Make sure we have a valid Name, Reference SkeletalMesh and Skeleton for each component
+	for (int32 ComponentIndex = 0; ComponentIndex < ActualRootObject->GetPrivate()->MutableMeshComponents.Num(); ++ComponentIndex)
+	{
+		FName ComponentName = ActualRootObject->GetPrivate()->MutableMeshComponents[ComponentIndex].Name;
+
+		if (ComponentName.IsNone())
+		{
+			CompilerLog(LOCTEXT("EmptyComponentName", "Error! Missing name in a component of the Customizable Object."), ActualRoot, EMessageSeverity::Error);
+			return nullptr;
+		}
+		else if (ComponentNames.Contains(ComponentName))
+		{
+			CompilerLog(FText::Format(LOCTEXT("RepeatedComponentName", "Error! Repeated name [{0}] used in more than one Component"),
+				FText::FromName(ComponentName)), ActualRoot, EMessageSeverity::Error);
+			return nullptr;
+		}
+
+		ComponentNames.Add(ComponentName);
+
 		USkeletalMesh* RefSkeletalMesh = ActualRootObject->GetRefSkeletalMesh(ComponentIndex);
 		if (!RefSkeletalMesh)
 		{
-			CompilerLog(LOCTEXT("NoReferenceMeshObjectTab", "Error! Missing reference mesh in the Object Properties Tab"), ActualRoot);
+			CompilerLog(LOCTEXT("NoReferenceMeshObjectTab", "Error! Missing reference mesh in the Object Properties Tab"), ActualRoot, EMessageSeverity::Error);
 			return nullptr;
 		}
 
@@ -674,20 +697,22 @@ mu::NodeObjectPtr FCustomizableObjectCompiler::GenerateMutableRoot(
 			FText Msg = FText::Format(LOCTEXT("NoReferenceSkeleton", "Error! Missing skeleton in the reference mesh [{0}]"),
 				FText::FromString(GenerationContext.CustomizableObjectWithCycle->GetPathName()));
 
-			CompilerLog(Msg, ActualRoot);
+			CompilerLog(Msg, ActualRoot, EMessageSeverity::Error);
 			return nullptr;
 		}
 
-		// Add a new entry to the list of Component Infos 
-		GenerationContext.ComponentInfos.Add(RefSkeletalMesh);
+		// Add a new entry to the list of Component Infos
+		GenerationContext.ComponentInfos.Add(FMutableComponentInfo(ComponentName, RefSkeletalMesh));
 
 		// Make sure the Skeleton from the reference mesh is added to the list of referenced Skeletons.
 		GenerationContext.ReferencedSkeletons.Add(RefSkeleton);
-	}
 
-	Object->ReferenceSkeletalMeshes = ActualRootObject->ReferenceSkeletalMeshes;
+		// Add reference meshes to the participating objects
+		GenerationContext.AddParticipatingObject(*RefSkeletalMesh);
+	}
 
-	GenerationContext.AddParticipatingObject(Object->ReferenceSkeletalMeshes);
+	// Copy component data to the object being compiled
+	Object->GetPrivate()->MutableMeshComponents = ActualRootObject->GetPrivate()->MutableMeshComponents;
 
     GenerationContext.RealTimeMorphTargetsOverrides = ActualRoot->RealTimeMorphSelectionOverrides;
     GenerationContext.RealTimeMorphTargetsOverrides.Reset();
diff --git a/Engine/Plugins/Experimental/Mutable/Source/CustomizableObjectEditor/Private/MuCOE/CustomizableObjectDetails.cpp b/Engine/Plugins/Experimental/Mutable/Source/CustomizableObjectEditor/Private/MuCOE/CustomizableObjectDetails.cpp
--- a/Engine/Plugins/Experimental/Mutable/Source/CustomizableObjectEditor/Private/MuCOE/CustomizableObjectDetails.cpp
+++ b/Engine/Plugins/Experimental/Mutable/Source/CustomizableObjectEditor/Private/MuCOE/CustomizableObjectDetails.cpp
@@ -89,7 +89,9 @@ void FCustomizableObjectDetails::CustomizeDetails( IDetailLayoutBuilder& DetailB
 		CustomizableObject = Cast<UCustomizableObject>(DetailsView->GetSelectedObjects()[0].Get());
 	}
 
+	IDetailCategoryBuilder& CustomizableObjectCategory = DetailBuilder.EditCategory("CustomizableObject");
 	IDetailCategoryBuilder& StatesCategory = DetailBuilder.EditCategory( "States" );
+	IDetailCategoryBuilder& VersioningCategory = DetailBuilder.EditCategory( "Versioning" );
 	//StatesCategory.CategoryIcon( "ActorClassIcon.CustomizableObject" );
 
 	DetailBuilder.HideProperty("States");
@@ -129,20 +131,6 @@ void FCustomizableObjectDetails::CustomizeDetails( IDetailLayoutBuilder& DetailB
 	];
 
 	StatesTree->SetIsRightClickScrollingEnabled(false);
-	
-	TSharedRef<IPropertyHandle> Property = DetailBuilder.GetProperty("ReferenceSkeletalMeshes");
-
-	if (Property->IsValidHandle() && CustomizableObject)
-	{
-		if (CustomizableObject->IsChildObject())
-		{
-			Property->MarkHiddenByCustomization();
-		}
-		else
-		{
-			Property->MarkResetToDefaultCustomized();
-		}
-	}
 
 	TSharedRef<IPropertyHandle> VersionBridgeProperty = DetailBuilder.GetProperty("VersionBridge");
 
@@ -171,6 +159,15 @@ void FCustomizableObjectDetails::CustomizeDetails( IDetailLayoutBuilder& DetailB
 			VersionStructProperty->MarkHiddenByCustomization();
 		}
 	}
+
+	// Add Private's components property.
+	if (CustomizableObject && !CustomizableObject->IsChildObject())
+	{
+		TArray<UObject*> Private;
+		Private.Add(CustomizableObject->GetPrivate());
+
+		CustomizableObjectCategory.AddExternalObjectProperty(Private, "MutableMeshComponents");
+	}
 }
 
 
diff --git a/Engine/Plugins/Experimental/Mutable/Source/CustomizableObjectEditor/Private/MuCOE/CustomizableObjectEditorModule.cpp b/Engine/Plugins/Experimental/Mutable/Source/CustomizableObjectEditor/Private/MuCOE/CustomizableObjectEditorModule.cpp
--- a/Engine/Plugins/Experimental/Mutable/Source/CustomizableObjectEditor/Private/MuCOE/CustomizableObjectEditorModule.cpp
+++ b/Engine/Plugins/Experimental/Mutable/Source/CustomizableObjectEditor/Private/MuCOE/CustomizableObjectEditorModule.cpp
@@ -37,6 +37,7 @@
 #include "MuCOE/Nodes/CustomizableObjectNodeExternalPinDetails.h"
 #include "MuCOE/Nodes/CustomizableObjectNodeLayoutBlocks.h"
 #include "MuCOE/Nodes/CustomizableObjectNodeLayoutBlocksDetails.h"
+#include "MuCOE/Nodes/CustomizableObjectNodeMaterialDetails.h"
 #include "MuCOE/Nodes/CustomizableObjectNodeMeshClipMorph.h"
 #include "MuCOE/Nodes/CustomizableObjectNodeMeshClipMorphDetails.h"
 #include "MuCOE/Nodes/CustomizableObjectNodeMeshClipWithMesh.h"
@@ -249,7 +250,7 @@ void FCustomizableObjectEditorModule::StartupModule()
 	RegisterCustomDetails(PropertyModule, UCustomizableObjectNodeMeshClipMorph::StaticClass(), FOnGetDetailCustomizationInstance::CreateStatic(&FCustomizableObjectNodeMeshClipMorphDetails::MakeInstance));
 	RegisterCustomDetails(PropertyModule, UCustomizableObjectNodeMeshClipWithMesh::StaticClass(), FOnGetDetailCustomizationInstance::CreateStatic(&FCustomizableObjectNodeMeshClipWithMeshDetails::MakeInstance));
 	RegisterCustomDetails(PropertyModule, UCustomizableObjectNodeExternalPin::StaticClass(), FOnGetDetailCustomizationInstance::CreateStatic(&FCustomizableObjectNodeExternalPinDetails::MakeInstance));
-	RegisterCustomDetails(PropertyModule, UCustomizableObjectNodeMaterial::StaticClass(), FOnGetDetailCustomizationInstance::CreateStatic(&FCustomizableObjectNodeDetails::MakeInstance));
+	RegisterCustomDetails(PropertyModule, UCustomizableObjectNodeMaterial::StaticClass(), FOnGetDetailCustomizationInstance::CreateStatic(&FCustomizableObjectNodeMaterialDetails::MakeInstance));
 	RegisterCustomDetails(PropertyModule, UCustomizableObjectNodeSkeletalMesh::StaticClass(), FOnGetDetailCustomizationInstance::CreateStatic(&FCustomizableObjectNodeSkeletalMeshDetails::MakeInstance));
 	RegisterCustomDetails(PropertyModule, UCustomizableObjectNodeStaticMesh::StaticClass(), FOnGetDetailCustomizationInstance::CreateStatic(&FCustomizableObjectNodeDetails::MakeInstance));
 	RegisterCustomDetails(PropertyModule, UCustomizableObjectNodeTable::StaticClass(), FOnGetDetailCustomizationInstance::CreateStatic(&FCustomizableObjectNodeTableDetails::MakeInstance));
diff --git a/Engine/Plugins/Experimental/Mutable/Source/CustomizableObjectEditor/Private/MuCOE/CustomizableObjectFactory.cpp b/Engine/Plugins/Experimental/Mutable/Source/CustomizableObjectEditor/Private/MuCOE/CustomizableObjectFactory.cpp
--- a/Engine/Plugins/Experimental/Mutable/Source/CustomizableObjectEditor/Private/MuCOE/CustomizableObjectFactory.cpp
+++ b/Engine/Plugins/Experimental/Mutable/Source/CustomizableObjectEditor/Private/MuCOE/CustomizableObjectFactory.cpp
@@ -94,13 +94,15 @@ UObject* UCustomizableObjectFactory::FactoryCreateNew(UClass* Class, UObject* In
 		}
 	}
 
-	if (!CreationSettings.bEmptyObject)
+	if (!BaseObjectNode)
 	{
-		if (!BaseObjectNode)
-		{
-			return NewObj;
-		}
+		return NewObj;
+	}
 
+	NewObj->GetPrivate()->MutableMeshComponents.SetNum(CreationSettings.NumMeshComponents);
+
+	if (!CreationSettings.bEmptyObject)
+	{
 		if (CreationSettings.bIsChildObject && CreationSettings.ParentObject.IsValid())
 		{
 			BaseObjectNode->ParentObject = CreationSettings.ParentObject.Get();
@@ -125,14 +127,11 @@ UObject* UCustomizableObjectFactory::FactoryCreateNew(UClass* Class, UObject* In
 		}
 		else
 		{
-			BaseObjectNode->NumMeshComponents = CreationSettings.NumMeshComponents;
+			check(CreationSettings.ReferenceSkeletalMeshes.Num() == CreationSettings.NumMeshComponents);
 
 			for (int32 MeshIndex = 0; MeshIndex < CreationSettings.ReferenceSkeletalMeshes.Num(); ++MeshIndex)
 			{
-				if (CreationSettings.ReferenceSkeletalMeshes[MeshIndex].IsValid())
-				{
-					NewObj->ReferenceSkeletalMeshes.Add(CreationSettings.ReferenceSkeletalMeshes[MeshIndex].Get());
-				}
+				NewObj->GetPrivate()->MutableMeshComponents[MeshIndex].ReferenceSkeletalMesh = CreationSettings.ReferenceSkeletalMeshes[MeshIndex].Get();
 			}
 		}
 	}
diff --git a/Engine/Plugins/Experimental/Mutable/Source/CustomizableObjectEditor/Private/MuCOE/CustomizableObjectInstanceFactory.cpp b/Engine/Plugins/Experimental/Mutable/Source/CustomizableObjectEditor/Private/MuCOE/CustomizableObjectInstanceFactory.cpp
--- a/Engine/Plugins/Experimental/Mutable/Source/CustomizableObjectEditor/Private/MuCOE/CustomizableObjectInstanceFactory.cpp
+++ b/Engine/Plugins/Experimental/Mutable/Source/CustomizableObjectEditor/Private/MuCOE/CustomizableObjectInstanceFactory.cpp
@@ -33,7 +33,7 @@ void UCustomizableObjectInstanceFactory::PostSpawnActor(UObject* Asset, AActor*
 	ACustomizableSkeletalMeshActor* NewCSMActor = CastChecked<ACustomizableSkeletalMeshActor>(NewActor);
 	if (NewCSMActor && COInstance)
 	{
-		int32 NumComponents = GetNumberOfComponents(COInstance);
+		int32 NumComponents = COInstance->GetNumComponents();
 
 		for (int32 ComponentIndex = 0; ComponentIndex < NumComponents; ++ComponentIndex)
 		{
@@ -59,6 +59,7 @@ void UCustomizableObjectInstanceFactory::PostSpawnActor(UObject* Asset, AActor*
 					CustomSkeletalComp->UnregisterComponent();
 					CustomSkeletalComp->CustomizableObjectInstance = COInstance;
 					CustomSkeletalComp->ComponentIndex = ComponentIndex;
+					//CustomSkeletalComp->ComponentName = GetComponentName(Asset,ComponentIndex);
 					CustomSkeletalComp->SetSkeletalMesh(SkeletalMesh);
 					CustomSkeletalComp->UpdateSkeletalMeshAsync();
 					CustomSkeletalComp->RegisterComponent();
@@ -135,6 +136,7 @@ bool UCustomizableObjectInstanceFactory::CanCreateActorFrom(const FAssetData& As
     return true;
 }
 
+
 USkeletalMesh* UCustomizableObjectInstanceFactory::GetSkeletalMeshFromAsset(UObject* Asset, int32 ComponentIndex) const
 {
     USkeletalMesh* SkeletalMesh = Cast<USkeletalMesh>(Asset);
@@ -157,36 +159,26 @@ USkeletalMesh* UCustomizableObjectInstanceFactory::GetSkeletalMeshFromAsset(UObj
 }
 
 
+//FName UCustomizableObjectInstanceFactory::GetComponentName(UObject* Asset, int32 ComponentIndex) const
+//{
+//	UCustomizableObjectInstance* CustomizableObjectInstance = Cast<UCustomizableObjectInstance>(Asset);
+//
+//	if (CustomizableObjectInstance != nullptr)
+//	{
+//		if (UCustomizableObject* CustomizableObject = CustomizableObjectInstance->GetCustomizableObject())
+//		{
+//			return CustomizableObject->GetPrivate()->Components[ComponentIndex].Name;
+//		}
+//	}
+//
+//	return FName();
+//}
+
+
 FQuat UCustomizableObjectInstanceFactory::AlignObjectToSurfaceNormal(const FVector& InSurfaceNormal, const FQuat& ActorRotation) const
 {
     // Meshes align the Z (up) axis with the surface normal
     return FindActorAlignmentRotation(ActorRotation, FVector(0.f, 0.f, 1.f), InSurfaceNormal);
 }
 
-
-int32 UCustomizableObjectInstanceFactory::GetNumberOfComponents(UCustomizableObjectInstance* COInstance)
-{
-	int32 NumMeshComponents = 0;
-
-	if (const UCustomizableObject* CustomizableObject = COInstance->GetCustomizableObject();
-		CustomizableObject &&
-		CustomizableObject->GetPrivate()->GetSource())
-	{
-		TArray<UCustomizableObjectNodeObject*> RootNodes;
-		CustomizableObject->GetPrivate()->GetSource()->GetNodesOfClass<UCustomizableObjectNodeObject>(RootNodes);
-
-		for (int32 i = 0; i < RootNodes.Num(); ++i)
-		{
-			if (RootNodes[i]->bIsBase)
-			{
-				NumMeshComponents = RootNodes[i]->NumMeshComponents;
-				break;
-			}
-		}
-	}
-
-	return NumMeshComponents;
-}
-
-
 #undef LOCTEXT_NAMESPACE
diff --git a/Engine/Plugins/Experimental/Mutable/Source/CustomizableObjectEditor/Private/MuCOE/CustomizableObjectLayout.cpp b/Engine/Plugins/Experimental/Mutable/Source/CustomizableObjectEditor/Private/MuCOE/CustomizableObjectLayout.cpp
--- a/Engine/Plugins/Experimental/Mutable/Source/CustomizableObjectEditor/Private/MuCOE/CustomizableObjectLayout.cpp
+++ b/Engine/Plugins/Experimental/Mutable/Source/CustomizableObjectEditor/Private/MuCOE/CustomizableObjectLayout.cpp
@@ -113,7 +113,7 @@ void UCustomizableObjectLayout::GenerateBlocksFromUVs()
 
 			GenerationContext.MeshGenerationFlags.Push(ShapeFlags);
 
-			GenerationContext.ComponentInfos.Add(SkeletalMesh);
+			GenerationContext.ComponentInfos.Add(FMutableComponentInfo(FName(), SkeletalMesh));
 			MutableMesh = ConvertSkeletalMeshToMutable(SkeletalMesh, TSoftClassPtr<UAnimInstance>(), LOD, Material, LOD, Material, GenerationContext, Node);
 		}
 		else if (UStaticMesh* StaticMesh = Cast<UStaticMesh>(Mesh))
diff --git a/Engine/Plugins/Experimental/Mutable/Source/CustomizableObjectEditor/Private/MuCOE/GenerateMutableSource/GenerateMutableSource.cpp b/Engine/Plugins/Experimental/Mutable/Source/CustomizableObjectEditor/Private/MuCOE/GenerateMutableSource/GenerateMutableSource.cpp
--- a/Engine/Plugins/Experimental/Mutable/Source/CustomizableObjectEditor/Private/MuCOE/GenerateMutableSource/GenerateMutableSource.cpp
+++ b/Engine/Plugins/Experimental/Mutable/Source/CustomizableObjectEditor/Private/MuCOE/GenerateMutableSource/GenerateMutableSource.cpp
@@ -446,8 +446,11 @@ mu::FBoneName FMutableGraphGenerationContext::GetBoneUnique(const FName& InBoneN
 
 FMutableComponentInfo& FMutableGraphGenerationContext::GetCurrentComponentInfo()
 {
-	check(ComponentInfos.IsValidIndex(CurrentMeshComponent));
-	return ComponentInfos[CurrentMeshComponent];
+	FMutableComponentInfo* CurrentComponentInfo = ComponentInfos.FindByPredicate(
+		[this](const FMutableComponentInfo& Component) { return Component.ComponentName == CurrentMeshComponent; });
+	check(CurrentComponentInfo);
+
+	return *CurrentComponentInfo;
 }
 
 
@@ -1007,7 +1010,7 @@ mu::Ptr<mu::NodeObject> GenerateMutableSource(const UEdGraphPin * Pin, FMutableG
 		// Mesh components per instance
 		if (!GenerationContext.NumMeshComponentsInRoot)
 		{
-			GenerationContext.NumMeshComponentsInRoot = TypedNodeObj->NumMeshComponents;
+			GenerationContext.NumMeshComponentsInRoot = GenerationContext.Object->GetPrivate()->MutableMeshComponents.Num();
 		}
 
 		const int32 NumLODsInRoot = GenerationContext.NumLODsInRoot;
@@ -1023,7 +1026,8 @@ mu::Ptr<mu::NodeObject> GenerateMutableSource(const UEdGraphPin * Pin, FMutableG
 		// Add implicit components
 		for (int32 CurrentComponent = 0; CurrentComponent < NumImplicitComponents; ++CurrentComponent)
 		{
-			GenerationContext.CurrentMeshComponent = CurrentComponent;
+			FName ComponentName = GenerationContext.Object->GetPrivate()->MutableMeshComponents[CurrentComponent].Name;
+			GenerationContext.CurrentMeshComponent = ComponentName;
 
 			mu::Ptr<mu::NodeComponent> ComponentNode = Invoke([&GenerationContext, CurrentComponent]() -> mu::Ptr<mu::NodeComponent>
 				{
@@ -1111,13 +1115,13 @@ mu::Ptr<mu::NodeObject> GenerateMutableSource(const UEdGraphPin * Pin, FMutableG
 							continue;
 						}
 
-						// Set it to -1 to indicate we don't care about component id.
-						GenerationContext.CurrentMeshComponent = -1;
+						// Set it to "None" to indicate we don't care about component id.
+						GenerationContext.CurrentMeshComponent = FName();
 
 						mu::NodeModifierPtr ModifierNode = GenerateMutableSourceModifier(ChildNodePin, GenerationContext);
 						LODNode->Modifiers.Add(ModifierNode);
 
-						GenerationContext.CurrentMeshComponent = CurrentComponent;
+						GenerationContext.CurrentMeshComponent = ComponentName;
 					}
 				}
 			}
@@ -1138,12 +1142,8 @@ mu::Ptr<mu::NodeObject> GenerateMutableSource(const UEdGraphPin * Pin, FMutableG
 
 			++GenerationContext.NumExplicitMeshComponents;
 
-			GenerationContext.CurrentMeshComponent = ObjectNode->Components.Num();
-
 			mu::Ptr<mu::NodeComponent> ComponentNode = GenerateMutableSourceComponent(ComponentNodePin, GenerationContext);
 			ObjectNode->Components.Add(ComponentNode);
-
-			GenerationContext.CurrentMeshComponent = -1;
 		}
 
 
@@ -1481,9 +1481,12 @@ bool AffectsCurrentComponent(const UEdGraphPin* Pin, FMutableGraphGenerationCont
 	check(Pin);
 	RETURN_ON_CYCLE(*Pin, GenerationContext)
 	
-	int32 ComponentIndex = INDEX_NONE;
-
+	FName ComponentName = FName();
 	const UCustomizableObjectNode* Node = CastChecked<UCustomizableObjectNode>(Pin->GetOwningNode());
+
+	// Parent material for extend and edit material nodes
+	UCustomizableObjectNodeMaterialBase* OriginalParentMaterialNode = nullptr;
+
 	if (const UCustomizableObjectNodeMaterialVariation* TypedNodeVar = Cast<UCustomizableObjectNodeMaterialVariation>(Node))
 	{
 		bool bAffectsCurrentComponent = false;
@@ -1494,7 +1497,7 @@ bool AffectsCurrentComponent(const UEdGraphPin* Pin, FMutableGraphGenerationCont
 			{
 				if (bAffectsCurrentComponent)
 				{
-					FString Msg = FString::Printf(TEXT("Error! One or more materials nodes linked to a material variation node have different component index"));
+					FString Msg = FString::Printf(TEXT("Error! One or more materials nodes linked to a material variation node have different component name"));
 					GenerationContext.Compiler->CompilerLog(FText::FromString(Msg), Node, EMessageSeverity::Error);
 				}
 
@@ -1514,7 +1517,7 @@ bool AffectsCurrentComponent(const UEdGraphPin* Pin, FMutableGraphGenerationCont
 					{
 						if (bAffectsCurrentComponent)
 						{
-							FString Msg = FString::Printf(TEXT("Error! One or more materials nodes linked to a material variation node have different component index"));
+							FString Msg = FString::Printf(TEXT("Error! One or more materials nodes linked to a material variation node have different component name"));
 							GenerationContext.Compiler->CompilerLog(FText::FromString(Msg), Node, EMessageSeverity::Error);
 						}
 
@@ -1544,7 +1547,7 @@ bool AffectsCurrentComponent(const UEdGraphPin* Pin, FMutableGraphGenerationCont
 					{
 						if (bAffectsCurrentComponent)
 						{
-							FString Msg = FString::Printf(TEXT("Error! One or more materials nodes linked to a material switch node have different component index"));
+							FString Msg = FString::Printf(TEXT("Error! One or more materials nodes linked to a material switch node have different component name"));
 							GenerationContext.Compiler->CompilerLog(FText::FromString(Msg), Node, EMessageSeverity::Error);
 						}
 
@@ -1561,24 +1564,24 @@ bool AffectsCurrentComponent(const UEdGraphPin* Pin, FMutableGraphGenerationCont
 
 	else if (const UCustomizableObjectNodeMaterialBase* TypedNodeMat = Cast<UCustomizableObjectNodeMaterialBase>(Node))
 	{
-		ComponentIndex = TypedNodeMat->GetMeshComponentIndex();
+		ComponentName = TypedNodeMat->GetMeshComponentName();
 	}
 	else if (const UCustomizableObjectNodeExtendMaterial* TypedNodeExt = Cast<UCustomizableObjectNodeExtendMaterial>(Node))
 	{
-		UCustomizableObjectNodeMaterialBase* OriginalParentMaterialNode = TypedNodeExt->GetParentMaterialNode();
-		ComponentIndex = OriginalParentMaterialNode ? OriginalParentMaterialNode->GetMeshComponentIndex() : GenerationContext.CurrentMeshComponent;
+		OriginalParentMaterialNode = TypedNodeExt->GetParentMaterialNode();
+		ComponentName = OriginalParentMaterialNode ? OriginalParentMaterialNode->GetMeshComponentName() : GenerationContext.CurrentMeshComponent;
 	}
 	else if (const UCustomizableObjectNodeEditMaterialBase* TypedNodeEdit = Cast<UCustomizableObjectNodeEditMaterialBase>(Node))
 	{
-		UCustomizableObjectNodeMaterialBase* OriginalParentMaterialNode = TypedNodeEdit->GetParentMaterialNode();
-		ComponentIndex = OriginalParentMaterialNode ? OriginalParentMaterialNode->GetMeshComponentIndex() : GenerationContext.CurrentMeshComponent;
+		OriginalParentMaterialNode = TypedNodeEdit->GetParentMaterialNode();
+		ComponentName = OriginalParentMaterialNode ? OriginalParentMaterialNode->GetMeshComponentName() : GenerationContext.CurrentMeshComponent;
 	}
 	else if (const UCustomizableObjectNodeModifierBase* TypedNodeModifier = Cast<UCustomizableObjectNodeModifierBase>(Node))
 	{
 		// Because of the current implementation, modifiers affect all components at lod level. If there is only one component it is ok, but otherwise rise an error.
 		if (GenerationContext.NumMeshComponentsInRoot == 1)
 		{
-			ComponentIndex = 0;
+			ComponentName = GenerationContext.ComponentInfos[0].ComponentName;
 			return true;
 		}
 		else
@@ -1586,7 +1589,7 @@ bool AffectsCurrentComponent(const UEdGraphPin* Pin, FMutableGraphGenerationCont
 			// This case is not supported yet
 			FString Msg = FString::Printf(TEXT("Error! Node has modifiers when using multiple components in root object. This is currently not supported."));
 			GenerationContext.Compiler->CompilerLog(FText::FromString(Msg), Node, EMessageSeverity::Error);
-			ComponentIndex = 0;
+			ComponentName = GenerationContext.ComponentInfos[0].ComponentName;
 			return false;
 		}
 	}
@@ -1596,7 +1599,14 @@ bool AffectsCurrentComponent(const UEdGraphPin* Pin, FMutableGraphGenerationCont
 		return false;
 	}
 
-	return ComponentIndex == GenerationContext.CurrentMeshComponent;
+	// Check if this material has a valid component
+	if (!GenerationContext.Object->GetPrivate()->MutableMeshComponents.ContainsByPredicate([&](const FMutableMeshComponentData& Component) { return Component.Name == ComponentName; }))
+	{
+		FString Msg = ("Error! Material Node with an invalid Component assigned.");
+		GenerationContext.Compiler->CompilerLog(FText::FromString(Msg), OriginalParentMaterialNode ? OriginalParentMaterialNode : Node, EMessageSeverity::Error);
+	}
+
+	return ComponentName == GenerationContext.CurrentMeshComponent;
 }
 
 int32 AddTagToMutableMeshUnique(mu::Mesh& MutableMesh, const FString& Tag)
@@ -2058,8 +2068,10 @@ void AddSocketTagsToMesh(const USkeletalMesh* SourceMesh, mu::Ptr<mu::Mesh> Muta
 }
 
 
-FMutableComponentInfo::FMutableComponentInfo(USkeletalMesh* InRefSkeletalMesh)
+FMutableComponentInfo::FMutableComponentInfo(FName InComponentName, USkeletalMesh* InRefSkeletalMesh)
 {
+	ComponentName = InComponentName;
+
 	if (!InRefSkeletalMesh || !InRefSkeletalMesh->GetSkeleton())
 	{
 		return;
diff --git a/Engine/Plugins/Experimental/Mutable/Source/CustomizableObjectEditor/Private/MuCOE/GenerateMutableSource/GenerateMutableSource.h b/Engine/Plugins/Experimental/Mutable/Source/CustomizableObjectEditor/Private/MuCOE/GenerateMutableSource/GenerateMutableSource.h
--- a/Engine/Plugins/Experimental/Mutable/Source/CustomizableObjectEditor/Private/MuCOE/GenerateMutableSource/GenerateMutableSource.h
+++ b/Engine/Plugins/Experimental/Mutable/Source/CustomizableObjectEditor/Private/MuCOE/GenerateMutableSource/GenerateMutableSource.h
@@ -393,10 +393,12 @@ inline uint32 GetTypeHash(const FMeshData& Key)
 /** Struct used to store info specific to each component during compilation */
 struct FMutableComponentInfo
 {
-	FMutableComponentInfo(USkeletalMesh* InRefSkeletalMesh);
+	FMutableComponentInfo(FName InComponentName, USkeletalMesh* InRefSkeletalMesh);
 
 	void AccumulateBonesToRemovePerLOD(const FComponentSettings& ComponentSettings, int32 NumLODs);
 
+	FName ComponentName;
+
 	// Each component must have a reference SkeletalMesh with a valid Skeleton
 	USkeletalMesh* RefSkeletalMesh = nullptr;
 	USkeleton* RefSkeleton = nullptr;
@@ -792,7 +794,7 @@ struct FMutableGraphGenerationContext
 
 	uint8 FromLOD = 0; // LOD to append to the CurrentLOD when using AutomaticLODs. 
 	uint8 CurrentLOD = 0;
-	int32 CurrentMeshComponent = 0;
+	FName CurrentMeshComponent;
 
 	/** If this is set, we are genreating materials for a "passthrough" component, with a fixed mesh. */
 	mu::Ptr<mu::NodeMesh> ComponentMeshOverride;
@@ -952,4 +954,3 @@ mu::Ptr<mu::Mesh> GenerateMeshConstant(const USkeletalMesh*, FMutableGraphGenera
 
 /** Generates a mutable image descriptor from an unreal engine texture */
 mu::FImageDesc GenerateImageDescriptor(UTexture* Texture);
-
diff --git a/Engine/Plugins/Experimental/Mutable/Source/CustomizableObjectEditor/Private/MuCOE/GenerateMutableSource/GenerateMutableSourceComponent.cpp b/Engine/Plugins/Experimental/Mutable/Source/CustomizableObjectEditor/Private/MuCOE/GenerateMutableSource/GenerateMutableSourceComponent.cpp
--- a/Engine/Plugins/Experimental/Mutable/Source/CustomizableObjectEditor/Private/MuCOE/GenerateMutableSource/GenerateMutableSourceComponent.cpp
+++ b/Engine/Plugins/Experimental/Mutable/Source/CustomizableObjectEditor/Private/MuCOE/GenerateMutableSource/GenerateMutableSourceComponent.cpp
@@ -21,7 +21,6 @@
 
 #define LOCTEXT_NAMESPACE "CustomizableObjectEditor"
 
-
 mu::Ptr<mu::NodeComponent> GenerateMutableSourceComponent(const UEdGraphPin * Pin, FMutableGraphGenerationContext& GenerationContext)
 {
 	check(Pin)
@@ -41,6 +40,15 @@ mu::Ptr<mu::NodeComponent> GenerateMutableSourceComponent(const UEdGraphPin * Pi
 	
 	if (const UCustomizableObjectNodeComponentMesh* TypedComponentMesh = Cast<UCustomizableObjectNodeComponentMesh>(Node))
 	{
+		GenerationContext.CurrentMeshComponent = TypedComponentMesh->ComponentName;
+
+		if (TypedComponentMesh->ComponentName.IsNone())
+		{
+			FString Msg = FString::Printf(TEXT("Invalid Component Name."));
+			GenerationContext.Compiler->CompilerLog(FText::FromString(Msg), TypedComponentMesh, EMessageSeverity::Warning);
+			return nullptr;
+		}
+
 		if (!TypedComponentMesh->Mesh.IsValid())
 		{
 			FString Msg = FString::Printf(TEXT("No mesh set for component node."));
@@ -71,7 +79,7 @@ mu::Ptr<mu::NodeComponent> GenerateMutableSourceComponent(const UEdGraphPin * Pi
 
 		// Create the component node
 		mu::Ptr<mu::NodeComponentNew> ComponentNode = new mu::NodeComponentNew;
-		ComponentNode->Id = GenerationContext.CurrentMeshComponent;
+		ComponentNode->Id = GenerationContext.NumMeshComponentsInRoot - 1 + GenerationContext.NumExplicitMeshComponents; // Last root component id + Num explicit components
 
 		// Create a LOD for each pass-through mesh LOD.
 		const FSkeletalMeshModel* Model = SkeletalMesh->GetImportedModel();
@@ -101,6 +109,7 @@ mu::Ptr<mu::NodeComponent> GenerateMutableSourceComponent(const UEdGraphPin * Pi
 			}
 		}
 
+		GenerationContext.CurrentMeshComponent = FName();
 		Result = ComponentNode;
 	}
 
diff --git a/Engine/Plugins/Experimental/Mutable/Source/CustomizableObjectEditor/Private/MuCOE/GenerateMutableSource/GenerateMutableSourceGroupProjector.cpp b/Engine/Plugins/Experimental/Mutable/Source/CustomizableObjectEditor/Private/MuCOE/GenerateMutableSource/GenerateMutableSourceGroupProjector.cpp
--- a/Engine/Plugins/Experimental/Mutable/Source/CustomizableObjectEditor/Private/MuCOE/GenerateMutableSource/GenerateMutableSourceGroupProjector.cpp
+++ b/Engine/Plugins/Experimental/Mutable/Source/CustomizableObjectEditor/Private/MuCOE/GenerateMutableSource/GenerateMutableSourceGroupProjector.cpp
@@ -411,8 +411,8 @@ bool GenerateMutableSourceGroupProjector(const UEdGraphPin* Pin, FMutableGraphGe
 			{
 				// Poses will only affect component 0 of the CO,
 				// TODO UE-206803
-				int32 OldCurrentMeshComponent = GenerationContext.CurrentMeshComponent;
-				GenerationContext.CurrentMeshComponent = 0;
+				FName OldCurrentMeshComponent = GenerationContext.CurrentMeshComponent;
+				GenerationContext.CurrentMeshComponent = GenerationContext.ComponentInfos[0].ComponentName;
 
 				PoseEnumParameterNode->SetValue(0, 0.f, "Default pose");
 
diff --git a/Engine/Plugins/Experimental/Mutable/Source/CustomizableObjectEditor/Private/MuCOE/GenerateMutableSource/GenerateMutableSourceMesh.cpp b/Engine/Plugins/Experimental/Mutable/Source/CustomizableObjectEditor/Private/MuCOE/GenerateMutableSource/GenerateMutableSourceMesh.cpp
--- a/Engine/Plugins/Experimental/Mutable/Source/CustomizableObjectEditor/Private/MuCOE/GenerateMutableSource/GenerateMutableSourceMesh.cpp
+++ b/Engine/Plugins/Experimental/Mutable/Source/CustomizableObjectEditor/Private/MuCOE/GenerateMutableSource/GenerateMutableSourceMesh.cpp
@@ -730,10 +730,10 @@ mu::MeshPtr ConvertSkeletalMeshToMutable(const USkeletalMesh* InSkeletalMesh, co
 
 	mu::MeshPtr MutableMesh = new mu::Mesh();
 
-	// CurrentMeshComponent < 0 implies IgnoreSkeleton flag.  
-	// CurrentMeshComponent < 0 will only happen with modifiers and, for now, any mesh generated from a modifier
+	// CurrentMeshComponent == None implies IgnoreSkeleton flag.
+	// CurrentMeshComponent == None will only happen with modifiers and, for now, any mesh generated from a modifier
 	// should ignore skinning.
-	check(!(GenerationContext.CurrentMeshComponent < 0) || bIgnoreSkeleton);
+	check(!(GenerationContext.CurrentMeshComponent.IsNone()) || bIgnoreSkeleton);
 
 	bool bBoneMapModified = false;
 	TArray<FBoneIndexType> BoneMap;
@@ -2687,7 +2687,7 @@ bool GetAndValidateReshapeBonesToDeform(
 	else if (SelectionMethod == EBoneDeformSelectionMethod::DEFORM_REF_SKELETON)
 	{
 		// Getting reference skeleton from the reference skeletal mesh of the current component
-		const FReferenceSkeleton RefSkeleton = GenerationContext.ComponentInfos[GenerationContext.CurrentMeshComponent].RefSkeletalMesh->GetRefSkeleton();
+		const FReferenceSkeleton RefSkeleton = GenerationContext.GetCurrentComponentInfo().RefSkeletalMesh->GetRefSkeleton(); //GenerationContext.ComponentInfos[GenerationContext.CurrentMeshComponent].RefSkeletalMesh->GetRefSkeleton();
 		int32 NumBones = RefSkeleton.GetRawBoneNum();
 
 		for (int32 BoneIndex = 0; BoneIndex < NumBones; ++BoneIndex)
@@ -2702,7 +2702,7 @@ bool GetAndValidateReshapeBonesToDeform(
 	else if (SelectionMethod == EBoneDeformSelectionMethod::DEFORM_NONE_REF_SKELETON)
 	{
 		// Getting reference skeleton from the reference skeletal mesh of the current component
-		const FReferenceSkeleton RefSkeleton = GenerationContext.ComponentInfos[GenerationContext.CurrentMeshComponent].RefSkeletalMesh->GetRefSkeleton();
+		const FReferenceSkeleton RefSkeleton = GenerationContext.GetCurrentComponentInfo().RefSkeletalMesh->GetRefSkeleton();
 
 		for (const MeshInfoType& Mesh : SkeletalMeshesInfo)
 		{
@@ -2800,7 +2800,7 @@ bool GetAndValidateReshapePhysicsToDeform(
 	if (bIsReferenceSkeletalMeshMethod)
 	{
 		const FReferenceSkeleton& RefSkeleton =
-			GenerationContext.ComponentInfos[GenerationContext.CurrentMeshComponent].RefSkeletalMesh->GetRefSkeleton();
+			GenerationContext.GetCurrentComponentInfo().RefSkeletalMesh->GetRefSkeleton();
 
 		const int32 RefSkeletonNumBones = RefSkeleton.GetRawBoneNum();
 			BoneNames.SetNum(RefSkeletonNumBones);
@@ -2831,10 +2831,10 @@ bool GetAndValidateReshapePhysicsToDeform(
 
 	for (const PhysicsInfoType& PhysicsInfo : ContributingPhysicsAssetsInfo)
 	{
-		check(GenerationContext.ComponentInfos[GenerationContext.CurrentMeshComponent].RefSkeletalMesh);
+		check(GenerationContext.GetCurrentComponentInfo().RefSkeletalMesh);
 
 		const FReferenceSkeleton& RefSkeleton = bIsReferenceSkeletalMeshMethod
-			? GenerationContext.ComponentInfos[GenerationContext.CurrentMeshComponent].RefSkeletalMesh->GetRefSkeleton()
+			? GenerationContext.GetCurrentComponentInfo().RefSkeletalMesh->GetRefSkeleton()
 			: PhysicsInfo.Get<const FReferenceSkeleton&>();
 
 		UPhysicsAsset* PhysicsAsset = PhysicsInfo.Get<UPhysicsAsset*>();
diff --git a/Engine/Plugins/Experimental/Mutable/Source/CustomizableObjectEditor/Private/MuCOE/GenerateMutableSource/GenerateMutableSourceSurface.cpp b/Engine/Plugins/Experimental/Mutable/Source/CustomizableObjectEditor/Private/MuCOE/GenerateMutableSource/GenerateMutableSourceSurface.cpp
--- a/Engine/Plugins/Experimental/Mutable/Source/CustomizableObjectEditor/Private/MuCOE/GenerateMutableSource/GenerateMutableSourceSurface.cpp
+++ b/Engine/Plugins/Experimental/Mutable/Source/CustomizableObjectEditor/Private/MuCOE/GenerateMutableSource/GenerateMutableSourceSurface.cpp
@@ -259,7 +259,7 @@ mu::Ptr<mu::NodeSurface> GenerateMutableSourceSurface(const UEdGraphPin * Pin, F
 	{
 		bool bGeneratingImplicitComponent = GenerationContext.ComponentMeshOverride.get()!=nullptr;
 
-		bool bValidForCurrentComponent = bGeneratingImplicitComponent || (TypedNodeMat->GetMeshComponentIndex() == GenerationContext.CurrentMeshComponent);
+		bool bValidForCurrentComponent = bGeneratingImplicitComponent || (TypedNodeMat->GetMeshComponentName() == GenerationContext.CurrentMeshComponent);
 		if (!bValidForCurrentComponent)
 		{
 			return Result;
diff --git a/Engine/Plugins/Experimental/Mutable/Source/CustomizableObjectEditor/Private/MuCOE/GraphTraversal.cpp b/Engine/Plugins/Experimental/Mutable/Source/CustomizableObjectEditor/Private/MuCOE/GraphTraversal.cpp
--- a/Engine/Plugins/Experimental/Mutable/Source/CustomizableObjectEditor/Private/MuCOE/GraphTraversal.cpp
+++ b/Engine/Plugins/Experimental/Mutable/Source/CustomizableObjectEditor/Private/MuCOE/GraphTraversal.cpp
@@ -296,13 +296,6 @@ bool GetParentsUntilRoot(UCustomizableObject* Object, TArray<UCustomizableObject
 		if (Root->ParentObject != nullptr)
 		{
 			bSuccess = GetParentsUntilRoot(Root->ParentObject, ArrayNodeObject, ArrayCustomizableObject);
-
-			UCustomizableObjectNodeObject* ParentRoot = GetRootNode(Root->ParentObject, MultipleBaseObjectsFound);
-
-			if (!MultipleBaseObjectsFound && (ParentRoot != nullptr))
-			{
-				Root->SetMeshComponentNumFromParent(ParentRoot->NumMeshComponents);
-			}
 		}
 	}
 
diff --git a/Engine/Plugins/Experimental/Mutable/Source/CustomizableObjectEditor/Private/MuCOE/Nodes/CustomizableObjectNodeCopyMaterial.cpp b/Engine/Plugins/Experimental/Mutable/Source/CustomizableObjectEditor/Private/MuCOE/Nodes/CustomizableObjectNodeCopyMaterial.cpp
--- a/Engine/Plugins/Experimental/Mutable/Source/CustomizableObjectEditor/Private/MuCOE/Nodes/CustomizableObjectNodeCopyMaterial.cpp
+++ b/Engine/Plugins/Experimental/Mutable/Source/CustomizableObjectEditor/Private/MuCOE/Nodes/CustomizableObjectNodeCopyMaterial.cpp
@@ -218,15 +218,15 @@ bool UCustomizableObjectNodeCopyMaterial::IsReuseMaterialBetweenLODs() const
 }
 
 
-int32 UCustomizableObjectNodeCopyMaterial::GetMeshComponentIndex() const
+FName UCustomizableObjectNodeCopyMaterial::GetMeshComponentName() const
 {
 	if (UCustomizableObjectNodeMaterial* NodeMaterial = GetMaterialNode())
 	{
-		return NodeMaterial->GetMeshComponentIndex();
+		return NodeMaterial->GetMeshComponentName();
 	}
 	else
 	{
-		return INDEX_NONE;
+		return FName();
 	}
 }
 
diff --git a/Engine/Plugins/Experimental/Mutable/Source/CustomizableObjectEditor/Private/MuCOE/Nodes/CustomizableObjectNodeMaterial.cpp b/Engine/Plugins/Experimental/Mutable/Source/CustomizableObjectEditor/Private/MuCOE/Nodes/CustomizableObjectNodeMaterial.cpp
--- a/Engine/Plugins/Experimental/Mutable/Source/CustomizableObjectEditor/Private/MuCOE/Nodes/CustomizableObjectNodeMaterial.cpp
+++ b/Engine/Plugins/Experimental/Mutable/Source/CustomizableObjectEditor/Private/MuCOE/Nodes/CustomizableObjectNodeMaterial.cpp
@@ -638,6 +638,11 @@ void UCustomizableObjectNodeMaterial::BackwardsCompatibleFixup()
 			}
 		}
 	}
+
+	if (CustomizableObjectCustomVersion < FCustomizableObjectCustomVersion::NewComponentOptions)
+	{
+		MeshComponentName = FName(FString::FromInt(MeshComponentIndex_DEPRECATED));
+	}
 }
 
 
@@ -1113,11 +1118,27 @@ bool UCustomizableObjectNodeMaterial::CustomRemovePin(UEdGraphPin& Pin)
 	return Super::CustomRemovePin(Pin);
 }
 
+
 void UCustomizableObjectNodeMaterial::ReconstructNode(UCustomizableObjectNodeRemapPins* RemapPinsMode)
 {
 	Super::ReconstructNode(RemapPinsMode);
 
 	UpdateAllImagesPinMode();
+
+	// When a material node is created, the first component is set as its mesh component
+	if (MeshComponentName.IsNone())
+	{
+		if (UCustomizableObject* CustomizableObject = Cast<UCustomizableObject>(GetOutermostObject()))
+		{
+			if (UCustomizableObject* RootObject = GetRootObject(CustomizableObject))
+			{
+				if (RootObject->GetPrivate()->MutableMeshComponents.Num())
+				{
+					MeshComponentName = RootObject->GetPrivate()->MutableMeshComponents[0].Name;
+				}
+			}
+		}
+	}
 }
 
 
@@ -1139,9 +1160,9 @@ bool UCustomizableObjectNodeMaterial::IsReuseMaterialBetweenLODs() const
 }
 
 
-int32 UCustomizableObjectNodeMaterial::GetMeshComponentIndex() const
+FName UCustomizableObjectNodeMaterial::GetMeshComponentName() const
 {
-	return MeshComponentIndex;	
+	return MeshComponentName;
 }
 
 
@@ -1405,4 +1426,11 @@ void UCustomizableObjectNodeMaterial::BreakExistingConnectionsPostConnection(UEd
 	}
 }
 
+
+void UCustomizableObjectNodeMaterial::SetComponentName(const FName& Name)
+{
+	MeshComponentName = Name;
+}
+
+
 #undef LOCTEXT_NAMESPACE
diff --git a/Engine/Plugins/Experimental/Mutable/Source/CustomizableObjectEditor/Private/MuCOE/Nodes/CustomizableObjectNodeMaterialDetails.cpp b/Engine/Plugins/Experimental/Mutable/Source/CustomizableObjectEditor/Private/MuCOE/Nodes/CustomizableObjectNodeMaterialDetails.cpp
new file mode 100644
--- /dev/null
+++ b/Engine/Plugins/Experimental/Mutable/Source/CustomizableObjectEditor/Private/MuCOE/Nodes/CustomizableObjectNodeMaterialDetails.cpp
@@ -0,0 +1,144 @@
+// Copyright Epic Games, Inc. All Rights Reserved.
+
+#include "MuCOE/Nodes/CustomizableObjectNodeMaterialDetails.h"
+
+#include "DetailCategoryBuilder.h"
+#include "DetailLayoutBuilder.h"
+#include "DetailWidgetRow.h"
+#include "MuCO/CustomizableObject.h"
+#include "MuCOE/GraphTraversal.h"
+#include "MuCOE/Nodes/CustomizableObjectNodeMaterial.h"
+#include "MuCOE/PinViewer/SPinViewer.h"
+#include "SSearchableComboBox.h"
+
+class IDetailLayoutBuilder;
+
+#define LOCTEXT_NAMESPACE "CustomizableObjectEditor"
+
+
+TSharedRef<IDetailCustomization> FCustomizableObjectNodeMaterialDetails::MakeInstance()
+{
+	return MakeShared<FCustomizableObjectNodeMaterialDetails>();
+}
+
+
+void FCustomizableObjectNodeMaterialDetails::CustomizeDetails(IDetailLayoutBuilder& DetailBuilder)
+{
+	// Adding Pin viewr to the details
+	FCustomizableObjectNodeDetails::CustomizeDetails(DetailBuilder);
+
+	// Get material node
+	if (DetailBuilder.GetObjectsOfTypeBeingCustomized<UCustomizableObjectNodeMaterial>().Num())
+	{
+		MaterialNode = DetailBuilder.GetObjectsOfTypeBeingCustomized<UCustomizableObjectNodeMaterial>()[0];
+	}
+	else
+	{
+		return;
+	}
+
+	DefaultTextColor = FCoreStyle::Get().GetWidgetStyle<FTextBlockStyle>("NormalText").ColorAndOpacity;
+
+	IDetailCategoryBuilder& CustomCategory = DetailBuilder.EditCategory("CustomizableObject");
+	TSharedRef<IPropertyHandle> MaterialProperty = DetailBuilder.GetProperty("Material");
+	TSharedRef<IPropertyHandle> ComponentNameProperty = DetailBuilder.GetProperty("ComponentName");
+
+	// Added Material property manually to display it above of the rest of properties
+	CustomCategory.AddProperty(MaterialProperty);
+
+	// Custom widget for the Components property
+	CustomCategory.AddProperty(ComponentNameProperty).CustomWidget()
+	.NameContent()
+	[
+		SNew(STextBlock)
+		.Text(LOCTEXT("ComponentComboboxText", "Component:"))
+		.Font(IDetailLayoutBuilder::GetDetailFont())
+	]
+	.ValueContent()
+	[
+		SAssignNew(ComponentComboBoxWidget, SSearchableComboBox)
+		.InitiallySelectedItem(SelectedComponentName)
+		.OptionsSource(&ComponentNames)
+		.OnSelectionChanged(this, &FCustomizableObjectNodeMaterialDetails::OnComponentNameSelectionChanged)
+		.OnGenerateWidget(this, &FCustomizableObjectNodeMaterialDetails::GenerateComponentNameComboEntryWidget)
+		.OnComboBoxOpening(this, &FCustomizableObjectNodeMaterialDetails::OnOpenComponentsCombobox)
+		.Content()
+		[
+			SAssignNew(ComponentNameTextWidget, STextBlock)
+			.Text(this, &FCustomizableObjectNodeMaterialDetails::GenerateSelectedComponentNameWidget)
+			.Font(IDetailLayoutBuilder::GetDetailFont())
+		]
+	];
+
+	// Init ComboBox
+	GetComponentNames();
+
+	if (TSharedPtr<FString>* ComponentName = ComponentNames.FindByPredicate([&](TSharedPtr<FString> ComponentName) { return MaterialNode->GetMeshComponentName() == *ComponentName; }))
+	{
+		SelectedComponentName = *ComponentName;
+		ComponentNameTextWidget->SetColorAndOpacity(DefaultTextColor);
+	}
+	else
+	{
+		SelectedComponentName = MakeShareable(new FString(MaterialNode->GetMeshComponentName().ToString()));
+		ComponentNameTextWidget->SetColorAndOpacity(FSlateColor(FLinearColor::Red));
+	}
+}
+
+
+void FCustomizableObjectNodeMaterialDetails::GetComponentNames()
+{
+	ComponentNames.Empty();
+
+	if (UCustomizableObject* RootObject = Cast<UCustomizableObject>(MaterialNode->GetOutermostObject()))
+	{
+		// ensure we get the root object of the tree
+		RootObject = GetRootObject(RootObject);
+
+		if (RootObject)
+		{
+			for (int32 ComponentIndex = 0; ComponentIndex < RootObject->GetPrivate()->MutableMeshComponents.Num(); ++ComponentIndex)
+			{
+				ComponentNames.Add(MakeShareable(new FString(RootObject->GetPrivate()->MutableMeshComponents[ComponentIndex].Name.ToString())));
+			}
+		}
+	}
+}
+
+
+void FCustomizableObjectNodeMaterialDetails::OnComponentNameSelectionChanged(TSharedPtr<FString> Selection, ESelectInfo::Type SelectInfo)
+{
+	if (Selection.IsValid())
+	{
+		SelectedComponentName = Selection;
+		MaterialNode->SetComponentName(FName(*Selection));
+
+		ComponentNameTextWidget->SetColorAndOpacity(DefaultTextColor);
+	}
+}
+
+
+TSharedRef<SWidget> FCustomizableObjectNodeMaterialDetails::GenerateComponentNameComboEntryWidget(TSharedPtr<FString> InItem) const
+{
+	return SNew(STextBlock).Text(FText::FromString(*InItem.Get()));
+}
+
+
+FText FCustomizableObjectNodeMaterialDetails::GenerateSelectedComponentNameWidget() const
+{
+	if (MaterialNode.IsValid() && SelectedComponentName.IsValid())
+	{
+		return FText::FromString(*SelectedComponentName.Get());
+	}
+
+	return FText();
+}
+
+
+void FCustomizableObjectNodeMaterialDetails::OnOpenComponentsCombobox()
+{
+	GetComponentNames();
+	ComponentComboBoxWidget->RefreshOptions();
+}
+
+#undef LOCTEXT_NAMESPACE
diff --git a/Engine/Plugins/Experimental/Mutable/Source/CustomizableObjectEditor/Private/MuCOE/Nodes/CustomizableObjectNodeMaterialDetails.h b/Engine/Plugins/Experimental/Mutable/Source/CustomizableObjectEditor/Private/MuCOE/Nodes/CustomizableObjectNodeMaterialDetails.h
new file mode 100644
--- /dev/null
+++ b/Engine/Plugins/Experimental/Mutable/Source/CustomizableObjectEditor/Private/MuCOE/Nodes/CustomizableObjectNodeMaterialDetails.h
@@ -0,0 +1,59 @@
+// Copyright Epic Games, Inc. All Rights Reserved.
+
+#pragma once
+
+#include "CustomizableObjectNodeDetails.h"
+#include "IDetailCustomization.h"
+#include "UObject/WeakObjectPtr.h"
+#include "Styling/SlateColor.h"
+
+class IDetailLayoutBuilder;
+class SSearchableComboBox;
+class STextBlock;
+class SWidget;
+class UCustomizableObjectNodeMaterial;
+
+namespace ESelectInfo { enum Type : int; }
+
+class FCustomizableObjectNodeMaterialDetails : public FCustomizableObjectNodeDetails
+{
+public:
+	/** Makes a new instance of this detail layout class for a specific detail view requesting it */
+	static TSharedRef<IDetailCustomization> MakeInstance();
+
+	// IDetailCustomization interface
+	virtual void CustomizeDetails(IDetailLayoutBuilder& DetailBuilder) override;
+
+private:
+
+	/** Fills an array of string with all the components names defined in the root CO. */
+	void GetComponentNames();
+
+	void OnComponentNameSelectionChanged(TSharedPtr<FString> Selection, ESelectInfo::Type SelectInfo);
+
+	TSharedRef<SWidget> GenerateComponentNameComboEntryWidget(TSharedPtr<FString> InItem) const;
+
+	FText GenerateSelectedComponentNameWidget() const;
+
+	void OnOpenComponentsCombobox();
+
+private:
+
+	// Weak pointer to the material node represented in the details
+	TWeakObjectPtr<UCustomizableObjectNodeMaterial> MaterialNode;
+
+	// All component names of the customizable object
+	TArray<TSharedPtr<FString>> ComponentNames;
+
+	/** Pointer to store which is the current selected component */
+	TSharedPtr<FString> SelectedComponentName;
+	
+	/** Widget to draw the name of the selected component inside the combobox */
+	TSharedPtr<STextBlock> ComponentNameTextWidget;
+	
+	/** Combobox widget to select a component name */
+	TSharedPtr<SSearchableComboBox> ComponentComboBoxWidget;
+
+	/** Color for the default text mode */
+	FSlateColor DefaultTextColor;
+};
diff --git a/Engine/Plugins/Experimental/Mutable/Source/CustomizableObjectEditor/Private/MuCOE/Nodes/CustomizableObjectNodeObject.cpp b/Engine/Plugins/Experimental/Mutable/Source/CustomizableObjectEditor/Private/MuCOE/Nodes/CustomizableObjectNodeObject.cpp
--- a/Engine/Plugins/Experimental/Mutable/Source/CustomizableObjectEditor/Private/MuCOE/Nodes/CustomizableObjectNodeObject.cpp
+++ b/Engine/Plugins/Experimental/Mutable/Source/CustomizableObjectEditor/Private/MuCOE/Nodes/CustomizableObjectNodeObject.cpp
@@ -94,6 +94,15 @@ void UCustomizableObjectNodeObject::BackwardsCompatibleFixup()
 			State.UIMetadata.ExtraAssets = State.StateUIMetadata_DEPRECATED.ExtraAssets;
 		}
 	}
+
+	if (CustomizableObjectCustomVersion < FCustomizableObjectCustomVersion::NewComponentOptions)
+	{
+		// Like we did in the CO components, we use the index of the component as the name of the component
+		for (int32 ComponentIndex = 0; ComponentIndex < ComponentSettings.Num(); ++ComponentIndex)
+		{
+			ComponentSettings[ComponentIndex].ComponentName = FString::FromInt(ComponentIndex);
+		}
+	}
 }
 
 
@@ -124,16 +133,6 @@ void UCustomizableObjectNodeObject::PostEditChangeProperty(FPropertyChangedEvent
 		ReconstructNode();
 	}
 
-	if (PropertyThatChanged && PropertyThatChanged->GetName() == TEXT("NumMeshComponents"))
-	{
-		ComponentSettings.SetNum(NumMeshComponents);
-
-		for (int32 CompSetIndex = 0; CompSetIndex < ComponentSettings.Num(); ++CompSetIndex)
-		{
-			ComponentSettings[CompSetIndex].LODReductionSettings.SetNum(NumLODs);
-		}
-	}
-	
 	Super::PostEditChangeProperty(PropertyChangedEvent);
 }
 
@@ -297,13 +296,16 @@ void UCustomizableObjectNodeObject::PostBackwardsCompatibleFixup()
 {
 	Super::PostBackwardsCompatibleFixup();
 
-	// Fix up ComponentSettings
-	if (ComponentSettings.IsEmpty())
+	// Fix up ComponentSettings. Only root nodes
+	if (ComponentSettings.IsEmpty() && bIsBase && !ParentObject)
 	{
 		FComponentSettings ComponentSettingsTemplate;
 		ComponentSettingsTemplate.LODReductionSettings.SetNum(NumLODs);
 
-		ComponentSettings.Init(ComponentSettingsTemplate, NumMeshComponents);
+		if (UCustomizableObject* CurrentObject = Cast<UCustomizableObject>(GetOutermostObject()))
+		{
+			ComponentSettings.Init(ComponentSettingsTemplate, CurrentObject->GetPrivate()->MutableMeshComponents.Num());
+		}
 	}
 
 	// Reconstruct in case any extension pins have changed
diff --git a/Engine/Plugins/Experimental/Mutable/Source/CustomizableObjectEditor/Private/MuCOE/Nodes/CustomizableObjectNodeObjectDetails.cpp b/Engine/Plugins/Experimental/Mutable/Source/CustomizableObjectEditor/Private/MuCOE/Nodes/CustomizableObjectNodeObjectDetails.cpp
--- a/Engine/Plugins/Experimental/Mutable/Source/CustomizableObjectEditor/Private/MuCOE/Nodes/CustomizableObjectNodeObjectDetails.cpp
+++ b/Engine/Plugins/Experimental/Mutable/Source/CustomizableObjectEditor/Private/MuCOE/Nodes/CustomizableObjectNodeObjectDetails.cpp
@@ -179,17 +179,33 @@ void FCustomizableObjectNodeObjectDetails::CustomizeDetails( IDetailLayoutBuilde
 	{
 		// Properties
 		TSharedRef<IPropertyHandle> StatesProperty = DetailBuilder.GetProperty("States");
+		//TODO(Max UE-215837)
+		//TSharedRef<IPropertyHandle> ComponentsProperty = DetailBuilder.GetProperty("Components");
 		TSharedRef<IPropertyHandle> ParentObjectProperty = DetailBuilder.GetProperty("ParentObject");
-		TSharedRef<IPropertyHandle> ComponentsProperty = DetailBuilder.GetProperty("NumMeshComponents");
 		TSharedRef<IPropertyHandle> LODsProperty = DetailBuilder.GetProperty("NumLODs");
 		TSharedRef<IPropertyHandle> ComponentSettingsProperty= DetailBuilder.GetProperty("ComponentSettings");
 
-		FName BonesToRemovePropertyPath = FName("ComponentSettings[" + FString::FromInt(BaseObjectNode->CurrentComponent) + "].LODReductionSettings[" + FString::FromInt(BaseObjectNode->CurrentLOD) + "].BonesToRemove");
+		// Index of the component shown in the Bones to edit widget
+		int32 CurrentComponentIndex = 0;
+
+		for (int32 ComponentIndex = 0; ComponentIndex < BaseObjectNode->ComponentSettings.Num(); ++ComponentIndex)
+		{
+			if (BaseObjectNode->CurrentComponent == BaseObjectNode->ComponentSettings[ComponentIndex].ComponentName)
+			{
+				CurrentComponentIndex = ComponentIndex;
+				break;
+			}
+		}
+
+		FName BonesToRemovePropertyPath = FName("ComponentSettings[" + FString::FromInt(CurrentComponentIndex) + "].LODReductionSettings[" + FString::FromInt(BaseObjectNode->CurrentLOD) + "].BonesToRemove");
 		TSharedRef<IPropertyHandle> BonesToRemoveProperty = DetailBuilder.GetProperty(BonesToRemovePropertyPath);
 
 		// Callbacks
 		StatesProperty->SetOnPropertyValueChanged(FSimpleDelegate::CreateSP(this, &FCustomizableObjectNodeObjectDetails::OnStatesPropertyChanged));
-		ComponentsProperty->SetOnPropertyValueChangedWithData(TDelegate<void(const FPropertyChangedEvent&)>::CreateSP(this, &FCustomizableObjectNodeObjectDetails::OnNumComponentsOrLODsChanged));
+
+		//TODO(Max UE-215837)
+		//ComponentsProperty->SetOnPropertyValueChanged(FSimpleDelegate::CreateSP(this, &FCustomizableObjectNodeObjectDetails::OnStatesPropertyChanged));
+		//NumComponentsProperty->SetOnPropertyValueChangedWithData(TDelegate<void(const FPropertyChangedEvent&)>::CreateSP(this, &FCustomizableObjectNodeObjectDetails::OnNumComponentsOrLODsChanged));
 		LODsProperty->SetOnPropertyValueChangedWithData(TDelegate<void(const FPropertyChangedEvent&)>::CreateSP(this, &FCustomizableObjectNodeObjectDetails::OnNumComponentsOrLODsChanged));
 
 		// Hidden Properties
@@ -257,6 +273,9 @@ void FCustomizableObjectNodeObjectDetails::CustomizeDetails( IDetailLayoutBuilde
 				];
 
 				DetailBuilder.HideProperty("NumMeshComponents");
+				DetailBuilder.HideProperty(StatesProperty);
+				//TODO(Max UE-215837)
+				//DetailBuilder.HideProperty(ComponentsProperty);
 			}
             else
             {
@@ -443,17 +462,20 @@ TSharedRef<SWidget> FCustomizableObjectNodeObjectDetails::OnGenerateComponentMen
 {
 	if (BaseObjectNode.IsValid())
 	{
-		int32 NumComponents = BaseObjectNode->NumMeshComponents;
-		FMenuBuilder MenuBuilder(true, NULL);
-
-		for (int32 ComponentIndex = 0; ComponentIndex < NumComponents; ++ComponentIndex)
+		if (const UCustomizableObject* ParentObject = Cast<UCustomizableObject>(BaseObjectNode->GetOutermostObject()))
 		{
-			FText ComponentString = FText::FromString((TEXT("Component ") + FString::FromInt(ComponentIndex)));
-			FUIAction Action(FExecuteAction::CreateSP(this, &FCustomizableObjectNodeObjectDetails::OnSelectedComponentChanged, ComponentIndex));
-			MenuBuilder.AddMenuEntry(ComponentString, FText::GetEmpty(), FSlateIcon(), Action);
-		}
+			int32 NumComponents = ParentObject->GetPrivate()->MutableMeshComponents.Num();
+			FMenuBuilder MenuBuilder(true, NULL);
 
-		return MenuBuilder.MakeWidget();
+			for (int32 ComponentIndex = 0; ComponentIndex < NumComponents; ++ComponentIndex)
+			{
+				FText ComponentString = FText::FromName(ParentObject->GetPrivate()->MutableMeshComponents[ComponentIndex].Name);
+				FUIAction Action(FExecuteAction::CreateSP(this, &FCustomizableObjectNodeObjectDetails::OnSelectedComponentChanged, ParentObject->GetPrivate()->MutableMeshComponents[ComponentIndex].Name));
+				MenuBuilder.AddMenuEntry(ComponentString, FText::GetEmpty(), FSlateIcon(), Action);
+			}
+
+			return MenuBuilder.MakeWidget();
+		}
 	}
 
 	return SNullWidget::NullWidget;
@@ -481,11 +503,11 @@ TSharedRef<SWidget> FCustomizableObjectNodeObjectDetails::OnGenerateLODMenuForPi
 }
 
 
-void FCustomizableObjectNodeObjectDetails::OnSelectedComponentChanged(int32 NewComponentIndex)
+void FCustomizableObjectNodeObjectDetails::OnSelectedComponentChanged(const FName NewComponentSelected)
 {
 	if (BaseObjectNode.IsValid())
 	{
-		BaseObjectNode->CurrentComponent = NewComponentIndex;
+		BaseObjectNode->CurrentComponent = NewComponentSelected;
 		BaseObjectNode->CurrentLOD = 0;
 	}
 	
@@ -510,7 +532,7 @@ FText FCustomizableObjectNodeObjectDetails::GetCurrentComponentName() const
 
 	if (BaseObjectNode.IsValid())
 	{
-		ComponentText = FText::FromString(FString(TEXT("Component ")) + FString::FromInt(BaseObjectNode->CurrentComponent));
+		ComponentText = FText::FromName(BaseObjectNode->CurrentComponent);
 	}
 
 	return ComponentText;
@@ -534,20 +556,21 @@ void FCustomizableObjectNodeObjectDetails::OnNumComponentsOrLODsChanged(const FP
 {
 	if (DetailBuilderPtr && PropertyChangedEvent.ChangeType == EPropertyChangeType::ValueSet)
 	{
-		int32 ComponentToSelect = FMath::Min(BaseObjectNode->NumMeshComponents - 1, BaseObjectNode->CurrentComponent);
-
-		if (ComponentToSelect != BaseObjectNode->CurrentComponent)
+		if (const UCustomizableObject* ParentObject = Cast<UCustomizableObject>(BaseObjectNode->GetOutermostObject()))
 		{
-			BaseObjectNode->CurrentComponent = ComponentToSelect;
+			if (ParentObject->GetPrivate()->MutableMeshComponents.FindByPredicate([&](const FMutableMeshComponentData& Component) { return Component.Name == BaseObjectNode->CurrentComponent; }) == nullptr)
+			{
+				BaseObjectNode->CurrentComponent = ParentObject->GetPrivate()->MutableMeshComponents.Last().Name;
 
-			// Reset the LOD selection
-			BaseObjectNode->CurrentLOD = 0;
-		}
-		else
-		{
-			BaseObjectNode->CurrentLOD = FMath::Min(BaseObjectNode->NumLODs - 1, BaseObjectNode->CurrentLOD);
+				// Reset the LOD selection
+				BaseObjectNode->CurrentLOD = 0;
+			}
+			else
+			{
+				BaseObjectNode->CurrentLOD = FMath::Min(BaseObjectNode->NumLODs - 1, BaseObjectNode->CurrentLOD);
+			}
 		}
-
+	
 		DetailBuilderPtr->ForceRefreshDetails();
 	}
 }
@@ -570,7 +593,7 @@ void FCustomizableObjectNodeObjectDetails::FillParameterNamesArray()
 	// Full tree graph of customizable objects
 	TSet<UCustomizableObject*> CustomObjectTree;
 
-	// Get the whole tree of customizable object
+	// Get and load the whole tree of customizable object
 	GetAllObjectsInGraph(RootObjet, CustomObjectTree);
 
 	// Array to store all the ids of group nodes of type toggle
diff --git a/Engine/Plugins/Experimental/Mutable/Source/CustomizableObjectEditor/Private/MuCOE/Nodes/CustomizableObjectNodeObjectDetails.h b/Engine/Plugins/Experimental/Mutable/Source/CustomizableObjectEditor/Private/MuCOE/Nodes/CustomizableObjectNodeObjectDetails.h
--- a/Engine/Plugins/Experimental/Mutable/Source/CustomizableObjectEditor/Private/MuCOE/Nodes/CustomizableObjectNodeObjectDetails.h
+++ b/Engine/Plugins/Experimental/Mutable/Source/CustomizableObjectEditor/Private/MuCOE/Nodes/CustomizableObjectNodeObjectDetails.h
@@ -111,7 +111,7 @@ private:
 	TSharedRef<SWidget> OnGenerateLODMenuForPicker();
 
 	// Component and LOD OnSelect callbacks
-	void OnSelectedComponentChanged(int32 NewComponentIndex);
+	void OnSelectedComponentChanged(const FName NewComponentSelected);
 	void OnSelectedLODChanged(int32 NewLODIndex);
 
 	// Component and LOD name generation callbacks
diff --git a/Engine/Plugins/Experimental/Mutable/Source/CustomizableObjectEditor/Private/MuCOE/StatePerformanceTesting.cpp b/Engine/Plugins/Experimental/Mutable/Source/CustomizableObjectEditor/Private/MuCOE/StatePerformanceTesting.cpp
--- a/Engine/Plugins/Experimental/Mutable/Source/CustomizableObjectEditor/Private/MuCOE/StatePerformanceTesting.cpp
+++ b/Engine/Plugins/Experimental/Mutable/Source/CustomizableObjectEditor/Private/MuCOE/StatePerformanceTesting.cpp
@@ -88,7 +88,7 @@ void FTestingCustomizableObject::StartTest(UCustomizableObject* CustomizableObje
 
 	Instance = CustomizableObject->CreateInstance();
 
-	const int32 NumComponents = CustomizableObject->ReferenceSkeletalMeshes.Num();
+	const int32 NumComponents = CustomizableObject->GetPrivate()->MutableMeshComponents.Num();
 	SkeletalMeshComponents.Reset(NumComponents);
 	CustomizableSkeletalComponents.Reset(NumComponents);
 
@@ -101,6 +101,8 @@ void FTestingCustomizableObject::StartTest(UCustomizableObject* CustomizableObje
 			
 			CustomizableSkeletalComponent->CustomizableObjectInstance = Instance;
 			CustomizableSkeletalComponent->ComponentIndex = ComponentIndex;
+			//CustomizableSkeletalComponent->ComponentName = CustomizableObject->GetPrivate()->Components[ComponentIndex].Name;
+
 
 			USkeletalMeshComponent* SkeletalMeshComponent = NewObject<USkeletalMeshComponent>(GetTransientPackage(), NAME_None, RF_Transient);
 			if (SkeletalMeshComponent)
diff --git a/Engine/Plugins/Experimental/Mutable/Source/CustomizableObjectEditor/Private/MuCOE/Widgets/CustomizableObjectLODReductionSettings.cpp b/Engine/Plugins/Experimental/Mutable/Source/CustomizableObjectEditor/Private/MuCOE/Widgets/CustomizableObjectLODReductionSettings.cpp
--- a/Engine/Plugins/Experimental/Mutable/Source/CustomizableObjectEditor/Private/MuCOE/Widgets/CustomizableObjectLODReductionSettings.cpp
+++ b/Engine/Plugins/Experimental/Mutable/Source/CustomizableObjectEditor/Private/MuCOE/Widgets/CustomizableObjectLODReductionSettings.cpp
@@ -150,9 +150,13 @@ const struct FReferenceSkeleton& FCustomizableObjectLODReductionSettings::GetRef
 	{
 		if (const UCustomizableObject* CustomizableObject = Cast<UCustomizableObject>(ObjectNode->GetCustomizableObjectGraph()->GetOuter()))
 		{
-			if (CustomizableObject->ReferenceSkeletalMeshes.IsValidIndex(ObjectNode->CurrentComponent))
+			if (const FMutableMeshComponentData* Component = CustomizableObject->GetPrivate()->MutableMeshComponents.FindByPredicate(
+				[&](const FMutableMeshComponentData& Component) 
+				{ 
+					return Component.Name == ObjectNode->CurrentComponent; 
+				}))
 			{
-				SkeletalMesh = CustomizableObject->ReferenceSkeletalMeshes[ObjectNode->CurrentComponent];
+				SkeletalMesh = Component->ReferenceSkeletalMesh;
 			}
 		}
 	}
diff --git a/Engine/Plugins/Experimental/Mutable/Source/CustomizableObjectEditor/Public/MuCOE/CustomizableObjectInstanceFactory.h b/Engine/Plugins/Experimental/Mutable/Source/CustomizableObjectEditor/Public/MuCOE/CustomizableObjectInstanceFactory.h
--- a/Engine/Plugins/Experimental/Mutable/Source/CustomizableObjectEditor/Public/MuCOE/CustomizableObjectInstanceFactory.h
+++ b/Engine/Plugins/Experimental/Mutable/Source/CustomizableObjectEditor/Public/MuCOE/CustomizableObjectInstanceFactory.h
@@ -24,6 +24,5 @@ protected:
 	//~ End UActorFactory Interface
 
 	USkeletalMesh* GetSkeletalMeshFromAsset(UObject* Asset, int32 ComponentIndex = 0) const;
-	int32 GetNumberOfComponents(class UCustomizableObjectInstance* COInstance);
 	FQuat AlignObjectToSurfaceNormal(const FVector& InSurfaceNormal, const FQuat& ActorRotation) const;
 };
diff --git a/Engine/Plugins/Experimental/Mutable/Source/CustomizableObjectEditor/Public/MuCOE/Nodes/CustomizableObjectNodeComponentMesh.h b/Engine/Plugins/Experimental/Mutable/Source/CustomizableObjectEditor/Public/MuCOE/Nodes/CustomizableObjectNodeComponentMesh.h
--- a/Engine/Plugins/Experimental/Mutable/Source/CustomizableObjectEditor/Public/MuCOE/Nodes/CustomizableObjectNodeComponentMesh.h
+++ b/Engine/Plugins/Experimental/Mutable/Source/CustomizableObjectEditor/Public/MuCOE/Nodes/CustomizableObjectNodeComponentMesh.h
@@ -66,7 +66,10 @@ class CUSTOMIZABLEOBJECTEDITOR_API UCustomizableObjectNodeComponentMesh : public
 {
 public:
 	GENERATED_BODY()
-		
+
+	UPROPERTY(EditAnywhere, Category = Mesh)
+	FName ComponentName;
+
 	UPROPERTY(EditAnywhere, Category = Mesh, Meta = (DisplayName = Mesh, AllowedClasses = "/Script/Engine.StaticMesh, /Script/Engine.SkeletalMesh"))
 	FSoftObjectPath Mesh;
 
diff --git a/Engine/Plugins/Experimental/Mutable/Source/CustomizableObjectEditor/Public/MuCOE/Nodes/CustomizableObjectNodeCopyMaterial.h b/Engine/Plugins/Experimental/Mutable/Source/CustomizableObjectEditor/Public/MuCOE/Nodes/CustomizableObjectNodeCopyMaterial.h
--- a/Engine/Plugins/Experimental/Mutable/Source/CustomizableObjectEditor/Public/MuCOE/Nodes/CustomizableObjectNodeCopyMaterial.h
+++ b/Engine/Plugins/Experimental/Mutable/Source/CustomizableObjectEditor/Public/MuCOE/Nodes/CustomizableObjectNodeCopyMaterial.h
@@ -35,7 +35,7 @@ public:
 	// UCustomizableObjectNodeMaterialBase interface
 	virtual UMaterialInterface* GetMaterial() const override;
 	virtual bool IsReuseMaterialBetweenLODs() const override;
-	virtual int32 GetMeshComponentIndex() const override;
+	virtual FName GetMeshComponentName() const override;
 	virtual TArray<FString> GetTags() const override;
 	virtual UEdGraphPin* GetMaterialAssetPin() const override;
 	virtual int32 GetNumParameters(EMaterialParameterType Type) const override;
diff --git a/Engine/Plugins/Experimental/Mutable/Source/CustomizableObjectEditor/Public/MuCOE/Nodes/CustomizableObjectNodeMaterial.h b/Engine/Plugins/Experimental/Mutable/Source/CustomizableObjectEditor/Public/MuCOE/Nodes/CustomizableObjectNodeMaterial.h
--- a/Engine/Plugins/Experimental/Mutable/Source/CustomizableObjectEditor/Public/MuCOE/Nodes/CustomizableObjectNodeMaterial.h
+++ b/Engine/Plugins/Experimental/Mutable/Source/CustomizableObjectEditor/Public/MuCOE/Nodes/CustomizableObjectNodeMaterial.h
@@ -139,7 +139,7 @@ public:
 	virtual UEdGraphPin* OutputPin() const override;
 	virtual UMaterialInterface* GetMaterial() const override;
 	virtual bool IsReuseMaterialBetweenLODs() const override;
-	virtual int32 GetMeshComponentIndex() const override;
+	virtual FName GetMeshComponentName() const override;
 	virtual TArray<FString> GetTags() const override;
 	virtual int32 GetNumParameters(EMaterialParameterType Type) const override;
 	virtual FGuid GetParameterId(EMaterialParameterType Type, int32 ParameterIndex) const override;
@@ -165,8 +165,7 @@ public:
 	// Own Interface
 
 	void SetMaterial(UMaterialInterface* InMaterial);
-	
-	/** Return true if a Material Parameter has changed on which we had a pin connected or data saved. */
+	void SetComponentName(const FName& Name);
 
 private:
 	/** Delegate called when a Texture Parameter Pin Mode changes. */
@@ -181,9 +180,12 @@ private:
 	UPROPERTY(EditAnywhere, Category = CustomizableObject)
 	TArray<FString> Tags;
 
+	UPROPERTY()
+	int32 MeshComponentIndex_DEPRECATED = 0;
+
 	/** Selects which Mesh component of the Instance this material belongs to */
-	UPROPERTY(EditAnywhere, Category = CustomizableObject, meta = (ClampMin = "0"))
-	int32 MeshComponentIndex = 0;
+	UPROPERTY(EditAnywhere, Category = CustomizableObject)
+	FName MeshComponentName;
 
 	/** Materials will be reused between LODs when possible unless explicitly disabled. */
 	UPROPERTY(EditAnywhere, Category = CustomizableObject)
diff --git a/Engine/Plugins/Experimental/Mutable/Source/CustomizableObjectEditor/Public/MuCOE/Nodes/CustomizableObjectNodeMaterialBase.h b/Engine/Plugins/Experimental/Mutable/Source/CustomizableObjectEditor/Public/MuCOE/Nodes/CustomizableObjectNodeMaterialBase.h
--- a/Engine/Plugins/Experimental/Mutable/Source/CustomizableObjectEditor/Public/MuCOE/Nodes/CustomizableObjectNodeMaterialBase.h
+++ b/Engine/Plugins/Experimental/Mutable/Source/CustomizableObjectEditor/Public/MuCOE/Nodes/CustomizableObjectNodeMaterialBase.h
@@ -41,7 +41,7 @@ public:
 	
 	virtual bool IsReuseMaterialBetweenLODs() const PURE_VIRTUAL(UCustomizableObjectNodeMaterial::IsReuseMaterialBetweenLODs, return {}; );
 
-	virtual int32 GetMeshComponentIndex() const PURE_VIRTUAL(UCustomizableObjectNodeMaterial::GetMeshComponentIndex, return {}; );
+	virtual FName GetMeshComponentName() const PURE_VIRTUAL(UCustomizableObjectNodeMaterial::GetComponentName, return {}; );
 	
 	virtual TArray<FString> GetTags() const PURE_VIRTUAL(UCustomizableObjectNodeMaterial::GetTags, return {}; );
 	
@@ -126,6 +126,7 @@ public:
 	
 	virtual FPostImagePinModeChangedDelegate* GetPostImagePinModeChangedDelegate() PURE_VIRTUAL(UCustomizableObjectNodeMaterial::OutputPin, return {}; );
 	
+	/** Return true if a Material Parameter has changed on which we had a pin connected or data saved. */
 	virtual bool RealMaterialDataHasChanged() const PURE_VIRTUAL(UCustomizableObjectNodeMaterial::OutputPin, return {}; );
 };
 
diff --git a/Engine/Plugins/Experimental/Mutable/Source/CustomizableObjectEditor/Public/MuCOE/Nodes/CustomizableObjectNodeObject.h b/Engine/Plugins/Experimental/Mutable/Source/CustomizableObjectEditor/Public/MuCOE/Nodes/CustomizableObjectNodeObject.h
--- a/Engine/Plugins/Experimental/Mutable/Source/CustomizableObjectEditor/Public/MuCOE/Nodes/CustomizableObjectNodeObject.h
+++ b/Engine/Plugins/Experimental/Mutable/Source/CustomizableObjectEditor/Public/MuCOE/Nodes/CustomizableObjectNodeObject.h
@@ -48,6 +48,9 @@ struct FComponentSettings
 {
 	GENERATED_USTRUCT_BODY()
 	
+	UPROPERTY(VisibleInstanceOnly, Category = CustomizableObject)
+	FString ComponentName;
+
 	UPROPERTY(EditAnywhere, Category = CustomizableObject)
 	TArray<FLODReductionSettings> LODReductionSettings;
 };
@@ -175,8 +178,8 @@ public:
 	UPROPERTY(EditAnywhere, Category = CustomizableObject)
 	ECustomizableObjectAutomaticLODStrategy AutoLODStrategy = ECustomizableObjectAutomaticLODStrategy::AutomaticFromMesh;
 
-	UPROPERTY(EditAnywhere, Category = CustomizableObject, meta = (ClampMin = "1"))
-	int32 NumMeshComponents = 1;
+	UPROPERTY()
+	int32 NumMeshComponents_DEPRECATED = 1;
 
 	UPROPERTY(EditAnywhere, Category=CustomizableObject ,meta = (TitleProperty = "Name"))
 	TArray<FCustomizableObjectState> States;
@@ -207,7 +210,7 @@ public:
 	TArray<FComponentSettings> ComponentSettings;
 	
 	// UObject interface.
-	void PostEditChangeProperty(FPropertyChangedEvent& PropertyChangedEvent) override;
+	virtual void PostEditChangeProperty(FPropertyChangedEvent& PropertyChangedEvent) override;
 
 	// EdGraphNode interface
 	FText GetNodeTitle(ENodeTitleType::Type TitleType) const override;
@@ -279,13 +282,8 @@ public:
 
 	bool IsSingleOutputNode() const override;
 
-	void SetMeshComponentNumFromParent(int32 MeshComponentNum)
-	{
-		NumMeshComponents = MeshComponentNum;
-	}
-
 	// Node Details Support
-	int32 CurrentComponent = 0;
+	FName CurrentComponent;
 	int32 CurrentLOD = 0;
 
 	// Array filled in the Details of the node to store all the parameter names of a CO graph (full tree)
