[Mutable] Add surface and mesh metadata to the CO side. RemoveMesh operation optimizations.

- Now mesh section bCastShadow property is set form data gathered at compile time.
- Optimizations to the remove mesh operations.
- Now remove mesh operations recompute mesh surfaces.
- ClipMorph operation now removes unused vertices. Before only index data was modified.

#rnx

diff --git a/Engine/Plugins/Experimental/Mutable/Source/CustomizableObject/Internal/MuCO/CustomizableObjectPrivate.h b/Engine/Plugins/Experimental/Mutable/Source/CustomizableObject/Internal/MuCO/CustomizableObjectPrivate.h
--- a/Engine/Plugins/Experimental/Mutable/Source/CustomizableObject/Internal/MuCO/CustomizableObjectPrivate.h
+++ b/Engine/Plugins/Experimental/Mutable/Source/CustomizableObject/Internal/MuCO/CustomizableObjectPrivate.h
@@ -510,6 +510,43 @@ struct FRealTimeMorphStreamable
 	}
 };
 
+USTRUCT()
+struct FMutableMeshMetadata
+{
+	GENERATED_USTRUCT_BODY()
+	
+	UPROPERTY()
+	uint32 MorphMetadataId = 0;
+
+	UPROPERTY()
+	uint32 ClothingMetadataId = 0;
+
+	friend FArchive& operator<<(FArchive& Ar, FMutableMeshMetadata& Elem)
+	{
+		Ar << Elem.MorphMetadataId;
+		Ar << Elem.ClothingMetadataId;
+
+		return Ar;
+	}
+};
+
+
+USTRUCT()
+struct FMutableSurfaceMetadata
+{
+	GENERATED_USTRUCT_BODY()
+	
+	UPROPERTY()
+	bool bCastShadow = false;
+
+	friend FArchive& operator<<(FArchive& Ar, FMutableSurfaceMetadata& Elem)
+	{
+		Ar << Elem.bCastShadow;
+
+		return Ar;
+	}
+};
+
 
 USTRUCT()
 struct FClothingStreamable
@@ -721,6 +758,12 @@ struct FModelResources
 	UPROPERTY()
 	TArray<FMutableModelImageProperties> ImageProperties;
 
+	UPROPERTY()
+	TMap<uint32, FMutableMeshMetadata> MeshMetadata;
+
+	UPROPERTY()
+	TMap<uint32, FMutableSurfaceMetadata> SurfaceMetadata;
+
 	/** Parameter UI metadata information for all the dependencies of this Customizable Object. */
 	UPROPERTY()
 	TMap<FString, FMutableParameterData> ParameterUIDataMap;
@@ -1117,6 +1160,8 @@ public:
 
 		LayoutNodeCleanup,
 
+		AddSurfaceAndMeshMetadata,
+
 		// -----<new versions can be added above this line>--------
 		LastCustomizableObjectVersion
 	};
diff --git a/Engine/Plugins/Experimental/Mutable/Source/CustomizableObject/Internal/MuCO/UnrealConversionUtils.h b/Engine/Plugins/Experimental/Mutable/Source/CustomizableObject/Internal/MuCO/UnrealConversionUtils.h
--- a/Engine/Plugins/Experimental/Mutable/Source/CustomizableObject/Internal/MuCO/UnrealConversionUtils.h
+++ b/Engine/Plugins/Experimental/Mutable/Source/CustomizableObject/Internal/MuCO/UnrealConversionUtils.h
@@ -18,12 +18,12 @@ namespace mu
 	struct FBoneName;
 }
 struct FReferenceSkeleton;
+struct FMutableSurfaceMetadata;
 
 class FSkeletalMeshLODRenderData;
 class USkeletalMesh;
 class USkeleton;
 
-
 namespace UnrealConversionUtils
 {
 	/*
@@ -38,13 +38,15 @@ namespace UnrealConversionUtils
 	 * @param InMutableMesh - Mutable mesh to be used as reference for the section data update on the Skeletal Mesh
 	 * @param InBoneMap - Bones to be set as part of the sections.
 	 * @param InFirstBoneMapIndex - Index to the first BoneMap bone that belongs to this LODResource.
+	 * @param SectionMetadata - Section metadata for each surface in InMutableMesh. 
 	 */
 	CUSTOMIZABLEOBJECT_API void SetupRenderSections(
 		FSkeletalMeshLODRenderData& LODResource,
 		const mu::Ptr<const mu::Mesh> InMutableMesh,
 		const TArray<mu::FBoneName>& InBoneMap,
 		const TMap<mu::FBoneName, TPair<FName, uint16>>& BoneInfoMap,
-		const int32 InFirstBoneMapIndex);
+		const int32 InFirstBoneMapIndex,
+		const TArray<const FMutableSurfaceMetadata*>& SectionMetadata);
 
 
 	/** Initializes the LODResource's VertexBuffers with dummy data to prepare it for streaming.
diff --git a/Engine/Plugins/Experimental/Mutable/Source/CustomizableObject/Private/MuCO/CustomizableObjectInstance.cpp b/Engine/Plugins/Experimental/Mutable/Source/CustomizableObject/Private/MuCO/CustomizableObjectInstance.cpp
--- a/Engine/Plugins/Experimental/Mutable/Source/CustomizableObject/Private/MuCO/CustomizableObjectInstance.cpp
+++ b/Engine/Plugins/Experimental/Mutable/Source/CustomizableObject/Private/MuCO/CustomizableObjectInstance.cpp
@@ -3990,8 +3990,8 @@ void UCustomizableInstancePrivate::BuildOrCopyClothingData(const TSharedRef<FUpd
 				const int32 SurfaceCount = MutableMesh->GetSurfaceCount();
 				for (int32 Section = 0; Section < SurfaceCount; ++Section)
 				{
-					int32 FirstVertex, VerticesCount, FirstIndex, IndicesCount;
-					MutableMesh->GetSurface(Section, &FirstVertex, &VerticesCount, &FirstIndex, &IndicesCount, nullptr, nullptr, nullptr);
+					int32 FirstVertex, VerticesCount, FirstIndex, IndicesCount, UnusedBoneIndex, UnusedBoneCount;
+					MutableMesh->GetSurface(Section, FirstVertex, VerticesCount, FirstIndex, IndicesCount, UnusedBoneIndex, UnusedBoneCount);
 
 					if (VerticesCount == 0 || IndicesCount == 0)
 					{
@@ -4699,8 +4699,8 @@ void UCustomizableInstancePrivate::BuildOrCopyClothingData(const TSharedRef<FUpd
 				for (int32 Section = 0; Section < SurfaceCount; ++Section)
 				{
 					// Check that is a valid surface.
-					int32 FirstVertex, VerticesCount, FirstIndex, IndicesCount;
-					MutableMesh->GetSurface(Section, &FirstVertex, &VerticesCount, &FirstIndex, &IndicesCount, nullptr, nullptr, nullptr);
+					int32 FirstVertex, VerticesCount, FirstIndex, IndicesCount, UnusedBoneIndex, UnusedBoneCount;
+					MutableMesh->GetSurface(Section, FirstVertex, VerticesCount, FirstIndex, IndicesCount, UnusedBoneIndex, UnusedBoneCount);
 
 					if (VerticesCount == 0 || IndicesCount == 0)
 					{
@@ -4929,13 +4929,34 @@ bool UCustomizableInstancePrivate::BuildOrCopyRenderData(const TSharedRef<FUpdat
 			LODResource.RequiredBones.Sort();
 		}
 
+
+		// Find referenced surface metadata.
+		const int32 MeshNumSurfaces = LOD.Mesh->Surfaces.Num();
+		TArray<const FMutableSurfaceMetadata*> MeshSurfacesMetadata;
+		MeshSurfacesMetadata.Init(nullptr, MeshNumSurfaces);
+
+		for (int32 MeshSectionIndex = 0; MeshSectionIndex < MeshNumSurfaces; ++MeshSectionIndex)
+		{
+			uint32 MeshSurfaceId = LOD.Mesh->GetSurfaceId(MeshSectionIndex);
+			uint32 InstanceSurfaceIndex = OperationData->MutableInstance->FindSurfaceById(ComponentIndex, LODIndex, MeshSurfaceId);
+			
+			if (InstanceSurfaceIndex < 0)
+			{
+				continue;
+			}
+			
+			uint32 SurfaceMetadataId = OperationData->MutableInstance->GetSurfaceCustomId(ComponentIndex, LODIndex, InstanceSurfaceIndex);
+			MeshSurfacesMetadata[MeshSectionIndex] = ModelResources.SurfaceMetadata.Find(SurfaceMetadataId);
+		}
+
 		// Set RenderSections
 		UnrealConversionUtils::SetupRenderSections(
 			LODResource,
 			LOD.Mesh,
 			OperationData->InstanceUpdateData.BoneMaps,
 			BoneInfoMap,
-			LOD.FirstBoneMap);
+			LOD.FirstBoneMap,
+			MeshSurfacesMetadata);
 
 		if (LODResource.bStreamedDataInlined) // Non-streamable LOD
 		{
diff --git a/Engine/Plugins/Experimental/Mutable/Source/CustomizableObject/Private/MuCO/CustomizableObjectSystem.cpp b/Engine/Plugins/Experimental/Mutable/Source/CustomizableObject/Private/MuCO/CustomizableObjectSystem.cpp
--- a/Engine/Plugins/Experimental/Mutable/Source/CustomizableObject/Private/MuCO/CustomizableObjectSystem.cpp
+++ b/Engine/Plugins/Experimental/Mutable/Source/CustomizableObject/Private/MuCO/CustomizableObjectSystem.cpp
@@ -1936,8 +1936,8 @@ namespace impl
 						Elem = 0;
 					}
 
-					int32 FirstVertex, VerticesCount, FirstIndex, IndiciesCount;
-					LOD.Mesh->GetSurface(Section, &FirstVertex, &VerticesCount, &FirstIndex, &IndiciesCount, nullptr, nullptr, nullptr);
+					int32 FirstVertex, VerticesCount, FirstIndex, IndiciesCount, UnusedBoneIndex, UnusedBoneCount;
+					LOD.Mesh->GetSurface(Section, FirstVertex, VerticesCount, FirstIndex, IndiciesCount, UnusedBoneIndex, UnusedBoneCount);
 
 					for (int32 VertexIdx = FirstVertex; VertexIdx < FirstVertex + VerticesCount;)
 					{
diff --git a/Engine/Plugins/Experimental/Mutable/Source/CustomizableObject/Private/MuCO/UnrealConversionUtils.cpp b/Engine/Plugins/Experimental/Mutable/Source/CustomizableObject/Private/MuCO/UnrealConversionUtils.cpp
--- a/Engine/Plugins/Experimental/Mutable/Source/CustomizableObject/Private/MuCO/UnrealConversionUtils.cpp
+++ b/Engine/Plugins/Experimental/Mutable/Source/CustomizableObject/Private/MuCO/UnrealConversionUtils.cpp
@@ -7,6 +7,7 @@
 #include "Engine/SkeletalMesh.h"
 #include "GPUSkinVertexFactory.h"
 #include "MuCO/CustomizableObject.h"
+#include "MuCO/CustomizableObjectPrivate.h"
 #include "MuR/Skeleton.h"
 #include "MuR/Mesh.h"
 #include "MuR/MeshBufferSet.h"
@@ -165,7 +166,8 @@ namespace UnrealConversionUtils
 		const mu::MeshPtrConst InMutableMesh,
 		const TArray<mu::FBoneName>& InBoneMap,
 		const TMap<mu::FBoneName, TPair<FName, uint16>>& BoneInfoMap,
-		const int32 InFirstBoneMapIndex)
+		const int32 InFirstBoneMapIndex,
+		const TArray<const FMutableSurfaceMetadata*>& SurfacesMetadata)
 	{
 		check(InMutableMesh);
 
@@ -193,8 +195,7 @@ namespace UnrealConversionUtils
 			int32 VertexCount;
 			int32 FirstBone;
 			int32 BoneCount;
-			bool bCastShadow;
-			InMutableMesh->GetSurface(SurfaceIndex, &FirstVertex, &VertexCount, &FirstIndex, &IndexCount, &FirstBone, &BoneCount, &bCastShadow);
+			InMutableMesh->GetSurface(SurfaceIndex, FirstVertex, VertexCount, FirstIndex, IndexCount, FirstBone, BoneCount);
 			FSkelMeshRenderSection& Section = LODResource.RenderSections[SurfaceIndex];
 
 			Section.DuplicatedVerticesBuffer.Init(1, TMap<int, TArray<int32>>());
@@ -211,8 +212,11 @@ namespace UnrealConversionUtils
 			Section.MaxBoneInfluences = NumBoneInfluences;
 			Section.NumVertices = VertexCount;
 
-			//TODO(Max): MTBL-1779
-			//Section.bCastShadow = bCastShadow;
+			check(SurfacesMetadata.Num() == InMutableMesh->Surfaces.Num());
+			if (SurfacesMetadata[SurfaceIndex])
+			{
+				Section.bCastShadow = SurfacesMetadata[SurfaceIndex]->bCastShadow;
+			}
 
 			// InBoneMaps may contain bonemaps from other sections. Copy the bones belonging to this mesh.
 			FirstBone += InFirstBoneMapIndex;
diff --git a/Engine/Plugins/Experimental/Mutable/Source/CustomizableObjectEditor/Private/MuCOE/CustomizableObjectCompiler.cpp b/Engine/Plugins/Experimental/Mutable/Source/CustomizableObjectEditor/Private/MuCOE/CustomizableObjectCompiler.cpp
--- a/Engine/Plugins/Experimental/Mutable/Source/CustomizableObjectEditor/Private/MuCOE/CustomizableObjectCompiler.cpp
+++ b/Engine/Plugins/Experimental/Mutable/Source/CustomizableObjectEditor/Private/MuCOE/CustomizableObjectCompiler.cpp
@@ -1101,6 +1101,9 @@ void FCustomizableObjectCompiler::CompileInternal(bool bAsync)
 			 ClothingAssetData.ConfigsData.RemoveAllSwap(IsSharedConfigData);
 		}
 
+		ModelResources.MeshMetadata = MoveTemp(GenerationContext.MeshMetadata);
+		ModelResources.SurfaceMetadata = MoveTemp(GenerationContext.SurfaceMetadata);
+
 		CurrentObject->GetPrivate()->GroupNodeMap = GenerationContext.GroupNodeMap;
 
 		if (GenerationContext.Options.OptimizationLevel == 0)
diff --git a/Engine/Plugins/Experimental/Mutable/Source/CustomizableObjectEditor/Private/MuCOE/GenerateMutableSource/GenerateMutableSource.h b/Engine/Plugins/Experimental/Mutable/Source/CustomizableObjectEditor/Private/MuCOE/GenerateMutableSource/GenerateMutableSource.h
--- a/Engine/Plugins/Experimental/Mutable/Source/CustomizableObjectEditor/Private/MuCOE/GenerateMutableSource/GenerateMutableSource.h
+++ b/Engine/Plugins/Experimental/Mutable/Source/CustomizableObjectEditor/Private/MuCOE/GenerateMutableSource/GenerateMutableSource.h
@@ -96,7 +96,10 @@ enum class EMutableMeshConversionFlags : uint32
 	IgnoreSkinning = 1 << 0,
 
 	// Ignore Physics assets
-	IgnorePhysics = 1 << 1
+	IgnorePhysics = 1 << 1,
+
+	// Prevent this mesh generation from adding per mesh metadata. 
+	DoNotCreateMeshMetadata = 1 << 2
 };
 
 ENUM_CLASS_FLAGS(EMutableMeshConversionFlags)
@@ -777,6 +780,9 @@ struct FMutableGraphGenerationContext
 	TArray<FCustomizableObjectClothingAssetData> ClothingAssetsData;
 	TMap<uint32, FClothingMeshData> ClothingPerMeshData;
 
+	TMap<uint32, FMutableMeshMetadata> MeshMetadata;
+	TMap<uint32, FMutableSurfaceMetadata> SurfaceMetadata;
+
 	// Data used for SkinWeightProfiles reconstruction
 	TArray<FMutableSkinWeightProfileInfo> SkinWeightProfilesInfo;
 
@@ -907,6 +913,61 @@ private:
 	FScopedPinData ScopedPinData(Context, Pin);
 
 
+namespace Private
+{
+	template<class HashableType, class HashDataSetType, class HashFuncType, class CompareFuncType>
+	uint32 GenerateUniquePersistentHash(const HashableType& HashableData, const HashDataSetType& HashDataSet, HashFuncType&& HashFunc, CompareFuncType&& CompareFunc)
+	{
+		constexpr uint32 InvalidResourceId = 0;
+		
+		const uint32 DataHash = HashFunc(HashableData);
+
+		uint32 UniqueHash = DataHash == InvalidResourceId ? DataHash + 1 : DataHash;
+
+		const HashableType* FoundHash = HashDataSet.Find(UniqueHash);
+
+		bool bIsDataAlreadyCollected = false;
+		
+		if (FoundHash)
+		{
+			bIsDataAlreadyCollected = CompareFunc(*FoundHash, HashableData); 
+		}
+
+		// NOTE: This way of unique hash generation guarantees all valid values can be used but given its 
+		// sequential nature a cascade of changes can occur if new meshes are added. Not many hash collisions 
+		// are expected so it should not be problematic.
+		if (FoundHash && !bIsDataAlreadyCollected)
+		{
+			uint32 NumTries = 0;
+			for (; NumTries < TNumericLimits<uint32>::Max(); ++NumTries)
+			{
+				FoundHash = HashDataSet.Find(UniqueHash);
+				
+				if (!FoundHash)
+				{
+					break;
+				}
+
+				bIsDataAlreadyCollected = CompareFunc(*FoundHash, HashableData);
+
+				if (bIsDataAlreadyCollected)
+				{
+					break;
+				}
+
+				UniqueHash = UniqueHash + 1 == InvalidResourceId ? InvalidResourceId + 1 : UniqueHash + 1;
+			}
+
+			if (NumTries == TNumericLimits<uint32>::Max())
+			{
+				UniqueHash = InvalidResourceId;
+			}	
+		}
+
+		return UniqueHash;
+	}
+} //Private
+
 //
 mu::Ptr<mu::NodeObject> GenerateMutableSource(const class UEdGraphPin* Pin, FMutableGraphGenerationContext& GenerationContext, bool bPartialCompilation);
 mu::Ptr<mu::NodeComponent> GenerateMutableSourceComponent(const UEdGraphPin* Pin, FMutableGraphGenerationContext& GenerationContext);
diff --git a/Engine/Plugins/Experimental/Mutable/Source/CustomizableObjectEditor/Private/MuCOE/GenerateMutableSource/GenerateMutableSourceMesh.cpp b/Engine/Plugins/Experimental/Mutable/Source/CustomizableObjectEditor/Private/MuCOE/GenerateMutableSource/GenerateMutableSourceMesh.cpp
--- a/Engine/Plugins/Experimental/Mutable/Source/CustomizableObjectEditor/Private/MuCOE/GenerateMutableSource/GenerateMutableSourceMesh.cpp
+++ b/Engine/Plugins/Experimental/Mutable/Source/CustomizableObjectEditor/Private/MuCOE/GenerateMutableSource/GenerateMutableSourceMesh.cpp
@@ -727,6 +727,7 @@ mu::MeshPtr ConvertSkeletalMeshToMutable(const USkeletalMesh* InSkeletalMesh, co
 	const EMutableMeshConversionFlags CurrentFlags = GenerationContext.MeshGenerationFlags.Last();
 	const bool bIgnoreSkeleton = EnumHasAnyFlags(CurrentFlags, EMutableMeshConversionFlags::IgnoreSkinning);
 	const bool bIgnorePhysics = EnumHasAnyFlags(CurrentFlags, EMutableMeshConversionFlags::IgnorePhysics);
+	const bool bDoNotCreateMeshMetadata = EnumHasAnyFlags(CurrentFlags, EMutableMeshConversionFlags::DoNotCreateMeshMetadata);
 
 	mu::MeshPtr MutableMesh = new mu::Mesh();
 
@@ -1856,14 +1857,6 @@ mu::MeshPtr ConvertSkeletalMeshToMutable(const USkeletalMesh* InSkeletalMesh, co
 		}
 	}
 
-	// Ensure Surface Data
-	mu::FMeshSurface MeshSurface;
-	MeshSurface.VertexCount = MutableMesh->VertexBuffers.GetElementCount();
-	MeshSurface.IndexCount = MutableMesh->IndexBuffers.GetElementCount();
-	MeshSurface.BoneMapCount = MutableMesh->BoneMap.Num();
-	MeshSurface.bCastShadow = MeshSection.bCastShadow;
-	MutableMesh->Surfaces.Add(MeshSurface);
-
 	if (!bIgnorePhysics && InSkeletalMesh->GetPhysicsAsset() && MutableMesh->GetSkeleton() && GenerationContext.Options.bPhysicsAssetMergeEnabled)
 	{
 		// Find BodySetups with relevant bones.
@@ -2105,6 +2098,52 @@ mu::MeshPtr ConvertSkeletalMeshToMutable(const USkeletalMesh* InSkeletalMesh, co
 		}
 	}
 
+	if (!bDoNotCreateMeshMetadata)
+	{
+		FMutableMeshMetadata MeshMetadata;
+		FMemory::Memzero(MeshMetadata);
+
+		// TODO: Store morph and cloth metadata in this struct, currenlty its added per vertex. With per mesh metadata
+		// this info can be stored in the CO per mesh.
+		MeshMetadata.MorphMetadataId = 0;
+		MeshMetadata.ClothingMetadataId = 0;
+
+		auto HashMeshMetadataFunc = [](const FMutableMeshMetadata& Data) -> uint32
+		{
+			return CityHash32(reinterpret_cast<const char*>(&Data), sizeof(FMutableMeshMetadata));
+		};
+
+		auto CompareMeshMetadataFunc = [](const FMutableMeshMetadata& A, const FMutableMeshMetadata& B) -> bool
+		{
+			return FMemory::Memcmp(&A, &B, sizeof(FMutableMeshMetadata)) == 0;
+		};
+
+		const uint32 MeshMetadataUniqueHash = Private::GenerateUniquePersistentHash(
+				MeshMetadata, GenerationContext.MeshMetadata, HashMeshMetadataFunc, CompareMeshMetadataFunc);
+		
+		if (MeshMetadataUniqueHash != 0)
+		{
+			GenerationContext.MeshMetadata.FindOrAdd(MeshMetadataUniqueHash, MeshMetadata);
+		}
+		else
+		{
+			UE_LOG(LogMutable, Error, TEXT("Maximum number of meshes reached."));
+		}
+
+		// Ensure Surface Data
+		mu::FMeshSurface& MeshSurface = MutableMesh->Surfaces.Emplace_GetRef();
+	
+		MeshSurface.SubMeshes.Emplace(
+				mu::FSurfaceSubMesh 
+				{
+					0, MutableMesh->VertexBuffers.GetElementCount(), 
+					0, MutableMesh->IndexBuffers.GetElementCount(),
+					MeshMetadataUniqueHash	
+				});
+
+		MeshSurface.BoneMapCount = MutableMesh->BoneMap.Num();
+	}
+
 	return MutableMesh;
 }
 
diff --git a/Engine/Plugins/Experimental/Mutable/Source/CustomizableObjectEditor/Private/MuCOE/GenerateMutableSource/GenerateMutableSourceSurface.cpp b/Engine/Plugins/Experimental/Mutable/Source/CustomizableObjectEditor/Private/MuCOE/GenerateMutableSource/GenerateMutableSourceSurface.cpp
--- a/Engine/Plugins/Experimental/Mutable/Source/CustomizableObjectEditor/Private/MuCOE/GenerateMutableSource/GenerateMutableSourceSurface.cpp
+++ b/Engine/Plugins/Experimental/Mutable/Source/CustomizableObjectEditor/Private/MuCOE/GenerateMutableSource/GenerateMutableSourceSurface.cpp
@@ -13,6 +13,7 @@
 #include "TextureCompressorModule.h"
 #include "Materials/MaterialInstance.h"
 #include "GPUSkinVertexFactory.h"
+#include "Rendering/SkeletalMeshLODModel.h"
 
 #include "MuCO/CustomizableObjectInstance.h"
 #include "MuCO/MutableMeshBufferUtils.h"
@@ -286,6 +287,7 @@ mu::Ptr<mu::NodeSurface> GenerateMutableSourceSurface(const UEdGraphPin * Pin, F
 		SharedSurface.bMakeUnique = !TypedNodeMat->IsReuseMaterialBetweenLODs();
 
 		int32 ReferencedMaterialsIndex = -1;
+		uint32 SurfaceMetadataUniqueHash = 0; // Value 0 is used as invalid hash.
 		if (TypedNodeMat->GetMaterial())
 		{
 			GenerationContext.AddParticipatingObject(*TypedNodeMat->GetMaterial());
@@ -307,15 +309,18 @@ mu::Ptr<mu::NodeSurface> GenerateMutableSourceSurface(const UEdGraphPin * Pin, F
 					if (const UEdGraphPin* SkeletalMeshPin = FindMeshBaseSource(*ConnectedMaterialPin, false))
 					{
 						FSkeletalMaterial* SkeletalMaterial = nullptr;
-
+						const FSkelMeshSection* ReferenceSkelMeshSection = nullptr;
+						
 						if (const UCustomizableObjectNodeSkeletalMesh* SkeletalMeshNode = Cast<UCustomizableObjectNodeSkeletalMesh>(SkeletalMeshPin->GetOwningNode()))
 						{
 							SkeletalMaterial = SkeletalMeshNode->GetSkeletalMaterialFor(*SkeletalMeshPin);
+							ReferenceSkelMeshSection = SkeletalMeshNode->GetSkeletalMeshSectionFor(*SkeletalMeshPin);
 						}
 
 						else if (const UCustomizableObjectNodeTable* TableNode = Cast<UCustomizableObjectNodeTable>(SkeletalMeshPin->GetOwningNode()))
 						{
 							SkeletalMaterial = TableNode->GetDefaultSkeletalMaterialFor(*SkeletalMeshPin);
+							ReferenceSkelMeshSection = TableNode->GetDefaultSkeletalMeshSectionFor(*SkeletalMeshPin);
 						}
 
 						if (SkeletalMaterial)
@@ -331,9 +336,40 @@ mu::Ptr<mu::NodeSurface> GenerateMutableSourceSurface(const UEdGraphPin * Pin, F
 								SlotNameFound = true;
 							}
 						}
+
+						if (ReferenceSkelMeshSection)
+						{
+							auto HashSurfaceMetadataFunc = [](const FMutableSurfaceMetadata& Data) -> uint32
+							{
+								return CityHash32(reinterpret_cast<const char*>(&Data), sizeof(FMutableSurfaceMetadata));
+							};
+
+							auto CompareSurfaceMetadataFunc = [](const FMutableSurfaceMetadata& A, const FMutableSurfaceMetadata& B)
+							{
+								return FMemory::Memcmp(&A, &B, sizeof(FMutableSurfaceMetadata)) == 0;
+							};
+
+							FMutableSurfaceMetadata SurfaceMetadata;
+							FMemory::Memzero(SurfaceMetadata);
+							
+							SurfaceMetadata.bCastShadow = ReferenceSkelMeshSection->bCastShadow;
+
+							SurfaceMetadataUniqueHash = Private::GenerateUniquePersistentHash(
+									SurfaceMetadata, GenerationContext.SurfaceMetadata, HashSurfaceMetadataFunc, CompareSurfaceMetadataFunc);
+
+							if (SurfaceMetadataUniqueHash != 0)
+							{
+								GenerationContext.SurfaceMetadata.FindOrAdd(SurfaceMetadataUniqueHash, SurfaceMetadata);
+							}
+							else
+							{
+								UE_LOG(LogMutable, Error, TEXT("Maximum number of surfaces reached."));
+							}
+						}
 					}
 				}
 
+
 				// No name was found, we need to keep index parity, so we add empty value. We may find the value later.
 				if (IsNewSlotName && !SlotNameFound)
 				{
@@ -375,6 +411,7 @@ mu::Ptr<mu::NodeSurface> GenerateMutableSourceSurface(const UEdGraphPin * Pin, F
 
 					MeshFormatNode->SetMessageContext(Node);
 
+					SurfNode->SetCustomID(SurfaceMetadataUniqueHash);
 					SurfNode->SetMesh(MeshFormatNode);
 				}
 				else
@@ -982,7 +1019,7 @@ mu::Ptr<mu::NodeSurface> GenerateMutableSourceSurface(const UEdGraphPin * Pin, F
 		{
 			mu::NodeSurfaceNewPtr SurfNode2 = new mu::NodeSurfaceNew;
 
-			SurfNode2->SetCustomID(ReferencedMaterialsIndex);
+			SurfNode2->SetCustomID(SurfaceMetadataUniqueHash);
 
 			SurfNode2->SetMesh(SurfNode->GetMesh());
 
diff --git a/Engine/Plugins/Experimental/Mutable/Source/CustomizableObjectEditor/Private/MuCOE/MutableMeshPreviewUtils.cpp b/Engine/Plugins/Experimental/Mutable/Source/CustomizableObjectEditor/Private/MuCOE/MutableMeshPreviewUtils.cpp
--- a/Engine/Plugins/Experimental/Mutable/Source/CustomizableObjectEditor/Private/MuCOE/MutableMeshPreviewUtils.cpp
+++ b/Engine/Plugins/Experimental/Mutable/Source/CustomizableObjectEditor/Private/MuCOE/MutableMeshPreviewUtils.cpp
@@ -382,6 +382,9 @@ namespace MutableMeshPreviewUtils
 
 			FSkeletalMeshLODRenderData& LODResource = OutSkeletalMesh->GetResourceForRendering()->LODRenderData[0];
 
+			TArray<const FMutableSurfaceMetadata*> MeshSurfacesMetadata;
+			MeshSurfacesMetadata.Init(nullptr, InMutableMesh->Surfaces.Num());
+
 			// Load buffer data found on the mutable model onto the out skeletal mesh
 			// It includes vertex and index buffers
 			UnrealConversionUtils::SetupRenderSections(
@@ -389,7 +392,8 @@ namespace MutableMeshPreviewUtils
 				InMutableMesh,
 				BoneMap,
 				BoneInfoMap,
-				0);
+				0,
+				MeshSurfacesMetadata);
 
 			UnrealConversionUtils::CopyMutableVertexBuffers(
 				LODResource,
diff --git a/Engine/Plugins/Experimental/Mutable/Source/CustomizableObjectEditor/Private/MuCOE/Nodes/CustomizableObjectNodeSkeletalMesh.cpp b/Engine/Plugins/Experimental/Mutable/Source/CustomizableObjectEditor/Private/MuCOE/Nodes/CustomizableObjectNodeSkeletalMesh.cpp
--- a/Engine/Plugins/Experimental/Mutable/Source/CustomizableObjectEditor/Private/MuCOE/Nodes/CustomizableObjectNodeSkeletalMesh.cpp
+++ b/Engine/Plugins/Experimental/Mutable/Source/CustomizableObjectEditor/Private/MuCOE/Nodes/CustomizableObjectNodeSkeletalMesh.cpp
@@ -386,6 +386,15 @@ int32 UCustomizableObjectNodeSkeletalMesh::GetSkeletalMaterialIndexFor(const UEd
 	return GetSkeletalMaterialIndexFor(LODIndex, SectionIndex);
 }
 
+const FSkelMeshSection* UCustomizableObjectNodeSkeletalMesh::GetSkeletalMeshSectionFor(const UEdGraphPin& Pin) const
+{
+	int32 LODIndex;
+	int32 SectionIndex;
+	int32 LayoutIndex;
+	GetPinSection(Pin, LODIndex, SectionIndex, LayoutIndex);
+	
+	return GetSkeletalMeshSectionFor(LODIndex, SectionIndex);
+}
 
 bool UCustomizableObjectNodeSkeletalMesh::IsPinRelevant(const UEdGraphPin* Pin) const
 {
@@ -739,6 +748,33 @@ FSkeletalMaterial* UCustomizableObjectNodeSkeletalMesh::GetSkeletalMaterialFor(c
 	return nullptr;
 }
 
+const FSkelMeshSection* UCustomizableObjectNodeSkeletalMesh::GetSkeletalMeshSectionFor(const int32 LODIndex, const int32 SectionIndex) const
+{	
+	if (!SkeletalMesh)
+	{
+		return nullptr;
+	}
+	const FSkeletalMeshModel* ImportedModel = SkeletalMesh->GetImportedModel();
+
+	if (!ImportedModel)
+	{
+		return nullptr;
+	}
+
+	if (LODIndex < 0 && LODIndex >= ImportedModel->LODModels.Num())
+	{
+		return nullptr;
+	}
+
+	const FSkeletalMeshLODModel& LODModel = ImportedModel->LODModels[LODIndex];
+
+	if (SectionIndex < 0 && SectionIndex >= LODModel.Sections.Num())
+	{
+		return nullptr;
+	}
+
+	return &LODModel.Sections[SectionIndex];
+}
 
 int32 UCustomizableObjectNodeSkeletalMesh::GetSkeletalMaterialIndexFor(const int32 LODIndex, const int32 SectionIndex) const
 {
diff --git a/Engine/Plugins/Experimental/Mutable/Source/CustomizableObjectEditor/Private/MuCOE/Nodes/CustomizableObjectNodeTable.cpp b/Engine/Plugins/Experimental/Mutable/Source/CustomizableObjectEditor/Private/MuCOE/Nodes/CustomizableObjectNodeTable.cpp
--- a/Engine/Plugins/Experimental/Mutable/Source/CustomizableObjectEditor/Private/MuCOE/Nodes/CustomizableObjectNodeTable.cpp
+++ b/Engine/Plugins/Experimental/Mutable/Source/CustomizableObjectEditor/Private/MuCOE/Nodes/CustomizableObjectNodeTable.cpp
@@ -1480,6 +1480,39 @@ TArray<FAssetData> UCustomizableObjectNodeTable::GetParentTables() const
 	return DataTableAssets;
 }
 
+const FSkelMeshSection* UCustomizableObjectNodeTable::GetDefaultSkeletalMeshSectionFor(const UEdGraphPin& MeshPin) const
+{
+	USkeletalMesh* SkeletalMesh = GetColumnDefaultAssetByType<USkeletalMesh>(&MeshPin);
+
+	if (!SkeletalMesh)
+	{
+		return nullptr;
+	}
+	const FSkeletalMeshModel* ImportedModel = SkeletalMesh->GetImportedModel();
+
+	if (!ImportedModel)
+	{
+		return nullptr;
+	}
+
+	int32 LODIndex;
+	int32 SectionIndex;
+	GetPinLODAndSection(&MeshPin, LODIndex, SectionIndex);
+	
+	if (!ImportedModel->LODModels.IsValidIndex(LODIndex))
+	{
+		return nullptr;
+	}
+
+	const FSkeletalMeshLODModel& LODModel = ImportedModel->LODModels[LODIndex];
+
+	if (!LODModel.Sections.IsValidIndex(SectionIndex))
+	{
+		return nullptr;
+	}
+
+	return &LODModel.Sections[SectionIndex];
+}
 
 FSkeletalMaterial* UCustomizableObjectNodeTable::GetDefaultSkeletalMaterialFor(const UEdGraphPin& MeshPin) const
 {
diff --git a/Engine/Plugins/Experimental/Mutable/Source/CustomizableObjectEditor/Public/MuCOE/Nodes/CustomizableObjectNodeSkeletalMesh.h b/Engine/Plugins/Experimental/Mutable/Source/CustomizableObjectEditor/Public/MuCOE/Nodes/CustomizableObjectNodeSkeletalMesh.h
--- a/Engine/Plugins/Experimental/Mutable/Source/CustomizableObjectEditor/Public/MuCOE/Nodes/CustomizableObjectNodeSkeletalMesh.h
+++ b/Engine/Plugins/Experimental/Mutable/Source/CustomizableObjectEditor/Public/MuCOE/Nodes/CustomizableObjectNodeSkeletalMesh.h
@@ -28,6 +28,7 @@ class USkeletalMesh;
 class UTexture2D;
 struct FPropertyChangedEvent;
 struct FSkeletalMaterial;
+struct FSkelMeshSection;
 struct FSlateBrush;
 
 
@@ -231,6 +232,9 @@ public:
 	FSkeletalMaterial* GetSkeletalMaterialFor(const UEdGraphPin& Pin) const;
 	int32 GetSkeletalMaterialIndexFor(const UEdGraphPin& Pin) const;
 
+	const FSkelMeshSection* GetSkeletalMeshSectionFor(const UEdGraphPin& Pin) const;
+
+
 	virtual bool ProvidesCustomPinRelevancyTest() const override { return true; }
 	virtual bool IsPinRelevant(const UEdGraphPin* Pin) const override;
 
@@ -253,6 +257,8 @@ private:
 	UMaterialInterface* GetMaterialInterfaceFor(const int32 LODIndex, const int32 MaterialIndex) const;
 	FSkeletalMaterial* GetSkeletalMaterialFor(const int32 LODIndex, const int32 SectionIndex) const;
 	int32 GetSkeletalMaterialIndexFor(const int32 LODIndex, const int32 SectionIndex) const;
+	
+	const FSkelMeshSection* GetSkeletalMeshSectionFor(const int32 LODIndex, const int32 SectionIndex) const;
 
 	// Deprecated
 	
diff --git a/Engine/Plugins/Experimental/Mutable/Source/CustomizableObjectEditor/Public/MuCOE/Nodes/CustomizableObjectNodeTable.h b/Engine/Plugins/Experimental/Mutable/Source/CustomizableObjectEditor/Public/MuCOE/Nodes/CustomizableObjectNodeTable.h
--- a/Engine/Plugins/Experimental/Mutable/Source/CustomizableObjectEditor/Public/MuCOE/Nodes/CustomizableObjectNodeTable.h
+++ b/Engine/Plugins/Experimental/Mutable/Source/CustomizableObjectEditor/Public/MuCOE/Nodes/CustomizableObjectNodeTable.h
@@ -23,7 +23,7 @@ class UTexture2DArray;
 class UAnimInstance;
 struct FGuid;
 struct FSkeletalMaterial;
-
+struct FSkelMeshSection;
 
 /** Enum class for the different types of image pins */
 UENUM()
@@ -395,7 +395,9 @@ public:
 
 	/** Returns the index of the skeletal material associated to the given skeletal mesh output pin. */
 	int32 GetDefaultSkeletalMaterialIndexFor(const UEdGraphPin& MeshPin) const;
-	
+
+	/** Returns the SkeletalMesh section associated to the given skeletal mesh output pin. */
+	const FSkelMeshSection* GetDefaultSkeletalMeshSectionFor(const UEdGraphPin& MeshPin) const;
 private:
 
 	/** Number of properties to know when the node needs an update */
diff --git a/Engine/Plugins/Experimental/Mutable/Source/MutableRuntime/Internal/MuR/OpMeshMerge.h b/Engine/Plugins/Experimental/Mutable/Source/MutableRuntime/Internal/MuR/OpMeshMerge.h
--- a/Engine/Plugins/Experimental/Mutable/Source/MutableRuntime/Internal/MuR/OpMeshMerge.h
+++ b/Engine/Plugins/Experimental/Mutable/Source/MutableRuntime/Internal/MuR/OpMeshMerge.h
@@ -458,12 +458,35 @@ namespace mu
 				}
 
 				FMeshSurface& NewSurface = Result->Surfaces.AddDefaulted_GetRef();
-				NewSurface.VertexCount = pFirst->GetVertexCount() + pSecond->GetVertexCount();
-				NewSurface.IndexCount = pFirst->GetIndexCount() + pSecond->GetIndexCount();
 				NewSurface.BoneMapCount = Result->BoneMap.Num();
+
+				int32 NumFirstSubMeshes = 0;
+				for (const FMeshSurface& Surf : pFirst->Surfaces)
+				{
+					NewSurface.SubMeshes.Append(Surf.SubMeshes);
+					NumFirstSubMeshes += Surf.SubMeshes.Num();
+				}
+
+				for (const FMeshSurface& Surf : pSecond->Surfaces)
+				{
+					NewSurface.SubMeshes.Append(Surf.SubMeshes);
+				}
+
+				// Fix surface Submesh ranges.
+				if (NumFirstSubMeshes > 0)
+				{
+					const int32 NumResultSubMeshes = NewSurface.SubMeshes.Num();
 				
-				//All merged surfaces will have the same bCastShadow value. Decided by the first merged mesh
-				NewSurface.bCastShadow = pFirst->Surfaces.Last().bCastShadow;
+					const FSurfaceSubMesh LastFromFirstMesh = pFirst->Surfaces.Last().SubMeshes.Last();
+	
+					for (int32 SecondSubMeshIndex = NumFirstSubMeshes; SecondSubMeshIndex < NumResultSubMeshes; ++SecondSubMeshIndex)
+					{
+						NewSurface.SubMeshes[SecondSubMeshIndex].VertexBegin += LastFromFirstMesh.VertexEnd;	
+						NewSurface.SubMeshes[SecondSubMeshIndex].VertexEnd += LastFromFirstMesh.VertexEnd;	
+						NewSurface.SubMeshes[SecondSubMeshIndex].IndexBegin += LastFromFirstMesh.IndexEnd;	
+						NewSurface.SubMeshes[SecondSubMeshIndex].IndexEnd += LastFromFirstMesh.IndexEnd;	
+					}
+				}
 			}
 			else
 			{
@@ -473,13 +496,20 @@ namespace mu
 				// Add pFirst surfaces
 				Result->Surfaces = pFirst->Surfaces;
 
-				const int32 FirstVertexIndex = pFirst->GetVertexCount();
-				const int32 FirstIndexIndex = pFirst->GetIndexCount();
+				const int32 FirstVertexEnd = pFirst->GetVertexCount();
+				const int32 FirstIndexEnd = pFirst->GetIndexCount();
 
 				check(pSecond->Surfaces.Num() == 1);
 				FMeshSurface& NewSurface = Result->Surfaces.Add_GetRef(pSecond->Surfaces[0]);
-				NewSurface.FirstVertex += FirstVertexIndex;
-				NewSurface.FirstIndex += FirstIndexIndex;
+				
+				for (FSurfaceSubMesh& SubMesh : NewSurface.SubMeshes)
+				{
+					SubMesh.VertexBegin += FirstVertexEnd;
+					SubMesh.VertexEnd += FirstVertexEnd;
+					SubMesh.IndexBegin += FirstIndexEnd;
+					SubMesh.IndexEnd += FirstIndexEnd;
+				}
+
 				NewSurface.BoneMapIndex += NumFirstBonesInBoneMap;
 			}
 
diff --git a/Engine/Plugins/Experimental/Mutable/Source/MutableRuntime/Private/MuR/Mesh.cpp b/Engine/Plugins/Experimental/Mutable/Source/MutableRuntime/Private/MuR/Mesh.cpp
--- a/Engine/Plugins/Experimental/Mutable/Source/MutableRuntime/Private/MuR/Mesh.cpp
+++ b/Engine/Plugins/Experimental/Mutable/Source/MutableRuntime/Private/MuR/Mesh.cpp
@@ -502,59 +502,63 @@ int Mesh::GetSurfaceCount() const
 }
 
 
-void Mesh::GetSurface( int32 surfaceIndex,
-                       int32* firstVertex, int32* vertexCount,
-                       int32* firstIndex, int32* indexCount,
-					   int32* BoneIndex, int32* BoneCount,
-					   bool* bCastShadow) const
+void Mesh::GetSurface(int32 SurfaceIndex,
+                      int32& OutFirstVertex, int32& OutVertexCount,
+                      int32& OutFirstIndex, int32& OutIndexCount,
+					  int32& OutBoneIndex, int32& OutBoneCount) const
 {
-    int count = GetSurfaceCount();
+    int32 Count = GetSurfaceCount();
 
-    if (surfaceIndex>=0 && surfaceIndex<count)
+    if (SurfaceIndex >= 0 && SurfaceIndex < Count)
     {
-        if (surfaceIndex<Surfaces.Num())
+        if (SurfaceIndex < Surfaces.Num())
         {
-            const FMeshSurface& surf = Surfaces[surfaceIndex];
-            if (firstVertex) *firstVertex = surf.FirstVertex;
-            if (vertexCount) *vertexCount = surf.VertexCount;
-            if (firstIndex) *firstIndex = surf.FirstIndex;
-            if (indexCount) *indexCount = surf.IndexCount;
-            if (BoneIndex) *BoneIndex = surf.BoneMapIndex;
-            if (BoneCount) *BoneCount = surf.BoneMapCount;
-            if (bCastShadow) *bCastShadow = surf.bCastShadow;
+            const FMeshSurface& Surf = Surfaces[SurfaceIndex];
+
+			check(Surf.SubMeshes.Num());
+
+			// Surfaces submeshes are sorted and have no gaps. 
+            OutFirstVertex = Surf.SubMeshes[0].VertexBegin;
+            OutVertexCount = Surf.SubMeshes.Last().VertexEnd - OutFirstVertex;
+            OutFirstIndex = Surf.SubMeshes[0].IndexBegin;
+            OutIndexCount = Surf.SubMeshes.Last().IndexEnd - OutFirstIndex;
+            OutBoneIndex = Surf.BoneMapIndex;
+            OutBoneCount = Surf.BoneMapCount;
         }
-        else
+        else if (!Surfaces.Num())
         {
             // No surfaces defined, means only one surface using all the mesh
-            if (firstVertex) *firstVertex = 0;
-            if (vertexCount) *vertexCount = GetVertexCount();
-            if (firstIndex) *firstIndex = 0;
-            if (indexCount) *indexCount = GetIndexCount();
-			if (BoneIndex) *BoneIndex = 0;
-			if (BoneCount) *BoneCount = BoneMap.Num();
-			if (bCastShadow) *bCastShadow = false;
+            OutFirstVertex = 0;
+            OutVertexCount = GetVertexCount();
+            OutFirstIndex = 0;
+            OutIndexCount = GetIndexCount();
+			OutBoneIndex = 0;
+			OutBoneCount = BoneMap.Num();
         }
+		else
+		{
+			check(false);
+		}
     }
     else
     {
-        check( false );
-        if (firstVertex) *firstVertex = 0;
-        if (vertexCount) *vertexCount = 0;
-        if (firstIndex) *firstIndex = 0;
-        if (indexCount) *indexCount = 0;
-		if (BoneIndex) *BoneIndex = 0;
-		if (BoneCount) *BoneCount = 0;
-		if (bCastShadow) *bCastShadow = false;
+        check(false);
+        OutFirstVertex = 0;
+        OutVertexCount = 0;
+        OutFirstIndex = 0;
+        OutIndexCount = 0;
+		OutBoneIndex = 0;
+		OutBoneCount = 0;
     }
 }
 
 
-uint32 Mesh::GetSurfaceId( int32 surfaceIndex ) const
+uint32 Mesh::GetSurfaceId(int32 SurfaceIndex) const
 {
-    if (surfaceIndex>=0 && surfaceIndex<Surfaces.Num())
+    if (SurfaceIndex >= 0 && SurfaceIndex < Surfaces.Num())
     {
-        const FMeshSurface& surf = Surfaces[surfaceIndex];
-        return surf.Id;
+        const FMeshSurface& Surf = Surfaces[SurfaceIndex];
+        return Surf.Id;
     }
 
     return 0;
@@ -938,11 +942,15 @@ void Mesh::EnsureSurfaceData()
 {
 	if (!Surfaces.Num() && VertexBuffers.GetElementCount())
 	{
-		FMeshSurface s;
-		s.VertexCount = VertexBuffers.GetElementCount();
-		s.IndexCount = IndexBuffers.GetElementCount();
-		s.BoneMapCount = BoneMap.Num();
-		Surfaces.Add(s);
+		FMeshSurface& NewSurface = Surfaces.Emplace_GetRef();
+
+		FSurfaceSubMesh& SubMesh = NewSurface.SubMeshes.Emplace_GetRef();
+		SubMesh.VertexBegin = 0;
+		SubMesh.VertexEnd = VertexBuffers.GetElementCount();
+		SubMesh.IndexBegin = 0;
+		SubMesh.IndexEnd = IndexBuffers.GetElementCount();
+
+		NewSurface.BoneMapCount = BoneMap.Num();
 	}
 }
 
@@ -953,72 +961,26 @@ void Mesh::ResetBufferIndices()
 	IndexBuffers.ResetBufferIndices();
 }
 
+MUTABLE_IMPLEMENT_POD_SERIALISABLE(FSurfaceSubMesh);
+MUTABLE_IMPLEMENT_POD_VECTOR_SERIALISABLE(FSurfaceSubMesh);
 
-void UnserialiseLegacySurfaces(InputArchive& arch, TArray<FMeshSurface>& OutMeshSurfaces)
+void FMeshSurface::Serialise(OutputArchive& Arch) const
 {
-	struct FMeshSurfaceLegacy
-	{
-		FMeshSurfaceLegacy()
-		{}
+	Arch << SubMeshes;
 
-		int32 FirstVertex = 0;
-		int32 VertexCount = 0;
-		int32 FirstIndex = 0;
-		int32 IndexCount = 0;
-		uint32 Id = 0;
-
-		void Unserialise(InputArchive& arch)
-		{
-			arch >> FirstVertex;
-			arch >> VertexCount;
-			arch >> FirstIndex;
-			arch >> IndexCount;
-			arch >> Id;
-		}
-	}; 
-
-	TArray<FMeshSurfaceLegacy> LegacyMeshSurfaces;
-	arch >> LegacyMeshSurfaces;
-	
-	const int32 NumSurfaces = LegacyMeshSurfaces.Num();
-	OutMeshSurfaces.SetNumZeroed(NumSurfaces);
-
-	for (int32 SurfaceIndex = 0; SurfaceIndex < NumSurfaces; ++SurfaceIndex)
-	{
-		FMeshSurfaceLegacy& LegacySurface = LegacyMeshSurfaces[SurfaceIndex];
-		FMeshSurface& Surface = OutMeshSurfaces[SurfaceIndex];
-		Surface.FirstVertex = LegacySurface.FirstVertex;
-		Surface.VertexCount = LegacySurface.VertexCount;
-		Surface.FirstIndex = LegacySurface.FirstIndex;
-		Surface.IndexCount = LegacySurface.IndexCount;
-		Surface.Id = LegacySurface.Id;
-	}
+	Arch << BoneMapIndex;
+	Arch << BoneMapCount;
+	Arch << Id;
 }
 
 
-void FMeshSurface::Serialise(OutputArchive& arch) const
+void FMeshSurface::Unserialise(InputArchive& Arch)
 {
-	arch << FirstVertex;
-	arch << VertexCount;
-	arch << FirstIndex;
-	arch << IndexCount;
-	arch << BoneMapIndex;
-	arch << BoneMapCount;
-	arch << bCastShadow;
-	arch << Id;
-}
+	Arch >> SubMeshes;
 
-
-void FMeshSurface::Unserialise(InputArchive& arch)
-{
-	arch >> FirstVertex;
-	arch >> VertexCount;
-	arch >> FirstIndex;
-	arch >> IndexCount;
-	arch >> BoneMapIndex;
-	arch >> BoneMapCount;
-	arch >> bCastShadow;
-	arch >> Id;
+	Arch >> BoneMapIndex;
+	Arch >> BoneMapCount;
+	Arch >> Id;
 }
 
 
@@ -1038,69 +1000,69 @@ void Mesh::FBonePose::Unserialise(InputArchive& arch)
 }
 
 
-void Mesh::Serialise(OutputArchive& arch) const
+void Mesh::Serialise(OutputArchive& Arch) const
 {
-	uint32 ver = 22;
-	arch << ver;
+	uint32 Version = 23;
+	Arch << Version;
 
-	arch << IndexBuffers;
-	arch << VertexBuffers;
-	arch << AdditionalBuffers;
-	arch << Layouts;
+	Arch << IndexBuffers;
+	Arch << VertexBuffers;
+	Arch << AdditionalBuffers;
+	Arch << Layouts;
 
-	arch << SkeletonIDs;
+	Arch << SkeletonIDs;
 
-	arch << Skeleton;
-	arch << PhysicsBody;
+	Arch << Skeleton;
+	Arch << PhysicsBody;
 
-	arch << uint32(Flags);
-	arch << Surfaces;
+	Arch << uint32(Flags);
+	Arch << Surfaces;
 
-	arch << Tags;
-	arch << StreamedResources;
+	Arch << Tags;
+	Arch << StreamedResources;
 
-	arch << BonePoses;
-	arch << BoneMap;
+	Arch << BonePoses;
+	Arch << BoneMap;
 
-	arch << AdditionalPhysicsBodies;
+	Arch << AdditionalPhysicsBodies;
 
-	arch << MeshIDPrefix;
-	arch << ReferenceID;
+	Arch << MeshIDPrefix;
+	Arch << ReferenceID;
 }
 
 
-void Mesh::Unserialise(InputArchive& arch)
+void Mesh::Unserialise(InputArchive& Arch)
 {
-	uint32 ver;
-	arch >> ver;
-	check(ver == 22);
+	uint32 Version;
+	Arch >> Version;
+	check(Version == 23);
 
-	arch >> IndexBuffers;
-	arch >> VertexBuffers;
-	arch >> AdditionalBuffers;
-	arch >> Layouts;
+	Arch >> IndexBuffers;
+	Arch >> VertexBuffers;
+	Arch >> AdditionalBuffers;
+	Arch >> Layouts;
 
-	arch >> SkeletonIDs;
+	Arch >> SkeletonIDs;
 
-	arch >> Skeleton;
-	arch >> PhysicsBody;
+	Arch >> Skeleton;
+	Arch >> PhysicsBody;
 
 	uint32 Temp;
-	arch >> Temp;
+	Arch >> Temp;
 	Flags = static_cast<EMeshFlags>(Temp);
 
-	arch >> Surfaces;
+	Arch >> Surfaces;
 
-	arch >> Tags;
-	arch >> StreamedResources;
+	Arch >> Tags;
+	Arch >> StreamedResources;
 
-	arch >> BonePoses;
-	arch >> BoneMap;
+	Arch >> BonePoses;
+	Arch >> BoneMap;
 
-	arch >> AdditionalPhysicsBodies;
+	Arch >> AdditionalPhysicsBodies;
 
-	arch >> MeshIDPrefix;
-	arch >> ReferenceID;
+	Arch >> MeshIDPrefix;
+	Arch >> ReferenceID;
 }
 
 
diff --git a/Engine/Plugins/Experimental/Mutable/Source/MutableRuntime/Private/MuR/OpImageProject.cpp b/Engine/Plugins/Experimental/Mutable/Source/MutableRuntime/Private/MuR/OpImageProject.cpp
--- a/Engine/Plugins/Experimental/Mutable/Source/MutableRuntime/Private/MuR/OpImageProject.cpp
+++ b/Engine/Plugins/Experimental/Mutable/Source/MutableRuntime/Private/MuR/OpImageProject.cpp
@@ -3102,7 +3102,7 @@ void MeshProject(Mesh* Result, const Mesh* pMesh, const FProjector& projector, b
 
 	if (bOutSuccess)
 	{
-		Result->Surfaces.SetNum(0);
+		Result->Surfaces.Empty();
 		Result->EnsureSurfaceData();
 		Result->ResetStaticFormatFlags();
 	}
diff --git a/Engine/Plugins/Experimental/Mutable/Source/MutableRuntime/Private/MuR/OpMeshApplyLayout.h b/Engine/Plugins/Experimental/Mutable/Source/MutableRuntime/Private/MuR/OpMeshApplyLayout.h
--- a/Engine/Plugins/Experimental/Mutable/Source/MutableRuntime/Private/MuR/OpMeshApplyLayout.h
+++ b/Engine/Plugins/Experimental/Mutable/Source/MutableRuntime/Private/MuR/OpMeshApplyLayout.h
@@ -208,15 +208,15 @@ namespace mu
 		if (VerticesToRemove.Num())
 		{
 			// Unpack vertices into a mask
-			TArray<uint8> VertexMask;
-			VertexMask.SetNumZeroed(pApplied->GetVertexCount());
-			for ( int32 VertexIndex: VerticesToRemove )
+			TBitArray<> VertexMask;
+			VertexMask.SetNum(pApplied->GetVertexCount(), false);
+			for (int32 VertexIndex : VerticesToRemove)
 			{
-				VertexMask[VertexIndex] = 1;
+				VertexMask[VertexIndex] = true;
 			}
 	
 			// Remove
-			MeshRemoveVerticesWithMap( pApplied, VertexMask.GetData(), (uint32)VertexMask.Num());
+			MeshRemoveVerticesWithCullSet(pApplied, VertexMask);
 		}
 
 		//
diff --git a/Engine/Plugins/Experimental/Mutable/Source/MutableRuntime/Private/MuR/OpMeshClipDeform.h b/Engine/Plugins/Experimental/Mutable/Source/MutableRuntime/Private/MuR/OpMeshClipDeform.h
--- a/Engine/Plugins/Experimental/Mutable/Source/MutableRuntime/Private/MuR/OpMeshClipDeform.h
+++ b/Engine/Plugins/Experimental/Mutable/Source/MutableRuntime/Private/MuR/OpMeshClipDeform.h
@@ -4,7 +4,7 @@
 
 
 #include "MuR/OpMeshBind.h"
-
+#include "MuR/OpMeshRemove.h"
 #include "MuR/MutableTrace.h"
 
 #include "MuR/MeshPrivate.h"
@@ -12,6 +12,7 @@
 #include "Algo/Copy.h"
 
 
+
 // TODO: Make the handling of rotations an option. It is more expensive on CPU and memory, and for some
 // cases it is not required at all.
 
@@ -271,8 +272,8 @@ namespace mu
 		// Update the result mesh positions
 		int32 MeshVertexCount = BaseMesh->GetVertexCount();
 
-		TArray<uint8> ExcludedVertices;
-		ExcludedVertices.SetNumZeroed(MeshVertexCount);
+		TBitArray<> VerticesToCull;
+		VerticesToCull.SetNum(MeshVertexCount, false);
 		{
 			MUTABLE_CPUPROFILER_SCOPE(UpdateClipDeformVertices);
 
@@ -309,7 +310,7 @@ namespace mu
 						ItTangent.SetFromVec3f(NewTangent);
 					}
 					
-					ExcludedVertices[MeshVertexIndex] = ClipWeight >= (1.0f - SMALL_NUMBER);	
+					VerticesToCull[MeshVertexIndex] = ClipWeight >= (1.0f - SMALL_NUMBER);	
 				}
 
 				++ItPosition;
@@ -326,133 +327,6 @@ namespace mu
 			}
 		}
 
-		// Remove excluded vertices
-		{
-			MUTABLE_CPUPROFILER_SCOPE(RemoveExcludedVertices);
-			
-			const int32 NumOriginalVerts = ExcludedVertices.Num();
-		
-			// Remove faces only if all vertices are gone.
-			UntypedMeshBufferIteratorConst ItSrc(Result->GetIndexBuffers(), MBS_VERTEXINDEX);
-			UntypedMeshBufferIterator ItDest(Result->GetIndexBuffers(), MBS_VERTEXINDEX);
-
-			UntypedMeshBufferIterator It(Result->GetIndexBuffers(), MBS_VERTEXINDEX);
-
-			// We cannot reintroduce vertices to the same set ExcludedVertices since some other triangle exclusion might
-			// depend on those.
-			TArray<uint8> FinalExcludedVertices;
-			FinalExcludedVertices.Init(1, ExcludedVertices.Num());
-
-			const SIZE_T TriangleSizeBytes =  ItSrc.GetElementSize() * 3;
-			const int32 SrcTriangleCount = Result->GetIndexCount() / 3;
-			for (int32 TriangleIndex = 0; TriangleIndex < SrcTriangleCount; ++TriangleIndex)
-			{
-				int32 A = It.GetAsUINT32();
-				++It;
-				int32 B = It.GetAsUINT32();
-				++It;
-				int32 C = It.GetAsUINT32();
-				++It;
-
-				const bool bAllVerticesExcluded = static_cast<bool>( ExcludedVertices[A] & ExcludedVertices[B] & ExcludedVertices[C] );
-
-				// Some sort of latching mechanisms is needed, if a vertex needs to be included for a vertex triangle
-				// but not for another, the vertex must stay.
-				FinalExcludedVertices[A] = FMath::Min((uint8)bAllVerticesExcluded, FinalExcludedVertices[A] );
-				FinalExcludedVertices[B] = FMath::Min((uint8)bAllVerticesExcluded, FinalExcludedVertices[B] );
-				FinalExcludedVertices[C] = FMath::Min((uint8)bAllVerticesExcluded, FinalExcludedVertices[C] );
-			 
-				if (!bAllVerticesExcluded)
-				{
-					if (ItSrc.ptr() != ItDest.ptr())
-					{
-						FMemory::Memcpy( ItDest.ptr(), ItSrc.ptr(),  TriangleSizeBytes);
-					}
-
-					ItDest += 3;
-				}
-			 
-				ItSrc += 3;
-			}
-
-			const int32 NumRemainingTris = SrcTriangleCount - ( ( ItSrc.ptr() - ItDest.ptr() ) / TriangleSizeBytes ); 
-			Result->GetIndexBuffers().SetElementCount( NumRemainingTris*3 );
-		
-			TArray<int32> IndexMap;
-			IndexMap.Init(-1, NumOriginalVerts);
-		
-			int32 NumRemainingVerts = 0;
-			for ( int32 I = 0; I < NumOriginalVerts; ++I )
-			{
-				if ( FinalExcludedVertices[I] == 0 )
-				{
-					IndexMap[I] = NumRemainingVerts++;
-				}
-			}
-		
-			UntypedMeshBufferIterator ResIndicesIt(Result->GetIndexBuffers(), MBS_VERTEXINDEX);
-		
-			// Remap Indices
-			for ( int32 TIdx = 0; TIdx < NumRemainingTris; ++TIdx )
-			{
-				const int32 A = ResIndicesIt.GetAsUINT32();
-				check(IndexMap[A] >= 0);
-				ResIndicesIt.SetFromUINT32(IndexMap[A]);
-				++ResIndicesIt;
-			
-				const int32 B = ResIndicesIt.GetAsUINT32();
-				check(IndexMap[B] >= 0);
-				ResIndicesIt.SetFromUINT32(IndexMap[B]);
-				++ResIndicesIt;
-
-				const int32 C = ResIndicesIt.GetAsUINT32();
-				check(IndexMap[C] >= 0);
-				ResIndicesIt.SetFromUINT32(IndexMap[C]);
-				++ResIndicesIt;
-			}
-
-
-			if ( NumRemainingVerts == NumOriginalVerts )
-			{
-				return;
-			}
-		
-			// We are guaranteed to find a removed vertex.
-			int32 SpanStart = 0; // points at the span first element
-			while ( FinalExcludedVertices[SpanStart++] == 1 );
-			--SpanStart;
-		
-			int32 DestEnd = 0;
-		
-			const int32 VertexBufferCount = Result->GetVertexBuffers().GetBufferCount();
-			while ( SpanStart < NumOriginalVerts )
-			{
-				// Find span to move.
-				int32 SpanEnd = SpanStart; // Points at the last span element + 1
-				for (; SpanEnd < NumOriginalVerts && FinalExcludedVertices[SpanEnd] == 0; SpanEnd++ );
-
-				const int32 SpanSize = SpanEnd - SpanStart;
-
-				// Assume we always have large spans, so it is more work to find the spans than iterate over the different
-				// buffers. This probably is not the best strategy and recomputing the spans for every buffer might be
-				// more performant.
-
-				if (SpanStart > 0)
-				{
-					for (int32 B = 0; B < VertexBufferCount; ++B)
-					{
-						const int32 ElementSize = Result->GetVertexBuffers().GetElementSize(B);
-						uint8* Ptr = Result->GetVertexBuffers().GetBufferData(B);
-						FMemory::Memmove( Ptr + DestEnd*ElementSize, Ptr + SpanStart*ElementSize , SpanSize*ElementSize );
-					}
-				}
-			
-				DestEnd += SpanSize;
-				
-				for (SpanStart = SpanEnd; SpanStart < NumOriginalVerts && FinalExcludedVertices[SpanStart] == 1; ++SpanStart);
-			}
-
-			Result->GetVertexBuffers().SetElementCount( DestEnd );
-		}
+		MeshRemoveVerticesWithCullSet(Result, VerticesToCull);
 	}
 }
diff --git a/Engine/Plugins/Experimental/Mutable/Source/MutableRuntime/Private/MuR/OpMeshClipMorphPlane.h b/Engine/Plugins/Experimental/Mutable/Source/MutableRuntime/Private/MuR/OpMeshClipMorphPlane.h
--- a/Engine/Plugins/Experimental/Mutable/Source/MutableRuntime/Private/MuR/OpMeshClipMorphPlane.h
+++ b/Engine/Plugins/Experimental/Mutable/Source/MutableRuntime/Private/MuR/OpMeshClipMorphPlane.h
@@ -6,17 +6,19 @@
 #include "MuR/ConvertData.h"
 #include "MuR/ParametersPrivate.h"
 #include "MuR/Platform.h"
+#include "MuR/OpMeshRemove.h"
 
+#include "Math/IntVector.h"
 
 namespace mu
 {
-	inline bool PointInBoundingBox(const FVector3f& point, const FShape& selectionShape)
+	inline bool PointInBoundingBox(const FVector3f& Point, const FShape& SelectionShape)
 	{
-		FVector3f v = point - selectionShape.position;
+		FVector3f V = Point - SelectionShape.position;
 
-		for (int i = 0; i < 3; ++i)
+		for (int32 i = 0; i < 3; ++i)
 		{
-			if (fabs(v[i]) > selectionShape.size[i])
+			if (FMath::Abs(V[i]) > SelectionShape.size[i])
 			{
 				return false;
 			}
@@ -25,39 +27,39 @@ namespace mu
 		return true;
 	}
 
-	inline bool VertexIsInMaxRadius(const FVector3f& vertex, const FVector3f& origin, float vertexSelectionBoneMaxRadius)
+	inline bool VertexIsInMaxRadius(const FVector3f& Position, const FVector3f& Origin, float VertexSelectionBoneMaxRadius)
 	{
-		if (vertexSelectionBoneMaxRadius < 0.f)
+		if (VertexSelectionBoneMaxRadius < 0.f)
 		{
 			return true;
 		}
 
-		FVector3f radiusVec = vertex - origin;
-		float radius = FVector3f::DotProduct(radiusVec, radiusVec);
+		FVector3f RadiusVec = Position - Origin;
+		float Radius2 = FVector3f::DotProduct(RadiusVec, RadiusVec);
 
-		return radius < vertexSelectionBoneMaxRadius* vertexSelectionBoneMaxRadius;
+		return Radius2 < VertexSelectionBoneMaxRadius*VertexSelectionBoneMaxRadius;
 	}
 
-	struct vertex_bone_info
+	struct FVertexBoneInfo
 	{
-		TArray<int32> bone_indices;
-		TArray<int32> bone_weights;
+		TArray<int32, TFixedAllocator<16>> BoneIndices;
+		TArray<int32, TFixedAllocator<16>> BoneWeights;
 	};
 
-	inline bool VertexIsAffectedByBone(int32 vertex_idx, const TArray<bool>& bone_is_affected, const TArray<vertex_bone_info>& vertex_info)
+	inline bool VertexIsAffectedByBone(int32 VertexIdx, const TBitArray<>& BoneIsAffected, const TArray<FVertexBoneInfo>& VertexInfo)
 	{
-		if (vertex_idx >= (int)vertex_info.Num())
+		if (VertexIdx >= VertexInfo.Num())
 		{
 			return false;
 		}
 
-		check(vertex_info[vertex_idx].bone_indices.Num() == vertex_info[vertex_idx].bone_weights.Num());		
+		check(VertexInfo[VertexIdx].BoneIndices.Num() == VertexInfo[VertexIdx].BoneWeights.Num());
 
-        for (size_t i = 0; i < vertex_info[vertex_idx].bone_indices.Num(); ++i)
+        for (int32 i = 0; i < VertexInfo[VertexIdx].BoneIndices.Num(); ++i)
 		{			
-            check(vertex_info[vertex_idx].bone_indices[i] <= (int)bone_is_affected.Num());
+            check(VertexInfo[VertexIdx].BoneIndices[i] <= BoneIsAffected.Num());
 
-			if (bone_is_affected[vertex_info[vertex_idx].bone_indices[i]] && vertex_info[vertex_idx].bone_weights[i] > 0)
+			if (BoneIsAffected[VertexInfo[VertexIdx].BoneIndices[i]] && VertexInfo[VertexIdx].BoneWeights[i] > 0)
 			{
 				return true;
 			}
@@ -69,47 +71,47 @@ namespace mu
 	//---------------------------------------------------------------------------------------------
 	//! Reference version
 	//---------------------------------------------------------------------------------------------
-	inline void MeshClipMorphPlane(Mesh* Result, const Mesh* pBase, const FVector3f& origin, const FVector3f& normal, float dist, float factor, float radius,
-		float radius2, float angle, const FShape& selectionShape, bool& bOutSuccess, const FBoneName* BoneId = nullptr, float vertexSelectionBoneMaxRadius = -1.f)
+	inline void MeshClipMorphPlane(Mesh* Result, const Mesh* pBase, const FVector3f& Origin, const FVector3f& Normal, float Dist, float Factor, float Radius,
+		float Radius2, float Angle, const FShape& SelectionShape, bool& bOutSuccess, const FBoneName* BoneId = nullptr, float VertexSelectionBoneMaxRadius = -1.f)
 	{
 		bOutSuccess = true;
-		//float radius = 8.f;
-		//float radius2 = 4.f;
-		//float factor = 1.f;
+		//float Radius = 8.f;
+		//float Radius2 = 4.f;
+		//float Factor = 1.f;
 
 		// Generate vector perpendicular to normal for ellipse rotation reference base
-		FVector3f aux_base(0.f, 1.f, 0.f);
+		FVector3f AuxBase(0.f, 1.f, 0.f);
 
-		if (FMath::Abs(FVector3f::DotProduct(normal, aux_base)) > 0.95f)		// fabs = absolute value
+		if (FMath::Abs(FVector3f::DotProduct(Normal, AuxBase)) > 0.95f)
 		{
-			aux_base = FVector3f(0.f, 0.f, 1.f);
+			AuxBase = FVector3f(0.f, 0.f, 1.f);
 		}
 		
-		FVector3f origin_radius_vector = FVector3f::CrossProduct(normal, aux_base);		// PERPENDICULAR VECTOR TO THE PLANE normal and aux base
-		check(FMath::Abs(FVector3f::DotProduct(normal, origin_radius_vector)) < 0.05f);
+		FVector3f OriginRadiusVector = FVector3f::CrossProduct(Normal, AuxBase);		// Perpendicular vector to the plane normal and aux base
+		check(FMath::Abs(FVector3f::DotProduct(Normal, OriginRadiusVector)) < 0.05f);
 
-        uint32 vcount = pBase->GetVertexBuffers().GetElementCount();
+        uint32 VertexCount = pBase->GetVertexBuffers().GetElementCount();
 
-		if (!vcount)
+		if (!VertexCount)
 		{
 			bOutSuccess = false;
 			return;
 		}
 
-		TArray<vertex_bone_info> vertex_info;
+		TArray<FVertexBoneInfo> VertexInfo;
 
         Ptr<const Skeleton> BaseSkeleton = pBase->GetSkeleton();
 
-		TArray<bool> AffectedBoneMapIndices;
+		TBitArray<> AffectedBoneMapIndices;
 		const int32 BaseBoneIndex = BaseSkeleton && BoneId ? BaseSkeleton->FindBone(*BoneId) : INDEX_NONE;
         if (BaseBoneIndex != INDEX_NONE)
 		{
 			const TArray<FBoneName>& BoneMap = pBase->BoneMap;
-			AffectedBoneMapIndices.SetNum(BoneMap.Num());
+			AffectedBoneMapIndices.SetNum(BoneMap.Num(), false);
 
 			const int32 BoneCount = BaseSkeleton->GetBoneCount();
-			TArray<bool> AffectedSkeletonBones;
-			AffectedSkeletonBones.SetNum(BoneCount);
+			TBitArray<> AffectedSkeletonBones;
+			AffectedSkeletonBones.SetNum(BoneCount, false);
 
             for (int32 BoneIndex = 0; BoneIndex < BoneCount; ++BoneIndex)
 			{
@@ -135,74 +137,77 @@ namespace mu
 			}
 
 			// Look for affected vertex indices
-			vertex_info.SetNum(vcount);
-			//int firstCount = pBase->GetVertexBuffers().GetElementCount();
+			VertexInfo.SetNum(VertexCount);
+			//int32 FirstCount = pBase->GetVertexBuffers().GetElementCount();
 
-			for (int32 vb = 0; vb < pBase->GetVertexBuffers().m_buffers.Num(); ++vb)
+			for (int32 BufferIndex = 0; BufferIndex < pBase->GetVertexBuffers().m_buffers.Num(); ++BufferIndex)
 			{
-				const FMeshBuffer& result = pBase->GetVertexBuffers().m_buffers[vb];
+				const FMeshBuffer& Buffer = pBase->GetVertexBuffers().m_buffers[BufferIndex];
 
-				int elemSize = pBase->GetVertexBuffers().GetElementSize((int)vb);
-				//int firstSize = firstCount * elemSize;
+				int32 ElemSize = pBase->GetVertexBuffers().GetElementSize(BufferIndex);
+				//int32 FirstSize = FirstCount * ElemSize;
 
-				for (int c = 0; c < pBase->GetVertexBuffers().GetBufferChannelCount((int)vb); ++c)
+				for (int32 ChannelIndex = 0; ChannelIndex < pBase->GetVertexBuffers().GetBufferChannelCount(BufferIndex); ++ChannelIndex)
 				{
 					// Get info about the destination channel
-					EMeshBufferSemantic semantic = MBS_NONE;
-					int semanticIndex = 0;
-					EMeshBufferFormat format = MBF_NONE;
-					int components = 0;
-					int offset = 0;
-					pBase->GetVertexBuffers().GetChannel((int)vb, c, &semantic, &semanticIndex, &format, &components, &offset);
+					EMeshBufferSemantic Semantic = MBS_NONE;
+					int32 SemanticIndex = 0;
+					EMeshBufferFormat Format = MBF_NONE;
+					int32 Components = 0;
+					int32 Offset = 0;
+					pBase->GetVertexBuffers().GetChannel(BufferIndex, ChannelIndex, &Semantic, &SemanticIndex, &Format, &Components, &Offset);
 
-					int secondOffset = offset;
-					//int resultOffset = firstSize + offset;
+					checkf(Components <= 16, TEXT("FVertexBoneInfo allocation is fixed to 16 elements."));
 
-					if (semantic == MBS_BONEINDICES)
+					//int32 ResultOffset = FirstSize + Offset;
+
+					if (Semantic == MBS_BONEINDICES)
 					{
-						for (int i = 0; i < pBase->GetVertexCount(); ++i)
+						int32 NumVertices = pBase->GetVertexCount();
 						{
-							switch (format)
+							switch (Format)
 							{
 							case MBF_INT8:
 							case MBF_UINT8:
 							{
-                                const uint8_t* pD = &result.m_data[secondOffset];
-								
-								for (int j = 0; j < components; ++j)
+								for (int32 i = 0; i < NumVertices; ++i)
 								{
-									vertex_info[i].bone_indices.Add(pD[j]);
+									const uint8* pD = &Buffer.m_data[i*ElemSize + Offset];
+									for (int32 j = 0; j < Components; ++j)
+									{
+										VertexInfo[i].BoneIndices.Add(pD[j]);
+									}
 								}
-
-								secondOffset += elemSize;
 								break;
 							}
 
 							case MBF_INT16:
 							case MBF_UINT16:
 							{
-								const uint16* pD = reinterpret_cast<const uint16*>(&result.m_data[secondOffset]);
-
-								for (int j = 0; j < components; ++j)
+								for (int32 i = 0; i < NumVertices; ++i)
 								{
-									vertex_info[i].bone_indices.Add(pD[j]);
-								}
+									const uint16* pD = reinterpret_cast<const uint16*>(&Buffer.m_data[i*ElemSize + Offset]);
 
-								secondOffset += elemSize;
+									for (int32 j = 0; j < Components; ++j)
+									{
+										VertexInfo[i].BoneIndices.Add(pD[j]);
+									}
+								}
 								break;
 							}
 
 							case MBF_INT32:
 							case MBF_UINT32:
 							{
-								const uint32* pD = reinterpret_cast<const uint32*>(&result.m_data[secondOffset]);
-
-								for (int j = 0; j < components; ++j)
+								for (int32 i = 0; i < NumVertices; ++i)
 								{
-									vertex_info[i].bone_indices.Add(pD[j]);
-								}
+									const uint32* pD = reinterpret_cast<const uint32*>(&Buffer.m_data[i*ElemSize + Offset]);
 
-								secondOffset += elemSize;
+									for (int32 j = 0; j < Components; ++j)
+									{
+										VertexInfo[i].BoneIndices.Add(pD[j]);
+									}
+								}
 								break;
 							}
 
@@ -211,24 +216,25 @@ namespace mu
 							}
 						}
 					}
-					else if (semantic == MBS_BONEWEIGHTS)
+					else if (Semantic == MBS_BONEWEIGHTS)
 					{
-						for (int i = 0; i < pBase->GetVertexCount(); ++i)
+						const int32 NumVertices = pBase->GetVertexCount();
 						{
-							switch (format)
+							switch (Format)
 							{
 							case MBF_INT8:
 							case MBF_UINT8:
 							case MBF_NUINT8:
 							{
-								const uint8_t* pD = &result.m_data[secondOffset];
-
-								for (int j = 0; j < components; ++j)
+								for (int32 i = 0; i < NumVertices; ++i)
 								{
-                                    vertex_info[i].bone_weights.Add(pD[j]);
-								}
+									const uint8* pD = &Buffer.m_data[i*ElemSize + Offset];
 
-								secondOffset += elemSize;
+									for (int32 j = 0; j < Components; ++j)
+									{
+										VertexInfo[i].BoneWeights.Add(pD[j]);
+									}
+								}
 								break;
 							}
 
@@ -236,14 +242,15 @@ namespace mu
 							case MBF_UINT16:
 							case MBF_NUINT16:
 							{
-								const uint16* pD = reinterpret_cast<const uint16*>(&result.m_data[secondOffset]);
-
-								for (int j = 0; j < components; ++j)
+								for (int32 i = 0; i < NumVertices; ++i)
 								{
-									vertex_info[i].bone_weights.Add(pD[j]);
-								}
+									const uint16* pD = reinterpret_cast<const uint16*>(&Buffer.m_data[i*ElemSize + Offset]);
 
-								secondOffset += elemSize;
+									for (int32 j = 0; j < Components; ++j)
+									{
+										VertexInfo[i].BoneWeights.Add(pD[j]);
+									}
+								}
 								break;
 							}
 
@@ -251,27 +258,29 @@ namespace mu
 							case MBF_UINT32:
 							case MBF_NUINT32:
 							{
-								const uint32* pD = reinterpret_cast<const uint32*>(&result.m_data[secondOffset]);
-
-								for (int j = 0; j < components; ++j)
+								for (int32 i = 0; i < NumVertices; ++i)
 								{
-									vertex_info[i].bone_weights.Add(pD[j]);
-								}
+									const uint32* pD = reinterpret_cast<const uint32*>(&Buffer.m_data[i*ElemSize + Offset]);
 
-								secondOffset += elemSize;
+									for (int32 j = 0; j < Components; ++j)
+									{
+										VertexInfo[i].BoneWeights.Add(pD[j]);
+									}
+								}
 								break;
 							}
 
                             case MBF_FLOAT32:
                             {
-								const float* pD = reinterpret_cast<const float*>(&result.m_data[secondOffset]);
-
-                                for (int j = 0; j < components; ++j)
-                                {
-                                    vertex_info[i].bone_weights.Add(pD[j]>0.0f);
-                                }
+								for (int32 i = 0; i < NumVertices; ++i)
+								{
+									const float* pD = reinterpret_cast<const float*>(&Buffer.m_data[i*ElemSize + Offset]);
 
-                                secondOffset += elemSize;
+									for (int32 j = 0; j < Components; ++j)
+									{
+										VertexInfo[i].BoneWeights.Add(pD[j] > 0.0f);
+									}
+								}
                                 break;
                             }
 
@@ -286,148 +295,83 @@ namespace mu
 
 		Result->CopyFrom(*pBase);
 
-		// TODO: Replace with an array of bools?
-        TArray<bool> RemovedVertices;
-		RemovedVertices.Init(false, Result->GetVertexCount());
-
-		const FMeshBufferSet& MBSPriv = Result->GetVertexBuffers();
-		for (int32 b = 0; b < MBSPriv.m_buffers.Num(); ++b)
-		{
-			for (int32 c = 0; c < MBSPriv.m_buffers[b].m_channels.Num(); ++c)
+		const int32 NumVertices = Result->GetVertexCount();
+        TBitArray<> VerticesToCull;
+		VerticesToCull.SetNum(NumVertices, false);
+		
+		// Positions can  be assumed to be in a FVector3f struct but changing the iterator to a untyped one should 
+		// work as well.
+		const MeshBufferIterator<MBF_FLOAT32, float, 3> PositionIterBegin(Result->GetVertexBuffers(), MBS_POSITION);
+		for (int32 VertexIndex = 0; VertexIndex < NumVertices; ++VertexIndex)
+		{	
+			FVector3f Position = (PositionIterBegin + VertexIndex).GetAsVec3f();
+
+			const bool bIsVertexAffectedBone = 
+					BaseBoneIndex != INDEX_NONE &&
+					VertexIsInMaxRadius(Position, Origin, VertexSelectionBoneMaxRadius) &&
+					VertexIsAffectedByBone(VertexIndex, AffectedBoneMapIndices, VertexInfo);
+			
+			const bool bIsVertexAffectedNoShape = 
+					BaseBoneIndex == INDEX_NONE && 
+					SelectionShape.type == (uint8)FShape::Type::None;
+
+			const bool bIsVertexAffectedBoundingBox = 
+					(SelectionShape.type == (uint8)FShape::Type::AABox && 
+					PointInBoundingBox(Position, SelectionShape));
+
+			bool bRemovedVertex = false;
+			if (bIsVertexAffectedBone || bIsVertexAffectedNoShape || bIsVertexAffectedBoundingBox)
 			{
-				EMeshBufferSemantic sem = MBSPriv.m_buffers[b].m_channels[c].m_semantic;
-				int semIndex = MBSPriv.m_buffers[b].m_channels[c].m_semanticIndex;
+				FVector3f MorphPlaneCenter = Origin; // Morph plane pos relative to root of the selected bone
+				FVector3f ClipPlaneCenter = Origin + Normal*Dist; // Clipping plane pos
+				FVector3f AuxMorph = Position - MorphPlaneCenter; // Morph plane --> current vertex
+				FVector3f AuxClip = Position - ClipPlaneCenter;   // Clipping plane --> current vertex
 
-				UntypedMeshBufferIterator it(Result->GetVertexBuffers(), sem, semIndex);
+				float DotMorph = FVector3f::DotProduct(AuxMorph, Normal); // Angle (morph plane to vertex and normal)
+				float DotCut = FVector3f::DotProduct(AuxClip, Normal);    // Angle (clipping plane to vertex and normal)
 
-				switch (sem)
+				// Check if clipping or morph should be computed for v vertex
+				if (DotMorph >= 0.f || DotCut >= 0.f)
 				{
-				case MBS_POSITION:
-                    for (uint32 v = 0; v < vcount; ++v)
-					{
-						FVector3f vertex(0.0f, 0.0f, 0.0f);
-						for (int i = 0; i < 3; ++i)
-						{
-							ConvertData(i, &vertex[0], MBF_FLOAT32, it.ptr(), it.GetFormat());
-						}
-
-						const bool bIsVertexAffectedBone = 
-								BaseBoneIndex != INDEX_NONE &&
-								VertexIsInMaxRadius(vertex, origin, vertexSelectionBoneMaxRadius) &&
-								VertexIsAffectedByBone(v, AffectedBoneMapIndices, vertex_info);
-						
-						const bool bIsVertexAffectedNoShape = 
-								BaseBoneIndex == INDEX_NONE && 
-								selectionShape.type == (uint8_t)FShape::Type::None;
-
-						const bool bIsVertexAffectedBoundingBox = 
-								(selectionShape.type == (uint8_t)FShape::Type::AABox && 
-								PointInBoundingBox(vertex, selectionShape));
-
-						if (bIsVertexAffectedBone || bIsVertexAffectedNoShape || bIsVertexAffectedBoundingBox)
-						{
-							FVector3f morph_plane_center = origin;					// MORPH PLANE POS relative to root of the selected bone
-							FVector3f clip_plane_center = origin + normal * dist;	// CPLIPPING PLANE POS
-							FVector3f aux_morph = vertex - morph_plane_center;		// MORPH PLANE --> CURRENT VERTEX
-							FVector3f aux_clip = vertex - clip_plane_center;		// CLIPPING PLANE --> CURRENT VERTEX
-
-							float dot_morph = FVector3f::DotProduct(aux_morph, normal);			// ANGLE (MORPH PLANE TO VERTEX AND NORMAL)
-							float dot_cut = FVector3f::DotProduct(aux_clip, normal);				// ANGLE (CLIPPING PLANE TO VERTEX AND NORMAL )
+					FVector3f CurrentCenter = MorphPlaneCenter + Normal*DotMorph; // Projected point from the mroph plane (the closer the dot value of normal and vertex the further it goes)
+					FVector3f RadiusVector = Position - CurrentCenter;	
+					float RadiusVectorLen = RadiusVector.Length();
+					FVector3f RadiusVectorUnit = RadiusVectorLen != 0.f ? RadiusVector / RadiusVectorLen : FVector3f(0.f, 0.f, 0.f); // Unitary vector that goes from the point to the vertex
 
-							if (dot_morph >= 0.f || dot_cut >= 0.f)				// CHECK IF CLIPPING OR MORPH SHOULD BE COMPUTED FOR V VERTEX
-							{
-								FVector3f current_center = morph_plane_center + normal * dot_morph;	// PROJECTED POINT FROM THE MROPH PLANE (THE CLOSER THE DOT VALUE OF NORMAL AND VERTEX THE FURTHER IT GOES )
-								FVector3f radius_vector = vertex - current_center;					// 	
-								float radius_vector_len = radius_vector.Length();
-								FVector3f radius_vector_unit = radius_vector_len != 0.f ? radius_vector / radius_vector_len : FVector3f(0.f, 0.f, 0.f); // UNITARY VECTOR THAT GOES FROM THE POINT TO THE VERTEX
-
-								float angle_from_origin = acosf(FVector3f::DotProduct(radius_vector_unit, origin_radius_vector));
-
-								// Cross product between the perpendicular vector from radius vector and origin radius and the normal vector
-								if (FVector3f::DotProduct(FVector3f::CrossProduct(radius_vector_unit, origin_radius_vector), normal) < 0)
-								{
-									angle_from_origin = -angle_from_origin;
-								}
-
-								angle_from_origin += angle * PI / 180.f;
+					float AngleFromOrigin = FMath::Acos(FVector3f::DotProduct(RadiusVectorUnit, OriginRadiusVector));
 
-								float term1 = radius2 * cosf(angle_from_origin);
-								float term2 = radius * sinf(angle_from_origin);
-								float ellipse_radius_at_angle = radius * radius2 / sqrtf(term1 * term1 + term2 * term2);
+					// Cross product between the perpendicular vector from radius vector and origin radius and the normal vector
+					if (FVector3f::DotProduct(FVector3f::CrossProduct(RadiusVectorUnit, OriginRadiusVector), Normal) < 0.0f)
+					{
+						AngleFromOrigin = -AngleFromOrigin;
+					}
 
-								FVector3f vertex_proj_ellipse = current_center + radius_vector_unit * ellipse_radius_at_angle;
-								//FVector3f vertex_proj_ellipse = current_center + radius_vector_unit * radius;
+					AngleFromOrigin += Angle * PI / 180.f;
 
-								float morph_alpha = dist != 0.f && dot_morph <= dist ? FMath::Clamp(powf(dot_morph / dist, factor),0.f, 1.f) : 1.f;
+					float Term1 = Radius2 * FMath::Cos(AngleFromOrigin);
+					float Term2 = Radius * FMath::Sin(AngleFromOrigin);
+					float EllipseRadiusAtAngle = Radius * Radius2 * FMath::InvSqrt(Term1*Term1 + Term2*Term2);
 
-								vertex = vertex * (1.f - morph_alpha) + vertex_proj_ellipse * morph_alpha;
+					FVector3f VertexProjEllipse = CurrentCenter + RadiusVectorUnit * EllipseRadiusAtAngle;
+					// FVector3f VertexProjEllipse = CurrentCenter + RadiusVectorUnit * Radius;
 
-								if (dot_cut >= 0.f)		// CHECK IF THE VERTEX SHOULD BE CLIPPED
-								{
-									FVector3f vert_displ = normal * -dot_cut;
-									vertex = vertex + vert_displ;
-									RemovedVertices[v] = true;
-								}
-							}
-						}
+					float MorphAlpha = Dist != 0.f && DotMorph <= Dist ? FMath::Clamp(FMath::Pow(DotMorph/Dist, Factor), 0.f, 1.f) : 1.f;
 
-						for (int i = 0; i < 3; ++i)
-						{
-							ConvertData(i, it.ptr(), it.GetFormat(), &vertex[0], MBF_FLOAT32);
-						}
+					Position = Position * (1.f - MorphAlpha) + VertexProjEllipse*MorphAlpha;
 
-						++it;
+					// check if the vertex should be clipped
+					if (DotCut >= 0.f)		
+					{
+						FVector3f VertDispl = Normal * -DotCut;
+						Position = Position + VertDispl;
+						VerticesToCull[VertexIndex] = true;
 					}
-					break;
-
-				default:
-					break;
 				}
 			}
-		}
-
-		// Now remove all the faces from the result mesh that have all vertices removed
-		UntypedMeshBufferIteratorConst itBase(Result->GetIndexBuffers(), MBS_VERTEXINDEX);
-		UntypedMeshBufferIterator itDest(Result->GetIndexBuffers(), MBS_VERTEXINDEX);
-		int32 aFaceCount = Result->GetFaceCount();
-
-		UntypedMeshBufferIteratorConst ito(Result->GetIndexBuffers(), MBS_VERTEXINDEX);
-		for (int f = 0; f < aFaceCount; ++f)
-		{
-            FUint32Vector3 ov;
-			ov[0] = ito.GetAsUINT32(); ++ito;
-			ov[1] = ito.GetAsUINT32(); ++ito;
-			ov[2] = ito.GetAsUINT32(); ++ito;
-		
-			bool all_vertexs_removed = RemovedVertices[ov[0]] && RemovedVertices[ov[1]] && RemovedVertices[ov[2]];
 
-			if (!all_vertexs_removed)
-			{
-				if (itDest.ptr() != itBase.ptr())
-				{
-					FMemory::Memcpy(itDest.ptr(), itBase.ptr(), itBase.GetElementSize() * 3);
-				}
-
-				itDest += 3;
-			}
-
-			itBase += 3;
+			(PositionIterBegin + VertexIndex).SetFromVec3f(Position);
 		}
 
-		SIZE_T removedIndices = itBase - itDest;
-		check(removedIndices % 3 == 0);
-
-		Result->GetIndexBuffers().SetElementCount(aFaceCount * 3 - (int32)removedIndices);
-
-		// TODO: Should redo/reorder the face buffer before SetElementCount since some deleted faces could be left and some remaining faces deleted.
-
-        // Fix the surface data if present.
-        if (Result->Surfaces.Num())
-        {
-            // We assume there will be only one.
-            check(Result->Surfaces.Num()==1);
-
-            Result->Surfaces[0].IndexCount -= (int32)removedIndices;
-        }
+        MeshRemoveVerticesWithCullSet(Result, VerticesToCull);
 	}
 }
diff --git a/Engine/Plugins/Experimental/Mutable/Source/MutableRuntime/Private/MuR/OpMeshDifference.h b/Engine/Plugins/Experimental/Mutable/Source/MutableRuntime/Private/MuR/OpMeshDifference.h
--- a/Engine/Plugins/Experimental/Mutable/Source/MutableRuntime/Private/MuR/OpMeshDifference.h
+++ b/Engine/Plugins/Experimental/Mutable/Source/MutableRuntime/Private/MuR/OpMeshDifference.h
@@ -204,7 +204,7 @@ namespace mu
 
 
         // Morphs are always single surfaced
-        Result->Surfaces.SetNum(0);
+        Result->Surfaces.Empty();
         Result->EnsureSurfaceData();
 	}
 
diff --git a/Engine/Plugins/Experimental/Mutable/Source/MutableRuntime/Private/MuR/OpMeshExtractLayoutBlock.h b/Engine/Plugins/Experimental/Mutable/Source/MutableRuntime/Private/MuR/OpMeshExtractLayoutBlock.h
--- a/Engine/Plugins/Experimental/Mutable/Source/MutableRuntime/Private/MuR/OpMeshExtractLayoutBlock.h
+++ b/Engine/Plugins/Experimental/Mutable/Source/MutableRuntime/Private/MuR/OpMeshExtractLayoutBlock.h
@@ -5,31 +5,31 @@
 #include "MuR/MeshPrivate.h"
 #include "MuR/Platform.h"
 
+#include "MuR/OpMeshRemove.h"
 
 namespace mu
 {
-
-
-    //---------------------------------------------------------------------------------------------
     inline void MeshExtractFromVertices( const Mesh* Source,
                                          Mesh* Result,
-                                         const TArray<int32>& oldToNew,
-                                         const TArray<int32>& newToOld )
+                                         const TArray<int32>& OldToNew,
+                                         const TArray<int32>& NewToOld )
     {
-        int32 ResultVertices = newToOld.Num();
+        int32 ResultVertices = NewToOld.Num();
 
         // Assemble the new vertex buffer
         Result->GetVertexBuffers().SetElementCount(ResultVertices);
-        for ( int b=0; b<Result->GetVertexBuffers().GetBufferCount(); ++b )
+        for (int32 BufferIndex = 0; BufferIndex < Result->GetVertexBuffers().GetBufferCount(); ++BufferIndex)
         {
-            const uint8* pSourceData = Source->GetVertexBuffers().GetBufferData(b);
-            uint8* pDest = Result->GetVertexBuffers().GetBufferData(b);
-            int32 size = Result->GetVertexBuffers().GetElementSize(b);
-            for ( int32 NewIndex =0; NewIndex < ResultVertices; ++NewIndex)
+            const uint8* pSourceData = Source->GetVertexBuffers().GetBufferData(BufferIndex);
+            uint8* pDest = Result->GetVertexBuffers().GetBufferData(BufferIndex);
+            int32 Size = Result->GetVertexBuffers().GetElementSize(BufferIndex);
+
+            for (int32 NewIndex = 0; NewIndex < ResultVertices; ++NewIndex)
             {
-                int32 OldIndex = newToOld[NewIndex];
-                FMemory::Memcpy( pDest, pSourceData + size*OldIndex, size );
-                pDest+=size;
+                int32 OldIndex = NewToOld[NewIndex];
+                FMemory::Memcpy(pDest, pSourceData + Size*OldIndex, Size);
+
+                pDest += Size;
             }
         }
 
@@ -54,68 +54,74 @@ namespace mu
 
 			for (int32 NewIndex = 0; NewIndex < ResultVertices; ++NewIndex)
 			{
-				uint32 OldIndex = newToOld[NewIndex];
+				uint32 OldIndex = NewToOld[NewIndex];
 				(*pIdData++) = OldIndex;
 			}
 		}
 
         // Assemble the new index buffers
-		TArray<bool> usedSourceFaces;
-		usedSourceFaces.SetNumZeroed(Source->GetFaceCount());
-        UntypedMeshBufferIteratorConst itIndex( Source->GetIndexBuffers(), MBS_VERTEXINDEX );
-        UntypedMeshBufferIterator itResultIndex( Result->GetIndexBuffers(), MBS_VERTEXINDEX );
-        int indexCount = 0;
-        if ( itIndex.GetFormat()==MBF_UINT32 )
+		TBitArray<> UsedSourceFaces;
+
+		UsedSourceFaces.SetNum(Source->GetFaceCount(), false);
+        UntypedMeshBufferIteratorConst itIndex(Source->GetIndexBuffers(), MBS_VERTEXINDEX);
+        UntypedMeshBufferIterator itResultIndex(Result->GetIndexBuffers(), MBS_VERTEXINDEX);
+        
+        int32 IndexCount = 0;
+        if (itIndex.GetFormat() == MBF_UINT32)
         {
-            const uint32* pIndices = reinterpret_cast<const uint32_t*>( itIndex.ptr() );
-            uint32* pDestIndices = reinterpret_cast<uint32_t*>( itResultIndex.ptr() );
-            for ( int32 i=0; i<Source->GetIndexCount()/3; ++i )
+            const uint32* pIndices = reinterpret_cast<const uint32*>(itIndex.ptr());
+            uint32* pDestIndices = reinterpret_cast<uint32*>(itResultIndex.ptr());
+            for (int32 i = 0; i < Source->GetIndexCount()/3; ++i)
             {
-                if ( oldToNew[ pIndices[i*3+0] ]>=0
-                     &&
-                     oldToNew[ pIndices[i*3+1] ]>=0
-                     &&
-                     oldToNew[ pIndices[i*3+2] ]>=0 )
+                if (OldToNew[pIndices[i*3 + 0]] >= 0 && OldToNew[pIndices[i*3 + 1]] >= 0 && OldToNew[pIndices[i*3 + 2]] >= 0)
                 {
-                    usedSourceFaces[i] = true;
+                    UsedSourceFaces[i] = true;
 
                     // Clamp in case triangles go across blocks
-                    pDestIndices[ indexCount++ ] = FMath::Max( 0, oldToNew[ pIndices[i*3+0] ] );
-                    pDestIndices[ indexCount++ ] = FMath::Max( 0, oldToNew[ pIndices[i*3+1] ] );
-                    pDestIndices[ indexCount++ ] = FMath::Max( 0, oldToNew[ pIndices[i*3+2] ] );
+                    pDestIndices[IndexCount++] = FMath::Max(0, OldToNew[pIndices[i*3 + 0]]);
+                    pDestIndices[IndexCount++] = FMath::Max(0, OldToNew[pIndices[i*3 + 1]]);
+                    pDestIndices[IndexCount++] = FMath::Max(0, OldToNew[pIndices[i*3 + 2]]);
                 }
             }
         }
-        else if ( itIndex.GetFormat()==MBF_UINT16 )
+        else if (itIndex.GetFormat() == MBF_UINT16)
         {
-            const uint16* pIndices = reinterpret_cast<const uint16*>( itIndex.ptr() );
-            uint16* pDestIndices = reinterpret_cast<uint16*>( itResultIndex.ptr() );
-            for ( int32 i=0; i<Source->GetIndexCount()/3; ++i )
+            const uint16* pIndices = reinterpret_cast<const uint16*>(itIndex.ptr());
+            uint16* pDestIndices = reinterpret_cast<uint16*>(itResultIndex.ptr());
+            for (int32 i = 0; i < Source->GetIndexCount()/3; ++i)
             {
-                if ( oldToNew[ pIndices[i*3+0] ]>=0
-                     &&
-                     oldToNew[ pIndices[i*3+1] ]>=0
-                     &&
-                     oldToNew[ pIndices[i*3+2] ]>=0 )
+                if (OldToNew[pIndices[i*3 + 0]] >= 0 && OldToNew[pIndices[i*3 + 1]] >= 0 && OldToNew[pIndices[i*3 + 2]] >= 0)
                 {
-                    usedSourceFaces[i] = true;
+                    UsedSourceFaces[i] = true;
 
                     // Clamp in case triangles go across blocks
-                    pDestIndices[ indexCount++ ] = (uint16)FMath::Max( 0, oldToNew[ pIndices[i*3+0] ] );
-                    pDestIndices[ indexCount++ ] = (uint16)FMath::Max( 0, oldToNew[ pIndices[i*3+1] ] );
-                    pDestIndices[ indexCount++ ] = (uint16)FMath::Max( 0, oldToNew[ pIndices[i*3+2] ] );
+                    pDestIndices[IndexCount++] = (uint16)FMath::Max(0, OldToNew[pIndices[i*3 + 0]]);
+                    pDestIndices[IndexCount++] = (uint16)FMath::Max(0, OldToNew[pIndices[i*3 + 1]]);
+                    pDestIndices[IndexCount++] = (uint16)FMath::Max(0, OldToNew[pIndices[i*3 + 2]]);
                 }
             }
         }
         else
         {
-            check( false );
+            check(false);
         }
-        Result->GetIndexBuffers().SetElementCount( indexCount );
+     
+        Result->GetIndexBuffers().SetElementCount(IndexCount);
+
+        const int32 NumOriginalVerts = OldToNew.Num();
+
+        TBitArray<> UsedVertices;
+        UsedVertices.SetNum(NumOriginalVerts, false);
+
+        for (int32 I = 0; I < NumOriginalVerts; ++I)
+        {
+            UsedVertices[I] = OldToNew[I] >= 0;
+        }
+
+        MeshRemoveRecreateSurface(Result, UsedVertices, UsedSourceFaces);
     }
 
 
-	//---------------------------------------------------------------------------------------------
     inline void MeshExtractLayoutBlock(Mesh* Result, const Mesh* Source,
                                            uint32 LayoutIndex,
                                            uint16 BlockCount,
@@ -123,42 +129,42 @@ namespace mu
 	{
 		check(Source);
 		bOutSuccess = true;
-		
+
 		// TODO: Optimise
 		Result->CopyFrom(*Source);
 
 		UntypedMeshBufferIteratorConst itBlocks(Source->GetVertexBuffers(), MBS_LAYOUTBLOCK, LayoutIndex);
 
-        if (itBlocks.GetFormat()!=MBF_NONE)
+        if (itBlocks.GetFormat() != MBF_NONE)
         {
             int32 ResultVertices = 0;
-			TArray<int32> oldToNew;
-			oldToNew.Init(-1,Source->GetVertexCount());
-			TArray<int32> newToOld;
-            newToOld.Reserve( Source->GetVertexCount() );
+			TArray<int32> OldToNew;
+			OldToNew.Init(-1, Source->GetVertexCount());
+			TArray<int32> NewToOld;
+            NewToOld.Reserve(Source->GetVertexCount());
 
-            if ( itBlocks.GetFormat()==MBF_UINT16 )
+            if (itBlocks.GetFormat() == MBF_UINT16)
             {
-                const uint16* pBlocks = reinterpret_cast<const uint16*>( itBlocks.ptr() );
+                const uint16* pBlocks = reinterpret_cast<const uint16*>(itBlocks.ptr());
                 for ( int32 i=0; i<Source->GetVertexCount(); ++i )
                 {
                     uint64 VertexBlockRelative = pBlocks[i];
 					uint64 VertexBlockId = (uint64(Source->MeshIDPrefix) << 32) | VertexBlockRelative;
 
-                    bool found = false;
-                    for ( int32 j=0; j< BlockCount; ++j)
+                    bool bFound = false;
+                    for (int32 j = 0; j < BlockCount; ++j)
                     {
                         if (VertexBlockId == BlockIds[j])
                         {
-                            found = true;
+                            bFound = true;
                             break;
                         }
                     }
 
-                    if ( found )
+                    if (bFound)
                     {
-                        oldToNew[i] = ResultVertices++;
-                        newToOld.Add( i );
+                        OldToNew[i] = ResultVertices++;
+                        NewToOld.Add(i);
                     }
                 }
             }
@@ -169,34 +175,30 @@ namespace mu
 				{
 					uint64 VertexBlockId = pBlocks[i];
 
-					bool found = false;
-					for (int j = 0; j < BlockCount; ++j)
+					bool bFound = false;
+					for (int32 j = 0; j < BlockCount; ++j)
 					{
 						if (VertexBlockId == BlockIds[j])
 						{
-							found = true;
+							bFound = true;
 							break;
 						}
 					}
 
-					if (found)
+					if (bFound)
 					{
-						oldToNew[i] = ResultVertices++;
-						newToOld.Add(i);
+						OldToNew[i] = ResultVertices++;
+						NewToOld.Add(i);
 					}
 				}
 			}
 			else
-
             {
                 check( false );
             }
 
-            MeshExtractFromVertices(Source, Result, oldToNew, newToOld);
+            MeshExtractFromVertices(Source, Result, OldToNew, NewToOld);
         }
-
-        Result->Surfaces.Empty();
-        Result->EnsureSurfaceData();
 	}
 
 }
diff --git a/Engine/Plugins/Experimental/Mutable/Source/MutableRuntime/Private/MuR/OpMeshOptimizeSkinning.h b/Engine/Plugins/Experimental/Mutable/Source/MutableRuntime/Private/MuR/OpMeshOptimizeSkinning.h
--- a/Engine/Plugins/Experimental/Mutable/Source/MutableRuntime/Private/MuR/OpMeshOptimizeSkinning.h
+++ b/Engine/Plugins/Experimental/Mutable/Source/MutableRuntime/Private/MuR/OpMeshOptimizeSkinning.h
@@ -151,32 +151,35 @@ namespace mu
 				{
 					const uint32 NumBonesInBoneMap = Surface.BoneMapCount;
 
-					for (int32 VertexIndex = 0; VertexIndex < Surface.VertexCount; ++VertexIndex)
+					for (const FSurfaceSubMesh SubMesh : Surface.SubMeshes)
 					{
-						FMemory::Memcpy(Data, SourceData, BoneIndexChannelOffset);
-						Data += BoneIndexChannelOffset;
-						SourceData += BoneIndexChannelOffset;
-
-						for (int32 ComponentIndex = 0; ComponentIndex < BoneIndexComponentCount; ++ComponentIndex)
+						for (int32 VertexIndex = SubMesh.VertexBegin; VertexIndex < SubMesh.VertexEnd; ++VertexIndex)
 						{
-							const uint16 SourceIndex = *((const uint16*)SourceData);
-							if (SourceIndex < NumBonesInBoneMap)
-							{
-								*Data = (uint8)SourceIndex;
-							}
-							else
+							FMemory::Memcpy(Data, SourceData, BoneIndexChannelOffset);
+							Data += BoneIndexChannelOffset;
+							SourceData += BoneIndexChannelOffset;
+
+							for (int32 ComponentIndex = 0; ComponentIndex < BoneIndexComponentCount; ++ComponentIndex)
 							{
-								*Data = 0;
+								const uint16 SourceIndex = *((const uint16*)SourceData);
+								if (SourceIndex < NumBonesInBoneMap)
+								{
+									*Data = (uint8)SourceIndex;
+								}
+								else
+								{
+									*Data = 0;
+								}
+								
+								Data += BoneIndexSize;
+								SourceData += SourceBoneIndexSize;
 							}
-							
-							Data += BoneIndexSize;
-							SourceData += SourceBoneIndexSize;
-						}
 
 
-						FMemory::Memcpy(Data, SourceData, TailSize);
-						Data += TailSize;
-						SourceData += TailSize;
+							FMemory::Memcpy(Data, SourceData, TailSize);
+							Data += TailSize;
+							SourceData += TailSize;
+						}
 					}
 				}
 			}
diff --git a/Engine/Plugins/Experimental/Mutable/Source/MutableRuntime/Private/MuR/OpMeshRemove.cpp b/Engine/Plugins/Experimental/Mutable/Source/MutableRuntime/Private/MuR/OpMeshRemove.cpp
--- a/Engine/Plugins/Experimental/Mutable/Source/MutableRuntime/Private/MuR/OpMeshRemove.cpp
+++ b/Engine/Plugins/Experimental/Mutable/Source/MutableRuntime/Private/MuR/OpMeshRemove.cpp
@@ -9,6 +9,64 @@
 
 namespace mu
 {
+    void MeshRemoveRecreateSurface(Mesh* Result, const TBitArray<>& UsedVertices, const TBitArray<>& UsedFaces)
+    {
+        MUTABLE_CPUPROFILER_SCOPE(MeshRemoveRecreateSurface);
+		TArray<FSurfaceSubMesh, TInlineAllocator<32>> OrigSubMeshes;
+        for (FMeshSurface& ResultSurf : Result->Surfaces)
+        {
+			OrigSubMeshes = ResultSurf.SubMeshes;
+			ResultSurf.SubMeshes.Reset();            
+
+			int32 PrevMeshRangeVertexEnd = 0;
+            int32 PrevMeshRangeIndexEnd = 0;
+            for (const FSurfaceSubMesh& SubMesh : OrigSubMeshes)
+            {
+                int32 MeshRangeVertexEnd = PrevMeshRangeVertexEnd; 
+                int32 MeshRangeIndexEnd = PrevMeshRangeIndexEnd;
+               
+                MeshRangeVertexEnd += UsedVertices.CountSetBits(SubMesh.VertexBegin, SubMesh.VertexEnd);
+				
+                // Only add the mesh if it has remaining vertices. 
+                if (MeshRangeVertexEnd > PrevMeshRangeVertexEnd)
+                {
+					check(SubMesh.IndexBegin % 3 == 0);
+					check((SubMesh.IndexEnd - SubMesh.IndexBegin) % 3 == 0);
+                    MeshRangeIndexEnd += UsedFaces.CountSetBits(SubMesh.IndexBegin/3, SubMesh.IndexEnd/3)*3;
+
+                    ResultSurf.SubMeshes.Emplace(FSurfaceSubMesh 
+                            {
+                                PrevMeshRangeVertexEnd, MeshRangeVertexEnd,
+                                PrevMeshRangeIndexEnd, MeshRangeIndexEnd,
+                                SubMesh.ExternalId 
+                            });
+                }
+
+                PrevMeshRangeVertexEnd = MeshRangeVertexEnd;
+                PrevMeshRangeIndexEnd = MeshRangeIndexEnd;
+            }
+        }
+		
+		// Remove Empty surfaces but always keep the first one. 
+		// The previous step has eliminated empty submeshes, so it is only needed to check if the surface has 
+		// any submesh.
+		for (int32 I = Result->Surfaces.Num() - 1; I >= 1; --I)
+		{
+			if (!Result->Surfaces[I].SubMeshes.Num())
+			{
+				Result->Surfaces.RemoveAt(I, EAllowShrinking::No);
+			}
+		}
+
+		check(Result->Surfaces.Num() >= 1);
+	
+		// Add a defaulted empty submesh if the surface is empty. A surface always needs a submesh 
+		// even if empty.
+		if (!Result->Surfaces[0].SubMeshes.Num())
+		{
+			Result->Surfaces[0].SubMeshes.Emplace();
+		}
+    }
 
     //---------------------------------------------------------------------------------------------
     struct FIdInterval
@@ -73,198 +131,432 @@ namespace mu
         return -1;
     }
 
+    // TODO: Remove once the new implementation is proven to be faster for all cases and does not present 
+    // any problem. For now keep it as a reference.
+	//void MeshRemoveVerticesWithMap( Mesh* Result, const TBitArray<>& RemovedVertices)
+	//{
+    //    uint32 RemovedVertexCount = RemovedVertices.Num();
+
+	//	int32 FirstFreeVertex = 0;
+	//	int32 RemovedIndices = 0;
+
+    //    // Rebuild index buffers
+
+    //    // Map from source vertex index, to new vertex index for used vertices.
+    //    // These are indices as in the index buffer, not the absolute vertex index as in the
+    //    // vertexbuffer MBS_VERTEXINDEX buffers.
+	//	TArray<int32> UsedVertexMap;
+    //    TBitArray<> UsedFaces;
+	//	UsedVertexMap.Init(-1, Result->GetVertexCount());
+    //    UsedFaces.Init(false, Result->GetIndexCount() / 3);  
+    //    {
+    //        if ( Result->GetIndexBuffers().GetElementSize(0)==4 )
+    //        {
+    //            MeshBufferIteratorConst<MBF_UINT32,uint32,1> itSource( Result->GetIndexBuffers(), MBS_VERTEXINDEX );
+    //            MeshBufferIterator<MBF_UINT32,uint32,1> itDest( Result->GetIndexBuffers(), MBS_VERTEXINDEX );
+
+    //            int32 IndexCount = Result->GetIndexCount();
+	//			check(IndexCount%3==0);
+    //            for ( int32 f=0; f<IndexCount/3; ++f )
+    //            {
+    //                uint32 sourceIndices[3];
+    //                sourceIndices[0] = (*itSource)[0];
+    //                sourceIndices[1] = (*itSource)[1];
+    //                sourceIndices[2] = (*itSource)[2];
+
+	//				check(sourceIndices[0] < RemovedVertexCount);
+	//				check(sourceIndices[1] < RemovedVertexCount);
+	//				check(sourceIndices[2] < RemovedVertexCount);
+
+	//				int32 RemoveCount =  (RemovedVertices[ sourceIndices[0] ]
+    //                        + RemovedVertices[ sourceIndices[1] ]
+    //                        + RemovedVertices[ sourceIndices[2] ]
+    //                        );
+	//				bool bFaceRemoved = RemoveCount > 0;
+    //                if (!bFaceRemoved)
+    //                {
+    //                    UsedFaces[f] = true;
+
+    //                    for (int32 i=0;i<3;++i)
+    //                    {
+    //                        uint32 sourceIndex = sourceIndices[i];
+
+    //                        if (UsedVertexMap[ sourceIndex ] < 0 )
+    //                        {
+	//							UsedVertexMap[ sourceIndex ] = FirstFreeVertex;
+    //                            FirstFreeVertex++;
+    //                        }
+
+    //                        uint32 destIndex = UsedVertexMap[ sourceIndex ];
+    //                        *(uint32*)itDest.ptr() = destIndex;
+
+    //                        itDest++;
+    //                    }
+    //                }
+
+    //                itSource+=3;
+    //            }
+
+    //            RemovedIndices = itSource - itDest;
+    //        }
+
+    //        else if ( Result->GetIndexBuffers().GetElementSize(0)==2 )
+    //        {
+    //            MeshBufferIteratorConst<MBF_UINT16,uint16,1> itSource( Result->GetIndexBuffers(), MBS_VERTEXINDEX );
+    //            MeshBufferIterator<MBF_UINT16,uint16,1> itDest( Result->GetIndexBuffers(), MBS_VERTEXINDEX );
+
+    //            int32 indexCount = Result->GetIndexCount();
+    //            for ( int32 f=0; f<indexCount/3; ++f )
+    //            {
+    //                uint16 sourceIndices[3];
+    //                sourceIndices[0] = (*itSource)[0];
+    //                sourceIndices[1] = (*itSource)[1];
+    //                sourceIndices[2] = (*itSource)[2];
+
+	//				check(sourceIndices[0] < RemovedVertexCount);
+	//				check(sourceIndices[1] < RemovedVertexCount);
+	//				check(sourceIndices[2] < RemovedVertexCount);
+
+	//				int32 RemoveCount =  (RemovedVertices[ sourceIndices[0] ]
+    //                        + RemovedVertices[ sourceIndices[1] ]
+    //                        + RemovedVertices[ sourceIndices[2] ]
+    //                        );
+	//				bool bFaceRemoved = RemoveCount > 0;
+
+    //                if (!bFaceRemoved)
+    //                {
+    //                    UsedFaces[f] = true;
+    //                    for (int32 i = 0; i < 3; ++i)
+    //                    {
+    //                        uint16 sourceIndex = sourceIndices[i];
+
+    //                        if (UsedVertexMap[ sourceIndex ] < 0 )
+    //                        {
+	//							UsedVertexMap[ sourceIndex ] = FirstFreeVertex;
+    //                            FirstFreeVertex++;
+    //                        }
+
+    //                        uint16 destIndex = (uint16)UsedVertexMap[ sourceIndex ];
+    //                        *(uint16*)itDest.ptr() = destIndex;
+
+    //                        itDest++;
+    //                    }
+    //                }
+
+    //                itSource+=3;
+    //            }
+
+    //            RemovedIndices = itSource - itDest;
+    //        }
+
+    //        else
+    //        {
+    //            // Index buffer format case not implemented
+    //            check( false );
+    //        }
+
+    //        check( RemovedIndices%3==0 );
+
+    //        int32 FaceCount = Result->GetFaceCount();
+    //        Result->GetIndexBuffers().SetElementCount( FaceCount*3-RemovedIndices );
+    //    }
+
+
+    //    // Rebuild the vertex buffers
+
+	//	// If we had implicit indices, make them explicit or relative to keep them valid
+	//	if (RemovedIndices && Result->AreVertexIdsImplicit())
+	//	{
+	//		Result->MakeVertexIdsRelative();
+	//	}
+
+	//	// The temp array is necessary because if the vertex buffer is not sorted according to the index buffer we cannot do it in-place
+	//	// This happens with some mesh import options.
+	//	TArray<uint8> Temp;
+	//	for ( int32 b=0; b<Result->GetVertexBuffers().GetBufferCount(); ++b )
+    //    {
+    //        int32 elemSize = Result->GetVertexBuffers().GetElementSize( b );
+    //        const uint8* SourceData = Result->GetVertexBuffers().GetBufferData( b );
+
+	//		Temp.SetNumUninitialized(FirstFreeVertex*elemSize,EAllowShrinking::No);
+    //        uint8* DestData = Temp.GetData();
+
+    //        for ( int32 v=0; v<Result->GetVertexCount(); ++v )
+    //        {
+    //            int32 span = 0;
+    //            for ( int32 s=0; v+s<Result->GetVertexCount(); ++s )
+    //            {
+    //                if (UsedVertexMap[v+s]>=0 )
+    //                {
+    //                    if (span==0)
+    //                    {
+    //                        ++span;
+    //                    }
+    //                    else
+    //                    {
+    //                        if (UsedVertexMap[v+s] == UsedVertexMap[v+s-1]+1 )
+    //                        {
+    //                            ++span;
+    //                        }
+    //                        else
+    //                        {
+    //                            break;
+    //                        }
+    //                    }
+    //                }
+    //                else
+    //                {
+    //                    break;
+    //                }
+    //            }
+
+
+    //            if (span>0)
+    //            {
+    //                FMemory::Memcpy( DestData+elemSize*UsedVertexMap[v], SourceData+elemSize*v, elemSize*span );
+
+    //                v += span-1;
+    //            }
+    //        }
+
+	//		// Copy from temp buffer to final vertex buffer
+	//		FMemory::Memcpy( Result->GetVertexBuffers().GetBufferData(b), DestData, FirstFreeVertex*elemSize);
+    //    }
+    //    
+    //    Result->GetVertexBuffers().SetElementCount(FirstFreeVertex);
+
+    //    //Temp fix, transform the vertex map to a bitset.
+    //    TBitArray<> UsedVertices;
+    //    UsedVertices.SetNum(RemovedVertexCount, false);
+
+    //    for (uint32 I = 0; I < RemovedVertexCount; ++I)
+    //    {
+    //        UsedVertices[I] = RemovedVertices[I] == 0;
+    //    }
+
+    //    MeshRemoveRecreateSurface(Result, UsedVertices, UsedFaces);
+    //}
+
+    void MeshRemoveVerticesWithCullSet(Mesh* Result, const TBitArray<>& VerticesToCull)
+    { 
+        MUTABLE_CPUPROFILER_SCOPE(MeshRemoveVerticesWithCullSet);
+		
+        const UntypedMeshBufferIterator IndicesBegin(Result->GetIndexBuffers(), MBS_VERTEXINDEX);
+		
+        const int32 NumFaces = Result->GetFaceCount();
+        const int32 NumVertices = Result->GetVertexCount();
+
+		TBitArray<> UsedVertices;
+		UsedVertices.SetNum(NumVertices, false);
+	
+		TBitArray<> UsedFaces;
+		UsedFaces.SetNum(NumFaces, false);
+
+        int32 NumUsedFaces = 0;
+		const uint32 IndexTypeSize = IndicesBegin.GetElementSize();
+
+        if (IndexTypeSize == 4)
+        {
+            for (int32 FaceIndex = 0; FaceIndex < NumFaces; ++FaceIndex)
+            {
+                const uint32* FaceIndicesData = reinterpret_cast<uint32*>((IndicesBegin + FaceIndex*3).ptr());
 
-	void MeshRemoveVerticesWithMap( Mesh* Result, const uint8* RemovedVertices, uint32 RemovedVertexCount)
-	{
-		int32 FirstFreeVertex = 0;
-		int32 RemovedIndices = 0;
+                const bool bAllVertsRemoved = 
+                        VerticesToCull[FaceIndicesData[0]] & 
+                        VerticesToCull[FaceIndicesData[1]] & 
+                        VerticesToCull[FaceIndicesData[2]];
 
-        // Rebuild index buffers
+                if (!bAllVertsRemoved)
+                {
+                    ++NumUsedFaces;
+                    UsedFaces[FaceIndex] = true;
 
-        // Map from source vertex index, to new vertex index for used vertices.
-        // These are indices as in the index buffer, not the absolute vertex index as in the
-        // vertexbuffer MBS_VERTEXINDEX buffers.
-		TArray<int32> UsedVertexMap;
-		UsedVertexMap.Init( -1, Result->GetVertexCount() );
+                    UsedVertices[FaceIndicesData[0]] = true;
+                    UsedVertices[FaceIndicesData[1]] = true;
+                    UsedVertices[FaceIndicesData[2]] = true;
+                }
+            }
+        }
+        else if (IndexTypeSize == 2)
         {
-            if ( Result->GetIndexBuffers().GetElementSize(0)==4 )
+            for (int32 FaceIndex = 0; FaceIndex < NumFaces; ++FaceIndex)
             {
-                MeshBufferIteratorConst<MBF_UINT32,uint32,1> itSource( Result->GetIndexBuffers(), MBS_VERTEXINDEX );
-                MeshBufferIterator<MBF_UINT32,uint32,1> itDest( Result->GetIndexBuffers(), MBS_VERTEXINDEX );
+                const uint16* FaceIndicesData = reinterpret_cast<uint16*>((IndicesBegin + FaceIndex*3).ptr());
 
-                int32 IndexCount = Result->GetIndexCount();
-				check(IndexCount%3==0);
-                for ( int32 f=0; f<IndexCount/3; ++f )
+                const bool bAllVertsRemoved = 
+                        VerticesToCull[FaceIndicesData[0]] & 
+                        VerticesToCull[FaceIndicesData[1]] & 
+                        VerticesToCull[FaceIndicesData[2]];
+
+                if (!bAllVertsRemoved)
                 {
-                    uint32 sourceIndices[3];
-                    sourceIndices[0] = (*itSource)[0];
-                    sourceIndices[1] = (*itSource)[1];
-                    sourceIndices[2] = (*itSource)[2];
-
-					check(sourceIndices[0] < RemovedVertexCount);
-					check(sourceIndices[1] < RemovedVertexCount);
-					check(sourceIndices[2] < RemovedVertexCount);
-
-					int32 RemoveCount =  (RemovedVertices[ sourceIndices[0] ]
-                            + RemovedVertices[ sourceIndices[1] ]
-                            + RemovedVertices[ sourceIndices[2] ]
-                            );
-					bool bFaceRemoved = RemoveCount > 0;
-
-                    if ( !bFaceRemoved)
-                    {
-                        for (int32 i=0;i<3;++i)
-                        {
-                            uint32 sourceIndex = sourceIndices[i];
+                    ++NumUsedFaces;
+                    UsedFaces[FaceIndex] = true;
 
-                            if (UsedVertexMap[ sourceIndex ] < 0 )
-                            {
-								UsedVertexMap[ sourceIndex ] = FirstFreeVertex;
-                                FirstFreeVertex++;
-                            }
+                    UsedVertices[FaceIndicesData[0]] = true;
+                    UsedVertices[FaceIndicesData[1]] = true;
+                    UsedVertices[FaceIndicesData[2]] = true;
+                }
+            }
+        }
+        else
+        {
+            check(false);
+        }
+ 
+		if (NumUsedFaces < NumFaces && Result->AreVertexIdsImplicit())
+		{
+			Result->MakeVertexIdsRelative();
+		}
 
-                            uint32 destIndex = UsedVertexMap[ sourceIndex ];
-                            *(uint32*)itDest.ptr() = destIndex;
+		TArray<int32> UsedVerticesMap;
+#if DO_CHECK
+        UsedVerticesMap.Init(-1, NumVertices);
+#else
+        // This data will only be accessed by indices that have been mapped, No need to initialize. 
+		UsedVerticesMap.SetNumUninitialized(NumVertices);
+#endif
+		FMeshBufferSet& VertexBufferSet = Result->GetVertexBuffers();
 
-                            itDest++;
-                        }
-                    }
+		const int32 NumBuffers = VertexBufferSet.GetBufferCount();
 
-                    itSource+=3;
+		// Compute vertices indices remap
+        int32 NumVerticesRemaining = 0;
+		if (NumUsedFaces > 0) 
+        {
+		    int32 LastFreeVertexIndex = 0;
+            for (int32 VertexIndex = UsedVertices.FindFrom(true, 0); VertexIndex >= 0;)
+			{
+                const int32 UsedSpanBegin = VertexIndex;
+                VertexIndex = UsedVertices.FindFrom(false, VertexIndex);
+				
+				// At the end of the buffer we may not find a false element, in that case
+				// FindForm returns INDEX_NONE, set the vertex at the range end. 
+                VertexIndex = VertexIndex >= 0 ? VertexIndex : NumVertices;
+                const int32 UsedSpanEnd = VertexIndex;
+
+				// VertexIndex may be one past the end of the array, VertexIndex will become INDEX_NONE
+				// and the loop will finish.
+                VertexIndex = UsedVertices.FindFrom(true, VertexIndex);
+				
+                for (int32 I = UsedSpanBegin; I < UsedSpanEnd; ++I)
+                {
+                    UsedVerticesMap[I] = LastFreeVertexIndex + I - UsedSpanBegin;
                 }
 
-                RemovedIndices = itSource - itDest;
-            }
+				LastFreeVertexIndex += UsedSpanEnd - UsedSpanBegin;
+			}
+
+            NumVerticesRemaining = LastFreeVertexIndex;
+		}
 
-            else if ( Result->GetIndexBuffers().GetElementSize(0)==2 )
+		// Copy move buffers. We are recomputing the spans for each buffer, should be ok as
+        // finding the span is fast compared to the data move.
+        if (NumVerticesRemaining > 0)
+        {
+            for (int32 BufferIndex = 0; BufferIndex < NumBuffers; ++BufferIndex)
             {
-                MeshBufferIteratorConst<MBF_UINT16,uint16,1> itSource( Result->GetIndexBuffers(), MBS_VERTEXINDEX );
-                MeshBufferIterator<MBF_UINT16,uint16,1> itDest( Result->GetIndexBuffers(), MBS_VERTEXINDEX );
+                uint8* BufferData = VertexBufferSet.GetBufferData(BufferIndex); 
+                const uint32 ElemSize = VertexBufferSet.GetElementSize(BufferIndex);
 
-                int32 indexCount = Result->GetIndexCount();
-                for ( int32 f=0; f<indexCount/3; ++f )
-                {
-                    uint16 sourceIndices[3];
-                    sourceIndices[0] = (*itSource)[0];
-                    sourceIndices[1] = (*itSource)[1];
-                    sourceIndices[2] = (*itSource)[2];
-
-					check(sourceIndices[0] < RemovedVertexCount);
-					check(sourceIndices[1] < RemovedVertexCount);
-					check(sourceIndices[2] < RemovedVertexCount);
-
-					int32 RemoveCount =  (RemovedVertices[ sourceIndices[0] ]
-                            + RemovedVertices[ sourceIndices[1] ]
-                            + RemovedVertices[ sourceIndices[2] ]
-                            );
-					bool bFaceRemoved = RemoveCount > 0;
-
-                    if ( !bFaceRemoved)
-                    {
-                        for (int i=0;i<3;++i)
-                        {
-                            uint16 sourceIndex = sourceIndices[i];
+                int32 LastFreeVertexIndex = 0;
 
-                            if (UsedVertexMap[ sourceIndex ] < 0 )
-                            {
-								UsedVertexMap[ sourceIndex ] = FirstFreeVertex;
-                                FirstFreeVertex++;
-                            }
+                for (int32 VertexIndex = UsedVertices.FindFrom(true, 0); VertexIndex >= 0;)
+                {
+                    const int32 UsedSpanBegin = VertexIndex;
+                    VertexIndex = UsedVertices.FindFrom(false, VertexIndex);
+                    VertexIndex = VertexIndex >= 0 ? VertexIndex : NumVertices;
+                    const int32 UsedSpanEnd = VertexIndex;
 
-                            uint16 destIndex = (uint16)UsedVertexMap[ sourceIndex ];
-                            *(uint16*)itDest.ptr() = destIndex;
+                    VertexIndex = UsedVertices.FindFrom(true, VertexIndex);
 
-                            itDest++;
-                        }
+                    // Copy vertex buffer span.	
+                    const int32 UsedSpanSize = UsedSpanEnd - UsedSpanBegin;
+                    
+                    if (LastFreeVertexIndex != UsedSpanBegin)
+                    {
+                        FMemory::Memmove(
+                                BufferData + LastFreeVertexIndex*ElemSize,
+                                BufferData + UsedSpanBegin*ElemSize,
+                                UsedSpanSize*ElemSize);
                     }
 
-                    itSource+=3;
+                    LastFreeVertexIndex += UsedSpanSize;
                 }
 
-                RemovedIndices = itSource - itDest;
-            }
-
-            else
-            {
-                // Index buffer format case not implemented
-                check( false );
+                check(LastFreeVertexIndex == NumVerticesRemaining);
             }
-
-            check( RemovedIndices%3==0 );
-
-            int32 FaceCount = Result->GetFaceCount();
-            Result->GetIndexBuffers().SetElementCount( FaceCount*3-RemovedIndices );
         }
+		Result->GetVertexBuffers().SetElementCount(NumVerticesRemaining);
 
-
-        // Rebuild the vertex buffers
-
-		// If we had implicit indices, make them explicit or relative to keep them valid
-		if (RemovedIndices && Result->AreVertexIdsImplicit())
-		{
-			Result->MakeVertexIdsRelative();
-		}
-
-		// The temp array is necessary because if the vertex buffer is not sorted according to the index buffer we cannot do it in-place
-		// This happens with some mesh import options.
-		TArray<uint8> Temp;
-		for ( int32 b=0; b<Result->GetVertexBuffers().GetBufferCount(); ++b )
+		int32 LastFreeFaceIndex = 0;
+        if (NumUsedFaces > 0)
         {
-            int32 elemSize = Result->GetVertexBuffers().GetElementSize( b );
-            const uint8* SourceData = Result->GetVertexBuffers().GetBufferData( b );
-
-			Temp.SetNumUninitialized(FirstFreeVertex*elemSize,EAllowShrinking::No);
-            uint8* DestData = Temp.GetData();
+            for (int32 FaceIndex = UsedFaces.FindFrom(true, 0); FaceIndex >= 0;)
+            {		
+                const int32 UsedSpanStart = FaceIndex;
+                FaceIndex = UsedFaces.FindFrom(false, FaceIndex);
+                FaceIndex = FaceIndex >= 0 ? FaceIndex : NumFaces;
+                const int32 UsedSpanEnd = FaceIndex;
+
+                FaceIndex = UsedFaces.FindFrom(true, FaceIndex);
+
+                const int32 UsedSpanSize = UsedSpanEnd - UsedSpanStart;
+        
+                if (LastFreeFaceIndex != UsedSpanStart)
+                {
+                    FMemory::Memmove(
+                            (IndicesBegin + LastFreeFaceIndex*3).ptr(), 
+                            (IndicesBegin + UsedSpanStart*3).ptr(), 
+                            UsedSpanSize*IndexTypeSize*3);
+                }
 
-            for ( int32 v=0; v<Result->GetVertexCount(); ++v )
-            {
-                int32 span = 0;
-                for ( int32 s=0; v+s<Result->GetVertexCount(); ++s )
+                // Remap vertices
+                if (IndexTypeSize == 4)
                 {
-                    if (UsedVertexMap[v+s]>=0 )
+                    for (int32 I = LastFreeFaceIndex; I < LastFreeFaceIndex + UsedSpanSize; ++I)
                     {
-                        if (span==0)
-                        {
-                            ++span;
-                        }
-                        else
-                        {
-                            if (UsedVertexMap[v+s] == UsedVertexMap[v+s-1]+1 )
-                            {
-                                ++span;
-                            }
-                            else
-                            {
-                                break;
-                            }
-                        }
+                        uint32* FaceIndicesData = reinterpret_cast<uint32*>((IndicesBegin + I*3).ptr());
+
+                        check(UsedVerticesMap[FaceIndicesData[0]] >= 0);
+                        check(UsedVerticesMap[FaceIndicesData[1]] >= 0);
+                        check(UsedVerticesMap[FaceIndicesData[2]] >= 0);
+
+                        FaceIndicesData[0] = UsedVerticesMap[FaceIndicesData[0]];
+                        FaceIndicesData[1] = UsedVerticesMap[FaceIndicesData[1]];
+                        FaceIndicesData[2] = UsedVerticesMap[FaceIndicesData[2]];
                     }
-                    else
+                }
+                else if (IndexTypeSize == 2)
+                {	
+                    for (int32 I = LastFreeFaceIndex; I < LastFreeFaceIndex + UsedSpanSize; ++I)
                     {
-                        break;
+                        uint16* FaceIndicesData = reinterpret_cast<uint16*>((IndicesBegin + I*3).ptr());
+
+                        check(UsedVerticesMap[FaceIndicesData[0]] >= 0);
+                        check(UsedVerticesMap[FaceIndicesData[1]] >= 0);
+                        check(UsedVerticesMap[FaceIndicesData[2]] >= 0);
+                        
+                        FaceIndicesData[0] = static_cast<uint16>(UsedVerticesMap[FaceIndicesData[0]]);
+                        FaceIndicesData[1] = static_cast<uint16>(UsedVerticesMap[FaceIndicesData[1]]);
+                        FaceIndicesData[2] = static_cast<uint16>(UsedVerticesMap[FaceIndicesData[2]]);
                     }
                 }
-
-
-                if (span>0)
+                else
                 {
-                    FMemory::Memcpy( DestData+elemSize*UsedVertexMap[v], SourceData+elemSize*v, elemSize*span );
-
-                    v += span-1;
+                    check(false);
                 }
-            }
 
-			// Copy from temp buffer to final vertex buffer
-			FMemory::Memcpy( Result->GetVertexBuffers().GetBufferData(b), DestData, FirstFreeVertex*elemSize);
+                LastFreeFaceIndex += UsedSpanSize;
+            }
         }
-        Result->GetVertexBuffers().SetElementCount( FirstFreeVertex );
 
-        // Rebuild surface data.
-        // \todo: For now make single surfaced
-        Result->Surfaces.Empty();
-        Result->EnsureSurfaceData();
-    }
+		check(LastFreeFaceIndex <= NumFaces);
 
+		Result->GetIndexBuffers().SetElementCount(LastFreeFaceIndex*3);
+
+		MeshRemoveRecreateSurface(Result, UsedVertices, UsedFaces);
+    }
 
 	void MeshRemoveMask(Mesh* Result, const Mesh* Source, const Mesh* Mask, bool& bOutSuccess)
     {
@@ -283,8 +575,8 @@ namespace mu
 
         // For each source vertex, true if it is removed.
         int32 ResultVertexCount = Result->GetVertexCount();
-		TArray<uint8> RemovedVertices;
-		RemovedVertices.SetNumZeroed(ResultVertexCount);
+		TBitArray<> RemovedVertices;
+		RemovedVertices.SetNum(ResultVertexCount, false);
         {
 			TArray<FIdInterval> Intervals;
             ExtractVertexIndexIntervals(Intervals, Result);
@@ -296,14 +588,13 @@ namespace mu
                 ++itMaskVI;
 
                 int32 IndexInSource = FindPositionInIntervals(Intervals, MaskVertexId);
-                if (IndexInSource >=0)
+                if (IndexInSource >= 0)
                 {
-					RemovedVertices[IndexInSource] = 1;
+					RemovedVertices[IndexInSource] = true;
                 }
             }
         }
-
-		MeshRemoveVerticesWithMap(Result, RemovedVertices.GetData(), (uint32)RemovedVertices.Num());
+		
+        MeshRemoveVerticesWithCullSet(Result, RemovedVertices);
 	}
-
 }
diff --git a/Engine/Plugins/Experimental/Mutable/Source/MutableRuntime/Private/MuR/OpMeshRemove.h b/Engine/Plugins/Experimental/Mutable/Source/MutableRuntime/Private/MuR/OpMeshRemove.h
--- a/Engine/Plugins/Experimental/Mutable/Source/MutableRuntime/Private/MuR/OpMeshRemove.h
+++ b/Engine/Plugins/Experimental/Mutable/Source/MutableRuntime/Private/MuR/OpMeshRemove.h
@@ -5,6 +5,7 @@
 #include "MuR/Ptr.h"
 #include "HAL/Platform.h"
 
+#include "Containers/BitArray.h"
 
 namespace mu
 {
@@ -14,6 +15,17 @@ namespace mu
 	extern void MeshRemoveMask(Mesh* Result, const Mesh* Source, const Mesh* Mask, bool& bOutSuccess);
 
 	/** Remove a list of vertices and related faces from a mesh. The list is stored as a bool map for every vertex in the mesh. */
-	extern void MeshRemoveVerticesWithMap(Mesh* Result, const uint8* RemovedVertices, uint32 RemovedVerticesCount);
+	//extern void MeshRemoveVerticesWithMap(Mesh* Result, const TBitArray<>& RemovedVertices);
 
+	/** 
+	 * Remove a set of vertices and related faces from a mesh in-place. VertexToCull is a bitset where if bit i-th is set, 
+	 * the vertex i-th will be removed if all faces referencing this vertex need to be removed. A face is remove if all its 
+	 * vertices have the bit set in VerticesToCull.
+	 */
+	extern void MeshRemoveVerticesWithCullSet(Mesh* Result, const TBitArray<>& VerticesToCull);
+
+	/**
+	 * Recreates the Surface and Surfaces Submeshes given a set of vertices and faces remaining after mesh removal.
+	 */
+	extern void MeshRemoveRecreateSurface(Mesh* Result, const TBitArray<>& UsedVertices, const TBitArray<>& UsedFaces);
 }
diff --git a/Engine/Plugins/Experimental/Mutable/Source/MutableRuntime/Public/MuR/Mesh.h b/Engine/Plugins/Experimental/Mutable/Source/MutableRuntime/Public/MuR/Mesh.h
--- a/Engine/Plugins/Experimental/Mutable/Source/MutableRuntime/Public/MuR/Mesh.h
+++ b/Engine/Plugins/Experimental/Mutable/Source/MutableRuntime/Public/MuR/Mesh.h
@@ -18,6 +18,8 @@
 #include "MuR/Skeleton.h"
 #include "Templates/Tuple.h"
 
+#include <type_traits>
+
 class FString;
 
 namespace mu
@@ -32,34 +34,45 @@ namespace mu
     typedef Ptr<Mesh> MeshPtr;
     typedef Ptr<const Mesh> MeshPtrConst;
 
-	struct FMeshSurface
+	struct FSurfaceSubMesh
 	{
-		int32 FirstVertex=0;
-		int32 VertexCount=0;
-		int32 FirstIndex=0;
-		int32 IndexCount=0;
-		uint32 Id=0;
 		
-		uint32 BoneMapIndex=0;
-		uint32 BoneMapCount=0;
+		int32 VertexBegin = 0;
+		int32 VertexEnd = 0;
+		int32 IndexBegin = 0;
+		int32 IndexEnd = 0;
 
-		bool bCastShadow=false;
+		uint32 ExternalId = 0;
 
-		inline bool operator==(const FMeshSurface& o) const
+		friend bool operator==(const FSurfaceSubMesh& Lhs, const FSurfaceSubMesh& Rhs)
 		{
-			return FirstVertex == o.FirstVertex
-				&& VertexCount == o.VertexCount
-				&& FirstIndex == o.FirstIndex
-				&& IndexCount == o.IndexCount
-				&& Id == o.Id
-				&& BoneMapIndex == o.BoneMapIndex
-				&& BoneMapCount == o.BoneMapCount
-				&& bCastShadow == o.bCastShadow;
+			return FMemory::Memcmp(&Lhs, &Rhs, sizeof(FSurfaceSubMesh)) == 0;
 		}
+	};
+	static_assert(std::has_unique_object_representations_v<FSurfaceSubMesh>);
+	
+	MUTABLE_DEFINE_POD_SERIALISABLE(FSurfaceSubMesh);
+	MUTABLE_DEFINE_POD_VECTOR_SERIALISABLE(FSurfaceSubMesh);
 
-		inline void Serialise(OutputArchive& arch) const;
-		inline void Unserialise(InputArchive& arch);
+	struct FMeshSurface
+	{
+		TArray<FSurfaceSubMesh, TInlineAllocator<1>> SubMeshes; 
+
+		uint32 BoneMapIndex = 0;
+		uint32 BoneMapCount = 0;
+		uint32 Id = 0;
+
+		friend bool operator==(const FMeshSurface& Lhs, const FMeshSurface& Rhs)
+		{
+			return
+				Lhs.Id == Rhs.Id &&
+				Lhs.BoneMapIndex == Rhs.BoneMapIndex &&
+				Lhs.BoneMapCount == Rhs.BoneMapCount &&
+				Lhs.SubMeshes == Rhs.SubMeshes;
+		}
 
+		inline void Serialise(OutputArchive& Arch) const;
+		inline void Unserialise(InputArchive& Arch);
 	};
 
 
@@ -225,16 +238,15 @@ namespace mu
 
         //! Get the number of surfaces defined in this mesh. Surfaces are buffer-contiguous mesh
         //! fragments that share common properties (usually material)
-		int32 GetSurfaceCount() const;
-        void GetSurface( int32 surfaceIndex,
-                         int32* FirstVertex, int32* VertexCount,
-                         int32* FirstIndex, int32* IndexCount,
-						 int32* FirstBone, int32* BoneCount,
-						 bool* bCastShadow) const;
+        int32 GetSurfaceCount() const;
+        void GetSurface(int32 SurfaceIndex,
+                        int32& OutFirstVertex, int32& OutVertexCount,
+                        int32& OutFirstIndex, int32& OutIndexCount,
+						int32& OutFirstBone, int32& OutBoneCount) const;
 
         //! Return an internal id that can be used to match mesh surfaces and instance surfaces.
         //! Only valid for meshes that are part of instances.
-        uint32 GetSurfaceId( int32 surfaceIndex ) const;
+        uint32 GetSurfaceId(int32 SurfaceIndex) const;
 
         //! \}
 
