[mutable] Fix for CustomizableObjectInstanceUsage ptr being set to nullptr during spawning of Blueprint based actor.
- When calling for a spawn of a Blueprint baed actor we ended up having a null InstanceUsage. By setting it also in the ReInit method we ensure, if not already set, that gets set.

#prelfight 667011f0836c0db7e65c25cb (build)

diff --git a/Engine/Plugins/Experimental/Mutable/Source/CustomizableObject/Private/MuCO/CustomizableSkeletalComponent.cpp b/Engine/Plugins/Experimental/Mutable/Source/CustomizableObject/Private/MuCO/CustomizableSkeletalComponent.cpp
--- a/Engine/Plugins/Experimental/Mutable/Source/CustomizableObject/Private/MuCO/CustomizableSkeletalComponent.cpp
+++ b/Engine/Plugins/Experimental/Mutable/Source/CustomizableObject/Private/MuCO/CustomizableSkeletalComponent.cpp
@@ -131,6 +131,12 @@ void UCustomizableSkeletalComponent::OnAttachmentChanged()
 void UCustomizableSkeletalComponent::PostInitProperties()
 {
 	Super::PostInitProperties();
+	CreateCustomizableObjectInstanceUsage();
+}
 
+void UCustomizableSkeletalComponent::PostReinitProperties()
+{
+	Super::PostReinitProperties();
 	CreateCustomizableObjectInstanceUsage();
 }
+
diff --git a/Engine/Plugins/Experimental/Mutable/Source/CustomizableObject/Public/MuCO/CustomizableSkeletalComponent.h b/Engine/Plugins/Experimental/Mutable/Source/CustomizableObject/Public/MuCO/CustomizableSkeletalComponent.h
--- a/Engine/Plugins/Experimental/Mutable/Source/CustomizableObject/Public/MuCO/CustomizableSkeletalComponent.h
+++ b/Engine/Plugins/Experimental/Mutable/Source/CustomizableObject/Public/MuCO/CustomizableSkeletalComponent.h
@@ -66,5 +66,6 @@ private:
 
 protected:
 	virtual void PostInitProperties() override;
+	virtual void PostReinitProperties() override;
 };
 
