[mutable] Reduced the warning severity of sections in table nodes.
+ Now, this warning will only appear if the requested section does not exists in the row mesh being compiled.

#rnx

diff --git a/Engine/Plugins/Experimental/Mutable/Source/CustomizableObjectEditor/Private/MuCOE/GenerateMutableSource/GenerateMutableSourceTable.cpp b/Engine/Plugins/Experimental/Mutable/Source/CustomizableObjectEditor/Private/MuCOE/GenerateMutableSource/GenerateMutableSourceTable.cpp
--- a/Engine/Plugins/Experimental/Mutable/Source/CustomizableObjectEditor/Private/MuCOE/GenerateMutableSource/GenerateMutableSourceTable.cpp
+++ b/Engine/Plugins/Experimental/Mutable/Source/CustomizableObjectEditor/Private/MuCOE/GenerateMutableSource/GenerateMutableSourceTable.cpp
@@ -213,15 +213,9 @@ bool FillTableColumn(const UCustomizableObjectNodeTable* TableNode, mu::TablePtr
 			
 			if (ImportedModel->LODModels.IsValidIndex(LODIndex)) // Ignore error since this Section is empty due to Automatic LODs From Mesh
 			{
-				int32 NumSections = ImportedModel->LODModels[LODIndex].Sections.Num();
-				int32 ReferenceNumMaterials = ImportedModel->LODModels[LODIndex].Sections.Num();
-
-				if (NumSections != ReferenceNumMaterials)
+				if (!ImportedModel->LODModels[LODIndex].Sections.IsValidIndex(SectionIndex))
 				{
-					FString Dif_1 = NumSections > ReferenceNumMaterials ? "more" : "less";
-					FString Dif_2 = NumSections > ReferenceNumMaterials ? "Some will be ignored" : "This can cause some compilation errors.";
-
-					FString msg = FString::Printf(TEXT("Mesh from column [%s] row [%s] has %s Sections than the reference mesh. %s"), *ColumnName, *RowName, *Dif_1, *Dif_2);
+					FString msg = FString::Printf(TEXT("Mesh from column [%s] row [%s] does not have section %d at LOD %d"), *ColumnName, *RowName, SectionIndexConnected, LODIndex);
 					LogRowGenerationMessage(TableNode, DataTablePtr, GenerationContext, msg, RowName);
 				}
 			}
