[Mutable] Fix static analyzer issue.

#rnx

diff --git a/Engine/Plugins/Experimental/Mutable/Source/CustomizableObject/Private/MuCO/CustomizableObjectInstance.cpp b/Engine/Plugins/Experimental/Mutable/Source/CustomizableObject/Private/MuCO/CustomizableObjectInstance.cpp
--- a/Engine/Plugins/Experimental/Mutable/Source/CustomizableObject/Private/MuCO/CustomizableObjectInstance.cpp
+++ b/Engine/Plugins/Experimental/Mutable/Source/CustomizableObject/Private/MuCO/CustomizableObjectInstance.cpp
@@ -4938,7 +4938,7 @@ bool UCustomizableInstancePrivate::BuildOrCopyRenderData(const TSharedRef<FUpdat
 		for (int32 MeshSectionIndex = 0; MeshSectionIndex < MeshNumSurfaces; ++MeshSectionIndex)
 		{
 			uint32 MeshSurfaceId = LOD.Mesh->GetSurfaceId(MeshSectionIndex);
-			uint32 InstanceSurfaceIndex = OperationData->MutableInstance->FindSurfaceById(ComponentIndex, LODIndex, MeshSurfaceId);
+			int32 InstanceSurfaceIndex = OperationData->MutableInstance->FindSurfaceById(ComponentIndex, LODIndex, MeshSurfaceId);
 			
 			if (InstanceSurfaceIndex < 0)
 			{
