[mutable] Added checks to prevent an static analisis warning.

#rnx

diff --git a/Engine/Plugins/Experimental/Mutable/Source/CustomizableObjectEditor/Private/MuCOE/CustomizableObjectFactory.cpp b/Engine/Plugins/Experimental/Mutable/Source/CustomizableObjectEditor/Private/MuCOE/CustomizableObjectFactory.cpp
--- a/Engine/Plugins/Experimental/Mutable/Source/CustomizableObjectEditor/Private/MuCOE/CustomizableObjectFactory.cpp
+++ b/Engine/Plugins/Experimental/Mutable/Source/CustomizableObjectEditor/Private/MuCOE/CustomizableObjectFactory.cpp
@@ -72,12 +72,14 @@ UObject* UCustomizableObjectFactory::FactoryCreateNew(UClass* Class, UObject* In
 	}
 	
 	UCustomizableObjectGraph* Source = NewObject<UCustomizableObjectGraph>(NewObj, NAME_None, RF_Transactional);
-	if (!Source)
+	UCustomizableObjectPrivate* ObjectPrivate = NewObj->GetPrivate(); // Needed to avoid a static analysis warning
+
+	if (!Source || !ObjectPrivate)
 	{
 		return NewObj;
 	}
 
-	NewObj->GetPrivate()->GetSource() = Source;
+	ObjectPrivate->GetSource() = Source;
 	
 	Source->AddEssentialGraphNodes();
 
@@ -99,7 +101,7 @@ UObject* UCustomizableObjectFactory::FactoryCreateNew(UClass* Class, UObject* In
 		return NewObj;
 	}
 
-	NewObj->GetPrivate()->MutableMeshComponents.SetNum(CreationSettings.NumMeshComponents);
+	ObjectPrivate->MutableMeshComponents.SetNum(CreationSettings.NumMeshComponents);
 
 	if (!CreationSettings.bEmptyObject)
 	{
@@ -120,10 +122,7 @@ UObject* UCustomizableObjectFactory::FactoryCreateNew(UClass* Class, UObject* In
 				}
 			}
 
-			if (NewObj->GetPrivate())
-			{
-				NewObj->GetPrivate()->SetIsChildObject(true);
-			}
+			ObjectPrivate->SetIsChildObject(true);
 		}
 		else
 		{
@@ -131,7 +130,7 @@ UObject* UCustomizableObjectFactory::FactoryCreateNew(UClass* Class, UObject* In
 
 			for (int32 MeshIndex = 0; MeshIndex < CreationSettings.ReferenceSkeletalMeshes.Num(); ++MeshIndex)
 			{
-				NewObj->GetPrivate()->MutableMeshComponents[MeshIndex].ReferenceSkeletalMesh = CreationSettings.ReferenceSkeletalMeshes[MeshIndex].Get();
+				ObjectPrivate->MutableMeshComponents[MeshIndex].ReferenceSkeletalMesh = CreationSettings.ReferenceSkeletalMeshes[MeshIndex].Get();
 			}
 		}
 	}
