[Mutable] Bump CustomizableObject version to force recompilation of possible stale data.

#rnx

diff --git a/Engine/Plugins/Experimental/Mutable/Source/CustomizableObject/Internal/MuCO/CustomizableObjectPrivate.h b/Engine/Plugins/Experimental/Mutable/Source/CustomizableObject/Internal/MuCO/CustomizableObjectPrivate.h
--- a/Engine/Plugins/Experimental/Mutable/Source/CustomizableObject/Internal/MuCO/CustomizableObjectPrivate.h
+++ b/Engine/Plugins/Experimental/Mutable/Source/CustomizableObject/Internal/MuCO/CustomizableObjectPrivate.h
@@ -1162,6 +1162,8 @@ public:
 
 		AddSurfaceAndMeshMetadata,
 
+		TablesPropertyNameBug,
+
 		// -----<new versions can be added above this line>--------
 		LastCustomizableObjectVersion
 	};
