[Mutable] Add method to get the current CustomizableObjectVersionBridge's version from the CO module.

#rnx

diff --git a/Engine/Plugins/Experimental/Mutable/Source/CustomizableObject/Internal/MuCO/ICustomizableObjectEditorModule.h b/Engine/Plugins/Experimental/Mutable/Source/CustomizableObject/Internal/MuCO/ICustomizableObjectEditorModule.h
--- a/Engine/Plugins/Experimental/Mutable/Source/CustomizableObject/Internal/MuCO/ICustomizableObjectEditorModule.h
+++ b/Engine/Plugins/Experimental/Mutable/Source/CustomizableObject/Internal/MuCO/ICustomizableObjectEditorModule.h
@@ -54,7 +54,11 @@ public:
 
 	/** See GraphTraversal::IsRootObject(...) */
 	virtual bool IsRootObject(const UCustomizableObject& Object) const = 0;
-	
+
+	/** Get the current VersionBridge's version for Object. 
+	  * @return Current version as string. */
+	virtual FString GetCurrentContentVersionForObject(const UCustomizableObject& Object) const = 0;
+
 	/**
 	 * Execute this method in order to bake the provided instance. It will schedule a special type of instance update before proceeding with the bake itself.
 	 * @param InTargetInstance The instance we want to bake
diff --git a/Engine/Plugins/Experimental/Mutable/Source/CustomizableObjectEditor/Private/MuCOE/CustomizableObjectEditorModule.cpp b/Engine/Plugins/Experimental/Mutable/Source/CustomizableObjectEditor/Private/MuCOE/CustomizableObjectEditorModule.cpp
--- a/Engine/Plugins/Experimental/Mutable/Source/CustomizableObjectEditor/Private/MuCOE/CustomizableObjectEditorModule.cpp
+++ b/Engine/Plugins/Experimental/Mutable/Source/CustomizableObjectEditor/Private/MuCOE/CustomizableObjectEditorModule.cpp
@@ -27,6 +27,7 @@
 #include "MuCOE/CustomizableObjectIdentifierCustomization.h"
 #include "MuCOE/CustomizableObjectInstanceEditor.h"
 #include "MuCOE/CustomizableObjectInstanceFactory.h"
+#include "MuCOE/CustomizableObjectVersionBridge.h"
 #include "MuCOE/CustomizableObjectNodeObjectGroupDetails.h"
 #include "MuCOE/Nodes/CustomizableObjectNodeCopyMaterial.h"
 #include "MuCOE/Nodes/CustomizableObjectNodeEditMaterial.h"
@@ -650,6 +651,21 @@ bool FCustomizableObjectEditorModule::IsRootObject(const UCustomizableObject& Ob
 	return GraphTraversal::IsRootObject(Object);
 }
 
+FString FCustomizableObjectEditorModule::GetCurrentContentVersionForObject(const UCustomizableObject& Object) const
+{
+	if (Object.VersionBridge && Object.VersionBridge->GetClass()->ImplementsInterface(UCustomizableObjectVersionBridgeInterface::StaticClass()))
+	{
+		ICustomizableObjectVersionBridgeInterface* CustomizableObjectVersionBridgeInterface = Cast<ICustomizableObjectVersionBridgeInterface>(Object.VersionBridge);
+
+		if (CustomizableObjectVersionBridgeInterface)
+		{
+			return CustomizableObjectVersionBridgeInterface->GetCurrentVersionAsString();
+		}
+	}
+
+	return FString();
+}
+
 
 void FCustomizableObjectEditorModule::BakeCustomizableObjectInstance(UCustomizableObjectInstance* InTargetInstance, const FBakingConfiguration& InBakingConfig)
 {
diff --git a/Engine/Plugins/Experimental/Mutable/Source/CustomizableObjectEditor/Private/MuCOE/CustomizableObjectEditorModule.h b/Engine/Plugins/Experimental/Mutable/Source/CustomizableObjectEditor/Private/MuCOE/CustomizableObjectEditorModule.h
--- a/Engine/Plugins/Experimental/Mutable/Source/CustomizableObjectEditor/Private/MuCOE/CustomizableObjectEditorModule.h
+++ b/Engine/Plugins/Experimental/Mutable/Source/CustomizableObjectEditor/Private/MuCOE/CustomizableObjectEditorModule.h
@@ -40,6 +40,7 @@ public:
 	virtual FCustomizableObjectEditorLogger& GetLogger() override;
 	virtual bool IsCompilationOutOfDate(const UCustomizableObject& Object, TArray<FName>* OutOfDatePackages) const override;
 	virtual bool IsRootObject(const UCustomizableObject& Object) const override;
+	virtual FString GetCurrentContentVersionForObject(const UCustomizableObject& Object) const override;
 	virtual void BakeCustomizableObjectInstance(UCustomizableObjectInstance* InTargetInstance, const FBakingConfiguration& InBakingConfig) override;
 
 	/** Request for a given customizable object to be compiled. Async compile requests will be queued and processed sequentially. 
diff --git a/Engine/Plugins/Experimental/Mutable/Source/CustomizableObjectEditor/Public/MuCOE/CustomizableObjectVersionBridge.h b/Engine/Plugins/Experimental/Mutable/Source/CustomizableObjectEditor/Public/MuCOE/CustomizableObjectVersionBridge.h
--- a/Engine/Plugins/Experimental/Mutable/Source/CustomizableObjectEditor/Public/MuCOE/CustomizableObjectVersionBridge.h
+++ b/Engine/Plugins/Experimental/Mutable/Source/CustomizableObjectEditor/Public/MuCOE/CustomizableObjectVersionBridge.h
@@ -41,4 +41,9 @@ public:
 	  * @return True if the table row is included in the current release
 	*/
 	virtual bool IsVersionPropertyIncludedInCurrentRelease(const FProperty& VersionProperty, const uint8* CellData) const { return false; };
+
+	/** Get the version as string to include it when building the derived data cache key.
+	  * @return current version as string.
+	  */
+	virtual FString GetCurrentVersionAsString() const { return {}; };
 };
