[Mutable] Sink ME_MASKDIFF operation when followed by ME_SWITCH

#rnx

diff --git a/Engine/Plugins/Experimental/Mutable/Source/MutableTools/Private/MuT/CodeOptimiser_Semantic.cpp b/Engine/Plugins/Experimental/Mutable/Source/MutableTools/Private/MuT/CodeOptimiser_Semantic.cpp
--- a/Engine/Plugins/Experimental/Mutable/Source/MutableTools/Private/MuT/CodeOptimiser_Semantic.cpp
+++ b/Engine/Plugins/Experimental/Mutable/Source/MutableTools/Private/MuT/CodeOptimiser_Semantic.cpp
@@ -406,7 +406,32 @@ namespace mu
 				at = NewAt;
 				break;
 			}
+			case OP_TYPE::ME_SWITCH:
+			{
+				// Move the mask diff down all the paths
+				Ptr<ASTOpSwitch> nop = mu::Clone<ASTOpSwitch>(Fragment.get());
+
+				if (nop->def)
+				{
+					Ptr<ASTOpFixed> defOp = mu::Clone<ASTOpFixed>(this);
+					defOp->SetChild(defOp->op.args.MeshMaskDiff.fragment, nop->def);
+					nop->def = defOp;
+				}
 
+				// We need to copy the options because we change them
+				for (int32 v = 0; v < nop->cases.Num(); ++v)
+				{
+					if (nop->cases[v].branch)
+					{
+						Ptr<ASTOpFixed> bOp = mu::Clone<ASTOpFixed>(this);
+						bOp->SetChild(bOp->op.args.MeshMaskDiff.fragment, nop->cases[v].branch);
+						nop->cases[v].branch = bOp;
+					}
+				}
+
+				at = nop;
+				break;
+			}
 			default:
 				break;
 			}
