[Mutable] Fix assert when preloading Node Table assets.

#rnx

diff --git a/Engine/Plugins/Experimental/Mutable/Source/CustomizableObjectEditor/Private/MuCOE/CustomizableObjectEditorUtilities.cpp b/Engine/Plugins/Experimental/Mutable/Source/CustomizableObjectEditor/Private/MuCOE/CustomizableObjectEditorUtilities.cpp
--- a/Engine/Plugins/Experimental/Mutable/Source/CustomizableObjectEditor/Private/MuCOE/CustomizableObjectEditorUtilities.cpp
+++ b/Engine/Plugins/Experimental/Mutable/Source/CustomizableObjectEditor/Private/MuCOE/CustomizableObjectEditorUtilities.cpp
@@ -28,10 +28,22 @@ bool CompareNames(const TSharedPtr<FString>& sp1, const TSharedPtr<FString>& sp2
 
 void ConditionalPostLoadReference(UObject& Object)
 {
-	if (Object.HasAnyFlags(RF_NeedLoad))
+	Object.ConditionalPostLoad();
+}
+
+
+void ConditionalPostLoadReference(FSoftObjectPtr& SoftObject)
+{
+	UObject* Object = SoftObject.LoadSynchronous();
+	if (!Object)
+	{
+		return;
+	}
+	
+	if (Object->HasAnyFlags(RF_NeedLoad))
 	{
-		Object.GetLinker()->Preload(&Object);
+		LoadPackage(nullptr, *SoftObject.GetLongPackageName(), LOAD_None);
 	}
 
-	Object.ConditionalPostLoad();
+	Object->ConditionalPostLoad();
 }
diff --git a/Engine/Plugins/Experimental/Mutable/Source/CustomizableObjectEditor/Private/MuCOE/CustomizableObjectEditorUtilities.h b/Engine/Plugins/Experimental/Mutable/Source/CustomizableObjectEditor/Private/MuCOE/CustomizableObjectEditorUtilities.h
--- a/Engine/Plugins/Experimental/Mutable/Source/CustomizableObjectEditor/Private/MuCOE/CustomizableObjectEditorUtilities.h
+++ b/Engine/Plugins/Experimental/Mutable/Source/CustomizableObjectEditor/Private/MuCOE/CustomizableObjectEditorUtilities.h
@@ -10,3 +10,5 @@ class UObject;
 bool CompareNames(const TSharedPtr<FString>& sp1, const TSharedPtr<FString>& sp2);
 
 void ConditionalPostLoadReference(UObject& Object);
+
+void ConditionalPostLoadReference(FSoftObjectPtr& SoftObject);
diff --git a/Engine/Plugins/Experimental/Mutable/Source/CustomizableObjectEditor/Private/MuCOE/Nodes/CustomizableObjectNodeTable.cpp b/Engine/Plugins/Experimental/Mutable/Source/CustomizableObjectEditor/Private/MuCOE/Nodes/CustomizableObjectNodeTable.cpp
--- a/Engine/Plugins/Experimental/Mutable/Source/CustomizableObjectEditor/Private/MuCOE/Nodes/CustomizableObjectNodeTable.cpp
+++ b/Engine/Plugins/Experimental/Mutable/Source/CustomizableObjectEditor/Private/MuCOE/Nodes/CustomizableObjectNodeTable.cpp
@@ -236,13 +236,8 @@ void UCustomizableObjectNodeTable::PostLoad()
 						continue;
 					}
 
-					UObject* Object = SoftObjectProperty->GetPropertyValue(CellData).LoadSynchronous();
-					if (!Object)
-					{
-						continue;
-					}
-
-					ConditionalPostLoadReference(*Object);
+					FSoftObjectPtr SoftObject = SoftObjectProperty->GetPropertyValue(CellData);
+					ConditionalPostLoadReference(SoftObject);
 				}
 			}
 		}
