[mutable] Fixed crash when the Slate tick asks a Mutable instance for parameter values while it's being compiled and the number of parameters changes.
#tests Editor
#rnx

diff --git a/Engine/Plugins/Experimental/Mutable/Source/CustomizableObjectEditor/Private/MuCOE/CustomizableInstanceDetails.cpp b/Engine/Plugins/Experimental/Mutable/Source/CustomizableObjectEditor/Private/MuCOE/CustomizableInstanceDetails.cpp
--- a/Engine/Plugins/Experimental/Mutable/Source/CustomizableObjectEditor/Private/MuCOE/CustomizableInstanceDetails.cpp
+++ b/Engine/Plugins/Experimental/Mutable/Source/CustomizableObjectEditor/Private/MuCOE/CustomizableInstanceDetails.cpp
@@ -1174,6 +1174,12 @@ TSharedRef<SWidget> FCustomizableInstanceDetails::GenerateFloatWidget(const UCus
 
 float FCustomizableInstanceDetails::GetFloatParameterValue(FString ParamName, int32 RangeIndex) const
 {
+	if (CustomInstance->GetCustomizableObject()->GetPrivate()->IsLocked())
+	{
+		// Prevent crashing if polling the float value during CO compilation
+		return -1.f;
+	}
+
 	if (RangeIndex == INDEX_NONE)
 	{
 		return CustomInstance->GetFloatParameterSelectedOption(ParamName, RangeIndex);
