[mutable] Fixed "Compile only selected" compilation option breaking Instance parameters.
#tests Editor
#rnx

diff --git a/Engine/Plugins/Experimental/Mutable/Source/CustomizableObject/Internal/MuCO/CustomizableObjectPrivate.h b/Engine/Plugins/Experimental/Mutable/Source/CustomizableObject/Internal/MuCO/CustomizableObjectPrivate.h
--- a/Engine/Plugins/Experimental/Mutable/Source/CustomizableObject/Internal/MuCO/CustomizableObjectPrivate.h
+++ b/Engine/Plugins/Experimental/Mutable/Source/CustomizableObject/Internal/MuCO/CustomizableObjectPrivate.h
@@ -14,6 +14,7 @@
 
 #if WITH_EDITOR
 #include "Misc/Guid.h"
+#include "Engine/DataTable.h"
 #endif
 
 #include "CustomizableObjectPrivate.generated.h"
@@ -794,6 +795,9 @@ struct FModelResources
 
 	UPROPERTY()
 	TArray<FCustomizableObjectMeshToMeshVertData> EditorOnlyClothingMeshToMeshVertData;
+
+	// Stores what param names use a certain table as a table can be used from multiple table nodes, useful for partial compilations to restrict params
+	TMap<TObjectPtr<const UDataTable>, TSet<FString>> TableToParamNames;
 #endif
 	
 	/** Map of Hash to Streaming blocks, used to stream a block of data representing a resource from the BulkData */
@@ -1164,6 +1168,8 @@ public:
 
 		TablesPropertyNameBug,
 
+		DataTablesParamTrackingForCompileOnlySelected,
+
 		// -----<new versions can be added above this line>--------
 		LastCustomizableObjectVersion
 	};
diff --git a/Engine/Plugins/Experimental/Mutable/Source/CustomizableObjectEditor/Private/MuCOE/CustomizableObjectCompiler.cpp b/Engine/Plugins/Experimental/Mutable/Source/CustomizableObjectEditor/Private/MuCOE/CustomizableObjectCompiler.cpp
--- a/Engine/Plugins/Experimental/Mutable/Source/CustomizableObjectEditor/Private/MuCOE/CustomizableObjectCompiler.cpp
+++ b/Engine/Plugins/Experimental/Mutable/Source/CustomizableObjectEditor/Private/MuCOE/CustomizableObjectCompiler.cpp
@@ -717,6 +717,11 @@ mu::NodeObjectPtr FCustomizableObjectCompiler::GenerateMutableRoot(
     GenerationContext.RealTimeMorphTargetsOverrides = ActualRoot->RealTimeMorphSelectionOverrides;
     GenerationContext.RealTimeMorphTargetsOverrides.Reset();
 
+	if (!GenerationContext.ParamNamesToSelectedOptions.IsEmpty())
+	{
+		GenerationContext.TableToParamNames = Object->GetPrivate()->GetModelResources().TableToParamNames;
+	}
+
 	// Generate the object expression
 	UE_LOG(LogMutable, Verbose, TEXT("PROFILE: [ %16.8f ] GenerateMutableSource start."), FPlatformTime::Seconds());
 	mu::NodeObjectPtr MutableRoot = GenerateMutableSource(ActualRoot->OutputPin(), GenerationContext, !bOutIsRootObject);
@@ -1129,6 +1134,12 @@ void FCustomizableObjectCompiler::CompileInternal(bool bAsync)
 
 #if WITH_EDITORONLY_DATA
 		CurrentObject->GetPrivate()->CustomizableObjectPathMap = GenerationContext.CustomizableObjectPathMap;
+
+		if (!CurrentRequest->GetParameterNamesToSelectedOptions().Num())
+		{
+			// Cache the tables that are used by more than one param so that CompileOnlySelected can work properly
+			ModelResources.TableToParamNames = GenerationContext.TableToParamNames;
+		}
 #endif
 
 		ModelResources.NumComponents = GenerationContext.NumMeshComponentsInRoot + GenerationContext.NumExplicitMeshComponents;
diff --git a/Engine/Plugins/Experimental/Mutable/Source/CustomizableObjectEditor/Private/MuCOE/GenerateMutableSource/GenerateMutableSource.h b/Engine/Plugins/Experimental/Mutable/Source/CustomizableObjectEditor/Private/MuCOE/GenerateMutableSource/GenerateMutableSource.h
--- a/Engine/Plugins/Experimental/Mutable/Source/CustomizableObjectEditor/Private/MuCOE/GenerateMutableSource/GenerateMutableSource.h
+++ b/Engine/Plugins/Experimental/Mutable/Source/CustomizableObjectEditor/Private/MuCOE/GenerateMutableSource/GenerateMutableSource.h
@@ -848,6 +848,9 @@ struct FMutableGraphGenerationContext
 	// Stores the only option of an Int Param that should be compiled in a partial compilation
 	TMap<FString, FString> ParamNamesToSelectedOptions;
 
+	// Stores what param names use a certain table as a table can be used from multiple table nodes, useful for partial compilations to restrict params
+	TMap<TObjectPtr<const UDataTable>, TSet<FString>> TableToParamNames;
+
 	TArray<const UEdGraphNode*> LimitedParameters;
 	int32 ParameterLimitationCount = 0;
 
diff --git a/Engine/Plugins/Experimental/Mutable/Source/CustomizableObjectEditor/Private/MuCOE/GenerateMutableSource/GenerateMutableSourceTable.cpp b/Engine/Plugins/Experimental/Mutable/Source/CustomizableObjectEditor/Private/MuCOE/GenerateMutableSource/GenerateMutableSourceTable.cpp
--- a/Engine/Plugins/Experimental/Mutable/Source/CustomizableObjectEditor/Private/MuCOE/GenerateMutableSource/GenerateMutableSourceTable.cpp
+++ b/Engine/Plugins/Experimental/Mutable/Source/CustomizableObjectEditor/Private/MuCOE/GenerateMutableSource/GenerateMutableSourceTable.cpp
@@ -769,26 +769,61 @@ TArray<FName> GetEnabledRows(const UDataTable& DataTable, const UCustomizableObj
 
 void RestrictRowNamesToSelectedOption(TArray<FName>& InOutRowNames, const UCustomizableObjectNodeTable& TableNode, FMutableGraphGenerationContext& GenerationContext)
 {
-	// If the param is in the map restrict to only the selected option
-	FString* SelectedOptionString = GenerationContext.ParamNamesToSelectedOptions.Find(TableNode.ParameterName);
-
-	if (SelectedOptionString)
+	if (!GenerationContext.ParamNamesToSelectedOptions.IsEmpty())
 	{
-		FName SelectedOptionName = FName(*SelectedOptionString);
+		TSet<FString>* ParamNames = GenerationContext.TableToParamNames.Find(TableNode.Table);
 
-		if (InOutRowNames.Contains(SelectedOptionName))
+		if (ParamNames && !ParamNames->IsEmpty())
 		{
-			FName AnotherOption = GetAnotherOption(SelectedOptionName, InOutRowNames);
-			InOutRowNames.Empty(2);
-			InOutRowNames.Add(SelectedOptionName);
+			TSet<FName> SelectedOptionNames;
 
-			// To prevent the optimization of the parameter for having just one option, which would prevent the restriction 
-			// of that parameter in the next compile only selected
-			InOutRowNames.Add(AnotherOption);
-		}
-		else
-		{
-			InOutRowNames.Empty(0);
+			for (const FString& ParamName : *ParamNames)
+			{
+				// If the param is in the map restrict to only the selected option
+				FString* SelectedOptionString = GenerationContext.ParamNamesToSelectedOptions.Find(ParamName);
+
+				if (SelectedOptionString)
+				{
+					SelectedOptionNames.Add(FName(*SelectedOptionString));
+				}
+			}
+
+			if (!SelectedOptionNames.IsEmpty())
+			{
+				bool bRowNamesContainsSelectedOptionName = false;
+
+				for (const FName& OptionName : SelectedOptionNames)
+				{
+					if (InOutRowNames.Contains(OptionName))
+					{
+						bRowNamesContainsSelectedOptionName = true;
+						break;
+					}
+				}
+
+				if (bRowNamesContainsSelectedOptionName)
+				{
+					if (SelectedOptionNames.Num() < 2)
+					{
+						FName AnotherOption = GetAnotherOption(SelectedOptionNames.Array()[0], InOutRowNames);
+
+						// To prevent the optimization of the parameter for having just one option, which would prevent the restriction 
+						// of that parameter in the next compile only selected
+						SelectedOptionNames.Add(AnotherOption);
+					}
+
+					InOutRowNames.Empty(SelectedOptionNames.Num());
+
+					for (const FName& OptionName : SelectedOptionNames)
+					{
+						InOutRowNames.Add(OptionName);
+					}
+				}
+				else
+				{
+					InOutRowNames.Empty(0);
+				}
+			}
 		}
 	}
 }
@@ -971,6 +1006,18 @@ mu::TablePtr GenerateMutableSourceTable(const UDataTable* DataTable, const UCust
 {
 	check(DataTable && TableNode);
 
+	if (GenerationContext.ParamNamesToSelectedOptions.IsEmpty())
+	{
+		TSet<FString>* ParamNames = GenerationContext.TableToParamNames.Find(DataTable);
+
+		if (!ParamNames)
+		{
+			ParamNames = &GenerationContext.TableToParamNames.Add(DataTable);
+		}
+
+		ParamNames->Add(TableNode->ParameterName);
+	}
+
 	// Checking if the table is in the cache
 	const FString TableName = DataTable->GetName();
 
