[mutable] Added float arithmetic operation nodes to the Mutable CO graph. It's possible to spawn a generic arithmetic node  or a concrete operation such as addition directly.
#tests Editor
#rnx

diff --git a/Engine/Plugins/Experimental/Mutable/Source/CustomizableObjectEditor/Private/MuCOE/EdGraphSchema_CustomizableObject.cpp b/Engine/Plugins/Experimental/Mutable/Source/CustomizableObjectEditor/Private/MuCOE/EdGraphSchema_CustomizableObject.cpp
--- a/Engine/Plugins/Experimental/Mutable/Source/CustomizableObjectEditor/Private/MuCOE/EdGraphSchema_CustomizableObject.cpp
+++ b/Engine/Plugins/Experimental/Mutable/Source/CustomizableObjectEditor/Private/MuCOE/EdGraphSchema_CustomizableObject.cpp
@@ -32,6 +32,7 @@
 #include "MuCOE/Nodes/CustomizableObjectNodeFloatParameter.h"
 #include "MuCOE/Nodes/CustomizableObjectNodeFloatSwitch.h"
 #include "MuCOE/Nodes/CustomizableObjectNodeFloatVariation.h"
+#include "MuCOE/Nodes/CustomizableObjectNodeFloatArithmeticOp.h"
 #include "MuCOE/Nodes/CustomizableObjectNodeGroupProjectorParameter.h"
 #include "MuCOE/Nodes/CustomizableObjectNodeLayoutBlocks.h"
 #include "MuCOE/Nodes/CustomizableObjectNodeMaterialVariation.h"
@@ -489,6 +490,7 @@ void UEdGraphSchema_CustomizableObject::GetGraphContextActions(FGraphContextMenu
 		{
 			ContextMenuBuilder.CreateTemplateNode<UCustomizableObjectNodeFloatConstant>(),
 			ContextMenuBuilder.CreateTemplateNode<UCustomizableObjectNodeFloatParameter>(),
+			ContextMenuBuilder.CreateTemplateNode<UCustomizableObjectNodeFloatArithmeticOp>(),
 			ContextMenuBuilder.CreateTemplateNode<UCustomizableObjectNodeFloatSwitch>(),
 			ContextMenuBuilder.CreateTemplateNode<UCustomizableObjectNodeFloatVariation>(),
 			ContextMenuBuilder.CreateTemplateNode<UCustomizableObjectNodeCurve>(),
@@ -496,6 +498,22 @@ void UEdGraphSchema_CustomizableObject::GetGraphContextActions(FGraphContextMenu
 
 
 		AddNewNodeCategoryActionsFiltered(FloatTemplateNodes, ContextMenuBuilder, TEXT("Float"), GeneralGrouping, Filter);
+
+		UCustomizableObjectNodeFloatArithmeticOp* Node = NewObject<UCustomizableObjectNodeFloatArithmeticOp>();
+		Node->Operation = EFloatArithmeticOperation::E_Add;
+		AddNewNodeActionFiltered(Node, ContextMenuBuilder, TEXT("Float"), LOCTEXT("Float_Addition", "Float Addition +"), GeneralGrouping, Filter);
+
+		Node = NewObject<UCustomizableObjectNodeFloatArithmeticOp>();
+		Node->Operation = EFloatArithmeticOperation::E_Sub;
+		AddNewNodeActionFiltered(Node, ContextMenuBuilder, TEXT("Float"), LOCTEXT("Float_Subtraction", "Float Subtraction -"), GeneralGrouping, Filter);
+
+		Node = NewObject<UCustomizableObjectNodeFloatArithmeticOp>();
+		Node->Operation = EFloatArithmeticOperation::E_Mul;
+		AddNewNodeActionFiltered(Node, ContextMenuBuilder, TEXT("Float"), LOCTEXT("Float_Multiplication", "Float Multiplication x"), GeneralGrouping, Filter);
+
+		Node = NewObject<UCustomizableObjectNodeFloatArithmeticOp>();
+		Node->Operation = EFloatArithmeticOperation::E_Div;
+		AddNewNodeActionFiltered(Node, ContextMenuBuilder, TEXT("Float"), LOCTEXT("Float_Division", "Float Division /"), GeneralGrouping, Filter);
 	}
 
 	{
diff --git a/Engine/Plugins/Experimental/Mutable/Source/CustomizableObjectEditor/Private/MuCOE/GenerateMutableSource/GenerateMutableSourceFloat.cpp b/Engine/Plugins/Experimental/Mutable/Source/CustomizableObjectEditor/Private/MuCOE/GenerateMutableSource/GenerateMutableSourceFloat.cpp
--- a/Engine/Plugins/Experimental/Mutable/Source/CustomizableObjectEditor/Private/MuCOE/GenerateMutableSource/GenerateMutableSourceFloat.cpp
+++ b/Engine/Plugins/Experimental/Mutable/Source/CustomizableObjectEditor/Private/MuCOE/GenerateMutableSource/GenerateMutableSourceFloat.cpp
@@ -16,12 +16,14 @@
 #include "MuCOE/Nodes/CustomizableObjectNodeFloatParameter.h"
 #include "MuCOE/Nodes/CustomizableObjectNodeFloatSwitch.h"
 #include "MuCOE/Nodes/CustomizableObjectNodeFloatVariation.h"
+#include "MuCOE/Nodes/CustomizableObjectNodeFloatArithmeticOp.h"
 #include "MuCOE/Nodes/CustomizableObjectNodeTable.h"
 #include "MuT/NodeScalarConstant.h"
 #include "MuT/NodeScalarCurve.h"
 #include "MuT/NodeScalarSwitch.h"
 #include "MuT/NodeScalarTable.h"
 #include "MuT/NodeScalarVariation.h"
+#include "MuT/NodeScalarArithmeticOperation.h"
 
 #define LOCTEXT_NAMESPACE "CustomizableObjectEditor"
 
@@ -262,6 +264,46 @@ mu::NodeScalarPtr GenerateMutableSourceFloat(const UEdGraphPin* Pin, FMutableGra
 		}
 	}
 
+	else if (const UCustomizableObjectNodeFloatArithmeticOp* TypedNodeFloatArith = Cast<UCustomizableObjectNodeFloatArithmeticOp>(Node))
+	{
+		mu::NodeScalarArithmeticOperationPtr OpNode = new mu::NodeScalarArithmeticOperation();
+		Result = OpNode;
+
+		if (const UEdGraphPin* ConnectedPin = FollowInputPin(*TypedNodeFloatArith->XPin()))
+		{
+			mu::NodeScalarPtr XNode = GenerateMutableSourceFloat(ConnectedPin, GenerationContext);
+			OpNode->SetA(XNode);
+		}
+
+		if (const UEdGraphPin* ConnectedPin = FollowInputPin(*TypedNodeFloatArith->YPin()))
+		{
+			mu::NodeScalarPtr YNode = GenerateMutableSourceFloat(ConnectedPin, GenerationContext);
+			OpNode->SetB(YNode);
+		}
+
+		switch (TypedNodeFloatArith->Operation)
+		{
+		case EFloatArithmeticOperation::E_Add:
+			OpNode->SetOperation(mu::NodeScalarArithmeticOperation::OPERATION::AO_ADD);
+			break;
+
+		case EFloatArithmeticOperation::E_Sub:
+			OpNode->SetOperation(mu::NodeScalarArithmeticOperation::OPERATION::AO_SUBTRACT);
+			break;
+
+		case EFloatArithmeticOperation::E_Mul:
+			OpNode->SetOperation(mu::NodeScalarArithmeticOperation::OPERATION::AO_MULTIPLY);
+			break;
+
+		case EFloatArithmeticOperation::E_Div:
+			OpNode->SetOperation(mu::NodeScalarArithmeticOperation::OPERATION::AO_DIVIDE);
+			break;
+
+		default:
+			unimplemented();
+		}
+	}
+
 	else if (const UCustomizableObjectNodeTable* TypedNodeTable = Cast<UCustomizableObjectNodeTable>(Node))
 	{
 		//This node will add a default value in case of error
diff --git a/Engine/Plugins/Experimental/Mutable/Source/CustomizableObjectEditor/Private/MuCOE/Nodes/CustomizableObjectNodeFloatArithmeticOp.cpp b/Engine/Plugins/Experimental/Mutable/Source/CustomizableObjectEditor/Private/MuCOE/Nodes/CustomizableObjectNodeFloatArithmeticOp.cpp
new file mode 100644
--- /dev/null
+++ b/Engine/Plugins/Experimental/Mutable/Source/CustomizableObjectEditor/Private/MuCOE/Nodes/CustomizableObjectNodeFloatArithmeticOp.cpp
@@ -0,0 +1,73 @@
+// Copyright Epic Games, Inc. All Rights Reserved.
+
+#include "MuCOE/Nodes/CustomizableObjectNodeFloatArithmeticOp.h"
+
+#include "MuCOE/EdGraphSchema_CustomizableObject.h"
+
+class UCustomizableObjectNodeRemapPins;
+struct FPropertyChangedEvent;
+
+#define LOCTEXT_NAMESPACE "CustomizableObjectEditor"
+
+
+UCustomizableObjectNodeFloatArithmeticOp::UCustomizableObjectNodeFloatArithmeticOp()
+	: Super()
+{
+	Operation = EFloatArithmeticOperation::E_Add;
+}
+
+
+void UCustomizableObjectNodeFloatArithmeticOp::AllocateDefaultPins(UCustomizableObjectNodeRemapPins* RemapPins)
+{
+	const UEdGraphSchema_CustomizableObject* Schema = GetDefault<UEdGraphSchema_CustomizableObject>();
+
+	UEdGraphPin* XPin = CustomCreatePin(EGPD_Input, Schema->PC_Float, FName("A"));
+	XPin->bDefaultValueIsIgnored = true;
+	XPin->PinFriendlyName = FText();
+
+	FName PinName = TEXT("B");
+	UEdGraphPin* YPin = CustomCreatePin(EGPD_Input, Schema->PC_Float, PinName);
+	YPin->bDefaultValueIsIgnored = true;
+	YPin->PinFriendlyName = FText();
+
+	PinName = TEXT("Result");
+	UEdGraphPin* ResultPin = CustomCreatePin(EGPD_Output, Schema->PC_Float, PinName);
+	ResultPin->bDefaultValueIsIgnored = true;
+	ResultPin->PinFriendlyName = FText();
+}
+
+
+FText UCustomizableObjectNodeFloatArithmeticOp::GetNodeTitle(ENodeTitleType::Type TitleType) const
+{
+	if (TitleType == ENodeTitleType::ListView)
+	{
+		return LOCTEXT("Float_Arithmetic_Operation", "Float Arithmetic Operation");
+	}
+
+	const UEnum* EnumPtr = FindObject<UEnum>(nullptr, TEXT("/Script/CustomizableObjectEditor.EFloatArithmeticOperation"), true);
+
+	if (!EnumPtr)
+	{
+		return FText::FromString(FString("Float Operation"));
+	}
+
+	const int32 index = EnumPtr->GetIndexByValue((int32)Operation);
+	return EnumPtr->GetDisplayNameTextByIndex(index);
+}
+
+
+FLinearColor UCustomizableObjectNodeFloatArithmeticOp::GetNodeTitleColor() const
+{
+	const UEdGraphSchema_CustomizableObject* Schema = GetDefault<UEdGraphSchema_CustomizableObject>();
+	return Schema->GetPinTypeColor(Schema->PC_Float);
+}
+
+
+FText UCustomizableObjectNodeFloatArithmeticOp::GetTooltipText() const
+{
+	return LOCTEXT("Float_Arithmetic_Tooltip", "Perform an arithmetic operation between two floats.");
+}
+
+
+#undef LOCTEXT_NAMESPACE
+
diff --git a/Engine/Plugins/Experimental/Mutable/Source/CustomizableObjectEditor/Public/MuCOE/Nodes/CustomizableObjectNodeFloatArithmeticOp.h b/Engine/Plugins/Experimental/Mutable/Source/CustomizableObjectEditor/Public/MuCOE/Nodes/CustomizableObjectNodeFloatArithmeticOp.h
new file mode 100644
--- /dev/null
+++ b/Engine/Plugins/Experimental/Mutable/Source/CustomizableObjectEditor/Public/MuCOE/Nodes/CustomizableObjectNodeFloatArithmeticOp.h
@@ -0,0 +1,62 @@
+// Copyright Epic Games, Inc. All Rights Reserved.
+
+#pragma once
+
+#include "MuCOE/Nodes/CustomizableObjectNode.h"
+
+#include "CustomizableObjectNodeFloatArithmeticOp.generated.h"
+
+namespace ENodeTitleType { enum Type : int; }
+
+class UCustomizableObjectNodeRemapPins;
+class UEdGraphPin;
+class UObject;
+struct FPropertyChangedEvent;
+
+
+UENUM(BlueprintType)
+enum class EFloatArithmeticOperation : uint8
+{
+	E_Add 	UMETA(DisplayName = "+"),
+	E_Sub 	UMETA(DisplayName = "-"),
+	E_Mul	UMETA(DisplayName = "x"),
+	E_Div	UMETA(DisplayName = "/")
+};
+
+UCLASS()
+class CUSTOMIZABLEOBJECTEDITOR_API UCustomizableObjectNodeFloatArithmeticOp : public UCustomizableObjectNode
+{
+public:
+	GENERATED_BODY()
+
+	UCustomizableObjectNodeFloatArithmeticOp();
+
+	// Begin EdGraphNode interface
+	FText GetNodeTitle(ENodeTitleType::Type TitleType) const override;
+	FLinearColor GetNodeTitleColor() const override;
+	FText GetTooltipText() const override;
+	bool ShouldOverridePinNames() const override { return true; }
+
+	// UCustomizableObjectNode interface
+	void AllocateDefaultPins(UCustomizableObjectNodeRemapPins* RemapPins) override;
+
+	// First Operand
+	UEdGraphPin* XPin() const
+	{
+		return FindPin(TEXT("A"));
+	}
+
+	// Second Operand
+	UEdGraphPin* YPin() const
+	{
+		return FindPin(TEXT("B"));
+	}
+
+	UPROPERTY(EditAnywhere, Category = FloatArithmeticOperation)
+	EFloatArithmeticOperation Operation;
+
+	// UCustomizableObjectNode interface
+	virtual bool IsAffectedByLOD() const { return false; }
+
+};
+
