[Backout] - CL34523035
Original CL Desc
-----------------------------------------------------------------
[Mutable] Fix assert when preloading Node Table assets.

#rnx

diff --git a/Engine/Plugins/Experimental/Mutable/Source/CustomizableObjectEditor/Private/MuCOE/CustomizableObjectEditorUtilities.cpp b/Engine/Plugins/Experimental/Mutable/Source/CustomizableObjectEditor/Private/MuCOE/CustomizableObjectEditorUtilities.cpp
--- a/Engine/Plugins/Experimental/Mutable/Source/CustomizableObjectEditor/Private/MuCOE/CustomizableObjectEditorUtilities.cpp
+++ b/Engine/Plugins/Experimental/Mutable/Source/CustomizableObjectEditor/Private/MuCOE/CustomizableObjectEditorUtilities.cpp
@@ -28,22 +28,10 @@ bool CompareNames(const TSharedPtr<FString>& sp1, const TSharedPtr<FString>& sp2
 
 void ConditionalPostLoadReference(UObject& Object)
 {
-	Object.ConditionalPostLoad();
-}
-
-
-void ConditionalPostLoadReference(FSoftObjectPtr& SoftObject)
-{
-	UObject* Object = SoftObject.LoadSynchronous();
-	if (!Object)
-	{
-		return;
-	}
-	
-	if (Object->HasAnyFlags(RF_NeedLoad))
+	if (Object.HasAnyFlags(RF_NeedLoad))
 	{
-		LoadPackage(nullptr, *SoftObject.GetLongPackageName(), LOAD_None);
+		Object.GetLinker()->Preload(&Object);
 	}
 
-	Object->ConditionalPostLoad();
+	Object.ConditionalPostLoad();
 }
diff --git a/Engine/Plugins/Experimental/Mutable/Source/CustomizableObjectEditor/Private/MuCOE/CustomizableObjectEditorUtilities.h b/Engine/Plugins/Experimental/Mutable/Source/CustomizableObjectEditor/Private/MuCOE/CustomizableObjectEditorUtilities.h
--- a/Engine/Plugins/Experimental/Mutable/Source/CustomizableObjectEditor/Private/MuCOE/CustomizableObjectEditorUtilities.h
+++ b/Engine/Plugins/Experimental/Mutable/Source/CustomizableObjectEditor/Private/MuCOE/CustomizableObjectEditorUtilities.h
@@ -10,5 +10,3 @@ class UObject;
 bool CompareNames(const TSharedPtr<FString>& sp1, const TSharedPtr<FString>& sp2);
 
 void ConditionalPostLoadReference(UObject& Object);
-
-void ConditionalPostLoadReference(FSoftObjectPtr& SoftObject);
diff --git a/Engine/Plugins/Experimental/Mutable/Source/CustomizableObjectEditor/Private/MuCOE/Nodes/CustomizableObjectNodeTable.cpp b/Engine/Plugins/Experimental/Mutable/Source/CustomizableObjectEditor/Private/MuCOE/Nodes/CustomizableObjectNodeTable.cpp
--- a/Engine/Plugins/Experimental/Mutable/Source/CustomizableObjectEditor/Private/MuCOE/Nodes/CustomizableObjectNodeTable.cpp
+++ b/Engine/Plugins/Experimental/Mutable/Source/CustomizableObjectEditor/Private/MuCOE/Nodes/CustomizableObjectNodeTable.cpp
@@ -236,8 +236,13 @@ void UCustomizableObjectNodeTable::PostLoad()
 						continue;
 					}
 
-					FSoftObjectPtr SoftObject = SoftObjectProperty->GetPropertyValue(CellData);
-					ConditionalPostLoadReference(SoftObject);
+					UObject* Object = SoftObjectProperty->GetPropertyValue(CellData).LoadSynchronous();
+					if (!Object)
+					{
+						continue;
+					}
+
+					ConditionalPostLoadReference(*Object);
 				}
 			}
 		}
