[mutable] Removed the "compile only selected" workaround after the parameter relevancy fix. Also fixed a case where "compile only selected" didn't work after a recompile.
#tests Editor
#rnx

diff --git a/Engine/Plugins/Experimental/Mutable/Source/CustomizableObjectEditor/Private/MuCOE/CustomizableObjectCompiler.cpp b/Engine/Plugins/Experimental/Mutable/Source/CustomizableObjectEditor/Private/MuCOE/CustomizableObjectCompiler.cpp
--- a/Engine/Plugins/Experimental/Mutable/Source/CustomizableObjectEditor/Private/MuCOE/CustomizableObjectCompiler.cpp
+++ b/Engine/Plugins/Experimental/Mutable/Source/CustomizableObjectEditor/Private/MuCOE/CustomizableObjectCompiler.cpp
@@ -1135,11 +1135,8 @@ void FCustomizableObjectCompiler::CompileInternal(bool bAsync)
 #if WITH_EDITORONLY_DATA
 		CurrentObject->GetPrivate()->CustomizableObjectPathMap = GenerationContext.CustomizableObjectPathMap;
 
-		if (!CurrentRequest->GetParameterNamesToSelectedOptions().Num())
-		{
-			// Cache the tables that are used by more than one param so that CompileOnlySelected can work properly
-			ModelResources.TableToParamNames = GenerationContext.TableToParamNames;
-		}
+		// Cache the tables that are used by more than one param so that CompileOnlySelected can work properly
+		ModelResources.TableToParamNames = GenerationContext.TableToParamNames;
 #endif
 
 		ModelResources.NumComponents = GenerationContext.NumMeshComponentsInRoot + GenerationContext.NumExplicitMeshComponents;
diff --git a/Engine/Plugins/Experimental/Mutable/Source/CustomizableObjectEditor/Private/MuCOE/GenerateMutableSource/GenerateMutableSourceTable.cpp b/Engine/Plugins/Experimental/Mutable/Source/CustomizableObjectEditor/Private/MuCOE/GenerateMutableSource/GenerateMutableSourceTable.cpp
--- a/Engine/Plugins/Experimental/Mutable/Source/CustomizableObjectEditor/Private/MuCOE/GenerateMutableSource/GenerateMutableSourceTable.cpp
+++ b/Engine/Plugins/Experimental/Mutable/Source/CustomizableObjectEditor/Private/MuCOE/GenerateMutableSource/GenerateMutableSourceTable.cpp
@@ -803,15 +803,6 @@ void RestrictRowNamesToSelectedOption(TArray<FName>& InOutRowNames, const UCusto
 
 				if (bRowNamesContainsSelectedOptionName)
 				{
-					if (SelectedOptionNames.Num() < 2)
-					{
-						FName AnotherOption = GetAnotherOption(SelectedOptionNames.Array()[0], InOutRowNames);
-
-						// To prevent the optimization of the parameter for having just one option, which would prevent the restriction 
-						// of that parameter in the next compile only selected
-						SelectedOptionNames.Add(AnotherOption);
-					}
-
 					InOutRowNames.Empty(SelectedOptionNames.Num());
 
 					for (const FName& OptionName : SelectedOptionNames)
