[mutable] Fixed Text conflict from LOCTEXT macro.

#rnx

diff --git a/Engine/Plugins/Experimental/Mutable/Source/CustomizableObjectEditor/Private/MuCOE/CustomizableObjectCompiler.cpp b/Engine/Plugins/Experimental/Mutable/Source/CustomizableObjectEditor/Private/MuCOE/CustomizableObjectCompiler.cpp
--- a/Engine/Plugins/Experimental/Mutable/Source/CustomizableObjectEditor/Private/MuCOE/CustomizableObjectCompiler.cpp
+++ b/Engine/Plugins/Experimental/Mutable/Source/CustomizableObjectEditor/Private/MuCOE/CustomizableObjectCompiler.cpp
@@ -659,7 +659,7 @@ mu::NodeObjectPtr FCustomizableObjectCompiler::GenerateMutableRoot(
 
 	if (ActualRootObject->GetPrivate()->MutableMeshComponents.IsEmpty())
 	{
-		CompilerLog(LOCTEXT("EmptyComponentName", "Error! There are no components defined in the Object Properties Tab."), ActualRoot, EMessageSeverity::Error);
+		CompilerLog(LOCTEXT("NoComponentsError", "Error! There are no components defined in the Object Properties Tab."), ActualRoot, EMessageSeverity::Error);
 		return nullptr;
 	}
 
@@ -672,7 +672,7 @@ mu::NodeObjectPtr FCustomizableObjectCompiler::GenerateMutableRoot(
 
 		if (ComponentName.IsNone())
 		{
-			CompilerLog(LOCTEXT("EmptyComponentName", "Error! Missing name in a component of the Customizable Object."), ActualRoot, EMessageSeverity::Error);
+			CompilerLog(LOCTEXT("EmptyComponentNameError", "Error! Missing name in a component of the Customizable Object."), ActualRoot, EMessageSeverity::Error);
 			return nullptr;
 		}
 		else if (ComponentNames.Contains(ComponentName))
