[Mutable] Enabled component selector in material nodes.
+Now the component of material nodes can be selected using a combobox

#rnx

diff --git a/Engine/Plugins/Experimental/Mutable/Source/CustomizableObjectEditor/Private/MuCOE/Nodes/CustomizableObjectNodeMaterialDetails.cpp b/Engine/Plugins/Experimental/Mutable/Source/CustomizableObjectEditor/Private/MuCOE/Nodes/CustomizableObjectNodeMaterialDetails.cpp
--- a/Engine/Plugins/Experimental/Mutable/Source/CustomizableObjectEditor/Private/MuCOE/Nodes/CustomizableObjectNodeMaterialDetails.cpp
+++ b/Engine/Plugins/Experimental/Mutable/Source/CustomizableObjectEditor/Private/MuCOE/Nodes/CustomizableObjectNodeMaterialDetails.cpp
@@ -41,7 +41,7 @@ void FCustomizableObjectNodeMaterialDetails::CustomizeDetails(IDetailLayoutBuild
 
 	IDetailCategoryBuilder& CustomCategory = DetailBuilder.EditCategory("CustomizableObject");
 	TSharedRef<IPropertyHandle> MaterialProperty = DetailBuilder.GetProperty("Material");
-	TSharedRef<IPropertyHandle> ComponentNameProperty = DetailBuilder.GetProperty("ComponentName");
+	TSharedRef<IPropertyHandle> ComponentNameProperty = DetailBuilder.GetProperty("MeshComponentName");
 
 	// Added Material property manually to display it above of the rest of properties
 	CustomCategory.AddProperty(MaterialProperty);
