[mutable] Fix mesh merge corner case for meshes merging to a fragment of themselves.
The handing of IDs for layout blocks and mesh vertices didn't work for cases where one fragment had been transformed and their format changed, and the other didn't.
#rnx

diff --git a/Engine/Plugins/Experimental/Mutable/Source/MutableRuntime/Internal/MuR/OpMeshMerge.h b/Engine/Plugins/Experimental/Mutable/Source/MutableRuntime/Internal/MuR/OpMeshMerge.h
--- a/Engine/Plugins/Experimental/Mutable/Source/MutableRuntime/Internal/MuR/OpMeshMerge.h
+++ b/Engine/Plugins/Experimental/Mutable/Source/MutableRuntime/Internal/MuR/OpMeshMerge.h
@@ -710,6 +710,22 @@ namespace mu
 
 		// This affects both vertex IDs and layout block ids.
 		bool bNeedsExplicitIds = pFirst->MeshIDPrefix != pSecond->MeshIDPrefix;
+
+		// These two extra checks are necessary for corner cases of meshes merging with fragments of themselves that
+		// undergo different operations.
+		if (!bNeedsExplicitIds)
+		{
+			UntypedMeshBufferIteratorConst FirstIDs(pFirst->VertexBuffers, MBS_VERTEXINDEX, 0);
+			UntypedMeshBufferIteratorConst SecondIDs(pSecond->VertexBuffers, MBS_VERTEXINDEX, 0);
+			bNeedsExplicitIds = FirstIDs.GetFormat() != SecondIDs.GetFormat();
+		}
+		if (!bNeedsExplicitIds)
+		{
+			UntypedMeshBufferIteratorConst FirstIDs(pFirst->VertexBuffers, MBS_LAYOUTBLOCK, 0);
+			UntypedMeshBufferIteratorConst SecondIDs(pSecond->VertexBuffers, MBS_LAYOUTBLOCK, 0);
+			bNeedsExplicitIds = FirstIDs.GetFormat() != SecondIDs.GetFormat();
+		}
+
 		if (!bNeedsExplicitIds)
 		{
 			// This is needed in case a mesh merges with itself.
