[mutable] Color nodes code cleanup.
This is necessary because they are used as reference for new nodes to implement.
#rnx

diff --git a/Engine/Plugins/Experimental/Mutable/Source/CustomizableObjectEditor/Private/MuCOE/GenerateMutableSource/GenerateMutableSourceColor.cpp b/Engine/Plugins/Experimental/Mutable/Source/CustomizableObjectEditor/Private/MuCOE/GenerateMutableSource/GenerateMutableSourceColor.cpp
--- a/Engine/Plugins/Experimental/Mutable/Source/CustomizableObjectEditor/Private/MuCOE/GenerateMutableSource/GenerateMutableSourceColor.cpp
+++ b/Engine/Plugins/Experimental/Mutable/Source/CustomizableObjectEditor/Private/MuCOE/GenerateMutableSource/GenerateMutableSourceColor.cpp
@@ -56,22 +56,22 @@ mu::NodeColourPtr GenerateMutableSourceColor(const UEdGraphPin* Pin, FMutableGra
 
 	if (const UCustomizableObjectNodeColorConstant* TypedNodeColorConst = Cast<UCustomizableObjectNodeColorConstant>(Node))
 	{
-		mu::NodeColourConstantPtr ColorNode = new mu::NodeColourConstant();
+		mu::Ptr<mu::NodeColourConstant> ColorNode = new mu::NodeColourConstant();
 		Result = ColorNode;
 
-		ColorNode->SetValue(TypedNodeColorConst->Value);
+		ColorNode->Value = TypedNodeColorConst->Value;
 	}
 
 	else if (const UCustomizableObjectNodeColorParameter* TypedNodeColorParam = Cast<UCustomizableObjectNodeColorParameter>(Node))
 	{
-		mu::NodeColourParameterPtr ColorNode = new mu::NodeColourParameter();
+		mu::Ptr<mu::NodeColourParameter> ColorNode = new mu::NodeColourParameter();
 		Result = ColorNode;
 
 		GenerationContext.AddParameterNameUnique(Node, TypedNodeColorParam->ParameterName);
 
-		ColorNode->SetName(TypedNodeColorParam->ParameterName);
-		ColorNode->SetUid(GenerationContext.GetNodeIdUnique(Node).ToString());
-		ColorNode->SetDefaultValue(TypedNodeColorParam->DefaultValue);
+		ColorNode->Name = TypedNodeColorParam->ParameterName;
+		ColorNode->Uid = GenerationContext.GetNodeIdUnique(Node).ToString();
+		ColorNode->DefaultValue = TypedNodeColorParam->DefaultValue;
 
 		GenerationContext.ParameterUIDataMap.Add(TypedNodeColorParam->ParameterName, FMutableParameterData(
 			TypedNodeColorParam->ParamUIMetadata,
@@ -126,16 +126,16 @@ mu::NodeColourPtr GenerateMutableSourceColor(const UEdGraphPin* Pin, FMutableGra
 				GenerationContext.Compiler->CompilerLog(Message, Node);
 			}
 
-			mu::NodeColourSwitchPtr SwitchNode = new mu::NodeColourSwitch;
-			SwitchNode->SetParameter(SwitchParam);
-			SwitchNode->SetOptionCount(NumSwitchOptions);
+			mu::Ptr<mu::NodeColourSwitch> SwitchNode = new mu::NodeColourSwitch;
+			SwitchNode->Parameter = SwitchParam;
+			SwitchNode->Options.SetNum(NumSwitchOptions);
 
 			for (int SelectorIndex = 0; SelectorIndex < NumSwitchOptions; ++SelectorIndex)
 			{
 				const UEdGraphPin* const ColorPin = TypedNodeColorSwitch->GetElementPin(SelectorIndex);
 				if (const UEdGraphPin* ConnectedPin = FollowInputPin(*ColorPin))
 				{
-					SwitchNode->SetOption(SelectorIndex, GenerateMutableSourceColor(ConnectedPin, GenerationContext));
+					SwitchNode->Options[SelectorIndex] = GenerateMutableSourceColor(ConnectedPin, GenerationContext);
 				}
 			}
 
@@ -145,98 +145,98 @@ mu::NodeColourPtr GenerateMutableSourceColor(const UEdGraphPin* Pin, FMutableGra
 
 	else if (const UCustomizableObjectNodeTextureSample* TypedNodeTexSample = Cast<UCustomizableObjectNodeTextureSample>(Node))
 	{
-		mu::NodeColourSampleImagePtr ColorNode = new mu::NodeColourSampleImage();
+		mu::Ptr<mu::NodeColourSampleImage> ColorNode = new mu::NodeColourSampleImage();
 		Result = ColorNode;
 
 		if (const UEdGraphPin* ConnectedPin = FollowInputPin(*TypedNodeTexSample->TexturePin()))
 		{
 			mu::NodeImagePtr TextureNode = GenerateMutableSourceImage(ConnectedPin, GenerationContext, 0);
-			ColorNode->SetImage(TextureNode);
+			ColorNode->Image = TextureNode;
 		}
 
 		if (const UEdGraphPin* ConnectedPin = FollowInputPin(*TypedNodeTexSample->XPin()))
 		{
 			mu::NodeScalarPtr XNode = GenerateMutableSourceFloat(ConnectedPin, GenerationContext);
-			ColorNode->SetX(XNode);
+			ColorNode->X = XNode;
 		}
 
 		if (const UEdGraphPin* ConnectedPin = FollowInputPin(*TypedNodeTexSample->YPin()))
 		{
 			mu::NodeScalarPtr YNode = GenerateMutableSourceFloat(ConnectedPin, GenerationContext);
-			ColorNode->SetY(YNode);
+			ColorNode->Y = YNode;
 		}
 	}
 
 	else if (const UCustomizableObjectNodeColorArithmeticOp* TypedNodeColorArith = Cast<UCustomizableObjectNodeColorArithmeticOp>(Node))
 	{
-		mu::NodeColourArithmeticOperationPtr OpNode = new mu::NodeColourArithmeticOperation();
+		mu::Ptr<mu::NodeColourArithmeticOperation> OpNode = new mu::NodeColourArithmeticOperation();
 		Result = OpNode;
 
 		if (const UEdGraphPin* ConnectedPin = FollowInputPin(*TypedNodeColorArith->XPin()))
 		{
 			mu::NodeColourPtr XNode = GenerateMutableSourceColor(ConnectedPin, GenerationContext);
-			OpNode->SetA(XNode);
+			OpNode->A = XNode;
 		}
 
 		if (const UEdGraphPin* ConnectedPin = FollowInputPin(*TypedNodeColorArith->YPin()))
 		{
 			mu::NodeColourPtr YNode = GenerateMutableSourceColor(ConnectedPin, GenerationContext);
-			OpNode->SetB(YNode);
+			OpNode->B = YNode;
 		}
 
 		switch (TypedNodeColorArith->Operation)
 		{
 		case EColorArithmeticOperation::E_Add:
-			OpNode->SetOperation(mu::NodeColourArithmeticOperation::OPERATION::AO_ADD);
+			OpNode->Operation = mu::NodeColourArithmeticOperation::EOperation::Add;
 			break;
 
 		case EColorArithmeticOperation::E_Sub:
-			OpNode->SetOperation(mu::NodeColourArithmeticOperation::OPERATION::AO_SUBTRACT);
+			OpNode->Operation = mu::NodeColourArithmeticOperation::EOperation::Subtract;
 			break;
 
 		case EColorArithmeticOperation::E_Mul:
-			OpNode->SetOperation(mu::NodeColourArithmeticOperation::OPERATION::AO_MULTIPLY);
+			OpNode->Operation = mu::NodeColourArithmeticOperation::EOperation::Multiply;
 			break;
 
 		case EColorArithmeticOperation::E_Div:
-			OpNode->SetOperation(mu::NodeColourArithmeticOperation::OPERATION::AO_DIVIDE);
+			OpNode->Operation = mu::NodeColourArithmeticOperation::EOperation::Divide;
 			break;
 		}
 	}
 
 	else if (const UCustomizableObjectNodeColorFromFloats* TypedNodeFrom = Cast<UCustomizableObjectNodeColorFromFloats>(Node))
 	{
-		mu::NodeColourFromScalarsPtr OpNode = new mu::NodeColourFromScalars();
+		mu::Ptr<mu::NodeColourFromScalars> OpNode = new mu::NodeColourFromScalars();
 		Result = OpNode;
 
 		if (const UEdGraphPin* ConnectedPin = FollowInputPin(*TypedNodeFrom->RPin()))
 		{
 			mu::NodeScalarPtr FloatNode = GenerateMutableSourceFloat(ConnectedPin, GenerationContext);
-			OpNode->SetX(FloatNode);
+			OpNode->X = FloatNode;
 		}
 
 		if (const UEdGraphPin* ConnectedPin = FollowInputPin(*TypedNodeFrom->GPin()))
 		{
 			mu::NodeScalarPtr FloatNode = GenerateMutableSourceFloat(ConnectedPin, GenerationContext);
-			OpNode->SetY(FloatNode);
+			OpNode->Y = FloatNode;
 		}
 
 		if (const UEdGraphPin* ConnectedPin = FollowInputPin(*TypedNodeFrom->BPin()))
 		{
 			mu::NodeScalarPtr FloatNode = GenerateMutableSourceFloat(ConnectedPin, GenerationContext);
-			OpNode->SetZ(FloatNode);
+			OpNode->Z = FloatNode;
 		}
 
 		if (const UEdGraphPin* ConnectedPin = FollowInputPin(*TypedNodeFrom->APin()))
 		{
 			mu::NodeScalarPtr FloatNode = GenerateMutableSourceFloat(ConnectedPin, GenerationContext);
-			OpNode->SetW(FloatNode);
+			OpNode->W = FloatNode;
 		}
 	}
 
 	else if (const UCustomizableObjectNodeColorVariation* TypedNodeColorVar = Cast<const UCustomizableObjectNodeColorVariation>(Node))
 	{
-		mu::NodeColourVariationPtr ColorNode = new mu::NodeColourVariation();
+		mu::Ptr<mu::NodeColourVariation> ColorNode = new mu::NodeColourVariation();
 		Result = ColorNode;
 
 		if (const UEdGraphPin* ConnectedPin = FollowInputPin(*TypedNodeColorVar->DefaultPin()))
@@ -244,7 +244,7 @@ mu::NodeColourPtr GenerateMutableSourceColor(const UEdGraphPin* Pin, FMutableGra
 			mu::NodeColourPtr ChildNode = GenerateMutableSourceColor(ConnectedPin, GenerationContext);
 			if (ChildNode)
 			{
-				ColorNode->SetDefaultColour(ChildNode.get());
+				ColorNode->DefaultColour = ChildNode;
 			}
 			else
 			{
@@ -253,17 +253,17 @@ mu::NodeColourPtr GenerateMutableSourceColor(const UEdGraphPin* Pin, FMutableGra
 		}
 
 		const int32 NumVariations = TypedNodeColorVar->GetNumVariations();
-		ColorNode->SetVariationCount(NumVariations);
+		ColorNode->Variations.SetNum(NumVariations);
 		for (int VariationIndex = 0; VariationIndex < NumVariations; ++VariationIndex)
 		{
 			const UEdGraphPin* VariationPin = TypedNodeColorVar->VariationPin(VariationIndex);
 			if (!VariationPin) continue;
 
-			ColorNode->SetVariationTag(VariationIndex, TypedNodeColorVar->GetVariation(VariationIndex).Tag);
+			ColorNode->Variations[VariationIndex].Tag = TypedNodeColorVar->GetVariation(VariationIndex).Tag;
 			if (const UEdGraphPin* ConnectedPin = FollowInputPin(*VariationPin))
 			{
 				mu::NodeColourPtr ChildNode = GenerateMutableSourceColor(ConnectedPin, GenerationContext);
-				ColorNode->SetVariationColour(VariationIndex, ChildNode.get());
+				ColorNode->Variations[VariationIndex].Colour = ChildNode;
 			}
 		}
 	}
@@ -271,8 +271,8 @@ mu::NodeColourPtr GenerateMutableSourceColor(const UEdGraphPin* Pin, FMutableGra
 	else if (const UCustomizableObjectNodeTable* TypedNodeTable = Cast<UCustomizableObjectNodeTable>(Node))
 	{
 		//This node will add a white color in case of error
-		mu::NodeColourConstantPtr WhiteColorNode = new mu::NodeColourConstant();
-		WhiteColorNode->SetValue(FVector4f(1.0f, 1.0f, 1.0f, 1.0f));
+		mu::Ptr<mu::NodeColourConstant> WhiteColorNode = new mu::NodeColourConstant();
+		WhiteColorNode->Value = FVector4f(1.0f, 1.0f, 1.0f, 1.0f);
 
 		Result = WhiteColorNode;
 
@@ -300,7 +300,7 @@ mu::NodeColourPtr GenerateMutableSourceColor(const UEdGraphPin* Pin, FMutableGra
 
 				if (Table)
 				{
-					mu::NodeColourTablePtr ColorTableNode = new mu::NodeColourTable();
+					mu::Ptr<mu::NodeColourTable> ColorTableNode = new mu::NodeColourTable();
 
 					// Generating a new Color column if not exists
 					if (Table->FindColumn(ColumnName) == INDEX_NONE)
@@ -321,10 +321,10 @@ mu::NodeColourPtr GenerateMutableSourceColor(const UEdGraphPin* Pin, FMutableGra
 						Result = ColorTableNode;
 
 						ColorTableNode->Table = Table;
-						ColorTableNode->SetColumn(ColumnName);
-						ColorTableNode->SetParameterName(TypedNodeTable->ParameterName);
-						ColorTableNode->SetNoneOption(TypedNodeTable->bAddNoneOption);
-						ColorTableNode->SetDefaultRowName(TypedNodeTable->DefaultRowName.ToString());
+						ColorTableNode->ColumnName = ColumnName;
+						ColorTableNode->ParameterName = TypedNodeTable->ParameterName;
+						ColorTableNode->bNoneOption = TypedNodeTable->bAddNoneOption;
+						ColorTableNode->DefaultRowName = TypedNodeTable->DefaultRowName.ToString();
 					}
 				}
 				else
diff --git a/Engine/Plugins/Experimental/Mutable/Source/CustomizableObjectEditor/Private/MuCOE/GenerateMutableSource/GenerateMutableSourceColor.h b/Engine/Plugins/Experimental/Mutable/Source/CustomizableObjectEditor/Private/MuCOE/GenerateMutableSource/GenerateMutableSourceColor.h
--- a/Engine/Plugins/Experimental/Mutable/Source/CustomizableObjectEditor/Private/MuCOE/GenerateMutableSource/GenerateMutableSourceColor.h
+++ b/Engine/Plugins/Experimental/Mutable/Source/CustomizableObjectEditor/Private/MuCOE/GenerateMutableSource/GenerateMutableSourceColor.h
@@ -8,4 +8,4 @@ class UEdGraphPin;
 struct FMutableGraphGenerationContext;
 
 /** Convert a CustomizableObject Source Graph into a mutable source graph. */
-mu::NodeColourPtr GenerateMutableSourceColor(const UEdGraphPin* Pin, FMutableGraphGenerationContext& GenerationContext);
+mu::Ptr<mu::NodeColour> GenerateMutableSourceColor(const UEdGraphPin* Pin, FMutableGraphGenerationContext& GenerationContext);
diff --git a/Engine/Plugins/Experimental/Mutable/Source/CustomizableObjectEditor/Private/MuCOE/GenerateMutableSource/GenerateMutableSourceGroupProjector.cpp b/Engine/Plugins/Experimental/Mutable/Source/CustomizableObjectEditor/Private/MuCOE/GenerateMutableSource/GenerateMutableSourceGroupProjector.cpp
--- a/Engine/Plugins/Experimental/Mutable/Source/CustomizableObjectEditor/Private/MuCOE/GenerateMutableSource/GenerateMutableSourceGroupProjector.cpp
+++ b/Engine/Plugins/Experimental/Mutable/Source/CustomizableObjectEditor/Private/MuCOE/GenerateMutableSource/GenerateMutableSourceGroupProjector.cpp
@@ -193,8 +193,8 @@ mu::NodeImagePtr GenerateMutableGroupProjection(const int32 NodeLOD, const int32
 		return mu::NodeImagePtr();
 	}
 	
-	mu::NodeColourConstantPtr ZeroColorNode = new mu::NodeColourConstant();
-	ZeroColorNode->SetValue(FVector4f(0.f, 0.f, 0.f, 1.0f));
+	mu::Ptr<mu::NodeColourConstant> ZeroColorNode = new mu::NodeColourConstant();
+	ZeroColorNode->Value = FVector4f(0.f, 0.f, 0.f, 1.0f);
 
 	mu::NodeImagePlainColourPtr ZeroPlainColourNode = new mu::NodeImagePlainColour;
 	ZeroPlainColourNode->SetSize(TextureSize, TextureSize);
@@ -237,11 +237,11 @@ mu::NodeImagePtr GenerateMutableGroupProjection(const int32 NodeLOD, const int32
 		ImageNodesAlphaChannelNode->SetSource(0, ImageNodes[i]);
 		ImageNodesAlphaChannelNode->SetSourceChannel(0, 3);
 
-		mu::NodeColourFromScalarsPtr ColourFromScalars = new mu::NodeColourFromScalars;
-		ColourFromScalars->SetX(ImageNodes_ProjectorTempData[i].NodeOpacityParameter);
-		ColourFromScalars->SetY(ImageNodes_ProjectorTempData[i].NodeOpacityParameter);
-		ColourFromScalars->SetZ(ImageNodes_ProjectorTempData[i].NodeOpacityParameter);
-		ColourFromScalars->SetW(OneConstantNode);
+		mu::Ptr<mu::NodeColourFromScalars> ColourFromScalars = new mu::NodeColourFromScalars;
+		ColourFromScalars->X = ImageNodes_ProjectorTempData[i].NodeOpacityParameter;
+		ColourFromScalars->Y = ImageNodes_ProjectorTempData[i].NodeOpacityParameter;
+		ColourFromScalars->Z = ImageNodes_ProjectorTempData[i].NodeOpacityParameter;
+		ColourFromScalars->W = OneConstantNode;
 
 		mu::NodeImageLayerColourPtr OpacityMultiLayerNode = new mu::NodeImageLayerColour;
 		OpacityMultiLayerNode->SetBlendType(mu::EBlendType::BT_MULTIPLY);
diff --git a/Engine/Plugins/Experimental/Mutable/Source/CustomizableObjectEditor/Private/MuCOE/GenerateMutableSource/GenerateMutableSourceImage.cpp b/Engine/Plugins/Experimental/Mutable/Source/CustomizableObjectEditor/Private/MuCOE/GenerateMutableSource/GenerateMutableSourceImage.cpp
--- a/Engine/Plugins/Experimental/Mutable/Source/CustomizableObjectEditor/Private/MuCOE/GenerateMutableSource/GenerateMutableSourceImage.cpp
+++ b/Engine/Plugins/Experimental/Mutable/Source/CustomizableObjectEditor/Private/MuCOE/GenerateMutableSource/GenerateMutableSourceImage.cpp
@@ -459,30 +459,26 @@ mu::NodeImagePtr GenerateMutableSourceImage(const UEdGraphPin* Pin, FMutableGrap
 
 	else if (const UCustomizableObjectNodeTextureFromFloats* TypedNodeFromFloats = Cast<UCustomizableObjectNodeTextureFromFloats>(Node))
 	{		
-		mu::NodeColourFromScalarsPtr Color = new mu::NodeColourFromScalars();
+		mu::Ptr<mu::NodeColourFromScalars> Color = new mu::NodeColourFromScalars();
 
 		if (const UEdGraphPin* RPin = FollowInputPin(*TypedNodeFromFloats->RPin()))
 		{
-			mu::NodeScalarPtr R = GenerateMutableSourceFloat(RPin, GenerationContext);
-			Color->SetX(R);
+			Color->X = GenerateMutableSourceFloat(RPin, GenerationContext);
 		}
 
 		if (const UEdGraphPin* GPin = FollowInputPin(*TypedNodeFromFloats->GPin()))
 		{
-			mu::NodeScalarPtr G = GenerateMutableSourceFloat(GPin, GenerationContext);
-			Color->SetY(G);
+			Color->Y = GenerateMutableSourceFloat(GPin, GenerationContext);
 		}
 
 		if (const UEdGraphPin* BPin = FollowInputPin(*TypedNodeFromFloats->BPin()))
 		{
-			mu::NodeScalarPtr B = GenerateMutableSourceFloat(BPin, GenerationContext);
-			Color->SetZ(B);
+			Color->Z = GenerateMutableSourceFloat(BPin, GenerationContext);
 		}
 
 		if (const UEdGraphPin* APin = FollowInputPin(*TypedNodeFromFloats->APin()))
 		{
-			mu::NodeScalarPtr A = GenerateMutableSourceFloat(APin, GenerationContext);
-			Color->SetW(A);
+			Color->W = GenerateMutableSourceFloat(APin, GenerationContext);
 		}
 
 		mu::NodeImagePlainColourPtr ImageFromColour = new mu::NodeImagePlainColour;
diff --git a/Engine/Plugins/Experimental/Mutable/Source/CustomizableObjectEditor/Private/MuCOE/SMutableGraphViewer.cpp b/Engine/Plugins/Experimental/Mutable/Source/CustomizableObjectEditor/Private/MuCOE/SMutableGraphViewer.cpp
--- a/Engine/Plugins/Experimental/Mutable/Source/CustomizableObjectEditor/Private/MuCOE/SMutableGraphViewer.cpp
+++ b/Engine/Plugins/Experimental/Mutable/Source/CustomizableObjectEditor/Private/MuCOE/SMutableGraphViewer.cpp
@@ -77,8 +77,6 @@
 #include "MuT/NodeModifierMeshClipDeformPrivate.h"
 #include "MuT/NodeModifierMeshClipWithMeshPrivate.h"
 #include "MuT/NodeModifierMeshClipWithUVMaskPrivate.h"
-#include "MuT/NodeColourParameterPrivate.h"
-#include "MuT/NodeColourSampleImagePrivate.h"
 #include "MuT/NodeImageInterpolatePrivate.h"
 #include "MuT/NodeImagePlainColourPrivate.h"
 #include "MuT/NodeImageProjectPrivate.h"
@@ -86,15 +84,12 @@
 #include "MuT/NodeMeshMorphPrivate.h"
 #include "MuT/NodeScalarEnumParameterPrivate.h"
 #include "MuT/NodeScalarParameterPrivate.h"
-#include "MuT/NodeColourPrivate.h"
 #include "MuT/NodeMeshMakeMorphPrivate.h"
 #include "MuT/NodeScalarCurvePrivate.h"
 #include "MuT/NodeProjectorPrivate.h"
-#include "MuT/NodeColourSwitchPrivate.h"
 #include "MuT/NodeImageInvertPrivate.h"
 #include "MuT/NodeImageSwizzlePrivate.h"
 #include "MuT/NodeImageMultiLayerPrivate.h"
-#include "MuT/NodeColourFromScalarsPrivate.h"
 #include "MuT/NodeScalarSwitchPrivate.h"
 
 class FExtender;
@@ -620,10 +615,9 @@ void SMutableGraphViewer::GetChildrenForInfo(TSharedPtr<FMutableGraphTreeElement
 	else if (ParentNode->GetType() == mu::NodeColourSampleImage::GetStaticType())
 	{
 		mu::NodeColourSampleImage* ColorSampleImageVar = StaticCast<mu::NodeColourSampleImage*>(ParentNode);
-		mu::NodeColourSampleImage::Private* Private = ColorSampleImageVar->GetPrivate();
-		AddChildFunc(Private->m_pImage.get(), FString::Printf(TEXT("IMAGE")));
-		AddChildFunc(Private->m_pX.get(), FString::Printf(TEXT("X")));
-		AddChildFunc(Private->m_pY.get(), FString::Printf(TEXT("Y")));
+		AddChildFunc(ColorSampleImageVar->Image.get(), FString::Printf(TEXT("IMAGE")));
+		AddChildFunc(ColorSampleImageVar->X.get(), FString::Printf(TEXT("X")));
+		AddChildFunc(ColorSampleImageVar->Y.get(), FString::Printf(TEXT("Y")));
 	}
 
 	else if (ParentNode->GetType() == mu::NodeImageInterpolate::GetStaticType())
@@ -655,10 +649,9 @@ void SMutableGraphViewer::GetChildrenForInfo(TSharedPtr<FMutableGraphTreeElement
 	else if (ParentNode->GetType() == mu::NodeColourParameter::GetStaticType())
 	{
 		mu::NodeColourParameter* ColorParameterVar = StaticCast<mu::NodeColourParameter*>(ParentNode);
-		mu::NodeColourParameter::Private* Private = ColorParameterVar->GetPrivate();
-		for (int32 RangeIndex = 0; RangeIndex < Private->m_ranges.Num(); RangeIndex++)
+		for (int32 RangeIndex = 0; RangeIndex < ColorParameterVar->Ranges.Num(); RangeIndex++)
 		{
-			AddChildFunc(Private->m_ranges[RangeIndex].get(), FString::Printf(TEXT("RANGE [%d]"), RangeIndex));
+			AddChildFunc(ColorParameterVar->Ranges[RangeIndex].get(), FString::Printf(TEXT("RANGE [%d]"), RangeIndex));
 		}
 	}
 
@@ -706,11 +699,10 @@ void SMutableGraphViewer::GetChildrenForInfo(TSharedPtr<FMutableGraphTreeElement
 	else if (ParentNode->GetType() == mu::NodeColourSwitch::GetStaticType())
 	{
 		mu::NodeColourSwitch* ColorSwitchVar = StaticCast<mu::NodeColourSwitch*>(ParentNode);
-		mu::NodeColourSwitch::Private* Private = ColorSwitchVar->GetPrivate();
-		AddChildFunc(Private->m_pParameter.get(), TEXT("PARAM"));
-		for (int32 OptionIndex = 0; OptionIndex < Private->m_options.Num(); ++OptionIndex)
+		AddChildFunc(ColorSwitchVar->Parameter.get(), TEXT("PARAM"));
+		for (int32 OptionIndex = 0; OptionIndex < ColorSwitchVar->Options.Num(); ++OptionIndex)
 		{
-			AddChildFunc(Private->m_options[OptionIndex].get(), FString::Printf(TEXT("OPTION [%d]"), OptionIndex));
+			AddChildFunc(ColorSwitchVar->Options[OptionIndex].get(), FString::Printf(TEXT("OPTION [%d]"), OptionIndex));
 		}
 	}
 
@@ -774,11 +766,10 @@ void SMutableGraphViewer::GetChildrenForInfo(TSharedPtr<FMutableGraphTreeElement
 	else if (ParentNode->GetType() == mu::NodeColourFromScalars::GetStaticType())
 	{
 		mu::NodeColourFromScalars* ScalarTableVar = StaticCast<mu::NodeColourFromScalars*>(ParentNode);
-		mu::NodeColourFromScalars::Private* Private = ScalarTableVar->GetPrivate();
-		AddChildFunc(Private->m_pX.get(), TEXT("X"));
-		AddChildFunc(Private->m_pY.get(), TEXT("Y"));
-		AddChildFunc(Private->m_pZ.get(), TEXT("Z"));
-		AddChildFunc(Private->m_pW.get(), TEXT("W"));
+		AddChildFunc(ScalarTableVar->X.get(), TEXT("X"));
+		AddChildFunc(ScalarTableVar->Y.get(), TEXT("Y"));
+		AddChildFunc(ScalarTableVar->Z.get(), TEXT("Z"));
+		AddChildFunc(ScalarTableVar->W.get(), TEXT("W"));
 	}
 	
 	else
diff --git a/Engine/Plugins/Experimental/Mutable/Source/MutableTools/Internal/MuT/NodeColourFromScalarsPrivate.h b/Engine/Plugins/Experimental/Mutable/Source/MutableTools/Internal/MuT/NodeColourFromScalarsPrivate.h
deleted file mode 100644
diff --git a/Engine/Plugins/Experimental/Mutable/Source/MutableTools/Internal/MuT/NodeColourParameterPrivate.h b/Engine/Plugins/Experimental/Mutable/Source/MutableTools/Internal/MuT/NodeColourParameterPrivate.h
deleted file mode 100644
diff --git a/Engine/Plugins/Experimental/Mutable/Source/MutableTools/Internal/MuT/NodeColourPrivate.h b/Engine/Plugins/Experimental/Mutable/Source/MutableTools/Internal/MuT/NodeColourPrivate.h
deleted file mode 100644
diff --git a/Engine/Plugins/Experimental/Mutable/Source/MutableTools/Internal/MuT/NodeColourSampleImagePrivate.h b/Engine/Plugins/Experimental/Mutable/Source/MutableTools/Internal/MuT/NodeColourSampleImagePrivate.h
deleted file mode 100644
diff --git a/Engine/Plugins/Experimental/Mutable/Source/MutableTools/Internal/MuT/NodeColourSwitchPrivate.h b/Engine/Plugins/Experimental/Mutable/Source/MutableTools/Internal/MuT/NodeColourSwitchPrivate.h
deleted file mode 100644
diff --git a/Engine/Plugins/Experimental/Mutable/Source/MutableTools/Private/MuT/CodeGenerator.cpp b/Engine/Plugins/Experimental/Mutable/Source/MutableTools/Private/MuT/CodeGenerator.cpp
--- a/Engine/Plugins/Experimental/Mutable/Source/MutableTools/Private/MuT/CodeGenerator.cpp
+++ b/Engine/Plugins/Experimental/Mutable/Source/MutableTools/Private/MuT/CodeGenerator.cpp
@@ -2369,7 +2369,7 @@ namespace mu
 		case mu::ETableColumnType::Color:
 		{
 			mu::Ptr<mu::NodeColourConstant> pNode = new NodeColourConstant();
-			pNode->SetValue(mu::DefaultMutableColorValue);
+			pNode->Value = mu::DefaultMutableColorValue;
 
 			FColorGenerationResult ChildResult;
 			FGenericGenerationOptions Options;
diff --git a/Engine/Plugins/Experimental/Mutable/Source/MutableTools/Private/MuT/CodeGenerator_Color.cpp b/Engine/Plugins/Experimental/Mutable/Source/MutableTools/Private/MuT/CodeGenerator_Color.cpp
--- a/Engine/Plugins/Experimental/Mutable/Source/MutableTools/Private/MuT/CodeGenerator_Color.cpp
+++ b/Engine/Plugins/Experimental/Mutable/Source/MutableTools/Private/MuT/CodeGenerator_Color.cpp
@@ -23,20 +23,13 @@
 #include "MuT/Node.h"
 #include "MuT/NodeColour.h"
 #include "MuT/NodeColourArithmeticOperation.h"
-#include "MuT/NodeColourArithmeticOperationPrivate.h"
 #include "MuT/NodeColourConstant.h"
-#include "MuT/NodeColourConstantPrivate.h"
 #include "MuT/NodeColourFromScalars.h"
-#include "MuT/NodeColourFromScalarsPrivate.h"
 #include "MuT/NodeColourParameter.h"
-#include "MuT/NodeColourParameterPrivate.h"
 #include "MuT/NodeColourSampleImage.h"
-#include "MuT/NodeColourSampleImagePrivate.h"
 #include "MuT/NodeColourSwitch.h"
-#include "MuT/NodeColourSwitchPrivate.h"
 #include "MuT/NodeColourTable.h"
 #include "MuT/NodeColourVariation.h"
-#include "MuT/NodeColourVariationPrivate.h"
 #include "MuT/NodeImage.h"
 #include "MuT/NodeImageProject.h"
 #include "MuT/NodeMeshGeometryOperation.h"
@@ -126,14 +119,14 @@ namespace mu
 	//-------------------------------------------------------------------------------------------------
 	void CodeGenerator::GenerateColor_Constant(FColorGenerationResult& result, const FGenericGenerationOptions& Options, const Ptr<const NodeColourConstant>& Typed)
 	{
-		const NodeColourConstant::Private& node = *Typed->GetPrivate();
+		const NodeColourConstant& node = *Typed;
 
 		Ptr<ASTOpFixed> op = new ASTOpFixed();
 		op->op.type = OP_TYPE::CO_CONSTANT;
-		op->op.args.ColourConstant.value[0] = node.m_value.X;
-		op->op.args.ColourConstant.value[1] = node.m_value.Y;
-		op->op.args.ColourConstant.value[2] = node.m_value.Z;
-		op->op.args.ColourConstant.value[3] = node.m_value.W;
+		op->op.args.ColourConstant.value[0] = node.Value.X;
+		op->op.args.ColourConstant.value[1] = node.Value.Y;
+		op->op.args.ColourConstant.value[2] = node.Value.Z;
+		op->op.args.ColourConstant.value[3] = node.Value.W;
 
 		result.op = op;
 	}
@@ -142,25 +135,25 @@ namespace mu
 	//-------------------------------------------------------------------------------------------------
 	void CodeGenerator::GenerateColor_Parameter(FColorGenerationResult& result, const FGenericGenerationOptions& Options, const Ptr<const NodeColourParameter>& Typed)
 	{
-		const NodeColourParameter::Private& node = *Typed->GetPrivate();
+		const NodeColourParameter& node = *Typed;
 
 		Ptr<ASTOpParameter> op;
 
-		Ptr<ASTOpParameter>* it = m_firstPass.ParameterNodes.Find(node.m_pNode);
+		Ptr<ASTOpParameter>* it = m_firstPass.ParameterNodes.Find(Typed);
 
 		if (!it)
 		{
 			FParameterDesc param;
-			param.m_name = node.m_name;
-			const TCHAR* CStr = ToCStr(node.m_uid);
+			param.m_name = node.Name;
+			const TCHAR* CStr = ToCStr(node.Uid);
 			param.m_uid.ImportTextItem(CStr, 0, nullptr, nullptr);
 			param.m_type = PARAMETER_TYPE::T_COLOUR;
 
 			ParamColorType Value;
-			Value[0] =  node.m_defaultValue[0];
-			Value[1] = node.m_defaultValue[1];
-			Value[2] = node.m_defaultValue[2];
-			Value[3] = node.m_defaultValue[3];
+			Value[0] =  node.DefaultValue[0];
+			Value[1] = node.DefaultValue[1];
+			Value[2] = node.DefaultValue[2];
+			Value[3] = node.DefaultValue[3];
 
 			param.m_defaultValue.Set<ParamColorType>(Value);
 			
@@ -169,14 +162,14 @@ namespace mu
 			op->parameter = param;
 
 			// Generate the code for the ranges
-			for (int32 a = 0; a < node.m_ranges.Num(); ++a)
+			for (int32 a = 0; a < node.Ranges.Num(); ++a)
 			{
 				FRangeGenerationResult rangeResult;
-				GenerateRange(rangeResult, Options, node.m_ranges[a]);
+				GenerateRange(rangeResult, Options, node.Ranges[a]);
 				op->ranges.Emplace(op.get(), rangeResult.sizeOp, rangeResult.rangeName, rangeResult.rangeUID);
 			}
 
-			m_firstPass.ParameterNodes.Add(node.m_pNode, op);
+			m_firstPass.ParameterNodes.Add(Typed, op);
 		}
 		else
 		{
@@ -190,11 +183,11 @@ namespace mu
 	//-------------------------------------------------------------------------------------------------
 	void CodeGenerator::GenerateColor_Switch(FColorGenerationResult& result, const FGenericGenerationOptions& Options, const Ptr<const NodeColourSwitch>& Typed)
 	{
-		const NodeColourSwitch::Private& node = *Typed->GetPrivate();
+		const NodeColourSwitch& node = *Typed;
 
 		MUTABLE_CPUPROFILER_SCOPE(NodeColourSwitch);
 
-		if (node.m_options.Num() == 0)
+		if (node.Options.Num() == 0)
 		{
 			// No options in the switch!
 			Ptr<ASTOp> missingOp = GenerateMissingColourCode(TEXT("Switch option"), Typed->GetMessageContext());
@@ -206,9 +199,9 @@ namespace mu
 		op->type = OP_TYPE::CO_SWITCH;
 
 		// Variable value
-		if (node.m_pParameter)
+		if (node.Parameter)
 		{
-			op->variable = Generate_Generic(node.m_pParameter.get(), Options);
+			op->variable = Generate_Generic(node.Parameter.get(), Options);
 		}
 		else
 		{
@@ -217,12 +210,12 @@ namespace mu
 		}
 
 		// Options
-		for (int32 t = 0; t < node.m_options.Num(); ++t)
+		for (int32 t = 0; t < node.Options.Num(); ++t)
 		{
 			Ptr<ASTOp> branch;
-			if (node.m_options[t])
+			if (node.Options[t])
 			{
-				branch = Generate_Generic(node.m_options[t].get(), Options);
+				branch = Generate_Generic(node.Options[t].get(), Options);
 			}
 			else
 			{
@@ -239,23 +232,23 @@ namespace mu
 	//-------------------------------------------------------------------------------------------------
 	void CodeGenerator::GenerateColor_Variation(FColorGenerationResult& result, const FGenericGenerationOptions& Options, const Ptr<const NodeColourVariation>& Typed)
 	{
-		const NodeColourVariation::Private& node = *Typed->GetPrivate();
+		const NodeColourVariation& node = *Typed;
 
 		Ptr<ASTOp> currentOp;
 
 		// Default case
-		if (node.m_defaultColour)
+		if (node.DefaultColour)
 		{
 			FColorGenerationResult BranchResults;
-			GenerateColor(BranchResults, Options, node.m_defaultColour);
+			GenerateColor(BranchResults, Options, node.DefaultColour);
 			currentOp = BranchResults.op;
 		}
 
 		// Process variations in reverse order, since conditionals are built bottom-up.
-		for (int t = node.m_variations.Num() - 1; t >= 0; --t)
+		for (int t = node.Variations.Num() - 1; t >= 0; --t)
 		{
 			int tagIndex = -1;
-			const FString& tag = node.m_variations[t].m_tag;
+			const FString& tag = node.Variations[t].Tag;
 			for (int i = 0; i < m_firstPass.m_tags.Num(); ++i)
 			{
 				if (m_firstPass.m_tags[i].tag == tag)
@@ -272,10 +265,10 @@ namespace mu
 			}
 
 			Ptr<ASTOp> variationOp;
-			if (node.m_variations[t].m_colour)
+			if (node.Variations[t].Colour)
 			{
 				FColorGenerationResult BranchResults;
-				GenerateColor(BranchResults,Options,node.m_variations[t].m_colour);
+				GenerateColor(BranchResults,Options,node.Variations[t].Colour);
 				variationOp = BranchResults.op;
 			}
 			else
@@ -301,7 +294,7 @@ namespace mu
 	//-------------------------------------------------------------------------------------------------
 	void CodeGenerator::GenerateColor_SampleImage(FColorGenerationResult& result, const FGenericGenerationOptions& Options, const Ptr<const NodeColourSampleImage>& Typed)
 	{
-		const NodeColourSampleImage::Private& node = *Typed->GetPrivate();
+		const NodeColourSampleImage& node = *Typed;
 
 		// Generate the code
 		Ptr<ASTOpFixed> op = new ASTOpFixed();
@@ -316,11 +309,11 @@ namespace mu
 		}
 
 		Ptr<ASTOp> base;
-		if (node.m_pImage)
+		if (node.Image)
 		{
 			// Generate
 			FImageGenerationResult MapResult;
-			GenerateImage(ImageOptions, MapResult, node.m_pImage);
+			GenerateImage(ImageOptions, MapResult, node.Image);
 			base = MapResult.op;
 		}
 		else
@@ -334,7 +327,7 @@ namespace mu
 		FScalarGenerationResult ChildResult;
 
 		// X
-		if (NodeScalar* pX = node.m_pX.get())
+		if (NodeScalar* pX = node.X.get())
 		{
 			GenerateScalar(ChildResult, Options, pX);
 			op->SetChild(op->op.args.ColourSampleImage.x, ChildResult.op);
@@ -350,7 +343,7 @@ namespace mu
 
 
 		// Y
-		if (NodeScalar* pY = node.m_pY.get())
+		if (NodeScalar* pY = node.Y.get())
 		{
 			GenerateScalar(ChildResult, Options, pY);
 			op->SetChild(op->op.args.ColourSampleImage.y, ChildResult.op);
@@ -374,7 +367,7 @@ namespace mu
 	//-------------------------------------------------------------------------------------------------
 	void CodeGenerator::GenerateColor_FromScalars(FColorGenerationResult& result, const FGenericGenerationOptions& Options, const Ptr<const NodeColourFromScalars>& Typed)
 	{
-		const NodeColourFromScalars::Private& node = *Typed->GetPrivate();
+		const NodeColourFromScalars& node = *Typed;
 
 		Ptr<ASTOpFixed> op = new ASTOpFixed();
 		op->op.type = OP_TYPE::CO_FROMSCALARS;
@@ -382,7 +375,7 @@ namespace mu
 		FScalarGenerationResult ChildResult;
 
 		// X
-		if (NodeScalar* pX = node.m_pX.get())
+		if (NodeScalar* pX = node.X.get())
 		{
 			GenerateScalar(ChildResult, Options, pX);
 			op->SetChild(op->op.args.ColourFromScalars.v[0], ChildResult.op );
@@ -396,7 +389,7 @@ namespace mu
 		}
 
 		// Y
-		if (NodeScalar* pY = node.m_pY.get())
+		if (NodeScalar* pY = node.Y.get())
 		{
 			GenerateScalar(ChildResult, Options, pY);
 			op->SetChild(op->op.args.ColourFromScalars.v[1], ChildResult.op);
@@ -410,7 +403,7 @@ namespace mu
 		}
 
 		// Z
-		if (NodeScalar* pZ = node.m_pZ.get())
+		if (NodeScalar* pZ = node.Z.get())
 		{
 			GenerateScalar(ChildResult, Options, pZ);
 			op->SetChild(op->op.args.ColourFromScalars.v[2], ChildResult.op);
@@ -424,7 +417,7 @@ namespace mu
 		}
 
 		// W
-		if (NodeScalar* pW = node.m_pW.get())
+		if (NodeScalar* pW = node.W.get())
 		{
 			GenerateScalar(ChildResult, Options, pW);
 			op->SetChild(op->op.args.ColourFromScalars.v[3], ChildResult.op);
@@ -444,17 +437,17 @@ namespace mu
 	//-------------------------------------------------------------------------------------------------
 	void CodeGenerator::GenerateColor_Arithmetic(FColorGenerationResult& result, const FGenericGenerationOptions& Options, const Ptr<const NodeColourArithmeticOperation>& Typed)
 	{
-		const NodeColourArithmeticOperation::Private& node = *Typed->GetPrivate();
+		const NodeColourArithmeticOperation& node = *Typed;
 
 		Ptr<ASTOpFixed> op = new ASTOpFixed();
 		op->op.type = OP_TYPE::CO_ARITHMETIC;
 
-		switch (node.m_operation)
+		switch (node.Operation)
 		{
-		case NodeColourArithmeticOperation::AO_ADD: op->op.args.ColourArithmetic.operation = OP::ArithmeticArgs::ADD; break;
-		case NodeColourArithmeticOperation::AO_SUBTRACT: op->op.args.ColourArithmetic.operation = OP::ArithmeticArgs::SUBTRACT; break;
-		case NodeColourArithmeticOperation::AO_MULTIPLY: op->op.args.ColourArithmetic.operation = OP::ArithmeticArgs::MULTIPLY; break;
-		case NodeColourArithmeticOperation::AO_DIVIDE: op->op.args.ColourArithmetic.operation = OP::ArithmeticArgs::DIVIDE; break;
+		case NodeColourArithmeticOperation::EOperation::Add: op->op.args.ColourArithmetic.operation = OP::ArithmeticArgs::ADD; break;
+		case NodeColourArithmeticOperation::EOperation::Subtract: op->op.args.ColourArithmetic.operation = OP::ArithmeticArgs::SUBTRACT; break;
+		case NodeColourArithmeticOperation::EOperation::Multiply: op->op.args.ColourArithmetic.operation = OP::ArithmeticArgs::MULTIPLY; break;
+		case NodeColourArithmeticOperation::EOperation::Divide: op->op.args.ColourArithmetic.operation = OP::ArithmeticArgs::DIVIDE; break;
 		default:
 			checkf(false, TEXT("Unknown arithmetic operation."));
 			op->op.args.ColourArithmetic.operation = OP::ArithmeticArgs::NONE;
@@ -464,7 +457,7 @@ namespace mu
 		FColorGenerationResult ChildResult;
 
 		// A
-		if (NodeColour* pA = node.m_pA.get())
+		if (NodeColour* pA = node.A.get())
 		{
 			GenerateColor(ChildResult, Options, pA );
 			op->SetChild(op->op.args.ColourArithmetic.a, ChildResult.op);
@@ -476,7 +469,7 @@ namespace mu
 		}
 
 		// B
-		if (NodeColour* pB = node.m_pB.get())
+		if (NodeColour* pB = node.B.get())
 		{
 			GenerateColor(ChildResult, Options, pB);
 			op->SetChild(op->op.args.ColourArithmetic.b, ChildResult.op);
@@ -501,7 +494,7 @@ namespace mu
 			{
 				Ptr<NodeColourConstant> CellData = new NodeColourConstant();
 				FVector4f Colour = node.Table->GetPrivate()->Rows[row].Values[colIndex].Color;
-				CellData->SetValue(Colour);
+				CellData->Value = Colour;
 				FColorGenerationResult BranchResults;
 				GenerateColor(BranchResults, Options, CellData );
 				return BranchResults.op;
@@ -517,8 +510,8 @@ namespace mu
 		m_pErrorLog->GetPrivate()->Add(Msg, ELMT_ERROR, errorContext);
 
 		// Create a constant colour node
-		NodeColourConstantPtr pNode = new NodeColourConstant();
-		pNode->SetValue(FVector4f(1, 1, 0, 1));
+		Ptr<NodeColourConstant> pNode = new NodeColourConstant();
+		pNode->Value = FVector4f(1, 1, 0, 1);
 
 		FColorGenerationResult Result;
 		FGenericGenerationOptions Options;
diff --git a/Engine/Plugins/Experimental/Mutable/Source/MutableTools/Private/MuT/CodeGenerator_Image.cpp b/Engine/Plugins/Experimental/Mutable/Source/MutableTools/Private/MuT/CodeGenerator_Image.cpp
--- a/Engine/Plugins/Experimental/Mutable/Source/MutableTools/Private/MuT/CodeGenerator_Image.cpp
+++ b/Engine/Plugins/Experimental/Mutable/Source/MutableTools/Private/MuT/CodeGenerator_Image.cpp
@@ -1606,11 +1606,11 @@ namespace mu
 
 			Ptr<NodeColourFromScalars> pPropsNode = new NodeColourFromScalars();
 
-			if (node.m_pAngleFadeStart) pPropsNode->SetX(node.m_pAngleFadeStart);
-			else pPropsNode->SetX(pDefaultFade);
+			if (node.m_pAngleFadeStart) pPropsNode->X = node.m_pAngleFadeStart;
+			else pPropsNode->X = pDefaultFade;
 
-			if (node.m_pAngleFadeEnd) pPropsNode->SetY(node.m_pAngleFadeEnd);
-			else pPropsNode->SetY(pDefaultFade);
+			if (node.m_pAngleFadeEnd) pPropsNode->Y = node.m_pAngleFadeEnd;
+			else pPropsNode->Y = pDefaultFade;
 
 			ImageRasterOp->angleFadeProperties = Generate_Generic(pPropsNode, Options);
 		}
@@ -2013,7 +2013,7 @@ namespace mu
     Ptr<ASTOp> CodeGenerator::GeneratePlainImageCode( const FVector4f& InColor, const FImageGenerationOptions& Options )
     {
 		Ptr<NodeColourConstant> ConstantColor = new NodeColourConstant();
-		ConstantColor->SetValue(InColor);
+		ConstantColor->Value = InColor;
 
         Ptr<NodeImagePlainColour> PlainNode = new NodeImagePlainColour();
 		PlainNode->SetColour(ConstantColor);
diff --git a/Engine/Plugins/Experimental/Mutable/Source/MutableTools/Private/MuT/Node.cpp b/Engine/Plugins/Experimental/Mutable/Source/MutableTools/Private/MuT/Node.cpp
--- a/Engine/Plugins/Experimental/Mutable/Source/MutableTools/Private/MuT/Node.cpp
+++ b/Engine/Plugins/Experimental/Mutable/Source/MutableTools/Private/MuT/Node.cpp
@@ -11,10 +11,22 @@
 #include "MuT/NodeExtensionData.h"
 #include "MuT/NodeExtensionDataConstant.h"
 #include "MuT/NodeImageTable.h"
+#include "MuT/NodeColour.h"
+#include "MuT/NodeColourSwitch.h"
+#include "MuT/NodeColourVariation.h"
+#include "MuT/NodeColourConstant.h"
+#include "MuT/NodeColourParameter.h"
+#include "MuT/NodeColourSampleImage.h"
+#include "MuT/NodeColourArithmeticOperation.h"
+#include "MuT/NodeColourFromScalars.h"
+#include "MuT/NodeColourTable.h"
 
 namespace mu
 {
 
+	// Static initialisation
+	FNodeType Node::StaticType = FNodeType(Node::EType::Node, nullptr);
+
 	FNodeType NodeComponent::StaticType = FNodeType(Node::EType::Component, Node::GetStaticType());
 	FNodeType NodeComponentNew::StaticType = FNodeType(Node::EType::ComponentNew, NodeComponent::GetStaticType());
 	FNodeType NodeLOD::StaticType = FNodeType(Node::EType::LOD, Node::GetStaticType());
@@ -22,33 +34,28 @@ namespace mu
 	FNodeType NodeExtensionDataConstant::StaticType = FNodeType(Node::EType::ExtensionDataConstant, NodeExtensionData::GetStaticType());
 	FNodeType NodeImageTable::StaticType = FNodeType(Node::EType::ImageTable, NodeImage::GetStaticType());
 
+	FNodeType NodeColour::StaticType = FNodeType(Node::EType::Color, Node::GetStaticType());
+	FNodeType NodeColourConstant::StaticType = FNodeType(Node::EType::ColorConstant, NodeColour::GetStaticType());
+	FNodeType NodeColourParameter::StaticType = FNodeType(Node::EType::ColorParameter, NodeColour::GetStaticType());
+	FNodeType NodeColourSwitch::StaticType = FNodeType(Node::EType::ColorSwitch, NodeColour::GetStaticType());
+	FNodeType NodeColourVariation::StaticType = FNodeType(Node::EType::ColorVariation, NodeColour::GetStaticType());
+	FNodeType NodeColourTable::StaticType = FNodeType(Node::EType::ColorTable, NodeColour::GetStaticType());
+	FNodeType NodeColourArithmeticOperation::StaticType = FNodeType(Node::EType::ColorArithmeticOperation, NodeColour::GetStaticType());
+	FNodeType NodeColourSampleImage::StaticType = FNodeType(Node::EType::ColorSampleImage, NodeColour::GetStaticType());
+	FNodeType NodeColourFromScalars::StaticType = FNodeType(Node::EType::ColorFromScalars, NodeColour::GetStaticType());
 
-	// Static initialisation
-	static FNodeType s_nodeType = FNodeType(Node::EType::Node, nullptr );
 
 	FNodeType::FNodeType()
 	{
 		Type = Node::EType::None;
-		m_pParent = nullptr;
+		Parent = nullptr;
 	}
 
 
 	FNodeType::FNodeType(Node::EType InType, const FNodeType* pParent )
 	{
 		Type = InType;
-		m_pParent = pParent;
-	}
-
-
-	const FNodeType* Node::GetType() const
-	{
-		return GetStaticType();
-	}
-
-
-	const FNodeType* Node::GetStaticType()
-	{
-		return &s_nodeType;
+		Parent = pParent;
 	}
 
 
diff --git a/Engine/Plugins/Experimental/Mutable/Source/MutableTools/Private/MuT/NodeColour.cpp b/Engine/Plugins/Experimental/Mutable/Source/MutableTools/Private/MuT/NodeColour.cpp
deleted file mode 100644
diff --git a/Engine/Plugins/Experimental/Mutable/Source/MutableTools/Private/MuT/NodeColourArithmeticOperation.cpp b/Engine/Plugins/Experimental/Mutable/Source/MutableTools/Private/MuT/NodeColourArithmeticOperation.cpp
deleted file mode 100644
diff --git a/Engine/Plugins/Experimental/Mutable/Source/MutableTools/Private/MuT/NodeColourArithmeticOperationPrivate.h b/Engine/Plugins/Experimental/Mutable/Source/MutableTools/Private/MuT/NodeColourArithmeticOperationPrivate.h
deleted file mode 100644
diff --git a/Engine/Plugins/Experimental/Mutable/Source/MutableTools/Private/MuT/NodeColourConstant.cpp b/Engine/Plugins/Experimental/Mutable/Source/MutableTools/Private/MuT/NodeColourConstant.cpp
deleted file mode 100644
diff --git a/Engine/Plugins/Experimental/Mutable/Source/MutableTools/Private/MuT/NodeColourConstantPrivate.h b/Engine/Plugins/Experimental/Mutable/Source/MutableTools/Private/MuT/NodeColourConstantPrivate.h
deleted file mode 100644
diff --git a/Engine/Plugins/Experimental/Mutable/Source/MutableTools/Private/MuT/NodeColourFromScalars.cpp b/Engine/Plugins/Experimental/Mutable/Source/MutableTools/Private/MuT/NodeColourFromScalars.cpp
deleted file mode 100644
diff --git a/Engine/Plugins/Experimental/Mutable/Source/MutableTools/Private/MuT/NodeColourParameter.cpp b/Engine/Plugins/Experimental/Mutable/Source/MutableTools/Private/MuT/NodeColourParameter.cpp
deleted file mode 100644
diff --git a/Engine/Plugins/Experimental/Mutable/Source/MutableTools/Private/MuT/NodeColourSampleImage.cpp b/Engine/Plugins/Experimental/Mutable/Source/MutableTools/Private/MuT/NodeColourSampleImage.cpp
deleted file mode 100644
diff --git a/Engine/Plugins/Experimental/Mutable/Source/MutableTools/Private/MuT/NodeColourSwitch.cpp b/Engine/Plugins/Experimental/Mutable/Source/MutableTools/Private/MuT/NodeColourSwitch.cpp
deleted file mode 100644
diff --git a/Engine/Plugins/Experimental/Mutable/Source/MutableTools/Private/MuT/NodeColourTable.cpp b/Engine/Plugins/Experimental/Mutable/Source/MutableTools/Private/MuT/NodeColourTable.cpp
deleted file mode 100644
diff --git a/Engine/Plugins/Experimental/Mutable/Source/MutableTools/Private/MuT/NodeColourVariation.cpp b/Engine/Plugins/Experimental/Mutable/Source/MutableTools/Private/MuT/NodeColourVariation.cpp
deleted file mode 100644
diff --git a/Engine/Plugins/Experimental/Mutable/Source/MutableTools/Private/MuT/NodeColourVariationPrivate.h b/Engine/Plugins/Experimental/Mutable/Source/MutableTools/Private/MuT/NodeColourVariationPrivate.h
deleted file mode 100644
diff --git a/Engine/Plugins/Experimental/Mutable/Source/MutableTools/Public/MuT/Node.h b/Engine/Plugins/Experimental/Mutable/Source/MutableTools/Public/MuT/Node.h
--- a/Engine/Plugins/Experimental/Mutable/Source/MutableTools/Public/MuT/Node.h
+++ b/Engine/Plugins/Experimental/Mutable/Source/MutableTools/Public/MuT/Node.h
@@ -148,13 +148,40 @@ namespace mu
 			Count
 		};
 
+
+		/** Information about the type of a node, to provide some means to the tools to deal generically with nodes. */
+		struct FNodeType
+		{
+			FNodeType();
+			FNodeType(Node::EType, const FNodeType* Parent);
+
+			Node::EType Type;
+			const FNodeType* Parent;
+
+			inline bool IsA(const FNodeType* CandidateType) const
+			{
+				if (CandidateType == this)
+				{
+					return true;
+				}
+
+				if (Parent)
+				{
+					return Parent->IsA(CandidateType);
+				}
+
+				return false;
+			}
+		};
+
+
 		//-----------------------------------------------------------------------------------------
 		// Own Interface
 		//-----------------------------------------------------------------------------------------
 
 		/** Node type hierarchy data. */
-        virtual const struct FNodeType* GetType() const;
-		static const struct FNodeType* GetStaticType();
+		virtual const FNodeType* GetType() const { return &StaticType; }
+		static const FNodeType* GetStaticType() { return &StaticType; }
 
 		/** Set the opaque context returned in messages in the compiler log. */
 		void SetMessageContext(const void* context);
@@ -164,43 +191,20 @@ namespace mu
         // Interface pattern
 		//-----------------------------------------------------------------------------------------
 		class Private;
-		//virtual Private* GetBasePrivate() const = 0;
 
 	protected:
 
 		inline ~Node() {}
 
-		//!
-		EType Type = EType::None;
-
 		/** This is an opaque context used to attach to reported error messages. */
 		const void* MessageContext = nullptr;
 
-	};
-
-	/** Information about the type of a node, to provide some means to the tools to deal generically with nodes. */
-	struct FNodeType
-	{
-		FNodeType();
-		FNodeType(Node::EType, const FNodeType* pParent);
+	private:
 
-		Node::EType Type;
-		const FNodeType* m_pParent;
+		static FNodeType StaticType;
 
-		inline bool IsA(const FNodeType* CandidateType) const
-		{
-			if (CandidateType == this)
-			{
-				return true;
-			}
-
-			if (m_pParent)
-			{
-				return m_pParent->IsA(CandidateType);
-			}
-
-			return false;
-		}
 	};
 
+
+	using FNodeType = Node::FNodeType;
 }
diff --git a/Engine/Plugins/Experimental/Mutable/Source/MutableTools/Public/MuT/NodeColour.h b/Engine/Plugins/Experimental/Mutable/Source/MutableTools/Public/MuT/NodeColour.h
--- a/Engine/Plugins/Experimental/Mutable/Source/MutableTools/Public/MuT/NodeColour.h
+++ b/Engine/Plugins/Experimental/Mutable/Source/MutableTools/Public/MuT/NodeColour.h
@@ -11,34 +11,25 @@
 namespace mu
 {
 
-	class NodeColour;
-	typedef Ptr<NodeColour> NodeColourPtr;
-	typedef Ptr<const NodeColour> NodeColourPtrConst;
-
-
-    //! %Base class of any node that outputs a colour.
-    //! \ingroup model
+    /** Base class of any node that outputs a colour.
+	*/
     class MUTABLETOOLS_API NodeColour : public Node
     {
     public:
 
-        //-----------------------------------------------------------------------------------------
         // Node Interface
-        //-----------------------------------------------------------------------------------------
-
-        const FNodeType* GetType() const override;
-        static const FNodeType* GetStaticType();
-
-
-        //-----------------------------------------------------------------------------------------
-        // Interface pattern
-        //-----------------------------------------------------------------------------------------
-        class Private;
+		virtual const FNodeType* GetType() const override { return &StaticType; }
+		static const FNodeType* GetStaticType() { return &StaticType; }
 
     protected:
 
         //! Forbidden. Manage with the Ptr<> template.
         inline ~NodeColour() {}
+
+	private:
+
+		static FNodeType StaticType;
+
     };
 
 }
diff --git a/Engine/Plugins/Experimental/Mutable/Source/MutableTools/Public/MuT/NodeColourArithmeticOperation.h b/Engine/Plugins/Experimental/Mutable/Source/MutableTools/Public/MuT/NodeColourArithmeticOperation.h
--- a/Engine/Plugins/Experimental/Mutable/Source/MutableTools/Public/MuT/NodeColourArithmeticOperation.h
+++ b/Engine/Plugins/Experimental/Mutable/Source/MutableTools/Public/MuT/NodeColourArithmeticOperation.h
@@ -4,79 +4,45 @@
 
 #include "MuR/Ptr.h"
 #include "MuR/RefCounted.h"
-#include "MuT/Node.h"
 #include "MuT/NodeColour.h"
 
 
 namespace mu
 {
 
-	// Forward definitions
-	class NodeColourArithmeticOperation;
-	typedef Ptr<NodeColourArithmeticOperation> NodeColourArithmeticOperationPtr;
-	typedef Ptr<const NodeColourArithmeticOperation> NodeColourArithmeticOperationPtrConst;
-
-
-    //! Perform a per-component arithmetic operation between two colours.
-	//! \ingroup model
+    /** Perform a per - component arithmetic operation between two colours. 
+	*/
 	class MUTABLETOOLS_API NodeColourArithmeticOperation : public NodeColour
 	{
 	public:
 
-		NodeColourArithmeticOperation();
-
-		//-----------------------------------------------------------------------------------------
-		// Node Interface
-		//-----------------------------------------------------------------------------------------
-
-		virtual const FNodeType* GetType() const override;
-		static const FNodeType* GetStaticType();
-
-		//-----------------------------------------------------------------------------------------
-		// Own Interface
-		//-----------------------------------------------------------------------------------------
-
-		//! Possible arithmetic operations
-		typedef enum
+		/** Possible arithmetic operations. */
+		enum class EOperation
 		{
-			AO_ADD,
-			AO_SUBTRACT,
-			AO_MULTIPLY,
-            AO_DIVIDE,
-            _AO_COUNT
-		} OPERATION;
+			Add,
+			Subtract,
+			Multiply,
+			Divide
+		};
 
-        static const char* s_opTypeName[ _AO_COUNT ];
+		EOperation Operation;
+		Ptr<NodeColour> A;
+		Ptr<NodeColour> B;
 
-		//! Get child selection type
-		OPERATION GetOperation() const;
-
-		//! Set the child selection type
-		void SetOperation(OPERATION);
-
-
-		//! Get the first operand
-		NodeColourPtr GetA() const;
-		void SetA(NodeColourPtr);
-
-		// Get the second operand
-		NodeColourPtr GetB() const;
-		void SetB(NodeColourPtr);
+	public:
 
-		//-----------------------------------------------------------------------------------------
-		// Interface pattern
-		//-----------------------------------------------------------------------------------------
-		class Private;
-		Private* GetPrivate() const;
+		/** Node type hierarchy data. */
+		virtual const FNodeType* GetType() const { return &StaticType; }
+		static const FNodeType* GetStaticType() { return &StaticType; }
 
 	protected:
 
 		//! Forbidden. Manage with the Ptr<> template.
-		~NodeColourArithmeticOperation();
+		~NodeColourArithmeticOperation() {}
 
 	private:
 
-		Private* m_pD;
+		static FNodeType StaticType;
 
 	};
 
diff --git a/Engine/Plugins/Experimental/Mutable/Source/MutableTools/Public/MuT/NodeColourConstant.h b/Engine/Plugins/Experimental/Mutable/Source/MutableTools/Public/MuT/NodeColourConstant.h
--- a/Engine/Plugins/Experimental/Mutable/Source/MutableTools/Public/MuT/NodeColourConstant.h
+++ b/Engine/Plugins/Experimental/Mutable/Source/MutableTools/Public/MuT/NodeColourConstant.h
@@ -6,55 +6,34 @@
 #include "MuR/RefCounted.h"
 #include "MuT/Node.h"
 #include "MuT/NodeColour.h"
-
+#include "Math/Vector4.h"
+#include "Math/MathFwd.h"
 
 namespace mu
 {
 
-	class NodeColourConstant;
-	typedef Ptr<NodeColourConstant> NodeColourConstantPtr;
-	typedef Ptr<const NodeColourConstant> NodeColourConstantPtrConst;
-
-
-	//! This node outputs a predefined colour value.
-	//! \ingroup model
+	/** This node outputs a predefined colour value.
+	*/
 	class MUTABLETOOLS_API NodeColourConstant : public NodeColour
 	{
 	public:
 
-		NodeColourConstant();
-
-		//-----------------------------------------------------------------------------------------
-		// Node Interface
-		//-----------------------------------------------------------------------------------------
+		FVector4f Value;
 
-        const FNodeType* GetType() const override;
-		static const FNodeType* GetStaticType();
-
-		//-----------------------------------------------------------------------------------------
-		// Own Interface
-		//-----------------------------------------------------------------------------------------
-
-		//! Get the value that this node returns
-		FVector4f GetValue() const;
-
-		//! Set the value to be returned by this node
-		void SetValue(FVector4f);
+	public:
 
-		//-----------------------------------------------------------------------------------------
-		// Interface pattern
-		//-----------------------------------------------------------------------------------------
-		class Private;
-		Private* GetPrivate() const;
+		// Node interface
+		virtual const FNodeType* GetType() const { return &StaticType; }
+		static const FNodeType* GetStaticType() { return &StaticType; }
 
 	protected:
 
 		//! Forbidden. Manage with the Ptr<> template.
-		~NodeColourConstant();
+		~NodeColourConstant() {}
 
 	private:
 
-		Private* m_pD;
+		static FNodeType StaticType;
 
 	};
 
diff --git a/Engine/Plugins/Experimental/Mutable/Source/MutableTools/Public/MuT/NodeColourFromScalars.h b/Engine/Plugins/Experimental/Mutable/Source/MutableTools/Public/MuT/NodeColourFromScalars.h
--- a/Engine/Plugins/Experimental/Mutable/Source/MutableTools/Public/MuT/NodeColourFromScalars.h
+++ b/Engine/Plugins/Experimental/Mutable/Source/MutableTools/Public/MuT/NodeColourFromScalars.h
@@ -6,68 +6,37 @@
 #include "MuR/RefCounted.h"
 #include "MuT/Node.h"
 #include "MuT/NodeColour.h"
+#include "MuT/NodeScalar.h"
 
 
 namespace mu
 {
 
-	// Forward definitions
-	class NodeScalar;
-	typedef Ptr<NodeScalar> NodeScalarPtr;
-	typedef Ptr<const NodeScalar> NodeScalarPtrConst;
-
-	class NodeColourFromScalars;
-	typedef Ptr<NodeColourFromScalars> NodeColourFromScalarsPtr;
-	typedef Ptr<const NodeColourFromScalars> NodeColourFromScalarsPtrConst;
-
-
-	//! Obtain a colour by sampling an image at specific homogeneous coordinates.
-	//! \ingroup model
+	/** Obtain a colour by sampling an image at specific homogeneous coordinates.
+	*/
 	class MUTABLETOOLS_API NodeColourFromScalars : public NodeColour
 	{
 	public:
 
-		NodeColourFromScalars();
-
-		//-----------------------------------------------------------------------------------------
-		// Node Interface
-		//-----------------------------------------------------------------------------------------
-
-		const FNodeType* GetType() const override;
-		static const FNodeType* GetStaticType();
+		Ptr<NodeScalar> X;
+		Ptr<NodeScalar> Y;
+		Ptr<NodeScalar> Z;
+		Ptr<NodeScalar> W;
 
-		//-----------------------------------------------------------------------------------------
-		// Own Interface
-		//-----------------------------------------------------------------------------------------
-
-		//! Get the node generating the x coordinate to sample, in the range of 0.0 to 1.0.
-		//! If it is not specified, a value of 1 will be used.
-		NodeScalarPtr GetX() const;
-		void SetX( NodeScalarPtr );
-
-		NodeScalarPtr GetY() const;
-		void SetY( NodeScalarPtr );
-
-		NodeScalarPtr GetZ() const;
-		void SetZ( NodeScalarPtr );
-
-		NodeScalarPtr GetW() const;
-		void SetW( NodeScalarPtr );
+	public:
 
-		//-----------------------------------------------------------------------------------------
-		// Interface pattern
-		//-----------------------------------------------------------------------------------------
-		class Private;
-		Private* GetPrivate() const;
+		// Node interface
+		virtual const FNodeType* GetType() const { return &StaticType; }
+		static const FNodeType* GetStaticType() { return &StaticType; }
 
 	protected:
 
 		//! Forbidden. Manage with the Ptr<> template.
-		~NodeColourFromScalars();
+		~NodeColourFromScalars() {}
 
 	private:
 
-		Private* m_pD;
+		static FNodeType StaticType;
 
 	};
 
diff --git a/Engine/Plugins/Experimental/Mutable/Source/MutableTools/Public/MuT/NodeColourParameter.h b/Engine/Plugins/Experimental/Mutable/Source/MutableTools/Public/MuT/NodeColourParameter.h
--- a/Engine/Plugins/Experimental/Mutable/Source/MutableTools/Public/MuT/NodeColourParameter.h
+++ b/Engine/Plugins/Experimental/Mutable/Source/MutableTools/Public/MuT/NodeColourParameter.h
@@ -6,65 +6,39 @@
 #include "MuR/RefCounted.h"
 #include "MuT/Node.h"
 #include "MuT/NodeColour.h"
+#include "MuT/NodeRange.h"
+#include "Math/Vector4.h"
+#include "Math/MathFwd.h"
 
 namespace mu
 {
 
-	class NodeColourParameter;
-    using NodeColourParameterPtr = Ptr<NodeColourParameter>;
-    using NodeColourParameterPtrConst = Ptr<const NodeColourParameter>;
-
-    class NodeRange;
-    using NodeRangePtr = Ptr<NodeRange>;
-    using NodeRangePtrConst = Ptr<const NodeRange>;
-
-
-	//! Node that defines a colour model parameter.
-	//! \ingroup model
+	/** Node that defines a colour model parameter.
+	*/
 	class MUTABLETOOLS_API NodeColourParameter : public NodeColour
 	{
 	public:
 
-		NodeColourParameter();
-
-		//-----------------------------------------------------------------------------------------
-		// Node Interface
-		//-----------------------------------------------------------------------------------------
-		const FNodeType* GetType() const override;
-		static const FNodeType* GetStaticType();
+		FVector4f DefaultValue;
+		FString Name;
+		FString Uid;
 
-		//-----------------------------------------------------------------------------------------
-		// Own Interface
-		//-----------------------------------------------------------------------------------------
+		TArray<Ptr<NodeRange>> Ranges;
 
-		//! Set the name of the parameter.
-		void SetName( const FString& );
-
-		//! Set the uid of the parameter.
-		void SetUid( const FString& );
-
-		//! Set the default value of the parameter.
-		void SetDefaultValue( FVector4f );
-
-        //! Set the number of ranges (dimensions) for this parameter.
-        //! By default a parameter has 0 ranges, meaning it only has one value.
-        void SetRangeCount( int i );
-        void SetRange( int i, NodeRangePtr pRange );
+	public:
 
-		//-----------------------------------------------------------------------------------------
-		// Interface pattern
-		//-----------------------------------------------------------------------------------------
-		class Private;
-		Private* GetPrivate() const;
+		// Node interface
+		virtual const FNodeType* GetType() const { return &StaticType; }
+		static const FNodeType* GetStaticType() { return &StaticType; }
 
 	protected:
 
 		//! Forbidden. Manage with the Ptr<> template.
-		~NodeColourParameter();
+		~NodeColourParameter() {}
 
 	private:
 
-		Private* m_pD;
+		static FNodeType StaticType;
 
 	};
 
diff --git a/Engine/Plugins/Experimental/Mutable/Source/MutableTools/Public/MuT/NodeColourSampleImage.h b/Engine/Plugins/Experimental/Mutable/Source/MutableTools/Public/MuT/NodeColourSampleImage.h
--- a/Engine/Plugins/Experimental/Mutable/Source/MutableTools/Public/MuT/NodeColourSampleImage.h
+++ b/Engine/Plugins/Experimental/Mutable/Source/MutableTools/Public/MuT/NodeColourSampleImage.h
@@ -6,73 +6,37 @@
 #include "MuR/RefCounted.h"
 #include "MuT/Node.h"
 #include "MuT/NodeColour.h"
+#include "MuT/NodeScalar.h"
+#include "MuT/NodeImage.h"
 
 
 namespace mu
 {
 
-	// Forward definitions
-	class NodeScalar;
-	typedef Ptr<NodeScalar> NodeScalarPtr;
-	typedef Ptr<const NodeScalar> NodeScalarPtrConst;
-
-	class NodeImage;
-	typedef Ptr<NodeImage> NodeImagePtr;
-	typedef Ptr<const NodeImage> NodeImagePtrConst;
-
-	class NodeColourSampleImage;
-	typedef Ptr<NodeColourSampleImage> NodeColourSampleImagePtr;
-	typedef Ptr<const NodeColourSampleImage> NodeColourSampleImagePtrConst;
-
-
-	//! Obtain a colour by sampling an image at specific homogeneous coordinates.
-	//! \ingroup model
+	/** Obtain a colour by sampling an image at specific homogeneous coordinates.
+	*/
 	class MUTABLETOOLS_API NodeColourSampleImage : public NodeColour
 	{
 	public:
 
-		NodeColourSampleImage();
-
-
-		//-----------------------------------------------------------------------------------------
-		// Node Interface
-		//-----------------------------------------------------------------------------------------
+		Ptr<NodeImage> Image;
+		Ptr<NodeScalar> X;
+		Ptr<NodeScalar> Y;
 
-		const FNodeType* GetType() const override;
-		static const FNodeType* GetStaticType();
-
-		//-----------------------------------------------------------------------------------------
-		// Own Interface
-		//-----------------------------------------------------------------------------------------
-
-		//! Get the node generating the x coordinate to sample, in the range of 0.0 to 1.0.
-		//! If it is not specified, a value of x =0.5 will be used.
-		NodeScalarPtr GetX() const;
-		void SetX( NodeScalarPtr );
-
-		//! Get the node generating the y coordinate to sample, in the range of 0.0 to 1.0.
-		//! If it is not specified, a value of y =0.5 will be used.
-		NodeScalarPtr GetY() const;
-		void SetY( NodeScalarPtr );
-
-		//! Get the node generating the source image to be sampled.
-		NodeImagePtr GetImage() const;
-		void SetImage( NodeImagePtr );
+	public:
 
-		//-----------------------------------------------------------------------------------------
-		// Interface pattern
-		//-----------------------------------------------------------------------------------------
-		class Private;
-		Private* GetPrivate() const;
+		// Node interface
+		virtual const FNodeType* GetType() const { return &StaticType; }
+		static const FNodeType* GetStaticType() { return &StaticType; }
 
 	protected:
 
 		//! Forbidden. Manage with the Ptr<> template.
-		~NodeColourSampleImage();
+		~NodeColourSampleImage() {}
 
 	private:
 
-		Private* m_pD;
+		static FNodeType StaticType;
 
 	};
 
diff --git a/Engine/Plugins/Experimental/Mutable/Source/MutableTools/Public/MuT/NodeColourSwitch.h b/Engine/Plugins/Experimental/Mutable/Source/MutableTools/Public/MuT/NodeColourSwitch.h
--- a/Engine/Plugins/Experimental/Mutable/Source/MutableTools/Public/MuT/NodeColourSwitch.h
+++ b/Engine/Plugins/Experimental/Mutable/Source/MutableTools/Public/MuT/NodeColourSwitch.h
@@ -6,68 +6,35 @@
 #include "MuR/RefCounted.h"
 #include "MuT/Node.h"
 #include "MuT/NodeColour.h"
+#include "MuT/NodeScalar.h"
 
 
 namespace mu
 {
 
-	// Forward definitions
-	class NodeScalar;
-	typedef Ptr<NodeScalar> NodeScalarPtr;
-	typedef Ptr<const NodeScalar> NodeScalarPtrConst;
-
-    class NodeColourSwitch;
-    typedef Ptr<NodeColourSwitch> NodeColourSwitchPtr;
-    typedef Ptr<const NodeColourSwitch> NodeColourSwitchPtrConst;
-
-
-    //! This node selects an output Colour from a set of input Colours based on a parameter.
-	//! \ingroup model
+    /** This node selects an output Colour from a set of input Colours based on a parameter. 
+	*/
     class MUTABLETOOLS_API NodeColourSwitch : public NodeColour
 	{
 	public:
 
-        NodeColourSwitch();
-
-
-		//-----------------------------------------------------------------------------------------
-		// Node Interface
-		//-----------------------------------------------------------------------------------------
-
-        const FNodeType* GetType() const override;
-		static const FNodeType* GetStaticType();
-
-		//-----------------------------------------------------------------------------------------
-		// Own Interface
-		//-----------------------------------------------------------------------------------------
-
-		//! Get the node generating the parameter used to select the option.
-		NodeScalarPtr GetParameter() const;
-		void SetParameter( NodeScalarPtr );
-
-        //! Set the number of option Colours. It will keep the currently set targets and initialise
-		//! the new ones as null.
-		void SetOptionCount( int );
-
-        //! Get the node generating the t-th option Colour.
-        NodeColourPtr GetOption( int t ) const;
-        void SetOption( int t, NodeColourPtr );
+		Ptr<NodeScalar> Parameter;
+		TArray<Ptr<NodeColour>> Options;
 
+	public:
 
-		//-----------------------------------------------------------------------------------------
-		// Interface pattern
-		//-----------------------------------------------------------------------------------------
-		class Private;
-		Private* GetPrivate() const;
+		/** Node type hierarchy data. */
+		virtual const FNodeType* GetType() const { return &StaticType; }
+		static const FNodeType* GetStaticType() { return &StaticType; }
 
 	protected:
 
 		//! Forbidden. Manage with the Ptr<> template.
-        ~NodeColourSwitch();
+		~NodeColourSwitch() {}
 
 	private:
 
-		Private* m_pD;
+		static FNodeType StaticType;
 
 	};
 
diff --git a/Engine/Plugins/Experimental/Mutable/Source/MutableTools/Public/MuT/NodeColourTable.h b/Engine/Plugins/Experimental/Mutable/Source/MutableTools/Public/MuT/NodeColourTable.h
--- a/Engine/Plugins/Experimental/Mutable/Source/MutableTools/Public/MuT/NodeColourTable.h
+++ b/Engine/Plugins/Experimental/Mutable/Source/MutableTools/Public/MuT/NodeColourTable.h
@@ -12,14 +12,8 @@
 namespace mu
 {
 
-	// Forward definitions
-	class NodeColourTable;
-	typedef Ptr<NodeColourTable> NodeColourTablePtr;
-	typedef Ptr<const NodeColourTable> NodeColourTablePtrConst;
-
-
-	//! This node provides the meshes stored in the column of a table.
-	//! \ingroup transform
+	/** This node provides the meshes stored in the column of a table.
+	*/
 	class MUTABLETOOLS_API NodeColourTable : public NodeColour
 	{
 	public:
@@ -36,23 +30,6 @@ namespace mu
 		virtual const FNodeType* GetType() const override { return GetStaticType(); }
 		static const FNodeType* GetStaticType() { return &StaticType; }
 
-		//-----------------------------------------------------------------------------------------
-		// Own Interface
-		//-----------------------------------------------------------------------------------------
-
-		//! Set the name of the implicit table parameter.
-		void SetParameterName( const FString& strName );
-
-		//!
-		void SetColumn( const FString& strName );
-
-		//!
-		void SetNoneOption(bool bAddNoneOption);
-
-		//! Set the row name to be used as default value
-		void SetDefaultRowName(const FString RowName);
-
-
 	protected:
 
 		//! Forbidden. Manage with the Ptr<> template.
diff --git a/Engine/Plugins/Experimental/Mutable/Source/MutableTools/Public/MuT/NodeColourVariation.h b/Engine/Plugins/Experimental/Mutable/Source/MutableTools/Public/MuT/NodeColourVariation.h
--- a/Engine/Plugins/Experimental/Mutable/Source/MutableTools/Public/MuT/NodeColourVariation.h
+++ b/Engine/Plugins/Experimental/Mutable/Source/MutableTools/Public/MuT/NodeColourVariation.h
@@ -11,62 +11,37 @@
 namespace mu
 {
 
-    // Forward definitions
-    class NodeColourVariation;
-    typedef Ptr<NodeColourVariation> NodeColourVariationPtr;
-    typedef Ptr<const NodeColourVariation> NodeColourVariationPtrConst;
-
-
-    //!
-    //! \ingroup model
+	/** Select different color subgraphs based on active tags. */
     class MUTABLETOOLS_API NodeColourVariation : public NodeColour
     {
-    public:
-
-        NodeColourVariation();
-
-        //-----------------------------------------------------------------------------------------
-        // Node Interface
-        //-----------------------------------------------------------------------------------------
-
-        const FNodeType* GetType() const override;
-        static const FNodeType* GetStaticType();
-
-        //-----------------------------------------------------------------------------------------
-        // Own Interface
-        //-----------------------------------------------------------------------------------------
+	public:
 
-        //!
-        void SetDefaultColour( NodeColour* Colour );
+		Ptr<NodeColour> DefaultColour;
 
-        //! Set the number of tags to consider in this variation
-        void SetVariationCount( int32 count );
+		struct FVariation
+		{
+			Ptr<NodeColour> Colour;
+			FString Tag;
+		};
 
-        //!
-        int GetVariationCount() const;
-
-        //! Set the tag or state name that will enable a specific vartiation
-        void SetVariationTag( int32 index, const FString& Tag );
-
-        //!
-        void SetVariationColour( int32 index, NodeColour* Colour );
-
-        //!}
+		TArray<FVariation> Variations;
 
+    public:
 
-        //-----------------------------------------------------------------------------------------
-        // Interface pattern
-        //-----------------------------------------------------------------------------------------
-        class Private;
-        Private* GetPrivate() const;
+		// Node Interface
+		virtual const FNodeType* GetType() const override { return &StaticType; }
+		static const FNodeType* GetStaticType() { return &StaticType; }
 
     protected:
+
         //! Forbidden. Manage with the Ptr<> template.
-        ~NodeColourVariation();
+		~NodeColourVariation() {}
+
+	private:
 
-    private:
-        Private* m_pD;
-    };
+		static FNodeType StaticType;
+	
+	};
 
 
 } // namespace mu
diff --git a/Engine/Plugins/Experimental/Mutable/Source/MutableTools/Public/MuT/NodeMeshClipMorphPlane.h b/Engine/Plugins/Experimental/Mutable/Source/MutableTools/Public/MuT/NodeMeshClipMorphPlane.h
--- a/Engine/Plugins/Experimental/Mutable/Source/MutableTools/Public/MuT/NodeMeshClipMorphPlane.h
+++ b/Engine/Plugins/Experimental/Mutable/Source/MutableTools/Public/MuT/NodeMeshClipMorphPlane.h
@@ -61,7 +61,6 @@ namespace mu
 		//-----------------------------------------------------------------------------------------
 		class Private;
 		Private* GetPrivate() const;
-		Node::Private* GetBasePrivate() const;
 
 	protected:
 
diff --git a/Engine/Plugins/Experimental/Mutable/Source/MutableTools/Public/MuT/NodeMeshClipWithMesh.h b/Engine/Plugins/Experimental/Mutable/Source/MutableTools/Public/MuT/NodeMeshClipWithMesh.h
--- a/Engine/Plugins/Experimental/Mutable/Source/MutableTools/Public/MuT/NodeMeshClipWithMesh.h
+++ b/Engine/Plugins/Experimental/Mutable/Source/MutableTools/Public/MuT/NodeMeshClipWithMesh.h
@@ -50,7 +50,6 @@ namespace mu
 		//-----------------------------------------------------------------------------------------
 		class Private;
 		Private* GetPrivate() const;
-		Node::Private* GetBasePrivate() const;
 
 	protected:
 
