[mutable] Fix mesh merge bug when the second mesh has additional texture channels.
- More defensive layout merging.
#rnx

diff --git a/Engine/Plugins/Experimental/Mutable/Source/MutableRuntime/Internal/MuR/OpMeshMerge.h b/Engine/Plugins/Experimental/Mutable/Source/MutableRuntime/Internal/MuR/OpMeshMerge.h
--- a/Engine/Plugins/Experimental/Mutable/Source/MutableRuntime/Internal/MuR/OpMeshMerge.h
+++ b/Engine/Plugins/Experimental/Mutable/Source/MutableRuntime/Internal/MuR/OpMeshMerge.h
@@ -332,16 +332,29 @@ namespace mu
 		{
 			MUTABLE_CPUPROFILER_SCOPE(Layouts);
 
-			Result->Layouts.SetNum(pFirst->Layouts.Num());
-			for (int32 LayoutIndex = 0; LayoutIndex < pFirst->Layouts.Num(); ++LayoutIndex)
+			int32 ResultLayoutCount = FMath::Max(pFirst->Layouts.Num(),pSecond->Layouts.Num());
+			Result->Layouts.SetNum(ResultLayoutCount);
+			for (int32 LayoutIndex = 0; LayoutIndex < ResultLayoutCount; ++LayoutIndex)
 			{
-				const Layout* pF = pFirst->Layouts[LayoutIndex].get();
-				Ptr<Layout> pR = pF->Clone();
+				Ptr<Layout> pR;
+				
+				if (LayoutIndex < pFirst->Layouts.Num())
+				{
+					const Layout* pF = pFirst->Layouts[LayoutIndex].get();
+					pR = pF->Clone();
+				}
 
 				if (LayoutIndex < pSecond->Layouts.Num())
 				{
 					const Layout* pS = pSecond->Layouts[LayoutIndex].get();
-					pR->Blocks.Append(pS->Blocks);
+					if (!pR)
+					{
+						pR = pS->Clone();
+					}
+					else
+					{
+						pR->Blocks.Append(pS->Blocks);
+					}
 				}
 
 				Result->Layouts[LayoutIndex] = pR;
@@ -825,10 +838,6 @@ namespace mu
 				{
 					// Expand component counts in vertex channels of the format mesh
 					FMeshBuffer& result = Result->GetVertexBuffers().m_buffers[BufferIndex];
-					const FMeshBuffer& first = pFirst->GetVertexBuffers().m_buffers[BufferIndex];
-
-					result.m_channels = first.m_channels;
-					result.m_elementSize = first.m_elementSize;
 
 					bool bResetOffsets = false;
 					for (int32 c = 0; c < result.m_channels.Num(); ++c)
