[Mutable] Typed Passthrough Pins in Node Material.

Due to inconsistencies with the UX/UI expirience, now Texture Parameter Pins will not change automatically.

- Fixed issues with pin icons.

#rnx

diff --git a/Engine/Plugins/Experimental/Mutable/Source/CustomizableObject/Internal/MuCO/CustomizableObjectCustomVersion.h b/Engine/Plugins/Experimental/Mutable/Source/CustomizableObject/Internal/MuCO/CustomizableObjectCustomVersion.h
--- a/Engine/Plugins/Experimental/Mutable/Source/CustomizableObject/Internal/MuCO/CustomizableObjectCustomVersion.h
+++ b/Engine/Plugins/Experimental/Mutable/Source/CustomizableObject/Internal/MuCO/CustomizableObjectCustomVersion.h
@@ -124,6 +124,8 @@ struct CUSTOMIZABLEOBJECT_API FCustomizableObjectCustomVersion
 
 		NewComponentOptions,
 
+		NodeMaterialTypedImagePins,
+
 		// -----<new versions can be added above this line>-------------------------------------------------
 		VersionPlusOne,
 		LatestVersion = VersionPlusOne - 1
diff --git a/Engine/Plugins/Experimental/Mutable/Source/CustomizableObjectEditor/Private/MuCOE/EdGraphSchema_CustomizableObject.cpp b/Engine/Plugins/Experimental/Mutable/Source/CustomizableObjectEditor/Private/MuCOE/EdGraphSchema_CustomizableObject.cpp
--- a/Engine/Plugins/Experimental/Mutable/Source/CustomizableObjectEditor/Private/MuCOE/EdGraphSchema_CustomizableObject.cpp
+++ b/Engine/Plugins/Experimental/Mutable/Source/CustomizableObjectEditor/Private/MuCOE/EdGraphSchema_CustomizableObject.cpp
@@ -1400,24 +1400,4 @@ FText UEdGraphSchema_CustomizableObject::GetPinCategoryName(const FName& PinCate
 }
 
 
-const FName& UEdGraphSchema_CustomizableObject::GetPinCategory(const EMaterialParameterType Type)
-{
-	switch(Type)
-	{
-	case EMaterialParameterType::Texture:
-		return UEdGraphSchema_CustomizableObject::PC_Image;
-
-	case EMaterialParameterType::Vector:
-		return UEdGraphSchema_CustomizableObject::PC_Color;
-
-	case EMaterialParameterType::Scalar:
-		return UEdGraphSchema_CustomizableObject::PC_Float;
-
-	default:
-		check(false); // Type not contemplated.
-		return PC_Image; // Fake result.
-	}
-}
-
-
 #undef LOCTEXT_NAMESPACE
diff --git a/Engine/Plugins/Experimental/Mutable/Source/CustomizableObjectEditor/Private/MuCOE/Nodes/CustomizableObjectNode.cpp b/Engine/Plugins/Experimental/Mutable/Source/CustomizableObjectEditor/Private/MuCOE/Nodes/CustomizableObjectNode.cpp
--- a/Engine/Plugins/Experimental/Mutable/Source/CustomizableObjectEditor/Private/MuCOE/Nodes/CustomizableObjectNode.cpp
+++ b/Engine/Plugins/Experimental/Mutable/Source/CustomizableObjectEditor/Private/MuCOE/Nodes/CustomizableObjectNode.cpp
@@ -2,6 +2,7 @@
 
 #include "MuCOE/Nodes/CustomizableObjectNode.h"
 
+#include "SCustomizableObjectNode.h"
 #include "Containers/Queue.h"
 #include "MuCO/CustomizableObjectCustomVersion.h"
 #include "MuCOE/CustomizableObjectEditor.h"
@@ -309,6 +310,12 @@ void UCustomizableObjectNode::DestroyNode()
 }
 
 
+TSharedPtr<SGraphNode> UCustomizableObjectNode::CreateVisualWidget()
+{
+	return SNew(SCustomizableObjectNode, this);
+}
+
+
 void UCustomizableObjectNode::AutowireNewNode(UEdGraphPin* FromPin)
 {
 	if (!FromPin)
diff --git a/Engine/Plugins/Experimental/Mutable/Source/CustomizableObjectEditor/Private/MuCOE/Nodes/CustomizableObjectNodeColorConstant.cpp b/Engine/Plugins/Experimental/Mutable/Source/CustomizableObjectEditor/Private/MuCOE/Nodes/CustomizableObjectNodeColorConstant.cpp
--- a/Engine/Plugins/Experimental/Mutable/Source/CustomizableObjectEditor/Private/MuCOE/Nodes/CustomizableObjectNodeColorConstant.cpp
+++ b/Engine/Plugins/Experimental/Mutable/Source/CustomizableObjectEditor/Private/MuCOE/Nodes/CustomizableObjectNodeColorConstant.cpp
@@ -75,15 +75,11 @@ TSharedPtr<SGraphNode> UCustomizableObjectNodeColorConstant::CreateVisualWidget(
 
 void SGraphNodeColorConstant::Construct(const FArguments& InArgs, UEdGraphNode* InGraphNode)
 {
-	GraphNode = InGraphNode;
 	NodeColorConstant = Cast< UCustomizableObjectNodeColorConstant >(InGraphNode);
 
 	WidgetStyle = FCoreStyle::Get().GetWidgetStyle<FSpinBoxStyle>("SpinBox");
 
-	if (NodeColorConstant)
-	{
-		UpdateGraphNode();
-	}
+	SCustomizableObjectNode::Construct({}, InGraphNode);
 }
 
 void SGraphNodeColorConstant::UpdateGraphNode()
diff --git a/Engine/Plugins/Experimental/Mutable/Source/CustomizableObjectEditor/Private/MuCOE/Nodes/CustomizableObjectNodeCopyMaterial.cpp b/Engine/Plugins/Experimental/Mutable/Source/CustomizableObjectEditor/Private/MuCOE/Nodes/CustomizableObjectNodeCopyMaterial.cpp
--- a/Engine/Plugins/Experimental/Mutable/Source/CustomizableObjectEditor/Private/MuCOE/Nodes/CustomizableObjectNodeCopyMaterial.cpp
+++ b/Engine/Plugins/Experimental/Mutable/Source/CustomizableObjectEditor/Private/MuCOE/Nodes/CustomizableObjectNodeCopyMaterial.cpp
@@ -332,6 +332,19 @@ const UEdGraphPin* UCustomizableObjectNodeCopyMaterial::GetParameterPin(EMateria
 }
 
 
+UEdGraphPin* UCustomizableObjectNodeCopyMaterial::GetParameterPin(const FGuid& ParameterId) const
+{
+	if (UCustomizableObjectNodeMaterial* NodeMaterial = GetMaterialNode())
+	{
+		return NodeMaterial->GetParameterPin(ParameterId);
+	}
+	else
+	{
+		return nullptr;	
+	}
+}
+
+
 bool UCustomizableObjectNodeCopyMaterial::IsImageMutableMode(int32 ImageIndex) const
 {
 	UCustomizableObjectNodeMaterial* NodeMaterial = GetMaterialNode();
@@ -350,33 +363,6 @@ bool UCustomizableObjectNodeCopyMaterial::IsImageMutableMode(const UEdGraphPin&
 }
 
 
-void UCustomizableObjectNodeCopyMaterial::UpdateImagePinMode(const FGuid ParameterId)
-{
-	if (UCustomizableObjectNodeMaterial* NodeMaterial = GetMaterialNode())
-	{
-		NodeMaterial->UpdateImagePinMode(ParameterId);
-	}
-}
-
-
-void UCustomizableObjectNodeCopyMaterial::UpdateImagePinMode(const UEdGraphPin& Pin)
-{
-	if (UCustomizableObjectNodeMaterial* NodeMaterial = GetMaterialNode())
-	{
-		NodeMaterial->UpdateImagePinMode(Pin);
-	}
-}
-
-
-void UCustomizableObjectNodeCopyMaterial::UpdateAllImagesPinMode()
-{
-	if (UCustomizableObjectNodeMaterial* NodeMaterial = GetMaterialNode())
-	{
-		NodeMaterial->UpdateAllImagesPinMode();
-	}
-}
-
-
 UTexture2D* UCustomizableObjectNodeCopyMaterial::GetImageReferenceTexture(int32 ImageIndex) const
 {
 	if (UCustomizableObjectNodeMaterial* NodeMaterial = GetMaterialNode())
diff --git a/Engine/Plugins/Experimental/Mutable/Source/CustomizableObjectEditor/Private/MuCOE/Nodes/CustomizableObjectNodeFloatConstant.cpp b/Engine/Plugins/Experimental/Mutable/Source/CustomizableObjectEditor/Private/MuCOE/Nodes/CustomizableObjectNodeFloatConstant.cpp
--- a/Engine/Plugins/Experimental/Mutable/Source/CustomizableObjectEditor/Private/MuCOE/Nodes/CustomizableObjectNodeFloatConstant.cpp
+++ b/Engine/Plugins/Experimental/Mutable/Source/CustomizableObjectEditor/Private/MuCOE/Nodes/CustomizableObjectNodeFloatConstant.cpp
@@ -73,15 +73,11 @@ TSharedPtr<SGraphNode> UCustomizableObjectNodeFloatConstant::CreateVisualWidget(
 
 void SGraphNodeFloatConstant::Construct(const FArguments& InArgs, UEdGraphNode* InGraphNode)
 {
-	GraphNode = InGraphNode;
 	NodeFloatConstant = Cast< UCustomizableObjectNodeFloatConstant >(InGraphNode);
 
 	WidgetStyle = FCoreStyle::Get().GetWidgetStyle<FSpinBoxStyle>("SpinBox");
 
-	if (NodeFloatConstant)
-	{
-		UpdateGraphNode();
-	}
+	SCustomizableObjectNode::Construct({}, InGraphNode);
 }
 
 
diff --git a/Engine/Plugins/Experimental/Mutable/Source/CustomizableObjectEditor/Private/MuCOE/Nodes/CustomizableObjectNodeMaterial.cpp b/Engine/Plugins/Experimental/Mutable/Source/CustomizableObjectEditor/Private/MuCOE/Nodes/CustomizableObjectNodeMaterial.cpp
--- a/Engine/Plugins/Experimental/Mutable/Source/CustomizableObjectEditor/Private/MuCOE/Nodes/CustomizableObjectNodeMaterial.cpp
+++ b/Engine/Plugins/Experimental/Mutable/Source/CustomizableObjectEditor/Private/MuCOE/Nodes/CustomizableObjectNodeMaterial.cpp
@@ -45,7 +45,7 @@ bool UCustomizableObjectNodeMaterialRemapPinsByName::Equal(const UCustomizableOb
 	const UCustomizableObjectNodeMaterialPinDataParameter* PinDataNewPin = Cast<UCustomizableObjectNodeMaterialPinDataParameter>(Node.GetPinData(NewPin));
 	if (PinDataOldPin && PinDataNewPin)
 	{
-		return PinDataOldPin->ParameterId == PinDataNewPin->ParameterId;
+		return PinDataOldPin->ParameterId == PinDataNewPin->ParameterId && (!OldPin.LinkedTo.Num() || OldPin.PinType == NewPin.PinType); // Pin type must match only if it was connected
 	}
 	else
 	{
@@ -66,8 +66,17 @@ void UCustomizableObjectNodeMaterialRemapPinsByName::RemapPins(const UCustomizab
 			{
 				bFound = true;
 
-				PinsToRemap.Add(OldPin, NewPin);
-				break;
+				if (UEdGraphPin** Result = PinsToRemap.Find(OldPin))
+				{
+					if ((*Result)->bOrphanedPin) // The node can have a deprecated and non-deprecated pin that should remap to the same new pin. Prioritize the non-deprecated
+					{
+						*Result = NewPin;
+					}
+				}
+				else
+				{
+					PinsToRemap.Add(OldPin, NewPin);
+				}
 			}
 		}
 
@@ -111,6 +120,7 @@ EPinMode UCustomizableObjectNodeMaterialPinDataImage::GetPinMode() const
 void UCustomizableObjectNodeMaterialPinDataImage::SetPinMode(const EPinMode InPinMode)
 {
 	FObjectEditorUtils::SetPropertyValue(this, GET_MEMBER_NAME_CHECKED(UCustomizableObjectNodeMaterialPinDataImage, PinMode), InPinMode);
+	NodeMaterial->GetPostImagePinModeChangedDelegate()->Broadcast();
 }
 
 
@@ -128,11 +138,6 @@ void UCustomizableObjectNodeMaterialPinDataImage::Copy(const UCustomizableObject
 		UVLayoutMode = PinDataOldPin->UVLayoutMode;
 		UVLayout = PinDataOldPin->UVLayout;
 		ReferenceTexture = PinDataOldPin->ReferenceTexture;
-
-		if (NodeMaterial)
-		{
-			NodeMaterial->UpdateImagePinMode(ParameterId);
-		}
 	}
 }
 
@@ -162,68 +167,44 @@ int32 UCustomizableObjectNodeMaterial::GetExpressionTextureCoordinate(UMaterial*
 }
 
 
-EPinMode UCustomizableObjectNodeMaterial::NodePinModeToImagePinMode(const ENodePinMode NodePinMode)
+FName UCustomizableObjectNodeMaterial::NodePinModeToImagePinMode(const ENodePinMode NodePinMode)
 {
 	switch (NodePinMode)
 	{
 	case ENodePinMode::Mutable:
-		return EPinMode::Mutable;
+		return UEdGraphSchema_CustomizableObject::PC_Image;
 	case ENodePinMode::Passthrough:
-		return EPinMode::Passthrough;
+		return UEdGraphSchema_CustomizableObject::PC_PassThroughImage;
 	default:
-		check(false); // Missing case.
-		return EPinMode::Mutable;
+		unimplemented();
+		return {};
 	}
 }
 
 
-EPinMode UCustomizableObjectNodeMaterial::GetImagePinMode(const UEdGraphPin& Pin) const
+FName UCustomizableObjectNodeMaterial::GetImagePinMode(const EPinMode PinMode) const
 {
-	const UEdGraphSchema_CustomizableObject* Schema = GetDefault<UEdGraphSchema_CustomizableObject>();
-
-	if (UEdGraphPin *FollowedPin = FollowInputPin(Pin))
-	{
-		if (FollowedPin->PinType.PinCategory == Schema->PC_PassThroughImage)
-		{
-			return EPinMode::Passthrough;
-		}
-
-		return EPinMode::Mutable;
-	}
-	else if (!Pin.LinkedTo.IsEmpty())
-	{
-		// This can happen if an external pin node is not connected to an expose pin node yet
-		if (UEdGraphPin* LinkedToPin = Pin.LinkedTo[0])
-		{
-			if (const UCustomizableObjectNodeExternalPin* ExternalPinNode = Cast<UCustomizableObjectNodeExternalPin>(LinkedToPin->GetOwningNode()))
-			{
-				if (ExternalPinNode->PinType == Schema->PC_Image)
-				{
-					return EPinMode::Mutable;
-				}
-				else if (ExternalPinNode->PinType == Schema->PC_PassThroughImage)
-				{
-					return EPinMode::Passthrough;
-				}
-			}
-		}
-	}
-	
-	switch (GetPinData<UCustomizableObjectNodeMaterialPinDataImage>(Pin).GetPinMode())
+	switch (PinMode)
 	{
 	case EPinMode::Default:
 		return NodePinModeToImagePinMode(TextureParametersMode);
 	case EPinMode::Mutable:
-		return EPinMode::Mutable;
+		return UEdGraphSchema_CustomizableObject::PC_Image;
 	case EPinMode::Passthrough:
-		return EPinMode::Passthrough;
+		return UEdGraphSchema_CustomizableObject::PC_PassThroughImage;
 	default:
-		check(false); // Missing case.
-		return EPinMode::Mutable;
+		unimplemented();
+		return {};
 	}
 }
 
 
+FName UCustomizableObjectNodeMaterial::GetImagePinMode(const UEdGraphPin& Pin) const
+{
+	return GetImagePinMode(GetPinData<UCustomizableObjectNodeMaterialPinDataImage>(Pin).GetPinMode());
+}
+
+
 int32 UCustomizableObjectNodeMaterial::GetImageUVLayoutFromMaterial(const int32 ImageIndex) const
 {
 	const FGuid ImageId = GetParameterId(EMaterialParameterType::Texture, ImageIndex);
@@ -292,12 +273,23 @@ void UCustomizableObjectNodeMaterialPinDataImage::PostEditChangeProperty(FProper
 	{
 		if (PropertyThatChanged->GetFName() == GET_MEMBER_NAME_CHECKED(UCustomizableObjectNodeMaterialPinDataImage, PinMode))
 		{
-			NodeMaterial->UpdateImagePinMode(ParameterId);
+			NodeMaterial->UCustomizableObjectNode::ReconstructNode();
 		}
 	}
 }
 
 
+bool UCustomizableObjectNodeMaterialPinDataImage::CanEditChange(const FProperty* InProperty) const
+{
+	if (InProperty->GetFName() == GET_MEMBER_NAME_CHECKED(UCustomizableObjectNodeMaterialPinDataImage, PinMode))
+	{
+		return !NodeMaterial->GetPin(*this)->LinkedTo.Num();
+	}
+	
+	return Super::CanEditChange(InProperty);
+}
+
+
 void UCustomizableObjectNodeMaterialPinDataImage::PostLoad()
 {
 	Super::PostLoad();
@@ -558,8 +550,6 @@ void UCustomizableObjectNodeMaterial::BackwardsCompatibleFixup()
 			{
 				UCustomizableObjectNodeMaterialPinDataImage* PinDataImage = Cast<UCustomizableObjectNodeMaterialPinDataImage>(GetPinData(*Pin));
 				PinDataImage->NodeMaterial = this;
-				
-				PinsImagePinMode.Add(Pin->PinId, GetImagePinMode(*Pin));
 			}
 		}
 	}
@@ -571,8 +561,6 @@ void UCustomizableObjectNodeMaterial::BackwardsCompatibleFixup()
 			if (UCustomizableObjectNodeMaterialPinDataImage* const PinDataImage = Cast<UCustomizableObjectNodeMaterialPinDataImage>(GetPinData(*Pin)))
 			{
 				PinDataImage->NodeMaterial = this;
-				
-				PinsImagePinMode.Add(Pin->PinId, GetImagePinMode(*Pin));
 			}
 		}
 	}
@@ -629,8 +617,6 @@ void UCustomizableObjectNodeMaterial::BackwardsCompatibleFixup()
 							if (UCustomizableObjectNodeMaterialPinDataImage* PinDataImage = Cast<UCustomizableObjectNodeMaterialPinDataImage>(GetPinData(*ImagePin)))
 							{
 								PinDataImage->PinMode = EPinMode::Mutable;
-
-								PinsImagePinMode.Add(ImagePin->PinId, GetImagePinMode(*ImagePin));
 							}
 						}
 					}
@@ -643,14 +629,54 @@ void UCustomizableObjectNodeMaterial::BackwardsCompatibleFixup()
 	{
 		MeshComponentName = FName(FString::FromInt(MeshComponentIndex_DEPRECATED));
 	}
-}
 
+	if (CustomizableObjectCustomVersion < FCustomizableObjectCustomVersion::NodeMaterialTypedImagePins)
+	{
+		const UEdGraphSchema_CustomizableObject* Schema = GetDefault<UEdGraphSchema_CustomizableObject>();
+		
+		for (UEdGraphPin* Pin : GetAllPins())
+		{
+			if (Pin->PinType.PinCategory != UEdGraphSchema_CustomizableObject::PC_Image)
+			{
+				continue;
+			}
+			
+			if (Pin->LinkedTo.Num())
+			{
+				const UEdGraphPin* LinkedPin = Pin->LinkedTo[0];
+
+				UCustomizableObjectNodeMaterialPinDataImage& PinData = GetPinData<UCustomizableObjectNodeMaterialPinDataImage>(*Pin);
 
-void UCustomizableObjectNodeMaterial::PostBackwardsCompatibleFixup()
-{
-	Super::PostBackwardsCompatibleFixup();
-	
-	UpdateAllImagesPinMode(); // Could be changed when not loaded.
+				if (LinkedPin->PinType.PinCategory == UEdGraphSchema_CustomizableObject::PC_Image)
+				{
+					Pin->PinType.PinCategory = UEdGraphSchema_CustomizableObject::PC_Image;
+
+					if (PinData.PinMode == EPinMode::Default &&
+						TextureParametersMode == ENodePinMode::Passthrough)
+					{
+						PinData.PinMode = EPinMode::Mutable;
+					}
+				}
+				else if (LinkedPin->PinType.PinCategory == UEdGraphSchema_CustomizableObject::PC_PassThroughImage)
+				{
+					Pin->PinType.PinCategory = UEdGraphSchema_CustomizableObject::PC_PassThroughImage;
+
+					if (PinData.PinMode == EPinMode::Default &&
+						TextureParametersMode == ENodePinMode::Mutable)
+					{
+						PinData.PinMode = EPinMode::Passthrough;
+					}
+				}
+			}
+			else
+			{
+				if (!IsImageMutableMode(*Pin))
+				{
+					Pin->PinType.PinCategory = UEdGraphSchema_CustomizableObject::PC_PassThroughImage;
+				}				
+			}
+		}
+	}
 }
 
 
@@ -685,7 +711,8 @@ void UCustomizableObjectNodeMaterial::PostEditChangeProperty(FPropertyChangedEve
 		}
 		else if (PropertyThatChanged->GetFName() == GET_MEMBER_NAME_CHECKED(UCustomizableObjectNodeMaterial, TextureParametersMode))
 		{
-			UpdateAllImagesPinMode();
+			Super::ReconstructNode();
+			PostImagePinModeChangedDelegate.Broadcast();
 		}
 	}
 }
@@ -782,59 +809,24 @@ bool UCustomizableObjectNodeMaterial::IsImageMutableMode(const int32 ImageIndex)
 	}
 	else
 	{
-		return NodePinModeToImagePinMode(TextureParametersMode) == EPinMode::Mutable;
+		return NodePinModeToImagePinMode(TextureParametersMode) == UEdGraphSchema_CustomizableObject::PC_Image;
 	}
 }
 
 
 bool UCustomizableObjectNodeMaterial::IsImageMutableMode(const UEdGraphPin& Pin) const
 {
-	return PinsImagePinMode[Pin.PinId] == EPinMode::Mutable;
-}
-
-
-void UCustomizableObjectNodeMaterial::UpdateImagePinMode(const FGuid ParameterId)
-{
-	UpdateImagePinMode(*PinsParameter[ParameterId].Get());
+	return GetImagePinMode(Pin) == UEdGraphSchema_CustomizableObject::PC_Image; // The ImageMutableMode is stored in the PinData, not in the PinCategory
 }
 
 
-void UCustomizableObjectNodeMaterial::UpdateImagePinMode(const UEdGraphPin& Pin)
+void UCustomizableObjectNodeMaterial::SetImagePinMode(UEdGraphPin& Pin, const EPinMode PinMode) const
 {
-	EPinMode* PinMode = PinsImagePinMode.Find(Pin.PinId);
-	const EPinMode OldPinMode = *PinMode;
+	UCustomizableObjectNodeMaterialPinDataImage& PinData = GetPinData<UCustomizableObjectNodeMaterialPinDataImage>(Pin);
 
-	*PinMode = GetImagePinMode(Pin);
+	Pin.PinType.PinCategory = GetImagePinMode(PinMode);	// Change the category so that it will remap correctly when reconstructing.
 	
-	if (*PinMode != OldPinMode)
-	{
-		PostImagePinModeChangedDelegate.Broadcast();
-	}
-}
-
-
-void UCustomizableObjectNodeMaterial::UpdateAllImagesPinMode()
-{
-	bool bChanged = false;
-	
-	const int32 NumImages = GetNumParameters(EMaterialParameterType::Texture);
-	for (int32 ImageIndex = 0; ImageIndex < NumImages; ++ImageIndex)
-	{
-		if (const UEdGraphPin* Pin = GetParameterPin(EMaterialParameterType::Texture, ImageIndex))
-		{
-			EPinMode* PinMode = PinsImagePinMode.Find(Pin->PinId);
-			const EPinMode OldPinMode = *PinMode;
-
-			*PinMode = GetImagePinMode(*Pin);
-
-			bChanged = bChanged || (*PinMode != OldPinMode);
-		}
-	}
-
-	if (bChanged)
-	{
-		PostImagePinModeChangedDelegate.Broadcast();
-	}
+	PinData.SetPinMode(PinMode); // Will trigger a reconstruct.
 }
 
 
@@ -1043,13 +1035,19 @@ bool UCustomizableObjectNodeMaterial::HasParameter(const FGuid& ParameterId) con
 }
 
 
-const UEdGraphPin* UCustomizableObjectNodeMaterial::GetParameterPin(const EMaterialParameterType Type, const int32 ParameterIndex) const
+UEdGraphPin* UCustomizableObjectNodeMaterial::GetParameterPin(const EMaterialParameterType Type, const int32 ParameterIndex) const
 {
 	const FGuid ParameterId = GetParameterId(Type, ParameterIndex);
 	
-	if (const FEdGraphPinReference* PinReference = PinsParameter.Find(ParameterId))
+	return GetParameterPin(ParameterId);
+}
+
+
+UEdGraphPin* UCustomizableObjectNodeMaterial::GetParameterPin(const FGuid& ParameterId) const
+{
+	if (const FEdGraphPinReference* Result = PinsParameter.Find(ParameterId))
 	{
-		return PinReference->Get();
+		return Result->Get();
 	}
 	else
 	{
@@ -1103,18 +1101,15 @@ TSharedPtr<IDetailsView> UCustomizableObjectNodeMaterial::CustomizePinDetails(co
 
 bool UCustomizableObjectNodeMaterial::CustomRemovePin(UEdGraphPin& Pin)
 {
-	if (const UCustomizableObjectNodeMaterialPinDataParameter* PinData = Cast<UCustomizableObjectNodeMaterialPinDataParameter>(GetPinData(Pin)))
+	for (TMap<FGuid, FEdGraphPinReference>::TIterator Iterator = PinsParameter.CreateIterator(); Iterator; ++Iterator)
 	{
-		if (const FEdGraphPinReference* Result = PinsParameter.Find(PinData->ParameterId);
-			Result &&
-			Result->Get() == &Pin)
+		if (Iterator->Value.Get() == &Pin)
 		{
-			PinsParameter.Remove(PinData->ParameterId);
+			Iterator.RemoveCurrent();
+			break;
 		}
 	}
 
-	PinsImagePinMode.Remove(Pin.PinId);
-	
 	return Super::CustomRemovePin(Pin);
 }
 
@@ -1123,8 +1118,6 @@ void UCustomizableObjectNodeMaterial::ReconstructNode(UCustomizableObjectNodeRem
 {
 	Super::ReconstructNode(RemapPinsMode);
 
-	UpdateAllImagesPinMode();
-
 	// When a material node is created, the first component is set as its mesh component
 	if (MeshComponentName.IsNone())
 	{
@@ -1264,24 +1257,44 @@ UCustomizableObjectNodeMaterialPinDataParameter* UCustomizableObjectNodeMaterial
 
 void UCustomizableObjectNodeMaterial::AllocateDefaultParameterPins(const EMaterialParameterType Type)
 {
-	const FName PinCategory = UEdGraphSchema_CustomizableObject::GetPinCategory(Type);
-
 	const int32 NumParameters = GetNumParameters(Type);
 	for (int32 ParameterIndex = 0; ParameterIndex < NumParameters; ++ParameterIndex)
 	{
 		UCustomizableObjectNodeMaterialPinDataParameter* PinData = CreatePinData(Type, ParameterIndex);
-		
+
 		const FName PinName = GetPinName(Type, ParameterIndex);
+
+		FName PinCategory;
+		switch(Type)
+		{
+		case EMaterialParameterType::Texture:
+			if (IsImageMutableMode(ParameterIndex)) // If a pin exists, we store the PinMode in its PinData.
+			{
+				PinCategory = UEdGraphSchema_CustomizableObject::PC_Image;				
+			}
+			else
+			{
+				PinCategory = UEdGraphSchema_CustomizableObject::PC_PassThroughImage;
+			}
+			break;
+
+		case EMaterialParameterType::Vector:
+			PinCategory = UEdGraphSchema_CustomizableObject::PC_Color;
+			break;
+
+		case EMaterialParameterType::Scalar:
+			PinCategory = UEdGraphSchema_CustomizableObject::PC_Float;
+			break;
+
+		default:
+			check(false); // Type not contemplated
+		}		
+		
 		UEdGraphPin* Pin = CustomCreatePin(EGPD_Input, PinCategory, PinName, PinData);
 		Pin->bHidden = true;
 		Pin->bDefaultValueIsIgnored = true;
 
 		PinsParameter.Add(PinData->ParameterId, FEdGraphPinReference(Pin));
-
-		if (Type == EMaterialParameterType::Texture)
-		{
-			PinsImagePinMode.Add(Pin->PinId, GetImagePinMode(*Pin));
-		}
 	}
 }
 
@@ -1340,13 +1353,31 @@ void UCustomizableObjectNodeMaterial::PinConnectionListChanged(UEdGraphPin* Pin)
 			}
 		}
 	}
-	else if (Cast<UCustomizableObjectNodeMaterialPinDataImage>(GetPinData(*Pin)))
-	{
-		UpdateImagePinMode(*Pin);
-	}
-	else if (Pin == GetMaterialAssetPin())
+	else if (Cast<UCustomizableObjectNodeMaterialPinDataImage>(GetPinData(*Pin))) // Image pin
 	{
-		UpdateAllImagesPinMode();	
+		// If necessary, automatically change the Pin Mode. Connected pin can never change its type.
+		if (Pin->LinkedTo.Num())
+		{
+			if (UEdGraphPin* LinkedPin = Pin->LinkedTo[0])
+			{
+				EPinMode PinMode = EPinMode::Default;
+		
+				if (LinkedPin->PinType.PinCategory == UEdGraphSchema_CustomizableObject::PC_Image)
+				{
+					PinMode = EPinMode::Mutable;
+				}
+				else if (LinkedPin->PinType.PinCategory == UEdGraphSchema_CustomizableObject::PC_PassThroughImage)
+				{
+					PinMode = EPinMode::Passthrough;
+				}
+				else
+				{
+					unimplemented()
+				}
+		
+				SetImagePinMode(*Pin, PinMode);
+			}
+		}
 	}
 }
 
@@ -1364,12 +1395,12 @@ bool UCustomizableObjectNodeMaterial::CanConnect(const UEdGraphPin* InOwnedInput
 
 	if (InOwnedInputPin && InOutputPin)
 	{
-		if (InOwnedInputPin->PinType.PinCategory == Schema->PC_Image &&
-			InOutputPin->PinType.PinCategory == Schema->PC_PassThroughImage)
+		if ((InOwnedInputPin->PinType.PinCategory == Schema->PC_Image && InOutputPin->PinType.PinCategory == Schema->PC_PassThroughImage) ||
+			(InOwnedInputPin->PinType.PinCategory == Schema->PC_PassThroughImage && InOutputPin->PinType.PinCategory == Schema->PC_Image))
 		{
 			return true;
 		}
-
+		
 		if (InOwnedInputPin->PinType.PinCategory == Schema->PC_Mesh &&
 			InOutputPin->PinType.PinCategory == Schema->PC_PassThroughMesh)
 		{
diff --git a/Engine/Plugins/Experimental/Mutable/Source/CustomizableObjectEditor/Private/MuCOE/Nodes/CustomizableObjectNodeReroute.cpp b/Engine/Plugins/Experimental/Mutable/Source/CustomizableObjectEditor/Private/MuCOE/Nodes/CustomizableObjectNodeReroute.cpp
--- a/Engine/Plugins/Experimental/Mutable/Source/CustomizableObjectEditor/Private/MuCOE/Nodes/CustomizableObjectNodeReroute.cpp
+++ b/Engine/Plugins/Experimental/Mutable/Source/CustomizableObjectEditor/Private/MuCOE/Nodes/CustomizableObjectNodeReroute.cpp
@@ -2,6 +2,7 @@
 
 #include "MuCOE/Nodes/CustomizableObjectNodeReroute.h"
 
+#include "SCustomizableObjectNodeReroute.h"
 #include "SGraphNodeKnot.h"
 #include "MuCOE/EdGraphSchema_CustomizableObject.h"
 #include "MuCOE/GraphTraversal.h"
@@ -82,7 +83,7 @@ bool UCustomizableObjectNodeReroute::ShouldDrawNodeAsControlPointOnly(int32& Out
 
 TSharedPtr<SGraphNode> UCustomizableObjectNodeReroute::CreateVisualWidget()
 {
-	return SNew(SGraphNodeKnot, this);
+	return SNew(SCustomizableObjectNodeReroute, this);
 }
 
 
diff --git a/Engine/Plugins/Experimental/Mutable/Source/CustomizableObjectEditor/Private/MuCOE/Nodes/CustomizableObjectNodeSkeletalMesh.cpp b/Engine/Plugins/Experimental/Mutable/Source/CustomizableObjectEditor/Private/MuCOE/Nodes/CustomizableObjectNodeSkeletalMesh.cpp
--- a/Engine/Plugins/Experimental/Mutable/Source/CustomizableObjectEditor/Private/MuCOE/Nodes/CustomizableObjectNodeSkeletalMesh.cpp
+++ b/Engine/Plugins/Experimental/Mutable/Source/CustomizableObjectEditor/Private/MuCOE/Nodes/CustomizableObjectNodeSkeletalMesh.cpp
@@ -813,8 +813,7 @@ int32 UCustomizableObjectNodeSkeletalMesh::GetSkeletalMaterialIndexFor(const int
 
 void SGraphNodeSkeletalMesh::Construct(const FArguments& InArgs, UEdGraphNode* InGraphNode)
 {
-	GraphNode = InGraphNode;
-	NodeSkeletalMesh = Cast< UCustomizableObjectNodeSkeletalMesh >(GraphNode);
+	NodeSkeletalMesh = Cast< UCustomizableObjectNodeSkeletalMesh >(InGraphNode);
 
 	WidgetSize = 128.0f;
 	ThumbnailSize = 128;
@@ -834,7 +833,7 @@ void SGraphNodeSkeletalMesh::Construct(const FArguments& InArgs, UEdGraphNode* I
 
 	SkeletalMeshSelector = PropPlugin.CreateSingleProperty(NodeSkeletalMesh, "SkeletalMesh", SingleDetails);
 
-	UpdateGraphNode();
+	SCustomizableObjectNode::Construct({}, InGraphNode);
 }
 
 
diff --git a/Engine/Plugins/Experimental/Mutable/Source/CustomizableObjectEditor/Private/MuCOE/Nodes/CustomizableObjectNodeStaticMesh.cpp b/Engine/Plugins/Experimental/Mutable/Source/CustomizableObjectEditor/Private/MuCOE/Nodes/CustomizableObjectNodeStaticMesh.cpp
--- a/Engine/Plugins/Experimental/Mutable/Source/CustomizableObjectEditor/Private/MuCOE/Nodes/CustomizableObjectNodeStaticMesh.cpp
+++ b/Engine/Plugins/Experimental/Mutable/Source/CustomizableObjectEditor/Private/MuCOE/Nodes/CustomizableObjectNodeStaticMesh.cpp
@@ -526,8 +526,7 @@ TSharedPtr<SGraphNode> UCustomizableObjectNodeStaticMesh::CreateVisualWidget()
 
 void SGraphNodeStaticMesh::Construct(const FArguments& InArgs, UEdGraphNode* InGraphNode)
 {
-	GraphNode = InGraphNode;
-	NodeStaticMesh = Cast< UCustomizableObjectNodeStaticMesh >(GraphNode);
+	NodeStaticMesh = Cast< UCustomizableObjectNodeStaticMesh >(InGraphNode);
 
 	WidgetSize = 128.0f;
 	ThumbnailSize = 128;
@@ -546,7 +545,7 @@ void SGraphNodeStaticMesh::Construct(const FArguments& InArgs, UEdGraphNode* InG
 
 	StaticMeshSelector = PropPlugin.CreateSingleProperty(NodeStaticMesh, "StaticMesh", SingleDetails);	
 
-	UpdateGraphNode();
+	SCustomizableObjectNode::Construct({}, InGraphNode);
 }
 
 
diff --git a/Engine/Plugins/Experimental/Mutable/Source/CustomizableObjectEditor/Private/MuCOE/Nodes/CustomizableObjectNodeTable.cpp b/Engine/Plugins/Experimental/Mutable/Source/CustomizableObjectEditor/Private/MuCOE/Nodes/CustomizableObjectNodeTable.cpp
--- a/Engine/Plugins/Experimental/Mutable/Source/CustomizableObjectEditor/Private/MuCOE/Nodes/CustomizableObjectNodeTable.cpp
+++ b/Engine/Plugins/Experimental/Mutable/Source/CustomizableObjectEditor/Private/MuCOE/Nodes/CustomizableObjectNodeTable.cpp
@@ -187,17 +187,6 @@ void UCustomizableObjectNodeTable::BackwardsCompatibleFixup()
 }
 
 
-void UCustomizableObjectNodeTable::PostBackwardsCompatibleFixup()
-{
-	Super::PostBackwardsCompatibleFixup();
-
-	if (Table)
-	{
-		OnTableChangedDelegateHandle = Table->OnDataTableChanged().AddUObject(this, &UCustomizableObjectNodeTable::OnTableChanged);
-	}
-}
-
-
 void UCustomizableObjectNodeTable::PostLoad()
 {
 	Super::PostLoad();
@@ -262,25 +251,6 @@ void UCustomizableObjectNodeTable::PreEditChange(FProperty* PropertyAboutToChang
 }
 
 
-void UCustomizableObjectNodeTable::PostEditChangeChainProperty(FPropertyChangedChainEvent& PropertyChangedEvent)
-{
-	Super::PostEditChangeChainProperty(PropertyChangedEvent);
-
-	if (const FProperty* PropertyThatChanged = PropertyChangedEvent.Property)
-	{
-		if (PropertyThatChanged->GetFName() == GET_MEMBER_NAME_CHECKED(UCustomizableObjectNodeTable, Table))
-		{
-			if (Table)
-			{
-				OnTableChangedDelegateHandle = Table->OnDataTableChanged().AddUObject(this, &UCustomizableObjectNodeTable::OnTableChanged);
-			}
-			
-			OnTableChanged();
-		}
-	}
-}
-
-
 FText UCustomizableObjectNodeTable::GetNodeTitle(ENodeTitleType::Type TitleType) const
 {
 	FFormatNamedArguments Args;
@@ -1091,24 +1061,6 @@ void UCustomizableObjectNodeTable::GetAnimationColumns(const FGuid& ColumnId, FS
 }
 
 
-void UCustomizableObjectNodeTable::OnTableChanged()
-{	
-	for (const UEdGraphPin* Pin : GetAllNonOrphanPins())
-	{
-		if (Pin->Direction == EGPD_Output)
-		{
-			for (const UEdGraphPin* ConnectedPin : FollowOutputPinArray(*Pin))
-        	{
-        		if (UCustomizableObjectNodeMaterial* NodeMaterial = Cast<UCustomizableObjectNodeMaterial>(ConnectedPin->GetOwningNode()))
-        		{
-        			NodeMaterial->UpdateAllImagesPinMode();
-        		}
-        	}
-		}
-	}
-}
-
-
 bool UCustomizableObjectNodeTable::CheckPinUpdated(const FString& PinName, const FName& PinType) const
 {
 	if (UEdGraphPin* Pin = FindPin(PinName))
diff --git a/Engine/Plugins/Experimental/Mutable/Source/CustomizableObjectEditor/Private/MuCOE/Nodes/CustomizableObjectNodeTexture.cpp b/Engine/Plugins/Experimental/Mutable/Source/CustomizableObjectEditor/Private/MuCOE/Nodes/CustomizableObjectNodeTexture.cpp
--- a/Engine/Plugins/Experimental/Mutable/Source/CustomizableObjectEditor/Private/MuCOE/Nodes/CustomizableObjectNodeTexture.cpp
+++ b/Engine/Plugins/Experimental/Mutable/Source/CustomizableObjectEditor/Private/MuCOE/Nodes/CustomizableObjectNodeTexture.cpp
@@ -81,8 +81,7 @@ FText UCustomizableObjectNodeTexture::GetTooltipText() const
 
 void SGraphNodeTexture::Construct(const FArguments& InArgs, UEdGraphNode* InGraphNode)
 {
-	GraphNode = InGraphNode;
-	NodeTexture = Cast<UCustomizableObjectNodeTextureBase>(GraphNode);
+	NodeTexture = Cast<UCustomizableObjectNodeTextureBase>(InGraphNode);
 
 	FPropertyEditorModule& PropPlugin = FModuleManager::LoadModuleChecked<FPropertyEditorModule>("PropertyEditor");
 
@@ -96,8 +95,8 @@ void SGraphNodeTexture::Construct(const FArguments& InArgs, UEdGraphNode* InGrap
 	TextureBrush.ImageSize.X = 128.0f;
 	TextureBrush.ImageSize.Y = 128.0f;
 	TextureBrush.DrawAs = ESlateBrushDrawType::Image;
-	
-	UpdateGraphNode();
+
+	SCustomizableObjectNode::Construct({}, InGraphNode);
 }
 
 
diff --git a/Engine/Plugins/Experimental/Mutable/Source/CustomizableObjectEditor/Private/MuCOE/Nodes/SCustomizableObjectNode.cpp b/Engine/Plugins/Experimental/Mutable/Source/CustomizableObjectEditor/Private/MuCOE/Nodes/SCustomizableObjectNode.cpp
new file mode 100644
--- /dev/null
+++ b/Engine/Plugins/Experimental/Mutable/Source/CustomizableObjectEditor/Private/MuCOE/Nodes/SCustomizableObjectNode.cpp
@@ -0,0 +1,18 @@
+﻿// Copyright Epic Games, Inc. All Rights Reserved.
+
+#include "SCustomizableObjectNode.h"
+
+#include "SCustomizableObjectNodePin.h"
+
+
+void SCustomizableObjectNode::Construct(const FArguments& InArgs, UEdGraphNode* InGraphNode)
+{
+    GraphNode = InGraphNode;
+    UpdateGraphNode();
+}
+
+
+TSharedPtr<SGraphPin> SCustomizableObjectNode::CreatePinWidget(UEdGraphPin* Pin) const
+{
+	return SNew(SCustomizableObjectNodePin, Pin);
+}
diff --git a/Engine/Plugins/Experimental/Mutable/Source/CustomizableObjectEditor/Private/MuCOE/Nodes/SCustomizableObjectNode.h b/Engine/Plugins/Experimental/Mutable/Source/CustomizableObjectEditor/Private/MuCOE/Nodes/SCustomizableObjectNode.h
new file mode 100644
--- /dev/null
+++ b/Engine/Plugins/Experimental/Mutable/Source/CustomizableObjectEditor/Private/MuCOE/Nodes/SCustomizableObjectNode.h
@@ -0,0 +1,22 @@
+﻿// Copyright Epic Games, Inc. All Rights Reserved.
+
+#pragma once
+
+#include "SGraphNode.h"
+
+class SGraphPin;
+class UEdGraphNode;
+class UEdGraphPin;
+
+
+class SCustomizableObjectNode : public SGraphNode
+{
+public:
+    SLATE_BEGIN_ARGS(SCustomizableObjectNode) {}
+    SLATE_END_ARGS();
+
+    void Construct(const FArguments& InArgs, UEdGraphNode* InGraphNode);
+
+	// SGraphNode interface
+    virtual TSharedPtr<SGraphPin> CreatePinWidget(UEdGraphPin* Pin) const override;
+};
diff --git a/Engine/Plugins/Experimental/Mutable/Source/CustomizableObjectEditor/Private/MuCOE/Nodes/SCustomizableObjectNodeMaterial.cpp b/Engine/Plugins/Experimental/Mutable/Source/CustomizableObjectEditor/Private/MuCOE/Nodes/SCustomizableObjectNodeMaterial.cpp
--- a/Engine/Plugins/Experimental/Mutable/Source/CustomizableObjectEditor/Private/MuCOE/Nodes/SCustomizableObjectNodeMaterial.cpp
+++ b/Engine/Plugins/Experimental/Mutable/Source/CustomizableObjectEditor/Private/MuCOE/Nodes/SCustomizableObjectNodeMaterial.cpp
@@ -4,20 +4,20 @@
 
 #include "MuCOE/EdGraphSchema_CustomizableObject.h"
 #include "MuCOE/Nodes/SCustomizableObjectNodeMaterialPinImage.h"
+#include "MuCOE/Nodes/CustomizableObjectNodeMaterial.h"
 
 class SGraphPin;
 
 
-void SCustomizableObjectNodeMaterial::Construct(const FArguments& InArgs, UEdGraphNode* InGraphNode)
+void SCustomizableObjectNodeMaterial::Construct(const FArguments& InArgs, UCustomizableObjectNodeMaterial* InGraphNode)
 {
-	GraphNode = InGraphNode;
-	UpdateGraphNode();
+	SCustomizableObjectNode::Construct({}, InGraphNode);
 }
 
 
 TSharedPtr<SGraphPin> SCustomizableObjectNodeMaterial::CreatePinWidget(UEdGraphPin* Pin) const
 {
-	if (Pin->PinType.PinCategory == UEdGraphSchema_CustomizableObject::PC_Image &&
+	if ((Pin->PinType.PinCategory == UEdGraphSchema_CustomizableObject::PC_Image || Pin->PinType.PinCategory == UEdGraphSchema_CustomizableObject::PC_PassThroughImage) &&
 		Pin->Direction == EGPD_Input)
 	{
 		return SNew(SCustomizableObjectNodeMaterialPinImage, Pin);
diff --git a/Engine/Plugins/Experimental/Mutable/Source/CustomizableObjectEditor/Private/MuCOE/Nodes/SCustomizableObjectNodeMaterial.h b/Engine/Plugins/Experimental/Mutable/Source/CustomizableObjectEditor/Private/MuCOE/Nodes/SCustomizableObjectNodeMaterial.h
--- a/Engine/Plugins/Experimental/Mutable/Source/CustomizableObjectEditor/Private/MuCOE/Nodes/SCustomizableObjectNodeMaterial.h
+++ b/Engine/Plugins/Experimental/Mutable/Source/CustomizableObjectEditor/Private/MuCOE/Nodes/SCustomizableObjectNodeMaterial.h
@@ -2,21 +2,23 @@
 
 #pragma once
 
+#include "SCustomizableObjectNode.h"
 #include "SGraphNode.h"
 
+class UCustomizableObjectNodeMaterial;
 class SGraphPin;
-class UEdGraphNode;
 class UEdGraphPin;
 
 
 /** Custom widget for the Material node. */
-class SCustomizableObjectNodeMaterial : public SGraphNode
+class SCustomizableObjectNodeMaterial : public SCustomizableObjectNode
 {
 public:
 	SLATE_BEGIN_ARGS(SCustomizableObjectNodeMaterial) {}
 	SLATE_END_ARGS();
 
-	void Construct(const FArguments& InArgs, UEdGraphNode* InGraphNode);
+	void Construct(const FArguments& InArgs, UCustomizableObjectNodeMaterial* InGraphNode);
 
-	TSharedPtr<SGraphPin> CreatePinWidget(UEdGraphPin* Pin) const override;
+	// SGraphNode interface
+	virtual TSharedPtr<SGraphPin> CreatePinWidget(UEdGraphPin* Pin) const override;
 };
diff --git a/Engine/Plugins/Experimental/Mutable/Source/CustomizableObjectEditor/Private/MuCOE/Nodes/SCustomizableObjectNodeMaterialPinImage.cpp b/Engine/Plugins/Experimental/Mutable/Source/CustomizableObjectEditor/Private/MuCOE/Nodes/SCustomizableObjectNodeMaterialPinImage.cpp
--- a/Engine/Plugins/Experimental/Mutable/Source/CustomizableObjectEditor/Private/MuCOE/Nodes/SCustomizableObjectNodeMaterialPinImage.cpp
+++ b/Engine/Plugins/Experimental/Mutable/Source/CustomizableObjectEditor/Private/MuCOE/Nodes/SCustomizableObjectNodeMaterialPinImage.cpp
@@ -2,6 +2,7 @@
 
 #include "MuCOE/Nodes/SCustomizableObjectNodeMaterialPinImage.h"
 
+#include "MuCOE/EdGraphSchema_CustomizableObject.h"
 #include "MuCOE/GraphTraversal.h"
 #include "MuCOE/Nodes/CustomizableObjectNodeMaterial.h"
 #include "MuCOE/UnrealEditorPortabilityHelpers.h"
@@ -17,18 +18,13 @@ struct FSlateBrush;
 
 void SCustomizableObjectNodeMaterialPinImage::Construct(const FArguments& InArgs, UEdGraphPin* InGraphPinObj)
 {
-	SGraphPin::Construct(SGraphPin::FArguments(), InGraphPinObj);
+	SCustomizableObjectNodePin::Construct(SCustomizableObjectNodePin::FArguments(), InGraphPinObj);
 
 	// Override previously defined tool tip.
 	TSharedPtr<IToolTip> TooltipWidget = SNew(SToolTip)
-		.Text(this, &SCustomizableObjectNodeMaterialPinImage::GetTooltipText);
+		.Text(this, &SCustomizableObjectNodeMaterialPinImage::GetPinTooltipText);
 
 	SetToolTip(TooltipWidget);
-
-	// Cache pin icons.
-	CachedPinMutableConnected = UE_MUTABLE_GET_BRUSH(TEXT("Graph.Pin.Connected"));
-	CachedPinMutableDisconnected = UE_MUTABLE_GET_BRUSH(TEXT("Graph.Pin.Disconnected"));
-	CachedPinPassthroughDisconnected = UE_MUTABLE_GET_BRUSH(TEXT("Graph.ExecPin.Disconnected"));
 }
 
 
@@ -46,34 +42,14 @@ TSharedRef<SWidget>	SCustomizableObjectNodeMaterialPinImage::GetDefaultValueWidg
 }
 
 
-const FSlateBrush* SCustomizableObjectNodeMaterialPinImage::GetPinIcon() const
-{
-	if (CastChecked<UCustomizableObjectNodeMaterial>(GraphPinObj->GetOwningNode())->IsImageMutableMode(*GraphPinObj))
-	{
-		if (GraphPinObj->LinkedTo.Num())
-		{
-			return CachedPinMutableConnected;
-		}
-		else
-		{
-			return CachedPinMutableDisconnected;
-		}
-	}
-	else
-	{
-		return CachedPinPassthroughDisconnected;
-	}
-}
-
-
-FText SCustomizableObjectNodeMaterialPinImage::GetTooltipText() const
+FText SCustomizableObjectNodeMaterialPinImage::GetPinTooltipText() const
 {
 	if (GraphPinObj->bOrphanedPin)
 	{
 		return LOCTEXT("PinModeMutableOrpahan", "Pin not disapearing due to being connected or having a property modified.");
 	}
 	
-	if (CastChecked<UCustomizableObjectNodeMaterial>(GraphPinObj->GetOwningNode())->IsImageMutableMode(*GraphPinObj))
+	if (GraphPinObj->PinType.PinCategory == UEdGraphSchema_CustomizableObject::PC_Image)
 	{
 		return LOCTEXT("PinModeMutableTooltip", "Texture Parameter goes through Mutable.");
 	}
@@ -86,7 +62,7 @@ FText SCustomizableObjectNodeMaterialPinImage::GetTooltipText() const
 
 FText SCustomizableObjectNodeMaterialPinImage::GetDefaultValueText() const
 {
-	if (CastChecked<UCustomizableObjectNodeMaterial>(GraphPinObj->GetOwningNode())->IsImageMutableMode(*GraphPinObj))
+	if (GraphPinObj->PinType.PinCategory == UEdGraphSchema_CustomizableObject::PC_Image)
 	{
 		return LOCTEXT("PinModeMutable", "mutable");
 	}
diff --git a/Engine/Plugins/Experimental/Mutable/Source/CustomizableObjectEditor/Private/MuCOE/Nodes/SCustomizableObjectNodeMaterialPinImage.h b/Engine/Plugins/Experimental/Mutable/Source/CustomizableObjectEditor/Private/MuCOE/Nodes/SCustomizableObjectNodeMaterialPinImage.h
--- a/Engine/Plugins/Experimental/Mutable/Source/CustomizableObjectEditor/Private/MuCOE/Nodes/SCustomizableObjectNodeMaterialPinImage.h
+++ b/Engine/Plugins/Experimental/Mutable/Source/CustomizableObjectEditor/Private/MuCOE/Nodes/SCustomizableObjectNodeMaterialPinImage.h
@@ -2,6 +2,7 @@
 
 #pragma once
 
+#include "SCustomizableObjectNodePin.h"
 #include "SGraphPin.h"
 
 class SWidget;
@@ -9,13 +10,8 @@ class UEdGraphPin;
 struct FSlateBrush;
 
 
-/** Material node custom Image pin. Allows to define a custom style to Image pins.
- *  
- * Currently these different styles allow the user to differentiate the different states an Image pin can have:
- * - passthrough: The texture does NOT go trough the mutable pipeline.
- * - mutable: The texture goes trough the mutable pipeline.
- */
-class SCustomizableObjectNodeMaterialPinImage : public SGraphPin
+/** Implements the "MUTABLE" and "PASSTHROUGH" text next to the pin name. */
+class SCustomizableObjectNodeMaterialPinImage : public SCustomizableObjectNodePin
 {
 public:
 
@@ -25,7 +21,6 @@ public:
 	void Construct(const FArguments& InArgs, UEdGraphPin* InGraphPinObj);
 
 	TSharedRef<SWidget>	GetDefaultValueWidget() override;
-	const FSlateBrush* GetPinIcon() const override;
 
 private:
 	/** Return pin state text. */
@@ -35,9 +30,5 @@ private:
 	EVisibility GetDefaultValueVisibility() const;
 
 	/** Return pin tool tip. */
-	FText GetTooltipText() const;
-
-	const FSlateBrush* CachedPinMutableConnected;
-	const FSlateBrush* CachedPinMutableDisconnected;
-	const FSlateBrush* CachedPinPassthroughDisconnected;
+	FText GetPinTooltipText() const;
 };
diff --git a/Engine/Plugins/Experimental/Mutable/Source/CustomizableObjectEditor/Private/MuCOE/Nodes/SCustomizableObjectNodePin.cpp b/Engine/Plugins/Experimental/Mutable/Source/CustomizableObjectEditor/Private/MuCOE/Nodes/SCustomizableObjectNodePin.cpp
new file mode 100644
--- /dev/null
+++ b/Engine/Plugins/Experimental/Mutable/Source/CustomizableObjectEditor/Private/MuCOE/Nodes/SCustomizableObjectNodePin.cpp
@@ -0,0 +1,31 @@
+﻿// Copyright Epic Games, Inc. All Rights Reserved.
+
+#include "SCustomizableObjectNodePin.h"
+
+#include "MuCOE/EdGraphSchema_CustomizableObject.h"
+#include "MuCOE/UnrealEditorPortabilityHelpers.h"
+
+
+void SCustomizableObjectNodePin::Construct(const FArguments& InArgs, UEdGraphPin* InGraphPinObj)
+{
+	SGraphPin::Construct(SGraphPin::FArguments(), InGraphPinObj);
+
+	// Cache pin icons.
+	PassThroughImageConnected = UE_MUTABLE_GET_BRUSH(TEXT("Graph.ExecPin.Connected"));
+	PassThroughImageDisconnected = UE_MUTABLE_GET_BRUSH(TEXT("Graph.ExecPin.Disconnected"));
+}
+
+
+const FSlateBrush* SCustomizableObjectNodePin::GetPinIcon() const
+{
+	if (GraphPinObj->PinType.PinCategory == UEdGraphSchema_CustomizableObject::PC_PassThroughImage)
+	{
+		return GraphPinObj->LinkedTo.Num() ?
+			PassThroughImageConnected :
+			PassThroughImageDisconnected;
+	}
+	else
+	{
+		return SGraphPin::GetPinIcon();
+	}
+}
diff --git a/Engine/Plugins/Experimental/Mutable/Source/CustomizableObjectEditor/Private/MuCOE/Nodes/SCustomizableObjectNodePin.h b/Engine/Plugins/Experimental/Mutable/Source/CustomizableObjectEditor/Private/MuCOE/Nodes/SCustomizableObjectNodePin.h
new file mode 100644
--- /dev/null
+++ b/Engine/Plugins/Experimental/Mutable/Source/CustomizableObjectEditor/Private/MuCOE/Nodes/SCustomizableObjectNodePin.h
@@ -0,0 +1,28 @@
+﻿// Copyright Epic Games, Inc. All Rights Reserved.
+
+#pragma once
+
+#include "SGraphPin.h"
+
+class SWidget;
+class UEdGraphPin;
+struct FSlateBrush;
+
+
+class SCustomizableObjectNodePin : public SGraphPin
+{
+public:
+
+	SLATE_BEGIN_ARGS(SCustomizableObjectNodePin) {}
+	SLATE_END_ARGS()
+
+	void Construct(const FArguments& InArgs, UEdGraphPin* InGraphPinObj);
+	
+	// SGraphPin interface
+protected:
+	virtual const FSlateBrush* GetPinIcon() const override;
+
+private:
+	const FSlateBrush* PassThroughImageConnected = nullptr;
+	const FSlateBrush* PassThroughImageDisconnected = nullptr;
+};
diff --git a/Engine/Plugins/Experimental/Mutable/Source/CustomizableObjectEditor/Private/MuCOE/Nodes/SCustomizableObjectNodeReroute.cpp b/Engine/Plugins/Experimental/Mutable/Source/CustomizableObjectEditor/Private/MuCOE/Nodes/SCustomizableObjectNodeReroute.cpp
new file mode 100644
--- /dev/null
+++ b/Engine/Plugins/Experimental/Mutable/Source/CustomizableObjectEditor/Private/MuCOE/Nodes/SCustomizableObjectNodeReroute.cpp
@@ -0,0 +1,18 @@
+﻿// Copyright Epic Games, Inc. All Rights Reserved.
+
+#include "SCustomizableObjectNodeReroute.h"
+
+#include "SCustomizableObjectNodeReroutePin.h"
+
+
+void SCustomizableObjectNodeReroute::Construct(const FArguments& InArgs, UEdGraphNode* InGraphNode)
+{
+	SGraphNodeKnot::Construct({}, InGraphNode);
+}
+
+
+TSharedPtr<SGraphPin> SCustomizableObjectNodeReroute::CreatePinWidget(UEdGraphPin* Pin) const
+{
+	return SNew(SCustomizableObjectNodeReroutePin, Pin);
+}
+
diff --git a/Engine/Plugins/Experimental/Mutable/Source/CustomizableObjectEditor/Private/MuCOE/Nodes/SCustomizableObjectNodeReroute.h b/Engine/Plugins/Experimental/Mutable/Source/CustomizableObjectEditor/Private/MuCOE/Nodes/SCustomizableObjectNodeReroute.h
new file mode 100644
--- /dev/null
+++ b/Engine/Plugins/Experimental/Mutable/Source/CustomizableObjectEditor/Private/MuCOE/Nodes/SCustomizableObjectNodeReroute.h
@@ -0,0 +1,24 @@
+﻿// Copyright Epic Games, Inc. All Rights Reserved.
+
+#pragma once
+
+#include "SGraphNodeKnot.h"
+#include "Templates/SharedPointer.h"
+
+class SGraphPin;
+class UEdGraphPin;
+
+
+// To avoid multiple inheritance SCustomizableObjectNodeReroute does not inherit from SCustomizableObjectNode but reimplements its functionalities.
+class SCustomizableObjectNodeReroute : public SGraphNodeKnot
+{
+public:
+	SLATE_BEGIN_ARGS(SGraphNodeKnot) {}
+	SLATE_END_ARGS()
+
+	void Construct(const FArguments& InArgs, UEdGraphNode* InKnot);
+	
+	// SGraphNode interface
+	virtual TSharedPtr<SGraphPin> CreatePinWidget(UEdGraphPin* Pin) const override;
+};
+
diff --git a/Engine/Plugins/Experimental/Mutable/Source/CustomizableObjectEditor/Private/MuCOE/Nodes/SCustomizableObjectNodeReroutePin.cpp b/Engine/Plugins/Experimental/Mutable/Source/CustomizableObjectEditor/Private/MuCOE/Nodes/SCustomizableObjectNodeReroutePin.cpp
new file mode 100644
--- /dev/null
+++ b/Engine/Plugins/Experimental/Mutable/Source/CustomizableObjectEditor/Private/MuCOE/Nodes/SCustomizableObjectNodeReroutePin.cpp
@@ -0,0 +1,31 @@
+﻿// Copyright Epic Games, Inc. All Rights Reserved.
+
+#include "SCustomizableObjectNodeReroutePin.h"
+
+#include "MuCOE/EdGraphSchema_CustomizableObject.h"
+#include "MuCOE/UnrealEditorPortabilityHelpers.h"
+
+
+void SCustomizableObjectNodeReroutePin::Construct(const FArguments& InArgs, UEdGraphPin* InPin)
+{
+	SGraphPinKnot::Construct({}, InPin);
+
+	// Cache pin icons.
+	PassThroughImageConnected = UE_MUTABLE_GET_BRUSH(TEXT("Graph.ExecPin.Connected"));
+	PassThroughImageDisconnected = UE_MUTABLE_GET_BRUSH(TEXT("Graph.ExecPin.Disconnected"));	
+}
+
+
+const FSlateBrush* SCustomizableObjectNodeReroutePin::GetPinIcon() const
+{
+	if (GraphPinObj->PinType.PinCategory == UEdGraphSchema_CustomizableObject::PC_PassThroughImage)
+	{
+		return GraphPinObj->LinkedTo.Num() ?
+			PassThroughImageConnected :
+			PassThroughImageDisconnected;
+	}
+	else
+	{
+		return SGraphPin::GetPinIcon();
+	}
+}
diff --git a/Engine/Plugins/Experimental/Mutable/Source/CustomizableObjectEditor/Private/MuCOE/Nodes/SCustomizableObjectNodeReroutePin.h b/Engine/Plugins/Experimental/Mutable/Source/CustomizableObjectEditor/Private/MuCOE/Nodes/SCustomizableObjectNodeReroutePin.h
new file mode 100644
--- /dev/null
+++ b/Engine/Plugins/Experimental/Mutable/Source/CustomizableObjectEditor/Private/MuCOE/Nodes/SCustomizableObjectNodeReroutePin.h
@@ -0,0 +1,24 @@
+﻿// Copyright Epic Games, Inc. All Rights Reserved.
+
+#pragma once
+
+#include "SGraphNodeKnot.h"
+
+
+// To avoid multiple inheritance SCustomizableObjectNodeReroutePin does not inherit from SCustomizableObjectNodePin but reimplements its functionalities.
+class SCustomizableObjectNodeReroutePin : public SGraphPinKnot
+{
+public:
+	SLATE_BEGIN_ARGS(SCustomizableObjectNodeReroutePin) {}
+	SLATE_END_ARGS()
+
+	void Construct(const FArguments& InArgs, UEdGraphPin* InPin);
+
+protected:
+	// SGraphPin interface
+	virtual const FSlateBrush* GetPinIcon() const override;
+
+	const FSlateBrush* PassThroughImageConnected = nullptr;
+	const FSlateBrush* PassThroughImageDisconnected = nullptr;
+};
+
diff --git a/Engine/Plugins/Experimental/Mutable/Source/CustomizableObjectEditor/Public/MuCOE/EdGraphSchema_CustomizableObject.h b/Engine/Plugins/Experimental/Mutable/Source/CustomizableObjectEditor/Public/MuCOE/EdGraphSchema_CustomizableObject.h
--- a/Engine/Plugins/Experimental/Mutable/Source/CustomizableObjectEditor/Public/MuCOE/EdGraphSchema_CustomizableObject.h
+++ b/Engine/Plugins/Experimental/Mutable/Source/CustomizableObjectEditor/Public/MuCOE/EdGraphSchema_CustomizableObject.h
@@ -166,9 +166,6 @@ public:
 	
 	/** Given a pin category name, get its friendly name (user readable). */
 	static FText GetPinCategoryName(const FName& PinCategory);
-
-	/** Given a Material parameter type, returns the Mutable pin category. */
-	static const FName& GetPinCategory(EMaterialParameterType Type);
 };
 
 
diff --git a/Engine/Plugins/Experimental/Mutable/Source/CustomizableObjectEditor/Public/MuCOE/Nodes/CustomizableObjectNode.h b/Engine/Plugins/Experimental/Mutable/Source/CustomizableObjectEditor/Public/MuCOE/Nodes/CustomizableObjectNode.h
--- a/Engine/Plugins/Experimental/Mutable/Source/CustomizableObjectEditor/Public/MuCOE/Nodes/CustomizableObjectNode.h
+++ b/Engine/Plugins/Experimental/Mutable/Source/CustomizableObjectEditor/Public/MuCOE/Nodes/CustomizableObjectNode.h
@@ -92,6 +92,7 @@ public:
 	virtual void NodeConnectionListChanged() override;
 	virtual void PinConnectionListChanged(UEdGraphPin* Pin) override;
 	virtual void DestroyNode() override;
+	virtual TSharedPtr<SGraphNode> CreateVisualWidget() override;
 
 	/** Allocates the default pins using the empty remap pins action. Usually called from CreateNode. */
 	virtual void AllocateDefaultPins() override final; // Final. Override AllocateDefaultPins(UCustomizableObjectNodeRemapPins* RemapPins) instead
diff --git a/Engine/Plugins/Experimental/Mutable/Source/CustomizableObjectEditor/Public/MuCOE/Nodes/CustomizableObjectNodeColorConstant.h b/Engine/Plugins/Experimental/Mutable/Source/CustomizableObjectEditor/Public/MuCOE/Nodes/CustomizableObjectNodeColorConstant.h
--- a/Engine/Plugins/Experimental/Mutable/Source/CustomizableObjectEditor/Public/MuCOE/Nodes/CustomizableObjectNodeColorConstant.h
+++ b/Engine/Plugins/Experimental/Mutable/Source/CustomizableObjectEditor/Public/MuCOE/Nodes/CustomizableObjectNodeColorConstant.h
@@ -3,7 +3,7 @@
 #pragma once
 
 #include "MuCOE/Nodes/CustomizableObjectNode.h"
-#include "SGraphNode.h"
+#include "MuCOE/Nodes/SCustomizableObjectNode.h"
 
 #include "CustomizableObjectNodeColorConstant.generated.h"
 
@@ -29,14 +29,14 @@ enum class ColorChannel
 };
 
 // Class create color inputs in the NodeColorConstant
-class SGraphNodeColorConstant : public SGraphNode
+class SGraphNodeColorConstant : public SCustomizableObjectNode
 {
 public:
 
 	SLATE_BEGIN_ARGS(SGraphNodeColorConstant) {}
 	SLATE_END_ARGS();
 
-	SGraphNodeColorConstant() : SGraphNode() {};
+	SGraphNodeColorConstant() : SCustomizableObjectNode() {}
 
 	// Builds the SGraphNodeFloatConstant when needed
 	void Construct(const FArguments& InArgs, UEdGraphNode* InGraphNode);
@@ -68,7 +68,7 @@ private:
 private:
 
 	// Pointer to the NodeFloatconstant that owns this SGraphNode
-	UCustomizableObjectNodeColorConstant* NodeColorConstant;
+	UCustomizableObjectNodeColorConstant* NodeColorConstant = nullptr;
 
 	// Style for the SpinBox
 	FSpinBoxStyle WidgetStyle;
diff --git a/Engine/Plugins/Experimental/Mutable/Source/CustomizableObjectEditor/Public/MuCOE/Nodes/CustomizableObjectNodeCopyMaterial.h b/Engine/Plugins/Experimental/Mutable/Source/CustomizableObjectEditor/Public/MuCOE/Nodes/CustomizableObjectNodeCopyMaterial.h
--- a/Engine/Plugins/Experimental/Mutable/Source/CustomizableObjectEditor/Public/MuCOE/Nodes/CustomizableObjectNodeCopyMaterial.h
+++ b/Engine/Plugins/Experimental/Mutable/Source/CustomizableObjectEditor/Public/MuCOE/Nodes/CustomizableObjectNodeCopyMaterial.h
@@ -45,11 +45,9 @@ public:
 	virtual FText GetParameterLayerName(EMaterialParameterType Type, int32 ParameterIndex) const override;
 	virtual bool HasParameter(const FGuid& ParameterId) const override;
 	virtual const UEdGraphPin* GetParameterPin(EMaterialParameterType Type, int32 ParameterIndex) const override;
+	virtual UEdGraphPin* GetParameterPin(const FGuid& ParameterId) const override;
 	virtual bool IsImageMutableMode(int32 ImageIndex) const override;
 	virtual bool IsImageMutableMode(const UEdGraphPin& Pin) const override;
-	virtual void UpdateImagePinMode(const FGuid ParameterId) override;
-	virtual void UpdateImagePinMode(const UEdGraphPin& Pin) override;
-	virtual void UpdateAllImagesPinMode() override;
 	virtual UTexture2D* GetImageReferenceTexture(int32 ImageIndex) const override;
 	virtual UTexture2D* GetImageValue(int32 ImageIndex) const override;
 	virtual int32 GetImageUVLayout(int32 ImageIndex) const override;
diff --git a/Engine/Plugins/Experimental/Mutable/Source/CustomizableObjectEditor/Public/MuCOE/Nodes/CustomizableObjectNodeFloatConstant.h b/Engine/Plugins/Experimental/Mutable/Source/CustomizableObjectEditor/Public/MuCOE/Nodes/CustomizableObjectNodeFloatConstant.h
--- a/Engine/Plugins/Experimental/Mutable/Source/CustomizableObjectEditor/Public/MuCOE/Nodes/CustomizableObjectNodeFloatConstant.h
+++ b/Engine/Plugins/Experimental/Mutable/Source/CustomizableObjectEditor/Public/MuCOE/Nodes/CustomizableObjectNodeFloatConstant.h
@@ -3,7 +3,7 @@
 #pragma once
 
 #include "MuCOE/Nodes/CustomizableObjectNode.h"
-#include "SGraphNode.h"
+#include "MuCOE/Nodes/SCustomizableObjectNode.h"
 
 #include "CustomizableObjectNodeFloatConstant.generated.h"
 
@@ -18,7 +18,7 @@ struct FPropertyChangedEvent;
 struct FSlateBrush;
 
 // Class create an input float in the NodeFloatConstant
-class SGraphNodeFloatConstant : public SGraphNode
+class SGraphNodeFloatConstant : public SCustomizableObjectNode
 {
 public:
 
@@ -26,7 +26,7 @@ public:
 
 	SLATE_END_ARGS();
 
-	SGraphNodeFloatConstant() : SGraphNode() {};
+	SGraphNodeFloatConstant() : SCustomizableObjectNode() {};
 
 	// Builds the SGraphNodeFloatConstant when needed
 	void Construct(const FArguments& InArgs, UEdGraphNode* InGraphNode);
diff --git a/Engine/Plugins/Experimental/Mutable/Source/CustomizableObjectEditor/Public/MuCOE/Nodes/CustomizableObjectNodeMaterial.h b/Engine/Plugins/Experimental/Mutable/Source/CustomizableObjectEditor/Public/MuCOE/Nodes/CustomizableObjectNodeMaterial.h
--- a/Engine/Plugins/Experimental/Mutable/Source/CustomizableObjectEditor/Public/MuCOE/Nodes/CustomizableObjectNodeMaterial.h
+++ b/Engine/Plugins/Experimental/Mutable/Source/CustomizableObjectEditor/Public/MuCOE/Nodes/CustomizableObjectNodeMaterial.h
@@ -72,12 +72,9 @@ enum class ENodePinMode
 UENUM()
 enum class EPinMode
 {
-	/** Use the defined node pin mode. Does not override it. */
-	Default,
-	/** Override the node pin mode. Set it to Mutable mode. */
-	Mutable,
-	/** Override the node pin mode. Set it to Pass-through mode. */
-	Passthrough
+	Default UMETA(DisplayName = "Node Defined", ToolTip = "Use node's \"Default Texture Parameter Mode\"."),
+	Mutable UMETA(ToolTip = "The Material Texture Parameters goes through Mutable."),
+	Passthrough UMETA(ToolTip = "The Material Texture Parameters is not modified by Mutable.")
 };
 
 
@@ -87,7 +84,7 @@ enum class EUVLayoutMode
 {
 	/** Does not override the UV Index specified in the Material. */
 	FromMaterial,
-	/* Texture should not be transformed by any layout. Theses textures will not be reduced automatically for LODs. */
+	/* Texture should not be transformed by any layout. These textures will not be reduced automatically for LODs. */
 	Ignore,
 	/** User specified UV Index. */
 	Index
@@ -119,7 +116,6 @@ public:
 
 	// UCustomizableObjectNode interface
 	virtual void BackwardsCompatibleFixup() override;
-	virtual void PostBackwardsCompatibleFixup() override;
 	virtual void AllocateDefaultPins(UCustomizableObjectNodeRemapPins* RemapPins) override;
 	virtual bool CanPinBeHidden(const UEdGraphPin& Pin) const override;
 	virtual bool HasPinViewer() const override;
@@ -127,11 +123,11 @@ public:
 	virtual bool ProvidesCustomPinRelevancyTest() const override { return true; }
 	virtual bool IsPinRelevant(const UEdGraphPin* Pin) const override;
 	virtual bool CustomRemovePin(UEdGraphPin& Pin) override;
-	virtual void ReconstructNode(UCustomizableObjectNodeRemapPins* RemapPinsMode) override;
 	virtual void BreakExistingConnectionsPostConnection(UEdGraphPin* InputPin, UEdGraphPin* OutputPin) override;
 	virtual bool IsNodeOutDatedAndNeedsRefresh() override;
 	virtual FString GetRefreshMessage() const override;
 	virtual TSharedPtr<IDetailsView> CustomizePinDetails(const UEdGraphPin& Pin) const override;
+	void ReconstructNode(UCustomizableObjectNodeRemapPins* RemapPinsMode);
 
 	
 	// UCustomizableObjectNodeMaterialBase interface
@@ -147,12 +143,10 @@ public:
 	virtual int32 GetParameterLayerIndex(EMaterialParameterType Type, int32 ParameterIndex) const override;
 	virtual FText GetParameterLayerName(EMaterialParameterType Type, int32 ParameterIndex) const override;
 	virtual bool HasParameter(const FGuid& ParameterId) const override;
-	virtual const UEdGraphPin* GetParameterPin(EMaterialParameterType Type, int32 ParameterIndex) const override;
+	virtual UEdGraphPin* GetParameterPin(EMaterialParameterType Type, int32 ParameterIndex) const override;
+	virtual UEdGraphPin* GetParameterPin(const FGuid& ParameterId) const override;
 	virtual bool IsImageMutableMode(int32 ImageIndex) const override;
 	virtual bool IsImageMutableMode(const UEdGraphPin& Pin) const override;
-	virtual void UpdateImagePinMode(const FGuid ParameterId) override;
-	virtual void UpdateImagePinMode(const UEdGraphPin& Pin) override;
-	virtual void UpdateAllImagesPinMode() override;
 	virtual UTexture2D* GetImageReferenceTexture(int32 ImageIndex) const override;
 	virtual UTexture2D* GetImageValue(int32 ImageIndex) const override;
 	virtual int32 GetImageUVLayout(int32 ImageIndex) const override;
@@ -168,13 +162,16 @@ public:
 	void SetComponentName(const FName& Name);
 
 private:
+	/** Set the Pin Mode of a Texture Parameter Pin. */
+	void SetImagePinMode(UEdGraphPin& Pin, EPinMode PinMode) const;
+	
 	/** Delegate called when a Texture Parameter Pin Mode changes. */
 	FPostImagePinModeChangedDelegate PostImagePinModeChangedDelegate;
 
 	UPROPERTY(EditAnywhere, Category=CustomizableObject)
 	TObjectPtr<UMaterialInterface> Material = nullptr;
 
-	UPROPERTY(EditAnywhere, Category=CustomizableObject, Meta = (ToolTip = "Set all Mateiral Texture Parameters to the specified mode. Each Texture Parameter Pin can override this mode."))
+	UPROPERTY(EditAnywhere, Category=CustomizableObject, DisplayName = "Default Texture Parameter Mode", Meta = (ToolTip = "All Mateiral Texture Parameters set to \"Node Defined\" will use this mode."))
 	ENodePinMode TextureParametersMode = ENodePinMode::Passthrough;
 
 	UPROPERTY(EditAnywhere, Category = CustomizableObject)
@@ -196,14 +193,10 @@ private:
 
 	static const TArray<EMaterialParameterType> ParameterTypes;
 
-	/** Relates a Parameter id (key) to a Pin (value). Only used to improve performance. */
+	/** Relates a Parameter id (key) to a Pin (value). Only used to improve performance.
+	  * If a deprecated pin and a non-deprecated pin have the same Parameter id, this the non-deprecated one prevails. */
 	UPROPERTY()
 	TMap<FGuid, FEdGraphPinReference> PinsParameter;
-
-	/** Relates an Image pin (key) to its Image Pin Mode (value). 
-		Represents the real mode of the pin. Required due to some node configurations can force the mode independently of what the user had previously selected */
-	UPROPERTY()
-	TMap<FGuid, EPinMode> PinsImagePinMode;
 	
 	/** Create the pin data of the given parameter type. */
 	UCustomizableObjectNodeMaterialPinDataParameter* CreatePinData(EMaterialParameterType Type, int32 ParameterIndex);
@@ -224,11 +217,14 @@ private:
 	/** Returns the texture coordinate of the given Material Expression. Returns -1 if not found. */
 	static int32 GetExpressionTextureCoordinate(UMaterial* Material, const FGuid &ImageId);
 
-	/** Converts node NodePinMode (does not include Default mode) to PinMode. */
-	static EPinMode NodePinModeToImagePinMode(ENodePinMode NodePinMode);
+	/** Return the Pin Category given the node NodePinMode. */
+	static FName NodePinModeToImagePinMode(ENodePinMode NodePinMode);
 
-	/** Returns the Image Pin Mode the pin should be at. It does not update its mode, to update it call UpdateImagePinMode. */
-	EPinMode GetImagePinMode(const UEdGraphPin& Pin) const;
+	/** Return the Pin Category given a PinMode. */
+	FName GetImagePinMode(EPinMode PinMode) const;
+
+	/** Return the Pin Category given a Pin. */
+	FName GetImagePinMode(const UEdGraphPin& Pin) const;
 
 	/** Get the UV Layout Index defined in the Material. */
 	int32 GetImageUVLayoutFromMaterial(int32 ImageIndex) const;
@@ -260,6 +256,8 @@ class CUSTOMIZABLEOBJECTEDITOR_API UCustomizableObjectNodeMaterialPinDataImage :
 public:
 	// UObject interface
 	virtual void PostEditChangeProperty(FPropertyChangedEvent& PropertyChangedEvent) override;
+	virtual bool CanEditChange(const FProperty* InProperty) const override;
+
 private:
 	virtual void PostLoad() override;
 	
@@ -281,7 +279,7 @@ public:
 	
 private:
 	/** Image pin mode. If is not default, overrides the defined node behaviour. */
-	UPROPERTY(EditAnywhere, Category = NoCategory)
+	UPROPERTY(EditAnywhere, Category = NoCategory, DisplayName = "Texture Parameter Mode")
 	EPinMode PinMode = EPinMode::Default;
 
 public:
diff --git a/Engine/Plugins/Experimental/Mutable/Source/CustomizableObjectEditor/Public/MuCOE/Nodes/CustomizableObjectNodeMaterialBase.h b/Engine/Plugins/Experimental/Mutable/Source/CustomizableObjectEditor/Public/MuCOE/Nodes/CustomizableObjectNodeMaterialBase.h
--- a/Engine/Plugins/Experimental/Mutable/Source/CustomizableObjectEditor/Public/MuCOE/Nodes/CustomizableObjectNodeMaterialBase.h
+++ b/Engine/Plugins/Experimental/Mutable/Source/CustomizableObjectEditor/Public/MuCOE/Nodes/CustomizableObjectNodeMaterialBase.h
@@ -76,12 +76,19 @@ public:
 	/** Returns true if the Material contains the given Material Parameter. */
 	virtual bool HasParameter(const FGuid& ParameterId) const PURE_VIRTUAL(UCustomizableObjectNodeMaterial::HasParameter, return {}; );
 	
-	/** Get the Vector pin for the given Material Vector Parameter.
-	 * Not all parameters have pin.
+	/** Get the pin for the given Material Parameter.
+	 * Not all parameters have pins.
 	 *
-	 * @param ParameterIndex Has to be valid.
+	 * @param ParameterIndex Material Parameter id.
 	 * @return Can return nullptr. */
 	virtual const UEdGraphPin* GetParameterPin(EMaterialParameterType Type, int32 ParameterIndex) const PURE_VIRTUAL(UCustomizableObjectNodeMaterial::GetMaterialNode, return {}; );
+
+	/** Get the pin for the given Material Parameter.
+	 * Not all parameters have pins.
+	 *
+	 * @param ParameterId Material Parameter id.
+	 * @return Can return nullptr. */
+	virtual UEdGraphPin* GetParameterPin(const FGuid& ParameterId) const PURE_VIRTUAL(UCustomizableObjectNodeMaterial::GetParameterPin, return {}; );
 	
 	// --------------------
 	// IMAGES PARAMETERS
@@ -94,15 +101,6 @@ public:
 
 	/** Given an Image pin, returns true if the Material Texture Parameter goes through Mutable. */
 	virtual bool IsImageMutableMode(const UEdGraphPin& Pin) const PURE_VIRTUAL(UCustomizableObjectNodeMaterial::IsImageMutableMode, return {}; );
-
-	/** Update a Material Texture Parameter Mode. */
-	virtual void UpdateImagePinMode(const FGuid ParameterId) PURE_VIRTUAL(UCustomizableObjectNodeMaterial::UpdateImagePinMode, );
-	
-	/** Update a Material Texture Parameter Mode. */
-	virtual void UpdateImagePinMode(const UEdGraphPin& Pin) PURE_VIRTUAL(UCustomizableObjectNodeMaterial::UpdateImagePinMode, );
-
-	/** Update all Material Texture Parameters Mode. */
-	virtual void UpdateAllImagesPinMode() PURE_VIRTUAL(UCustomizableObjectNodeMaterial::UpdateAllImagesPinMode, );
 	
 	/** Returns the reference texture assigned to a Material Texture Parameter.
 	 *
diff --git a/Engine/Plugins/Experimental/Mutable/Source/CustomizableObjectEditor/Public/MuCOE/Nodes/CustomizableObjectNodeSkeletalMesh.h b/Engine/Plugins/Experimental/Mutable/Source/CustomizableObjectEditor/Public/MuCOE/Nodes/CustomizableObjectNodeSkeletalMesh.h
--- a/Engine/Plugins/Experimental/Mutable/Source/CustomizableObjectEditor/Public/MuCOE/Nodes/CustomizableObjectNodeSkeletalMesh.h
+++ b/Engine/Plugins/Experimental/Mutable/Source/CustomizableObjectEditor/Public/MuCOE/Nodes/CustomizableObjectNodeSkeletalMesh.h
@@ -5,7 +5,7 @@
 #include "GameplayTagContainer.h"
 #include "MuCOE/Nodes/CustomizableObjectNodeMesh.h"
 #include "MuCOE/CustomizableObjectEditor_Deprecated.h"
-#include "SGraphNode.h"
+#include "MuCOE/Nodes/SCustomizableObjectNode.h"
 #include "MuCOE/RemapPins/CustomizableObjectNodeRemapPinsByNameDefaultPin.h"
 
 #include "CustomizableObjectNodeSkeletalMesh.generated.h"
@@ -33,14 +33,14 @@ struct FSlateBrush;
 
 
 // Class to render the Skeletal Mesh thumbnail of a CustomizableObjectNodeSkeletalMesh
-class SGraphNodeSkeletalMesh : public SGraphNode
+class SGraphNodeSkeletalMesh : public SCustomizableObjectNode
 {
 public:
 
 	SLATE_BEGIN_ARGS(SGraphNodeSkeletalMesh) {}
 	SLATE_END_ARGS();
 
-	SGraphNodeSkeletalMesh() : SGraphNode() {};
+	SGraphNodeSkeletalMesh() : SCustomizableObjectNode() {};
 
 	// Builds the SGraphNodeSkeletalMesh when needed
 	void Construct(const FArguments& InArgs, UEdGraphNode* InGraphNode);
diff --git a/Engine/Plugins/Experimental/Mutable/Source/CustomizableObjectEditor/Public/MuCOE/Nodes/CustomizableObjectNodeStaticMesh.h b/Engine/Plugins/Experimental/Mutable/Source/CustomizableObjectEditor/Public/MuCOE/Nodes/CustomizableObjectNodeStaticMesh.h
--- a/Engine/Plugins/Experimental/Mutable/Source/CustomizableObjectEditor/Public/MuCOE/Nodes/CustomizableObjectNodeStaticMesh.h
+++ b/Engine/Plugins/Experimental/Mutable/Source/CustomizableObjectEditor/Public/MuCOE/Nodes/CustomizableObjectNodeStaticMesh.h
@@ -3,7 +3,7 @@
 #pragma once
 
 #include "MuCOE/Nodes/CustomizableObjectNodeMesh.h"
-#include "SGraphNode.h"
+#include "MuCOE/Nodes/SCustomizableObjectNode.h"
 
 #include "CustomizableObjectNodeStaticMesh.generated.h"
 
@@ -26,14 +26,14 @@ struct FSlateBrush;
 
 
 // Class to render the Static Mesh thumbnail of a CustomizableObjectNodeStaticMesh
-class SGraphNodeStaticMesh : public SGraphNode
+class SGraphNodeStaticMesh : public SCustomizableObjectNode
 {
 public:
 
 	SLATE_BEGIN_ARGS(SGraphNodeStaticMesh) {}
 	SLATE_END_ARGS();
 
-	SGraphNodeStaticMesh() : SGraphNode() {};
+	SGraphNodeStaticMesh() : SCustomizableObjectNode() {};
 
 	// Builds the SGraphNodeStaticMesh when needed
 	void Construct(const FArguments& InArgs, UEdGraphNode* InGraphNode);
@@ -58,7 +58,7 @@ public:
 	TSharedPtr<class ISinglePropertyView> StaticMeshSelector;
 
 	// Pointer to the NodeStaticMesh that owns this SGraphNode
-	class UCustomizableObjectNodeStaticMesh* NodeStaticMesh;
+	class UCustomizableObjectNodeStaticMesh* NodeStaticMesh = nullptr;
 
 private:
 
diff --git a/Engine/Plugins/Experimental/Mutable/Source/CustomizableObjectEditor/Public/MuCOE/Nodes/CustomizableObjectNodeTable.h b/Engine/Plugins/Experimental/Mutable/Source/CustomizableObjectEditor/Public/MuCOE/Nodes/CustomizableObjectNodeTable.h
--- a/Engine/Plugins/Experimental/Mutable/Source/CustomizableObjectEditor/Public/MuCOE/Nodes/CustomizableObjectNodeTable.h
+++ b/Engine/Plugins/Experimental/Mutable/Source/CustomizableObjectEditor/Public/MuCOE/Nodes/CustomizableObjectNodeTable.h
@@ -235,7 +235,6 @@ public:
 	// UObject interface
 	virtual void PostLoad() override;
 	virtual void PreEditChange(FProperty* PropertyAboutToChange) override;
-	virtual void PostEditChangeChainProperty(FPropertyChangedChainEvent& PropertyChangedEvent) override;
 	virtual void PostEditChangeProperty(FPropertyChangedEvent& PropertyChangedEvent) override;
 	
 	// EdGraphNode interface
@@ -246,7 +245,6 @@ public:
 	virtual bool GetCanRenameNode() const override { return true; }
 	
 	// UCustomizableObjectNode interface
-	virtual void PostBackwardsCompatibleFixup() override;
 	virtual void BackwardsCompatibleFixup() override;
 	void AllocateDefaultPins(UCustomizableObjectNodeRemapPins* RemapPins) override;
 	bool IsNodeOutDatedAndNeedsRefresh() override;
@@ -281,9 +279,6 @@ public:
 	// Get the anim blueprint and anim slot columns related to a mesh
 	void GetAnimationColumns(const FGuid& ColumnId, FString& AnimBPColumnName, FString& AnimSlotColumnName, FString& AnimTagColumnName) const;
 
-	/** Callback called when the Table property or its contents has changed. */
-	void OnTableChanged();
-	
 	template<class T> 
 	T* GetColumnDefaultAssetByType(FString ColumnName) const
 	{
diff --git a/Engine/Plugins/Experimental/Mutable/Source/CustomizableObjectEditor/Public/MuCOE/Nodes/CustomizableObjectNodeTexture.h b/Engine/Plugins/Experimental/Mutable/Source/CustomizableObjectEditor/Public/MuCOE/Nodes/CustomizableObjectNodeTexture.h
--- a/Engine/Plugins/Experimental/Mutable/Source/CustomizableObjectEditor/Public/MuCOE/Nodes/CustomizableObjectNodeTexture.h
+++ b/Engine/Plugins/Experimental/Mutable/Source/CustomizableObjectEditor/Public/MuCOE/Nodes/CustomizableObjectNodeTexture.h
@@ -3,7 +3,7 @@
 #pragma once
 
 #include "MuCOE/Nodes/CustomizableObjectNode.h"
-#include "SGraphNode.h"
+#include "MuCOE/Nodes/SCustomizableObjectNode.h"
 
 #include "CustomizableObjectNodeTexture.generated.h"
 
@@ -18,14 +18,14 @@ class UTexture2D;
 struct FGeometry;
 
 // Class to render the Texture thumbnail of a CustomizableObjectNodeTexture
-class SGraphNodeTexture : public SGraphNode
+class SGraphNodeTexture : public SCustomizableObjectNode
 {
 public:
 
 	SLATE_BEGIN_ARGS(SGraphNodeTexture) {}
 	SLATE_END_ARGS();
 
-	SGraphNodeTexture() : SGraphNode() {};
+	SGraphNodeTexture() : SCustomizableObjectNode() {}
 
 	// Builds the SGraphNodeTexture when needed
 	void Construct(const FArguments& InArgs, UEdGraphNode* InGraphNode);
@@ -48,7 +48,7 @@ public:
 public:
 
 	// Pointer to the NodeTexture that owns this SGraphNode
-	UCustomizableObjectNodeTextureBase* NodeTexture;
+	UCustomizableObjectNodeTextureBase* NodeTexture = nullptr;
 
 	// Single property that only draws the combo box widget of the Texture
 	TSharedPtr<class ISinglePropertyView> TextureSelector;
