[mutable] Better enforcement of coding standards and naming convention in the code generator.
Added stub nodes for component switch and variation nodes.
#rnx

diff --git a/Engine/Plugins/Experimental/Mutable/Source/CustomizableObjectEditor/Private/MuCOE/GenerateMutableSource/GenerateMutableSource.cpp b/Engine/Plugins/Experimental/Mutable/Source/CustomizableObjectEditor/Private/MuCOE/GenerateMutableSource/GenerateMutableSource.cpp
--- a/Engine/Plugins/Experimental/Mutable/Source/CustomizableObjectEditor/Private/MuCOE/GenerateMutableSource/GenerateMutableSource.cpp
+++ b/Engine/Plugins/Experimental/Mutable/Source/CustomizableObjectEditor/Private/MuCOE/GenerateMutableSource/GenerateMutableSource.cpp
@@ -318,7 +318,7 @@ void FMutableGraphGenerationContext::GenerateClippingCOInternalTags()
 
 					for (int32 k = 0; k < It2->Value.Num(); ++k)
 					{
-						It2->Value[k]->AddTag(TagName);
+						It2->Value[k]->Tags.Add(TagName);
 					}
 				}
 			}
@@ -352,7 +352,7 @@ void FMutableGraphGenerationContext::GenerateSharedSurfacesUniqueIds()
 			}
 
 			FSharedSurface& CurrentSharedSurface = SharedSurfaces[SurfaceIndex];
-			CurrentSharedSurface.NodeSurfaceNew->SetSharedSurfaceId(UniqueId);
+			CurrentSharedSurface.NodeSurfaceNew->SharedSurfaceId = UniqueId;
 			VisitedSurfaces[SurfaceIndex] = true;
 
 			for (int32 AuxSurfaceIndex = SurfaceIndex; AuxSurfaceIndex < NumSurfaces && !CurrentSharedSurface.bMakeUnique; ++AuxSurfaceIndex)
@@ -367,7 +367,7 @@ void FMutableGraphGenerationContext::GenerateSharedSurfacesUniqueIds()
 					continue;
 				}
 
-				SharedSurfaces[AuxSurfaceIndex].NodeSurfaceNew->SetSharedSurfaceId(UniqueId);
+				SharedSurfaces[AuxSurfaceIndex].NodeSurfaceNew->SharedSurfaceId = UniqueId;
 				VisitedSurfaces[AuxSurfaceIndex] = true;
 			}
 
@@ -2140,7 +2140,7 @@ void FMutableComponentInfo::AccumulateBonesToRemovePerLOD(const FComponentSettin
 
 #undef LOCTEXT_NAMESPACE
 
-FMutableGraphGenerationContext::FSharedSurface::FSharedSurface(uint8 InLOD, const mu::NodeSurfaceNewPtr& InNodeSurfaceNew)
+FMutableGraphGenerationContext::FSharedSurface::FSharedSurface(uint8 InLOD, const mu::Ptr<mu::NodeSurfaceNew>& InNodeSurfaceNew)
 {
 	LOD = InLOD;
 	NodeSurfaceNew = InNodeSurfaceNew;
diff --git a/Engine/Plugins/Experimental/Mutable/Source/CustomizableObjectEditor/Private/MuCOE/GenerateMutableSource/GenerateMutableSource.h b/Engine/Plugins/Experimental/Mutable/Source/CustomizableObjectEditor/Private/MuCOE/GenerateMutableSource/GenerateMutableSource.h
--- a/Engine/Plugins/Experimental/Mutable/Source/CustomizableObjectEditor/Private/MuCOE/GenerateMutableSource/GenerateMutableSource.h
+++ b/Engine/Plugins/Experimental/Mutable/Source/CustomizableObjectEditor/Private/MuCOE/GenerateMutableSource/GenerateMutableSource.h
@@ -332,9 +332,9 @@ struct FGroupNodeIdsTempData
 
 struct FGroupProjectorImageInfo
 {
-	mu::NodeImagePtr ImageNode;
-	mu::NodeImagePtr ImageResizeNode;
-	mu::NodeSurfaceNewPtr SurfNode;
+	mu::Ptr<mu::NodeImage> ImageNode;
+	mu::Ptr<mu::NodeImage> ImageResizeNode;
+	mu::Ptr<mu::NodeSurfaceNew> SurfNode;
 	UCustomizableObjectNodeMaterialBase* TypedNodeMat;
 	FString TextureName;
 	FString RealTextureName;
@@ -344,7 +344,7 @@ struct FGroupProjectorImageInfo
 	int32 UVLayout = 0;
 
 	FGroupProjectorImageInfo(mu::NodeImagePtr InImageNode, const FString& InTextureName, const FString& InRealTextureName, UCustomizableObjectNodeMaterialBase* InTypedNodeMat,
-		float InAlternateProjectionResolutionFactor, const FString& InAlternateResStateName, mu::NodeSurfaceNewPtr InSurfNode, int32 InUVLayout)
+		float InAlternateProjectionResolutionFactor, const FString& InAlternateResStateName, mu::Ptr<mu::NodeSurfaceNew> InSurfNode, int32 InUVLayout)
 		: TypedNodeMat(InTypedNodeMat), TextureName(InTextureName), RealTextureName(InRealTextureName),
 		AlternateResStateName(InAlternateResStateName), AlternateProjectionResolutionFactor(InAlternateProjectionResolutionFactor), 
 		UVLayout(InUVLayout)
@@ -768,7 +768,7 @@ struct FMutableGraphGenerationContext
 
 	// Data structures used for clipping feature by assigning a Customizable Object
 	// Map with pairs (Unreal Mutable material node Guid, array with its corresponding Mutable surface node) to add tags from clipping nodes
-	TMap<class UCustomizableObjectNodeMaterial*, TArray<mu::NodeSurfaceNewPtr>> MapMaterialNodeToMutableSurfaceNodeArray;
+	TMap<class UCustomizableObjectNodeMaterial*, TArray<mu::Ptr<mu::NodeSurfaceNew>>> MapMaterialNodeToMutableSurfaceNodeArray;
 
 	// Map with pairs (Unreal Mutable clip mesh node, array with its corresponding Mutable mesh modifier nodes) to add tags
 	TMap<class UCustomizableObjectNodeMeshClipWithMesh*, TArray<mu::Ptr<mu::NodeModifierMeshClipWithMesh>>> MapClipMeshNodeToMutableClipMeshNodeArray;
@@ -871,12 +871,12 @@ struct FMutableGraphGenerationContext
 
 	struct FSharedSurface
 	{
-		FSharedSurface(uint8 InLOD, const mu::NodeSurfaceNewPtr& InNodeSurfaceNew);
+		FSharedSurface(uint8 InLOD, const mu::Ptr<mu::NodeSurfaceNew>& InNodeSurfaceNew);
 
 		bool operator==(const FSharedSurface& o) const;
 
 		uint8 LOD = 0;
-		mu::NodeSurfaceNewPtr NodeSurfaceNew;
+		mu::Ptr<mu::NodeSurfaceNew> NodeSurfaceNew;
 
 		bool bMakeUnique = false;
 		TArray<SIZE_T> NodeModifierIDs;
diff --git a/Engine/Plugins/Experimental/Mutable/Source/CustomizableObjectEditor/Private/MuCOE/GenerateMutableSource/GenerateMutableSourceColor.cpp b/Engine/Plugins/Experimental/Mutable/Source/CustomizableObjectEditor/Private/MuCOE/GenerateMutableSource/GenerateMutableSourceColor.cpp
--- a/Engine/Plugins/Experimental/Mutable/Source/CustomizableObjectEditor/Private/MuCOE/GenerateMutableSource/GenerateMutableSourceColor.cpp
+++ b/Engine/Plugins/Experimental/Mutable/Source/CustomizableObjectEditor/Private/MuCOE/GenerateMutableSource/GenerateMutableSourceColor.cpp
@@ -30,7 +30,7 @@
 #define LOCTEXT_NAMESPACE "CustomizableObjectEditor"
 
 
-mu::NodeColourPtr GenerateMutableSourceColor(const UEdGraphPin* Pin, FMutableGraphGenerationContext& GenerationContext)
+mu::Ptr<mu::NodeColour> GenerateMutableSourceColor(const UEdGraphPin* Pin, FMutableGraphGenerationContext& GenerationContext)
 {
 	check(Pin)
 	RETURN_ON_CYCLE(*Pin, GenerationContext)
@@ -52,7 +52,7 @@ mu::NodeColourPtr GenerateMutableSourceColor(const UEdGraphPin* Pin, FMutableGra
 		Node->SetRefreshNodeWarning();
 	}
 
-	mu::NodeColourPtr Result;
+	mu::Ptr<mu::NodeColour> Result;
 
 	if (const UCustomizableObjectNodeColorConstant* TypedNodeColorConst = Cast<UCustomizableObjectNodeColorConstant>(Node))
 	{
@@ -80,7 +80,7 @@ mu::NodeColourPtr GenerateMutableSourceColor(const UEdGraphPin* Pin, FMutableGra
 
 	else if (const UCustomizableObjectNodeColorSwitch* TypedNodeColorSwitch = Cast<UCustomizableObjectNodeColorSwitch>(Node))
 	{
-		Result = [&]()
+		Result = [&]() -> mu::Ptr<mu::NodeColourSwitch>
 		{
 			if (const int32 NumParameters = FollowInputPinArray(*TypedNodeColorSwitch->SwitchParameter()).Num();
 				NumParameters != 1)
@@ -90,7 +90,7 @@ mu::NodeColourPtr GenerateMutableSourceColor(const UEdGraphPin* Pin, FMutableGra
 					: LOCTEXT("InvalidEnumInSwitch", "Switch nodes must have a single enum with all the options inside. Please remove all the enums but one and refresh the switch node.");
 
 				GenerationContext.Compiler->CompilerLog(Message, Node);
-				return Result;
+				return nullptr;
 			}
 
 			const UEdGraphPin* EnumPin = FollowInputPin(*TypedNodeColorSwitch->SwitchParameter());
@@ -106,7 +106,7 @@ mu::NodeColourPtr GenerateMutableSourceColor(const UEdGraphPin* Pin, FMutableGra
 					GenerationContext.Compiler->CompilerLog(Message, Node);
 				}
 
-				return Result;
+				return nullptr;
 			}
 
 			if (SwitchParam->GetType() != mu::NodeScalarEnumParameter::GetStaticType())
@@ -114,7 +114,7 @@ mu::NodeColourPtr GenerateMutableSourceColor(const UEdGraphPin* Pin, FMutableGra
 				const FText Message = LOCTEXT("WrongSwitchParamType", "Switch parameter of incorrect type.");
 				GenerationContext.Compiler->CompilerLog(Message, Node);
 
-				return Result;
+				return nullptr;
 			}
 
 			const int32 NumSwitchOptions = TypedNodeColorSwitch->GetNumElements();
@@ -139,7 +139,7 @@ mu::NodeColourPtr GenerateMutableSourceColor(const UEdGraphPin* Pin, FMutableGra
 				}
 			}
 
-			return static_cast<mu::NodeColourPtr>(SwitchNode);
+			return SwitchNode;
 		}(); // invoke lambda;
 	}
 
@@ -174,14 +174,12 @@ mu::NodeColourPtr GenerateMutableSourceColor(const UEdGraphPin* Pin, FMutableGra
 
 		if (const UEdGraphPin* ConnectedPin = FollowInputPin(*TypedNodeColorArith->XPin()))
 		{
-			mu::NodeColourPtr XNode = GenerateMutableSourceColor(ConnectedPin, GenerationContext);
-			OpNode->A = XNode;
+			OpNode->A = GenerateMutableSourceColor(ConnectedPin, GenerationContext);
 		}
 
 		if (const UEdGraphPin* ConnectedPin = FollowInputPin(*TypedNodeColorArith->YPin()))
 		{
-			mu::NodeColourPtr YNode = GenerateMutableSourceColor(ConnectedPin, GenerationContext);
-			OpNode->B = YNode;
+			OpNode->B = GenerateMutableSourceColor(ConnectedPin, GenerationContext);
 		}
 
 		switch (TypedNodeColorArith->Operation)
@@ -241,7 +239,7 @@ mu::NodeColourPtr GenerateMutableSourceColor(const UEdGraphPin* Pin, FMutableGra
 
 		if (const UEdGraphPin* ConnectedPin = FollowInputPin(*TypedNodeColorVar->DefaultPin()))
 		{
-			mu::NodeColourPtr ChildNode = GenerateMutableSourceColor(ConnectedPin, GenerationContext);
+			mu::Ptr<mu::NodeColour> ChildNode = GenerateMutableSourceColor(ConnectedPin, GenerationContext);
 			if (ChildNode)
 			{
 				ColorNode->DefaultColour = ChildNode;
@@ -262,7 +260,7 @@ mu::NodeColourPtr GenerateMutableSourceColor(const UEdGraphPin* Pin, FMutableGra
 			ColorNode->Variations[VariationIndex].Tag = TypedNodeColorVar->GetVariation(VariationIndex).Tag;
 			if (const UEdGraphPin* ConnectedPin = FollowInputPin(*VariationPin))
 			{
-				mu::NodeColourPtr ChildNode = GenerateMutableSourceColor(ConnectedPin, GenerationContext);
+				mu::Ptr<mu::NodeColour> ChildNode = GenerateMutableSourceColor(ConnectedPin, GenerationContext);
 				ColorNode->Variations[VariationIndex].Colour = ChildNode;
 			}
 		}
diff --git a/Engine/Plugins/Experimental/Mutable/Source/CustomizableObjectEditor/Private/MuCOE/GenerateMutableSource/GenerateMutableSourceFloat.cpp b/Engine/Plugins/Experimental/Mutable/Source/CustomizableObjectEditor/Private/MuCOE/GenerateMutableSource/GenerateMutableSourceFloat.cpp
--- a/Engine/Plugins/Experimental/Mutable/Source/CustomizableObjectEditor/Private/MuCOE/GenerateMutableSource/GenerateMutableSourceFloat.cpp
+++ b/Engine/Plugins/Experimental/Mutable/Source/CustomizableObjectEditor/Private/MuCOE/GenerateMutableSource/GenerateMutableSourceFloat.cpp
@@ -114,7 +114,7 @@ mu::NodeScalarPtr GenerateMutableSourceFloat(const UEdGraphPin* Pin, FMutableGra
 	else if (const UCustomizableObjectNodeFloatSwitch* TypedNodeFloatSwitch = Cast<UCustomizableObjectNodeFloatSwitch>(Node))
 	{
 		// Using a lambda so control flow is easier to manage.
-		Result = [&]()
+		Result = [&]() -> mu::Ptr<mu::NodeScalar>
 		{
 			const UEdGraphPin* SwitchParameter = TypedNodeFloatSwitch->SwitchParameter();
 
@@ -127,7 +127,7 @@ mu::NodeScalarPtr GenerateMutableSourceFloat(const UEdGraphPin* Pin, FMutableGra
 					: LOCTEXT("InvalidEnumInSwitch", "Switch nodes must have a single enum with all the options inside. Please remove all the enums but one and refresh the switch node.");
 
 				GenerationContext.Compiler->CompilerLog(Message, Node);
-				return Result;
+				return nullptr;
 			}
 
 			const UEdGraphPin* EnumPin = FollowInputPin(*SwitchParameter);
@@ -143,7 +143,7 @@ mu::NodeScalarPtr GenerateMutableSourceFloat(const UEdGraphPin* Pin, FMutableGra
 					GenerationContext.Compiler->CompilerLog(Message, Node);
 				}
 
-				return Result;
+				return nullptr;
 			}
 
 			if (SwitchParam->GetType() != mu::NodeScalarEnumParameter::GetStaticType())
@@ -151,7 +151,7 @@ mu::NodeScalarPtr GenerateMutableSourceFloat(const UEdGraphPin* Pin, FMutableGra
 				const FText Message = LOCTEXT("WrongSwitchParamType", "Switch parameter of incorrect type.");
 				GenerationContext.Compiler->CompilerLog(Message, Node);
 
-				return Result;
+				return nullptr;
 			}
 
 			const int32 NumSwitchOptions = TypedNodeFloatSwitch->GetNumElements();
@@ -163,9 +163,9 @@ mu::NodeScalarPtr GenerateMutableSourceFloat(const UEdGraphPin* Pin, FMutableGra
 				GenerationContext.Compiler->CompilerLog(Message, Node);
 			}
 
-			mu::NodeScalarSwitchPtr SwitchNode = new mu::NodeScalarSwitch;
-			SwitchNode->SetParameter(SwitchParam);
-			SwitchNode->SetOptionCount(NumSwitchOptions);
+			mu::Ptr<mu::NodeScalarSwitch> SwitchNode = new mu::NodeScalarSwitch;
+			SwitchNode->Parameter = SwitchParam;
+			SwitchNode->Options.SetNum(NumSwitchOptions);
 
 			for (int SelectorIndex = 0; SelectorIndex < NumSwitchOptions; ++SelectorIndex)
 			{
@@ -173,13 +173,12 @@ mu::NodeScalarPtr GenerateMutableSourceFloat(const UEdGraphPin* Pin, FMutableGra
 				{
 					if (const UEdGraphPin* ConnectedPin = FollowInputPin(*FloatPin))
 					{
-						SwitchNode->SetOption(SelectorIndex, GenerateMutableSourceFloat(ConnectedPin, GenerationContext));
+						SwitchNode->Options[SelectorIndex] = GenerateMutableSourceFloat(ConnectedPin, GenerationContext);
 					}
 				}
 			}
 
-			Result = SwitchNode;
-			return Result;
+			return SwitchNode;
 		}(); // invoke lambda.
 	}
 
diff --git a/Engine/Plugins/Experimental/Mutable/Source/CustomizableObjectEditor/Private/MuCOE/GenerateMutableSource/GenerateMutableSourceSurface.cpp b/Engine/Plugins/Experimental/Mutable/Source/CustomizableObjectEditor/Private/MuCOE/GenerateMutableSource/GenerateMutableSourceSurface.cpp
--- a/Engine/Plugins/Experimental/Mutable/Source/CustomizableObjectEditor/Private/MuCOE/GenerateMutableSource/GenerateMutableSourceSurface.cpp
+++ b/Engine/Plugins/Experimental/Mutable/Source/CustomizableObjectEditor/Private/MuCOE/GenerateMutableSource/GenerateMutableSourceSurface.cpp
@@ -383,7 +383,7 @@ mu::Ptr<mu::NodeSurface> GenerateMutableSourceSurface(const UEdGraphPin * Pin, F
 		if (bGeneratingImplicitComponent)
 		{
 			MeshNode = GenerationContext.ComponentMeshOverride;
-			SurfNode->SetMesh(MeshNode);
+			SurfNode->Mesh = MeshNode;
 
 			if (const UEdGraphPin* ConnectedPin = FollowInputPin(*TypedNodeMat->GetMeshPin()))
 			{
@@ -411,8 +411,8 @@ mu::Ptr<mu::NodeSurface> GenerateMutableSourceSurface(const UEdGraphPin * Pin, F
 
 					MeshFormatNode->SetMessageContext(Node);
 
-					SurfNode->SetCustomID(SurfaceMetadataUniqueHash);
-					SurfNode->SetMesh(MeshFormatNode);
+					SurfNode->ExternalId = SurfaceMetadataUniqueHash;
+					SurfNode->Mesh = MeshFormatNode;
 				}
 				else
 				{
@@ -461,7 +461,7 @@ mu::Ptr<mu::NodeSurface> GenerateMutableSourceSurface(const UEdGraphPin * Pin, F
 		}
 
 		int32 NumImages = TypedNodeMat->GetNumParameters(EMaterialParameterType::Texture);
-		SurfNode->SetImageCount(NumImages);
+		SurfNode->Images.SetNum(NumImages);
 
 		if (!GenerationContext.Options.TargetPlatform || GenerationContext.Options.TargetPlatform->IsServerOnly())
 		{
@@ -494,13 +494,14 @@ mu::Ptr<mu::NodeSurface> GenerateMutableSourceSurface(const UEdGraphPin * Pin, F
 
 					// This is a connected pass-through texture that simply has to be passed to the core
 					mu::Ptr<mu::NodeImage> PassThroughImagePtr = GenerateMutableSourceImage(ConnectedPin, GenerationContext, 0);
-					SurfNode->SetImage(ImageIndex, PassThroughImagePtr);
+					SurfNode->Images[ImageIndex].Image = PassThroughImagePtr;
 
 					check(Props.ImagePropertiesIndex != INDEX_NONE);
 					const FString SurfNodeImageName = FString::Printf(TEXT("%d"), Props.ImagePropertiesIndex);
-					SurfNode->SetImageName(ImageIndex, SurfNodeImageName);
-					SurfNode->SetImageLayoutIndex(ImageIndex, -1);
-					SurfNode->SetImageAdditionalNames(ImageIndex, TypedNodeMat->GetMaterial()->GetName(), Props.TextureParameterName);
+					SurfNode->Images[ImageIndex].Name = SurfNodeImageName;
+					SurfNode->Images[ImageIndex].LayoutIndex = -1;
+					SurfNode->Images[ImageIndex].MaterialName = TypedNodeMat->GetMaterial()->GetName();
+					SurfNode->Images[ImageIndex].MaterialParameterName = Props.TextureParameterName;
 
 				}
 			}
@@ -864,7 +865,7 @@ mu::Ptr<mu::NodeSurface> GenerateMutableSourceSurface(const UEdGraphPin * Pin, F
 							ImageNode = LastImage;
 						}
 
-						SurfNode->SetImage(ImageIndex, ImageNode);
+						SurfNode->Images[ImageIndex].Image = ImageNode;
 
 						check(Props.ImagePropertiesIndex != INDEX_NONE);
 						const FString SurfNodeImageName = FString::Printf(TEXT("%d"), Props.ImagePropertiesIndex);
@@ -873,7 +874,7 @@ mu::Ptr<mu::NodeSurface> GenerateMutableSourceSurface(const UEdGraphPin * Pin, F
 						const int32 LayerIndex = TypedNodeMat->GetParameterLayerIndex(EMaterialParameterType::Texture, ImageIndex);
 						const FString LayerEncoding = LayerIndex != INDEX_NONE ? "-MutableLayerParam:" + FString::FromInt(LayerIndex) : "";
 						
-						SurfNode->SetImageName(ImageIndex, SurfNodeImageName + LayerEncoding);
+						SurfNode->Images[ImageIndex].Name = SurfNodeImageName + LayerEncoding;
 
 						// If we are generating an implicit component (with a passthrough mesh) we don't apply any layout.
 						int32 UVLayout = -1;
@@ -881,8 +882,9 @@ mu::Ptr<mu::NodeSurface> GenerateMutableSourceSurface(const UEdGraphPin * Pin, F
 						{
 							UVLayout = TypedNodeMat->GetImageUVLayout(ImageIndex);;
 						}
-						SurfNode->SetImageLayoutIndex(ImageIndex, UVLayout);
-						SurfNode->SetImageAdditionalNames(ImageIndex, TypedNodeMat->GetMaterial()->GetName(), ImageName);
+						SurfNode->Images[ImageIndex].LayoutIndex = UVLayout;
+						SurfNode->Images[ImageIndex].MaterialName = TypedNodeMat->GetMaterial()->GetName();
+						SurfNode->Images[ImageIndex].MaterialParameterName = ImageName;
 
 						if (bShareProjectionTexturesBetweenLODs && bIsGroupProjectorImage)
 						{
@@ -898,10 +900,10 @@ mu::Ptr<mu::NodeSurface> GenerateMutableSourceSurface(const UEdGraphPin * Pin, F
 				else
 				{
 					ensure(LOD > GenerationContext.FirstLODAvailable);
-					check(ProjectorInfo->SurfNode->GetImage(ImageIndex) == ProjectorInfo->ImageNode);
-					SurfNode->SetImage(ImageIndex, ProjectorInfo->ImageNode);
-					SurfNode->SetImageName(ImageIndex, ProjectorInfo->TextureName);
-					SurfNode->SetImageLayoutIndex(ImageIndex, ProjectorInfo->UVLayout);
+					check(ProjectorInfo->SurfNode->Images[ImageIndex].Image == ProjectorInfo->ImageNode);
+					SurfNode->Images[ImageIndex].Image = ProjectorInfo->ImageNode;
+					SurfNode->Images[ImageIndex].Name = ProjectorInfo->TextureName;
+					SurfNode->Images[ImageIndex].LayoutIndex = ProjectorInfo->UVLayout;
 
 					TextureNameToProjectionResFactor.Add(ProjectorInfo->RealTextureName, ProjectorInfo->AlternateProjectionResolutionFactor);
 					AlternateResStateName = ProjectorInfo->AlternateResStateName;
@@ -910,7 +912,7 @@ mu::Ptr<mu::NodeSurface> GenerateMutableSourceSurface(const UEdGraphPin * Pin, F
 		}
 
 		const int32 NumVectors = TypedNodeMat->GetNumParameters(EMaterialParameterType::Vector);
-		SurfNode->SetVectorCount(NumVectors);
+		SurfNode->Vectors.SetNum(NumVectors);
 		for (int32 VectorIndex = 0; VectorIndex < NumVectors; ++VectorIndex)
 		{
 			const UEdGraphPin* VectorPin = TypedNodeMat->GetParameterPin(EMaterialParameterType::Vector, VectorIndex);
@@ -923,7 +925,7 @@ mu::Ptr<mu::NodeSurface> GenerateMutableSourceSurface(const UEdGraphPin * Pin, F
 			{				
 				if (const UEdGraphPin* ConnectedPin = FollowInputPin(*VectorPin))
 				{
-					mu::NodeColourPtr ColorNode = GenerateMutableSourceColor(ConnectedPin, GenerationContext);
+					mu::Ptr<mu::NodeColour> ColorNode = GenerateMutableSourceColor(ConnectedPin, GenerationContext);
 
 					// Encoding material layer in mutable name
 					if (const int32 LayerIndex = TypedNodeMat->GetParameterLayerIndex(EMaterialParameterType::Vector, VectorIndex); LayerIndex != INDEX_NONE)
@@ -931,14 +933,14 @@ mu::Ptr<mu::NodeSurface> GenerateMutableSourceSurface(const UEdGraphPin * Pin, F
 						VectorName += "-MutableLayerParam:" + FString::FromInt(LayerIndex);
 					}
 
-					SurfNode->SetVector(VectorIndex, ColorNode);
-					SurfNode->SetVectorName(VectorIndex, VectorName);
+					SurfNode->Vectors[VectorIndex].Vector = ColorNode;
+					SurfNode->Vectors[VectorIndex].Name = VectorName;
 				}
 			}
 		}
 
 		const int32 NumScalar = TypedNodeMat->GetNumParameters(EMaterialParameterType::Scalar);
-		SurfNode->SetScalarCount(NumScalar);
+		SurfNode->Scalars.SetNum(NumScalar);
 		for (int32 ScalarIndex = 0; ScalarIndex < NumScalar; ++ScalarIndex)
 		{
 			const UEdGraphPin* ScalarPin = TypedNodeMat->GetParameterPin(EMaterialParameterType::Scalar, ScalarIndex);
@@ -959,8 +961,8 @@ mu::Ptr<mu::NodeSurface> GenerateMutableSourceSurface(const UEdGraphPin * Pin, F
 						ScalarName += "-MutableLayerParam:" + FString::FromInt(LayerIndex);
 					}
 
-					SurfNode->SetScalar(ScalarIndex, ScalarNode);
-					SurfNode->SetScalarName(ScalarIndex, ScalarName);
+					SurfNode->Scalars[ScalarIndex].Scalar = ScalarNode;
+					SurfNode->Scalars[ScalarIndex].Name = ScalarName;
 				}
 			}
 		}
@@ -968,7 +970,7 @@ mu::Ptr<mu::NodeSurface> GenerateMutableSourceSurface(const UEdGraphPin * Pin, F
 		// New method to pass the surface id as a scalar parameter
 		{
 			int32 MaterialIndex = NumScalar;
-			SurfNode->SetScalarCount(NumScalar + 1);
+			SurfNode->Scalars.SetNum(NumScalar + 1);
 
 			const UEdGraphPin* MaterialPin = TypedNodeMat->GetMaterialAssetPin();
 
@@ -982,8 +984,8 @@ mu::Ptr<mu::NodeSurface> GenerateMutableSourceSurface(const UEdGraphPin * Pin, F
 					GenerationContext.CurrentMaterialTableParameterId = MaterialName;
 					mu::NodeScalarPtr ScalarNode = GenerateMutableSourceFloat(ConnectedPin, GenerationContext);
 
-					SurfNode->SetScalar(MaterialIndex, ScalarNode);
-					SurfNode->SetScalarName(MaterialIndex, MaterialName);
+					SurfNode->Scalars[MaterialIndex].Scalar = ScalarNode;
+					SurfNode->Scalars[MaterialIndex].Name = MaterialName;
 				}
 			}
 			else
@@ -991,21 +993,21 @@ mu::Ptr<mu::NodeSurface> GenerateMutableSourceSurface(const UEdGraphPin * Pin, F
 				mu::NodeScalarConstantPtr ScalarNode = new mu::NodeScalarConstant();
 				ScalarNode->SetValue(ReferencedMaterialsIndex);
 
-				SurfNode->SetScalar(MaterialIndex, ScalarNode);
-				SurfNode->SetScalarName(MaterialIndex, MaterialName);
+				SurfNode->Scalars[MaterialIndex].Scalar = ScalarNode;
+				SurfNode->Scalars[MaterialIndex].Name = MaterialName;
 			}
 		}
 		
 
 		for (const FString& Tag : TypedNodeMat->GetTags())
 		{
-			SurfNode->AddTag(Tag);
+			SurfNode->Tags.Add(Tag);
 		}
 
-		TArray<mu::NodeSurfaceNewPtr>* ArraySurfaceNodePtr = GenerationContext.MapMaterialNodeToMutableSurfaceNodeArray.Find(TypedNodeMat->GetMaterialNode());
+		TArray<mu::Ptr<mu::NodeSurfaceNew>>* ArraySurfaceNodePtr = GenerationContext.MapMaterialNodeToMutableSurfaceNodeArray.Find(TypedNodeMat->GetMaterialNode());
 		if (ArraySurfaceNodePtr == nullptr)
 		{
-			TArray<mu::NodeSurfaceNewPtr> ArraySurfaceNode;
+			TArray<mu::Ptr<mu::NodeSurfaceNew>> ArraySurfaceNode;
 			ArraySurfaceNode.Add(SurfNode);
 			ArraySurfaceNodePtr = &GenerationContext.MapMaterialNodeToMutableSurfaceNodeArray.Add(TypedNodeMat->GetMaterialNode(), ArraySurfaceNode);
 		}
@@ -1017,36 +1019,18 @@ mu::Ptr<mu::NodeSurface> GenerateMutableSourceSurface(const UEdGraphPin * Pin, F
 		// If an alternate resolution for a particular state is present, clone the surface node, add the image resizing and inject the surface variation node
 		if (TextureNameToProjectionResFactor.Num() > 0 && !AlternateResStateName.IsEmpty())
 		{
-			mu::NodeSurfaceNewPtr SurfNode2 = new mu::NodeSurfaceNew;
+			mu::Ptr<mu::NodeSurfaceNew> SurfNode2 = new mu::NodeSurfaceNew;
 
-			SurfNode2->SetCustomID(SurfaceMetadataUniqueHash);
+			SurfNode2->ExternalId = SurfaceMetadataUniqueHash;
 
-			SurfNode2->SetMesh(SurfNode->GetMesh());
+			SurfNode2->Mesh = SurfNode->Mesh;
+			SurfNode2->Tags = SurfNode->Tags;
+			SurfNode2->Vectors = SurfNode->Vectors;
+			SurfNode2->Scalars = SurfNode->Scalars;
+			SurfNode2->Strings = SurfNode->Strings;
+			SurfNode2->Images = SurfNode->Images;
 
-			SurfNode2->SetVectorCount(SurfNode->GetVectorCount());
-
-			for (int32 VectorParamIndex = 0; VectorParamIndex < SurfNode->GetVectorCount(); ++VectorParamIndex)
-			{
-				SurfNode2->SetVector(VectorParamIndex, SurfNode->GetVector(VectorParamIndex));
-				SurfNode2->SetVectorName(VectorParamIndex, SurfNode->GetVectorName(VectorParamIndex));
-			}
-
-			SurfNode2->SetScalarCount(SurfNode->GetScalarCount());
-
-			for (int32 ScalarParamIndex = 0; ScalarParamIndex < SurfNode->GetScalarCount(); ++ScalarParamIndex)
-			{
-				SurfNode2->SetScalar(ScalarParamIndex, SurfNode->GetScalar(ScalarParamIndex));
-				SurfNode2->SetScalarName(ScalarParamIndex, SurfNode->GetScalarName(ScalarParamIndex));
-			}
-
-			for (const FString& Tag : TypedNodeMat->GetTags())
-			{
-				SurfNode2->AddTag(Tag);
-			}
-
-			SurfNode2->SetImageCount(SurfNode->GetImageCount());
-
-			for (int32 ImageIndex = 0; ImageIndex < SurfNode->GetImageCount(); ++ImageIndex)
+			for (int32 ImageIndex = 0; ImageIndex < SurfNode2->Images.Num(); ++ImageIndex)
 			{
 				const FString ImageName = TypedNodeMat->GetParameterName(EMaterialParameterType::Texture, ImageIndex).ToString();
 				if (float* ResolutionFactor = TextureNameToProjectionResFactor.Find(ImageName))
@@ -1059,9 +1043,9 @@ mu::Ptr<mu::NodeSurface> GenerateMutableSourceSurface(const UEdGraphPin * Pin, F
 						mu::NodeImageResizePtr NodeImageResize = new mu::NodeImageResize;
 						NodeImageResize->SetRelative(true);
 						NodeImageResize->SetSize(*ResolutionFactor, *ResolutionFactor);
-						NodeImageResize->SetBase(SurfNode->GetImage(ImageIndex));
+						NodeImageResize->SetBase(SurfNode2->Images[ImageIndex].Image);
 
-						SurfNode2->SetImage(ImageIndex, NodeImageResize);
+						SurfNode2->Images[ImageIndex].Image = NodeImageResize;
 
 						if (ProjectorInfo)
 						{
@@ -1074,27 +1058,20 @@ mu::Ptr<mu::NodeSurface> GenerateMutableSourceSurface(const UEdGraphPin * Pin, F
 					{
 						ensure(LOD > GenerationContext.FirstLODAvailable);
 						check(ProjectorInfo->bIsAlternateResolutionResized);
-						SurfNode2->SetImage(ImageIndex, ProjectorInfo->ImageResizeNode);
+						SurfNode2->Images[ImageIndex].Image = ProjectorInfo->ImageResizeNode;
 					}
 				}
-				else
-				{
-					SurfNode2->SetImage(ImageIndex, SurfNode->GetImage(ImageIndex));
-				}
-
-				SurfNode2->SetImageName(ImageIndex, SurfNode->GetImageName(ImageIndex));
-				SurfNode2->SetImageLayoutIndex(ImageIndex, SurfNode->GetImageLayoutIndex(ImageIndex));
 			}
 
 			ArraySurfaceNodePtr->AddUnique(SurfNode2);
 
-			mu::NodeSurfaceVariationPtr SurfaceVariation = new mu::NodeSurfaceVariation;
-			SurfaceVariation->SetVariationType(mu::NodeSurfaceVariation::VariationType::State);
-			SurfaceVariation->SetVariationCount(1);
-			SurfaceVariation->SetVariationTag(0, AlternateResStateName);
+			mu::Ptr<mu::NodeSurfaceVariation> SurfaceVariation = new mu::NodeSurfaceVariation;
+			SurfaceVariation->Type = mu::NodeSurfaceVariation::VariationType::State;
+			SurfaceVariation->Variations.SetNum(1);
+			SurfaceVariation->Variations[0].Tag = AlternateResStateName;
 
-			SurfaceVariation->AddDefaultSurface(&*SurfNode);
-			SurfaceVariation->AddVariationSurface(0, &*SurfNode2);
+			SurfaceVariation->DefaultSurfaces.Add(SurfNode);
+			SurfaceVariation->Variations[0].Surfaces.Add(SurfNode2);
 
 			Result = SurfaceVariation;
 		}
@@ -1102,7 +1079,7 @@ mu::Ptr<mu::NodeSurface> GenerateMutableSourceSurface(const UEdGraphPin * Pin, F
 
 	else if (UCustomizableObjectNodeExtendMaterial* TypedNodeExt = Cast<UCustomizableObjectNodeExtendMaterial>(Node))
 	{
-		mu::NodeSurfaceEditPtr SurfNode = new mu::NodeSurfaceEdit();
+		mu::Ptr<mu::NodeSurfaceEdit> SurfNode = new mu::NodeSurfaceEdit();
 		Result = SurfNode;
 
 		[&] // Using a lambda so control flow is easier to manage.
@@ -1122,7 +1099,7 @@ mu::Ptr<mu::NodeSurface> GenerateMutableSourceSurface(const UEdGraphPin * Pin, F
 
 			// Parent, probably generated, will be retrieved from the cache
 			mu::NodeSurfacePtr ParentNode = GenerateMutableSourceSurface(ParentMaterialNode->OutputPin(), GenerationContext);
-			SurfNode->SetParent(ParentNode.get());
+			SurfNode->Parent = ParentNode;
 
 			mu::NodeMeshPtr AddMeshNode;
 			FMutableGraphMeshGenerationData MeshData;
@@ -1170,12 +1147,12 @@ mu::Ptr<mu::NodeSurface> GenerateMutableSourceSurface(const UEdGraphPin * Pin, F
 				mu::NodePatchMeshPtr MeshPatch = new mu::NodePatchMesh();
 				MeshPatch->SetAdd(MeshFormat.get());
 
-				SurfNode->SetMesh(MeshPatch.get());
+				SurfNode->Mesh = MeshPatch;
 				MeshPatch->SetMessageContext(Node);
 			}
 			
 			const int32 NumImages = ParentMaterialNode->GetNumParameters(EMaterialParameterType::Texture);
-			SurfNode->SetImageCount(NumImages);
+			SurfNode->Textures.SetNum(NumImages);
 			for (int32 ImageIndex = 0; ImageIndex < NumImages; ++ImageIndex)
 			{
 				mu::NodeImagePtr ImageNode;
@@ -1188,7 +1165,7 @@ mu::Ptr<mu::NodeSurface> GenerateMutableSourceSurface(const UEdGraphPin * Pin, F
 					if (ProjectorInfo)
 					{
 						ensure(LOD > GenerationContext.FirstLODAvailable);
-						check(ProjectorInfo->SurfNode->GetImage(ImageIndex) == ProjectorInfo->ImageNode);
+						check(ProjectorInfo->SurfNode->Images[ImageIndex].Image == ProjectorInfo->ImageNode);
 						ImageNode = ProjectorInfo->ImageNode;
 
 						//TextureNameToProjectionResFactor.Add(ProjectorInfo->RealTextureName, ProjectorInfo->AlternateProjectionResolutionFactor);
@@ -1229,7 +1206,7 @@ mu::Ptr<mu::NodeSurface> GenerateMutableSourceSurface(const UEdGraphPin * Pin, F
 
 				if (ImageNode)
 				{
-					SurfNode->SetImage(ImageIndex, ImageNode);
+					SurfNode->Textures[ImageIndex].Extend = ImageNode;
 				}
 
 			}
@@ -1237,14 +1214,14 @@ mu::Ptr<mu::NodeSurface> GenerateMutableSourceSurface(const UEdGraphPin * Pin, F
 		
 			for (const FString& Tag : TypedNodeExt->Tags)
 			{
-				SurfNode->AddTag(Tag);
+				SurfNode->Tags.Add(Tag);
 			}
 		}();
 	}
 
 	else if (const UCustomizableObjectNodeRemoveMesh* TypedNodeRem = Cast<UCustomizableObjectNodeRemoveMesh>(Node))
 	{
-		mu::NodeSurfaceEditPtr SurfNode = new mu::NodeSurfaceEdit();
+		mu::Ptr<mu::NodeSurfaceEdit> SurfNode = new mu::NodeSurfaceEdit();
 		Result = SurfNode;
 
 		UCustomizableObjectNodeMaterialBase* ParentMaterialNode = TypedNodeRem->GetParentMaterialNode();
@@ -1256,7 +1233,7 @@ mu::Ptr<mu::NodeSurface> GenerateMutableSourceSurface(const UEdGraphPin * Pin, F
 		{
 			// Parent, probably generated, will be retrieved from the cache
 			mu::NodeSurfacePtr ParentNode = GenerateMutableSourceSurface(ParentMaterialNode->OutputPin(), GenerationContext);
-			SurfNode->SetParent(ParentNode.get());
+			SurfNode->Parent = ParentNode;
 
 			mu::NodeMeshPtr RemoveMeshNode;
 
@@ -1270,7 +1247,7 @@ mu::Ptr<mu::NodeSurface> GenerateMutableSourceSurface(const UEdGraphPin * Pin, F
 			{
 				mu::NodePatchMeshPtr MeshPatch = new mu::NodePatchMesh();
 				MeshPatch->SetRemove(RemoveMeshNode.get());
-				SurfNode->SetMesh(MeshPatch.get());
+				SurfNode->Mesh = MeshPatch;
 				MeshPatch->SetMessageContext(Node);
 			}
 
@@ -1280,7 +1257,7 @@ mu::Ptr<mu::NodeSurface> GenerateMutableSourceSurface(const UEdGraphPin * Pin, F
 
 	else if (const UCustomizableObjectNodeRemoveMeshBlocks* TypedNodeRemBlocks = Cast<UCustomizableObjectNodeRemoveMeshBlocks>(Node))
 	{
-		mu::NodeSurfaceEditPtr SurfNode = new mu::NodeSurfaceEdit();
+		mu::Ptr<mu::NodeSurfaceEdit> SurfNode = new mu::NodeSurfaceEdit();
 		Result = SurfNode;
 
 		UCustomizableObjectNodeMaterialBase* ParentMaterialNode = TypedNodeRemBlocks->GetParentMaterialNode();
@@ -1292,7 +1269,7 @@ mu::Ptr<mu::NodeSurface> GenerateMutableSourceSurface(const UEdGraphPin * Pin, F
 		{
 			// Parent, probably generated, will be retrieved from the cache
 			mu::NodeSurfacePtr ParentNode = GenerateMutableSourceSurface(ParentMaterialNode->OutputPin(), GenerationContext);
-			SurfNode->SetParent(ParentNode.get());
+			SurfNode->Parent = ParentNode.get();
 
 			const UEdGraphPin* BaseSourcePin = FindMeshBaseSource(*ParentMaterialNode->OutputPin(), false);
 
@@ -1352,7 +1329,7 @@ mu::Ptr<mu::NodeSurface> GenerateMutableSourceSurface(const UEdGraphPin * Pin, F
 
 				mu::NodePatchMeshPtr MeshPatch = new mu::NodePatchMesh();
 				MeshPatch->SetRemove(MeshFrag.get());
-				SurfNode->SetMesh(MeshPatch.get());
+				SurfNode->Mesh = MeshPatch.get();
 				MeshPatch->SetMessageContext(Node);
 			}
 
@@ -1362,7 +1339,7 @@ mu::Ptr<mu::NodeSurface> GenerateMutableSourceSurface(const UEdGraphPin * Pin, F
 
 	else if (UCustomizableObjectNodeEditMaterial* TypedNodeEdit = Cast<UCustomizableObjectNodeEditMaterial>(Node))
 	{
-		mu::NodeSurfaceEditPtr SurfNode = new mu::NodeSurfaceEdit();
+		mu::Ptr<mu::NodeSurfaceEdit> SurfNode = new mu::NodeSurfaceEdit();
 		Result = SurfNode;
 
 		UCustomizableObjectNodeMaterialBase* ParentMaterialNode = TypedNodeEdit->GetParentMaterialNode();
@@ -1374,11 +1351,10 @@ mu::Ptr<mu::NodeSurface> GenerateMutableSourceSurface(const UEdGraphPin * Pin, F
 		{
 			// Parent, probably generated, will be retrieved from the cache
 			mu::NodeSurfacePtr ParentNode = GenerateMutableSourceSurface(ParentMaterialNode->OutputPin(), GenerationContext);
-			SurfNode->SetParent(ParentNode.get());
-
+			SurfNode->Parent = ParentNode;
 
 			const int32 NumImages = ParentMaterialNode->GetNumParameters(EMaterialParameterType::Texture);
-			SurfNode->SetImageCount(NumImages);
+			SurfNode->Textures.SetNum(NumImages);
 			for (int32 ImageIndex = 0; ImageIndex < NumImages; ++ImageIndex)
 			{
 				const FGuid ImageId = ParentMaterialNode->GetParameterId(EMaterialParameterType::Texture, ImageIndex);
@@ -1389,7 +1365,7 @@ mu::Ptr<mu::NodeSurface> GenerateMutableSourceSurface(const UEdGraphPin * Pin, F
 					
 					const UEdGraphPin* ConnectedImagePin = FollowInputPin(*TypedNodeEdit->GetUsedImagePin(ImageId));
 					
-					mu::NodePatchImagePtr ImagePatchNode = new mu::NodePatchImage;
+					mu::Ptr<mu::NodePatchImage> ImagePatchNode = new mu::NodePatchImage;
 					ImagePatchNode->SetMessageContext(Node);
 
 					// \todo: expose these two options?
@@ -1457,7 +1433,7 @@ mu::Ptr<mu::NodeSurface> GenerateMutableSourceSurface(const UEdGraphPin * Pin, F
 						ImagePatchNode->SetBlock(BlockIndex, FoundBlockIndices[BlockIndex]);
 					}
 
-					SurfNode->SetPatch(ImageIndex, ImagePatchNode);
+					SurfNode->Textures[ImageIndex].Patch = ImagePatchNode;
 				}
 
 				AddModifierToSharedSurface(GenerationContext, ParentMaterialNode, *TypedNodeEdit);
@@ -1467,7 +1443,7 @@ mu::Ptr<mu::NodeSurface> GenerateMutableSourceSurface(const UEdGraphPin * Pin, F
 
 	else if (const UCustomizableObjectNodeMorphMaterial* TypedNodeMorph = Cast<UCustomizableObjectNodeMorphMaterial>(Node))
 	{
-		mu::NodeSurfaceEditPtr SurfNode = new mu::NodeSurfaceEdit();
+		mu::Ptr<mu::NodeSurfaceEdit> SurfNode = new mu::NodeSurfaceEdit();
 		Result = SurfNode;
 
 		UCustomizableObjectNodeMaterialBase* ParentMaterialNode = TypedNodeMorph->GetParentMaterialNode();
@@ -1479,7 +1455,7 @@ mu::Ptr<mu::NodeSurface> GenerateMutableSourceSurface(const UEdGraphPin * Pin, F
 		{
 			// Parent, probably generated, will be retrieved from the cache
 			mu::NodeSurfacePtr ParentNode = GenerateMutableSourceSurface(ParentMaterialNode->OutputPin(), GenerationContext);
-			SurfNode->SetParent(ParentNode.get());
+			SurfNode->Parent = ParentNode;
 
 			const UEdGraphPin* BaseSourcePin = FindMeshBaseSource(*ParentMaterialNode->OutputPin(), false);
 			if (!BaseSourcePin)
@@ -1496,7 +1472,7 @@ mu::Ptr<mu::NodeSurface> GenerateMutableSourceSurface(const UEdGraphPin * Pin, F
 				MorphedSourceMeshNode->SetMessageContext(Node);					
 				MorphedSourceMeshNode->SetValue(MorphedSourceMesh);
 				
-				SurfNode->SetMorph(MorphedSourceMeshNode);
+				SurfNode->Morph = MorphedSourceMeshNode;
 				
 				if (const UEdGraphPin* ConnectedPin = FollowInputPin(*TypedNodeMorph->FactorPin()))
 				{
@@ -1536,7 +1512,7 @@ mu::Ptr<mu::NodeSurface> GenerateMutableSourceSurface(const UEdGraphPin * Pin, F
 					if (validStaticFactor)
 					{
 						mu::NodeScalarPtr FactorNode = GenerateMutableSourceFloat(ConnectedPin, GenerationContext);
-						SurfNode->SetFactor(FactorNode);
+						SurfNode->MorphFactor = FactorNode;
 					}
 				}
 			}
@@ -1547,7 +1523,7 @@ mu::Ptr<mu::NodeSurface> GenerateMutableSourceSurface(const UEdGraphPin * Pin, F
 
 	else if (const UCustomizableObjectNodeMaterialVariation* TypedNodeVar = Cast<UCustomizableObjectNodeMaterialVariation>(Node))
 	{
-		mu::NodeSurfaceVariationPtr SurfNode = new mu::NodeSurfaceVariation();
+		mu::Ptr<mu::NodeSurfaceVariation> SurfNode = new mu::NodeSurfaceVariation();
 		Result = SurfNode;
 
 		mu::NodeSurfaceVariation::VariationType muType = mu::NodeSurfaceVariation::VariationType::Tag;
@@ -1559,7 +1535,7 @@ mu::Ptr<mu::NodeSurface> GenerateMutableSourceSurface(const UEdGraphPin * Pin, F
 			check(false);
 			break;
 		}
-		SurfNode->SetVariationType(muType);
+		SurfNode->Type = muType;
 
 		for (const UEdGraphPin* ConnectedPin : FollowInputPinArray(*TypedNodeVar->DefaultPin()))
 		{
@@ -1567,7 +1543,7 @@ mu::Ptr<mu::NodeSurface> GenerateMutableSourceSurface(const UEdGraphPin * Pin, F
 			mu::NodeSurfacePtr ChildNode = GenerateMutableSourceSurface(ConnectedPin, GenerationContext);
 			if (ChildNode)
 			{
-				SurfNode->AddDefaultSurface(ChildNode.get());
+				SurfNode->DefaultSurfaces.Add(ChildNode);
 			}
 			else
 			{
@@ -1576,21 +1552,21 @@ mu::Ptr<mu::NodeSurface> GenerateMutableSourceSurface(const UEdGraphPin * Pin, F
 		}
 
 		const int32 NumVariations = TypedNodeVar->GetNumVariations();
-		SurfNode->SetVariationCount(NumVariations);
+		SurfNode->Variations.SetNum(NumVariations);
 		for (int VariationIndex = 0; VariationIndex < NumVariations; ++VariationIndex)
 		{
 			mu::NodeSurfacePtr VariationSurfaceNode;
 
 			if (UEdGraphPin* VariationPin = TypedNodeVar->VariationPin(VariationIndex))
 			{
-				SurfNode->SetVariationTag(VariationIndex, TypedNodeVar->GetVariation(VariationIndex).Tag);
+				SurfNode->Variations[VariationIndex].Tag = TypedNodeVar->GetVariation(VariationIndex).Tag;
 				for (const UEdGraphPin* ConnectedPin : FollowInputPinArray(*VariationPin))
 				{
 					// Is it a modifier?
 					mu::NodeSurfacePtr ChildNode = GenerateMutableSourceSurface(ConnectedPin, GenerationContext);
 					if (ChildNode)
 					{
-						SurfNode->AddVariationSurface(VariationIndex, ChildNode.get());
+						SurfNode->Variations[VariationIndex].Surfaces.Add( ChildNode );
 					}
 					else
 					{
@@ -1644,8 +1620,8 @@ mu::Ptr<mu::NodeSurface> GenerateMutableSourceSurface(const UEdGraphPin * Pin, F
 				}
 
 				mu::Ptr<mu::NodeSurfaceSwitch> SwitchNode = new mu::NodeSurfaceSwitch;
-				SwitchNode->SetParameter(SwitchParam);
-				SwitchNode->SetOptionCount(NumSwitchOptions);
+				SwitchNode->Parameter = SwitchParam;
+				SwitchNode->Options.SetNum(NumSwitchOptions);
 
 				for (int32 SelectorIndex = 0; SelectorIndex < NumSwitchOptions; ++SelectorIndex)
 				{
@@ -1654,7 +1630,7 @@ mu::Ptr<mu::NodeSurface> GenerateMutableSourceSurface(const UEdGraphPin * Pin, F
 						mu::NodeSurfacePtr ChildNode = GenerateMutableSourceSurface(ConnectedPin, GenerationContext);
 						if (ChildNode)
 						{
-							SwitchNode->SetOption(SelectorIndex, ChildNode.get());
+							SwitchNode->Options[SelectorIndex] = ChildNode;
 						}
 						else
 						{
diff --git a/Engine/Plugins/Experimental/Mutable/Source/CustomizableObjectEditor/Private/MuCOE/SMutableGraphViewer.cpp b/Engine/Plugins/Experimental/Mutable/Source/CustomizableObjectEditor/Private/MuCOE/SMutableGraphViewer.cpp
--- a/Engine/Plugins/Experimental/Mutable/Source/CustomizableObjectEditor/Private/MuCOE/SMutableGraphViewer.cpp
+++ b/Engine/Plugins/Experimental/Mutable/Source/CustomizableObjectEditor/Private/MuCOE/SMutableGraphViewer.cpp
@@ -38,6 +38,7 @@
 #include "MuT/NodeImageTable.h"
 #include "MuT/NodeObjectGroup.h"
 #include "MuT/NodeObjectNew.h"
+#include "MuT/NodeSurfaceNew.h"
 #include "MuT/NodeSurfaceEdit.h"
 #include "MuT/NodeSurfaceSwitch.h"
 #include "MuT/NodeSurfaceVariation.h"
@@ -58,10 +59,6 @@
 #include "MuT/NodeObjectNew.h"
 
 #include "MuT/NodeObjectGroupPrivate.h"
-#include "MuT/NodeSurfaceNewPrivate.h"
-#include "MuT/NodeSurfaceEditPrivate.h"
-#include "MuT/NodeSurfaceSwitchPrivate.h"
-#include "MuT/NodeSurfaceVariationPrivate.h"
 #include "MuT/NodeModifierPrivate.h"
 #include "MuT/NodeImageFormatPrivate.h"
 #include "MuT/NodeMeshFormatPrivate.h"
@@ -90,7 +87,6 @@
 #include "MuT/NodeImageInvertPrivate.h"
 #include "MuT/NodeImageSwizzlePrivate.h"
 #include "MuT/NodeImageMultiLayerPrivate.h"
-#include "MuT/NodeScalarSwitchPrivate.h"
 
 class FExtender;
 class FReferenceCollector;
@@ -348,74 +344,75 @@ void SMutableGraphViewer::GetChildrenForInfo(TSharedPtr<FMutableGraphTreeElement
 	else if (ParentNode->GetType() == mu::NodeSurfaceNew::GetStaticType())
 	{
 		mu::NodeSurfaceNew* SurfaceNew = StaticCast<mu::NodeSurfaceNew*>(ParentNode);
-		mu::NodeSurfaceNew::Private* Private = SurfaceNew->GetPrivate();
-		AddChildFunc(Private->Mesh.get(), TEXT("MESH"));
+		AddChildFunc(SurfaceNew->Mesh.get(), TEXT("MESH"));
 
-		for (int32 l = 0; l < Private->m_images.Num(); ++l)
+		for (int32 l = 0; l < SurfaceNew->Images.Num(); ++l)
 		{
-			AddChildFunc(Private->m_images[l].m_pImage.get(), FString::Printf(TEXT("IMAGE [%s]"), *Private->m_images[l].m_name));
+			AddChildFunc(SurfaceNew->Images[l].Image.get(), FString::Printf(TEXT("IMAGE [%s]"), *SurfaceNew->Images[l].Name));
 		}
 
-		for (int32 l = 0; l < Private->m_vectors.Num(); ++l)
+		for (int32 l = 0; l < SurfaceNew->Vectors.Num(); ++l)
 		{
-			AddChildFunc(Private->m_vectors[l].m_pVector.get(), FString::Printf(TEXT("VECTOR [%s]"), *Private->m_vectors[l].m_name));
+			AddChildFunc(SurfaceNew->Vectors[l].Vector.get(), FString::Printf(TEXT("VECTOR [%s]"), *SurfaceNew->Vectors[l].Name));
 		}
 
-		for (int32 l = 0; l < Private->m_scalars.Num(); ++l)
+		for (int32 l = 0; l < SurfaceNew->Scalars.Num(); ++l)
 		{
-			AddChildFunc(Private->m_scalars[l].m_pScalar.get(), FString::Printf(TEXT("SCALAR [%s]"), *Private->m_scalars[l].m_name));
+			AddChildFunc(SurfaceNew->Scalars[l].Scalar.get(), FString::Printf(TEXT("SCALAR [%s]"), *SurfaceNew->Scalars[l].Name));
+		}
+
+		for (int32 l = 0; l < SurfaceNew->Strings.Num(); ++l)
+		{
+			AddChildFunc(SurfaceNew->Strings[l].String.get(), FString::Printf(TEXT("STRING [%s]"), *SurfaceNew->Strings[l].Name));
 		}
 	}
 
 	else if (ParentNode->GetType() == mu::NodeSurfaceEdit::GetStaticType())
 	{
 		mu::NodeSurfaceEdit* SurfaceEdit = StaticCast<mu::NodeSurfaceEdit*>(ParentNode);
-		mu::NodeSurfaceEdit::Private* Private = SurfaceEdit->GetPrivate();
-		AddChildFunc(Private->m_pMesh.get(), TEXT("MESH"));
-		AddChildFunc(Private->m_pMorph.get(), TEXT("MORPH"));
-		AddChildFunc(Private->m_pFactor.get(), TEXT("MORPH_FACTOR"));
+		AddChildFunc(SurfaceEdit->Mesh.get(), TEXT("MESH"));
+		AddChildFunc(SurfaceEdit->Morph.get(), TEXT("MORPH"));
+		AddChildFunc(SurfaceEdit->MorphFactor.get(), TEXT("MORPH_FACTOR"));
 
-		for (int32 l = 0; l < Private->m_textures.Num(); ++l)
+		for (int32 l = 0; l < SurfaceEdit->Textures.Num(); ++l)
 		{
-			AddChildFunc(Private->m_textures[l].m_pExtend.get(), FString::Printf(TEXT("EXTEND [%d]"), l));
-			AddChildFunc(Private->m_textures[l].m_pPatch.get(), FString::Printf(TEXT("PATCH [%d]"), l));
+			AddChildFunc(SurfaceEdit->Textures[l].Extend.get(), FString::Printf(TEXT("EXTEND [%d]"), l));
+			AddChildFunc(SurfaceEdit->Textures[l].Patch.get(), FString::Printf(TEXT("PATCH [%d]"), l));
 		}
 	}
 
 	else if (ParentNode->GetType() == mu::NodeSurfaceSwitch::GetStaticType())
 	{
 		mu::NodeSurfaceSwitch* SurfaceSwitch = StaticCast<mu::NodeSurfaceSwitch*>(ParentNode);
-		mu::NodeSurfaceSwitch::Private* Private = SurfaceSwitch->GetPrivate();
-		AddChildFunc(Private->Parameter.get(), TEXT("PARAM"));
-		for (int32 l = 0; l < Private->Options.Num(); ++l)
+		AddChildFunc(SurfaceSwitch->Parameter.get(), TEXT("PARAM"));
+		for (int32 l = 0; l < SurfaceSwitch->Options.Num(); ++l)
 		{
-			AddChildFunc(Private->Options[l].get(), FString::Printf(TEXT("OPTION [%d]"), l));
+			AddChildFunc(SurfaceSwitch->Options[l].get(), FString::Printf(TEXT("OPTION [%d]"), l));
 		}
 	}
 
 	else if (ParentNode->GetType() == mu::NodeSurfaceVariation::GetStaticType())
 	{
 		mu::NodeSurfaceVariation* SurfaceVar = StaticCast<mu::NodeSurfaceVariation*>(ParentNode);
-		mu::NodeSurfaceVariation::Private* Private = SurfaceVar->GetPrivate();
-		for (int32 l = 0; l < Private->m_defaultSurfaces.Num(); ++l)
+		for (int32 l = 0; l < SurfaceVar->DefaultSurfaces.Num(); ++l)
 		{
-			AddChildFunc(Private->m_defaultSurfaces[l].get(), FString::Printf(TEXT("DEF SURF [%d]"), l));
+			AddChildFunc(SurfaceVar->DefaultSurfaces[l].get(), FString::Printf(TEXT("DEF SURF [%d]"), l));
 		}
-		for (int32 l = 0; l < Private->m_defaultModifiers.Num(); ++l)
+		for (int32 l = 0; l < SurfaceVar->DefaultModifiers.Num(); ++l)
 		{
-			AddChildFunc(Private->m_defaultModifiers[l].get(), FString::Printf(TEXT("DEF MOD [%d]"), l));
+			AddChildFunc(SurfaceVar->DefaultModifiers[l].get(), FString::Printf(TEXT("DEF MOD [%d]"), l));
 		}
 
-		for (int32 v = 0; v < Private->m_variations.Num(); ++v)
+		for (int32 v = 0; v < SurfaceVar->Variations.Num(); ++v)
 		{
-			const mu::NodeSurfaceVariation::Private::FVariation Var = Private->m_variations[v];
-			for (int32 l = 0; l < Var.m_surfaces.Num(); ++l)
+			const mu::NodeSurfaceVariation::FVariation Var = SurfaceVar->Variations[v];
+			for (int32 l = 0; l < Var.Surfaces.Num(); ++l)
 			{
-				AddChildFunc(Var.m_surfaces[l].get(), FString::Printf(TEXT("VAR [%s] SURF [%d]"), *Var.m_tag, l));
+				AddChildFunc(Var.Surfaces[l].get(), FString::Printf(TEXT("VAR [%s] SURF [%d]"), *Var.Tag, l));
 			}
-			for (int32 l = 0; l < Var.m_modifiers.Num(); ++l)
+			for (int32 l = 0; l < Var.Modifiers.Num(); ++l)
 			{
-				AddChildFunc(Var.m_modifiers[l].get(), FString::Printf(TEXT("VAR [%s] MOD [%d]"), *Var.m_tag, l));
+				AddChildFunc(Var.Modifiers[l].get(), FString::Printf(TEXT("VAR [%s] MOD [%d]"), *Var.Tag, l));
 			}
 		}
 	}
@@ -755,11 +752,10 @@ void SMutableGraphViewer::GetChildrenForInfo(TSharedPtr<FMutableGraphTreeElement
 	else if (ParentNode->GetType() == mu::NodeScalarSwitch::GetStaticType())
 	{
 		mu::NodeScalarSwitch* ScalarSwitchVar = StaticCast<mu::NodeScalarSwitch*>(ParentNode);
-		mu::NodeScalarSwitch::Private* Private = ScalarSwitchVar->GetPrivate();
-		AddChildFunc(Private->m_pParameter.get(), TEXT("PARAM"));
-		for (int32 OptionIndex = 0; OptionIndex < Private->m_options.Num(); ++OptionIndex)
+		AddChildFunc(ScalarSwitchVar->Parameter.get(), TEXT("PARAM"));
+		for (int32 OptionIndex = 0; OptionIndex < ScalarSwitchVar->Options.Num(); ++OptionIndex)
 		{
-			AddChildFunc(Private->m_options[OptionIndex].get(), FString::Printf(TEXT("OPTION [%d]"), OptionIndex));
+			AddChildFunc(ScalarSwitchVar->Options[OptionIndex].get(), FString::Printf(TEXT("OPTION [%d]"), OptionIndex));
 		}
 	}
 	
diff --git a/Engine/Plugins/Experimental/Mutable/Source/MutableTools/Internal/MuT/NodeScalarSwitchPrivate.h b/Engine/Plugins/Experimental/Mutable/Source/MutableTools/Internal/MuT/NodeScalarSwitchPrivate.h
deleted file mode 100644
diff --git a/Engine/Plugins/Experimental/Mutable/Source/MutableTools/Internal/MuT/NodeSurfaceEditPrivate.h b/Engine/Plugins/Experimental/Mutable/Source/MutableTools/Internal/MuT/NodeSurfaceEditPrivate.h
deleted file mode 100644
diff --git a/Engine/Plugins/Experimental/Mutable/Source/MutableTools/Internal/MuT/NodeSurfaceNewPrivate.h b/Engine/Plugins/Experimental/Mutable/Source/MutableTools/Internal/MuT/NodeSurfaceNewPrivate.h
deleted file mode 100644
diff --git a/Engine/Plugins/Experimental/Mutable/Source/MutableTools/Internal/MuT/NodeSurfaceSwitchPrivate.h b/Engine/Plugins/Experimental/Mutable/Source/MutableTools/Internal/MuT/NodeSurfaceSwitchPrivate.h
deleted file mode 100644
diff --git a/Engine/Plugins/Experimental/Mutable/Source/MutableTools/Internal/MuT/NodeSurfaceVariationPrivate.h b/Engine/Plugins/Experimental/Mutable/Source/MutableTools/Internal/MuT/NodeSurfaceVariationPrivate.h
deleted file mode 100644
diff --git a/Engine/Plugins/Experimental/Mutable/Source/MutableTools/Private/MuT/CodeGenerator.cpp b/Engine/Plugins/Experimental/Mutable/Source/MutableTools/Private/MuT/CodeGenerator.cpp
--- a/Engine/Plugins/Experimental/Mutable/Source/MutableTools/Private/MuT/CodeGenerator.cpp
+++ b/Engine/Plugins/Experimental/Mutable/Source/MutableTools/Private/MuT/CodeGenerator.cpp
@@ -81,10 +81,10 @@
 #include "MuT/NodeScalar.h"
 #include "MuT/NodeScalarConstant.h"
 #include "MuT/NodeSurface.h"
-#include "MuT/NodeSurfaceEditPrivate.h"
-#include "MuT/NodeSurfaceNewPrivate.h"
-#include "MuT/NodeSurfaceVariationPrivate.h"
-#include "MuT/NodeSurfaceSwitchPrivate.h"
+#include "MuT/NodeSurfaceEdit.h"
+#include "MuT/NodeSurfaceNew.h"
+#include "MuT/NodeSurfaceVariation.h"
+#include "MuT/NodeSurfaceSwitch.h"
 #include "MuT/TablePrivate.h"
 #include "Trace/Detail/Channel.h"
 
@@ -97,13 +97,13 @@ namespace mu
 	//---------------------------------------------------------------------------------------------
 	CodeGenerator::CodeGenerator(CompilerOptions::Private* options)
 	{
-		m_compilerOptions = options;
+		CompilerOptions = options;
 
 		// Create the message log
-		m_pErrorLog = new ErrorLog;
+		ErrorLog = new mu::ErrorLog;
 
 		// Add the parent at the top of the hierarchy
-		m_currentParents.Add(FParentKey());
+		CurrentParents.Add(FParentKey());
 	}
 
 	//---------------------------------------------------------------------------------------------
@@ -112,11 +112,11 @@ namespace mu
 		MUTABLE_CPUPROFILER_SCOPE(Generate);
 
 		// First pass
-		m_firstPass.Generate(m_pErrorLog, pNode.get(), m_compilerOptions->bIgnoreStates, this);
+		FirstPass.Generate(ErrorLog, pNode.get(), CompilerOptions->bIgnoreStates, this);
 
 		// Second pass
-		SecondPassGenerator SecondPass(&m_firstPass, m_compilerOptions);
-		bool bSuccess = SecondPass.Generate(m_pErrorLog, pNode.get());
+		SecondPassGenerator SecondPass(&FirstPass, CompilerOptions);
+		bool bSuccess = SecondPass.Generate(ErrorLog, pNode.get());
 		if (!bSuccess)
 		{
 			return;
@@ -127,7 +127,7 @@ namespace mu
 			MUTABLE_CPUPROFILER_SCOPE(MainPass);
 
             int32 CurrentStateIndex = 0;
-			for (const TPair<FObjectState, const Node*>& s : m_firstPass.m_states)
+			for (const TPair<FObjectState, const Node*>& s : FirstPass.States)
 			{
 				MUTABLE_CPUPROFILER_SCOPE(MainPassState);
 
@@ -135,7 +135,7 @@ namespace mu
 				Options.State = CurrentStateIndex;
 
 				Ptr<ASTOp> stateRoot = Generate_Generic(pNode, Options);
-				m_states.Emplace(s.Key, stateRoot);
+				States.Emplace(s.Key, stateRoot);
 
 				++CurrentStateIndex;
 			}
@@ -378,7 +378,7 @@ namespace mu
 		op->type = OP_TYPE::NU_PARAMETER;
 		op->parameter = param;
 
-		m_firstPass.ParameterNodes.Add(InNode, op);
+		FirstPass.ParameterNodes.Add(InNode, op);
 
 		return op;
 	}
@@ -547,7 +547,7 @@ namespace mu
 		{
 			if (const NodeLOD* LODNode = node.LODs[LODIndex].get())
 			{
-				m_currentParents.Last().m_lod = LODIndex;
+				CurrentParents.Last().Lod = LODIndex;
 
 				FLODGenerationOptions LODOptions(Options, LODIndex, InNode );
 				FGenericGenerationResult LODResult;
@@ -570,12 +570,12 @@ namespace mu
 	Ptr<ASTOp> CodeGenerator::ApplyTiling(Ptr<ASTOp> Source, UE::Math::TIntVector2<int32> Size, EImageFormat Format)
 	{
 		// For now always apply tiling
-		if (m_compilerOptions->ImageTiling==0)
+		if (CompilerOptions->ImageTiling==0)
 		{
 			return Source;
 		}
 
-		int32 TileSize = m_compilerOptions->ImageTiling;
+		int32 TileSize = CompilerOptions->ImageTiling;
 
 		int32 TilesX = FMath::DivideAndRoundUp<int32>(Size[0], TileSize);
 		int32 TilesY = FMath::DivideAndRoundUp<int32>(Size[1], TileSize);
@@ -631,7 +631,7 @@ namespace mu
     {
         MUTABLE_CPUPROFILER_SCOPE(GenerateSurface);
 
-        const NodeSurfaceNew::Private& node = *surfaceNode->GetPrivate();
+        const NodeSurfaceNew& node = *surfaceNode;
 
         // Build a series of operations to assemble the surface
         Ptr<ASTOp> lastSurfOp;
@@ -648,22 +648,22 @@ namespace mu
 		// Do we need to generate the mesh? Or was it already generated for state conditions 
 		// accepting the current state?
 		FirstPassGenerator::FSurface* targetSurface = nullptr;
-		for (FirstPassGenerator::FSurface& its : m_firstPass.surfaces)
+		for (FirstPassGenerator::FSurface& its : FirstPass.Surfaces)
 		{
-            if (its.node->GetPrivate() == &node)
+            if (its.Node == &node)
             {
                 // Check state conditions
                 bool surfaceValidForThisState =
-                    Options.State >= its.stateCondition.Num()
+                    Options.State >= its.StateCondition.Num()
                     ||
-                    its.stateCondition[Options.State];
+                    its.StateCondition[Options.State];
 
                 if (surfaceValidForThisState)
                 {
-                    if (its.resultSurfaceOp)
+                    if (its.ResultSurfaceOp)
                     {
                         // Reuse the entire surface
-                        result.surfaceOp = its.resultSurfaceOp;
+                        result.surfaceOp = its.ResultSurfaceOp;
                         return;
                     }
                     else
@@ -693,13 +693,13 @@ namespace mu
             Ptr<ASTOp> lastMeshOp;
 
 			// Store the data necessary to apply modifiers for the pre-normal operations stage.
-			m_activeTags.Add(node.m_tags);
+			ActiveTags.Add(node.Tags);
 
             // Generate the mesh
 			FMeshGenerationOptions MeshOptions;
 			MeshOptions.bLayouts = true;
 			MeshOptions.State = Options.State;
-			MeshOptions.ActiveTags = node.m_tags;
+			MeshOptions.ActiveTags = node.Tags;
 
 			const FMeshGenerationResult* SharedMeshResults = nullptr;
 			if (bShareSurface)
@@ -713,7 +713,7 @@ namespace mu
 			}
 
 			// Normalize UVs if we're going to work with images and layouts.
-			const bool bNormalizeUVs = node.m_images.Num() && edits.Num();
+			const bool bNormalizeUVs = node.Images.Num() && edits.Num();
 			MeshOptions.bNormalizeUVs = bNormalizeUVs;
 
 			// Ensure UV islands remain within their main layout block on lower LODs to avoid unexpected reordering 
@@ -722,10 +722,10 @@ namespace mu
 			MeshOptions.bClampUVIslands = bShareSurface && bNormalizeUVs;
 
             GenerateMesh(MeshOptions, meshResults, node.Mesh);
-            lastMeshOp = meshResults.meshOp;
-            meshResults.extraMeshLayouts.SetNum( edits.Num() );
+            lastMeshOp = meshResults.MeshOp;
+            meshResults.ExtraMeshLayouts.SetNum( edits.Num() );
 
-			m_activeTags.Pop();
+			ActiveTags.Pop();
 
             // Let's remember the base mesh of each added mesh, so that we can use them for the
             // "remove mesh" operations that apply to them instead of the base.
@@ -737,60 +737,60 @@ namespace mu
             {
                 const FirstPassGenerator::FSurface::FEdit& e = edits[editIndex];
 
-                if ( e.node->GetPrivate()->m_pMesh)
+                if ( e.Node->Mesh)
                 {
-                    if ( NodeMeshPtr pAdd = e.node->GetPrivate()->m_pMesh->GetAdd() )
+                    if ( NodeMeshPtr pAdd = e.Node->Mesh->GetAdd() )
                     {
 						// Store the data necessary to apply modifiers for the pre-normal operations stage.
-						m_activeTags.Add(e.node->GetPrivate()->m_tags);
+						ActiveTags.Add(e.Node->Tags);
 						
 						FMeshGenerationOptions MergedMeshOptions;
 						MergedMeshOptions.bLayouts = true;
 						MergedMeshOptions.bClampUVIslands = bShareSurface && bNormalizeUVs;
 						MergedMeshOptions.bNormalizeUVs = bNormalizeUVs;
 						MergedMeshOptions.State = Options.State;
-						MergedMeshOptions.ActiveTags = e.node->GetPrivate()->m_tags;
+						MergedMeshOptions.ActiveTags = e.Node->Tags;
 
 						if (SharedMeshResults)
 						{
-							check(SharedMeshResults->extraMeshLayouts.Num()>editIndex);
-							MergedMeshOptions.OverrideLayouts = SharedMeshResults->extraMeshLayouts[editIndex].GeneratedLayouts;
+							check(SharedMeshResults->ExtraMeshLayouts.Num()>editIndex);
+							MergedMeshOptions.OverrideLayouts = SharedMeshResults->ExtraMeshLayouts[editIndex].GeneratedLayouts;
 						}
 
 						FMeshGenerationResult addResults;
                         GenerateMesh(MergedMeshOptions, addResults, pAdd);
 
-						m_activeTags.Pop();
+						ActiveTags.Pop();
 
-                        baseMeshesForEachAddedMesh[e.node] = addResults.baseMeshOp;
+                        baseMeshesForEachAddedMesh[e.Node] = addResults.BaseMeshOp;
 
 						// Apply the modifier for the post-normal operations stage to the added mesh
 						bool bModifiersForBeforeOperations = false;
 						FGenericGenerationOptions ModifierOptions(Options);
-						ModifierOptions.ActiveTags = e.node->GetPrivate()->m_tags;
+						ModifierOptions.ActiveTags = e.Node->Tags;
 						lastMeshOp = ApplyMeshModifiers(Options, lastMeshOp, bModifiersForBeforeOperations, surfaceNode->GetMessageContext());
 
                         FMeshGenerationResult::FExtraLayouts data;
 						data.GeneratedLayouts = addResults.GeneratedLayouts;
-						data.condition = e.condition;
-                        data.meshFragment = addResults.meshOp;
-                        meshResults.extraMeshLayouts[editIndex] = data;
+						data.Condition = e.Condition;
+                        data.MeshFragment = addResults.MeshOp;
+                        meshResults.ExtraMeshLayouts[editIndex] = data;
 
                         Ptr<ASTOpFixed> mop = new ASTOpFixed();
                         mop->op.type = OP_TYPE::ME_MERGE;
                         mop->SetChild(mop->op.args.MeshMerge.base, lastMeshOp );
-                        mop->SetChild(mop->op.args.MeshMerge.added, addResults.meshOp );
+                        mop->SetChild(mop->op.args.MeshMerge.added, addResults.MeshOp );
                         // will merge the meshes under the same surface
                         mop->op.args.MeshMerge.newSurfaceID = 0;
 
                         // Condition to apply
-                        if (e.condition)
+                        if (e.Condition)
                         {
                             Ptr<ASTOpConditional> conditionalAd = new ASTOpConditional();
                             conditionalAd->type = OP_TYPE::ME_CONDITIONAL;
                             conditionalAd->no = lastMeshOp;
                             conditionalAd->yes = mop;
-                            conditionalAd->condition = e.condition;
+                            conditionalAd->condition = e.Condition;
                             lastMeshOp = conditionalAd;
                         }
                         else
@@ -809,15 +809,15 @@ namespace mu
             Ptr<ASTOpMeshRemoveMask> rop;
             for ( const FirstPassGenerator::FSurface::FEdit& e: edits )
             {
-                if ( e.node->GetPrivate()->m_pMesh )
+                if ( e.Node->Mesh )
                 {
-                    if ( NodeMeshPtr pRemove = e.node->GetPrivate()->m_pMesh->GetRemove() )
+                    if ( NodeMeshPtr pRemove = e.Node->Mesh->GetRemove() )
                     {
                         FMeshGenerationResult removeResults;
 						FMeshGenerationOptions RemoveMeshOptions;
 						RemoveMeshOptions.bLayouts = false;
 						RemoveMeshOptions.State = Options.State;
-						RemoveMeshOptions.ActiveTags = e.node->GetPrivate()->m_tags;
+						RemoveMeshOptions.ActiveTags = e.Node->Tags;
 
                         GenerateMesh(RemoveMeshOptions, removeResults, pRemove );
 
@@ -825,11 +825,11 @@ namespace mu
                         maskOp->op.type = OP_TYPE::ME_MASKDIFF;
 
                         // By default, remove from the base
-                        Ptr<ASTOp> removeFrom = meshResults.baseMeshOp;
+                        Ptr<ASTOp> removeFrom = meshResults.BaseMeshOp;
                         // See if we want to remove from an added mesh instead.
-                        if ( e.node->GetPrivate()->m_pParent )
+                        if ( e.Node->Parent )
                         {
-                            auto addedBaseMeshIt = baseMeshesForEachAddedMesh.find(e.node->GetPrivate()->m_pParent.get());
+                            auto addedBaseMeshIt = baseMeshesForEachAddedMesh.find(e.Node->Parent.get());
                             if (addedBaseMeshIt!=baseMeshesForEachAddedMesh.end())
                             {
                                 removeFrom = addedBaseMeshIt->second;
@@ -837,7 +837,7 @@ namespace mu
                         }
 
                         maskOp->SetChild(maskOp->op.args.MeshMaskDiff.source, removeFrom );
-                        maskOp->SetChild(maskOp->op.args.MeshMaskDiff.fragment, removeResults.meshOp );
+                        maskOp->SetChild(maskOp->op.args.MeshMaskDiff.fragment, removeResults.MeshOp );
 
                         if (!rop)
                         {
@@ -845,7 +845,7 @@ namespace mu
                             rop->source = lastMeshOp;
                         }
 
-                        rop->AddRemove( e.condition, maskOp );
+                        rop->AddRemove( e.Condition, maskOp );
                     }
                 }
             }
@@ -858,7 +858,7 @@ namespace mu
             // Apply mesh morphs from child objects "edit surface" nodes
             for ( const FirstPassGenerator::FSurface::FEdit& e: edits )
             {
-                if ( NodeMeshPtr pMorph = e.node->GetPrivate()->m_pMorph )
+                if ( NodeMeshPtr pMorph = e.Node->Morph )
                 {
 					// Not needed because it has been generated before already.
 					// Base mesh
@@ -874,11 +874,11 @@ namespace mu
                     GenerateMesh(MorphTargetMeshOptions, morphResult, pMorph );
 
 					// Morph generation through mesh diff
-					Ptr<ASTOp> targetAd = morphResult.meshOp;
+					Ptr<ASTOp> targetAd = morphResult.MeshOp;
                     Ptr<ASTOpMeshDifference> diffAd;
                     {
                         Ptr<ASTOpMeshDifference> op = new ASTOpMeshDifference();
-                        op->Base = meshResults.baseMeshOp;
+                        op->Base = meshResults.BaseMeshOp;
                         op->Target = targetAd;
 					
                         // Morphing tex coords here is not supported:
@@ -894,10 +894,10 @@ namespace mu
                         Ptr<ASTOpMeshMorph> op = new ASTOpMeshMorph();
 
 						// Factor
-						if (e.node->GetPrivate()->m_pFactor)
+						if (e.Node->MorphFactor)
 						{
 							FScalarGenerationResult ChildResult;
-							GenerateScalar(ChildResult, Options, e.node->GetPrivate()->m_pFactor);
+							GenerateScalar(ChildResult, Options, e.Node->MorphFactor);
 							op->Factor = ChildResult.op;
 						}
 						else
@@ -919,13 +919,13 @@ namespace mu
                     }
 
                     // Condition to apply the morph
-                    if (e.condition)
+                    if (e.Condition)
                     {
                         Ptr<ASTOpConditional> conditionalAd = new ASTOpConditional();
                         conditionalAd->type = OP_TYPE::ME_CONDITIONAL;
                         conditionalAd->no = lastMeshOp;
                         conditionalAd->yes = morphAd;
-                        conditionalAd->condition = e.condition;
+                        conditionalAd->condition = e.Condition;
                         lastMeshOp = conditionalAd;
                     }
                     else
@@ -938,7 +938,7 @@ namespace mu
 			// Apply the modifier for the post-normal operations stage.
 			bool bModifiersForBeforeOperations = false;
 			FGenericGenerationOptions ModifierOptions(Options);
-			ModifierOptions.ActiveTags = node.m_tags;
+			ModifierOptions.ActiveTags = node.Tags;
 			lastMeshOp = ApplyMeshModifiers(Options, lastMeshOp, bModifiersForBeforeOperations, surfaceNode->GetMessageContext());
 
             // Layouts
@@ -951,8 +951,8 @@ namespace mu
                 {
 					if (SharedMeshResults)
 					{
-						check(SharedMeshResults->layoutOps.Num()>LayoutIndex);
-						layoutOp = SharedMeshResults->layoutOps[LayoutIndex];
+						check(SharedMeshResults->LayoutOps.Num()>LayoutIndex);
+						layoutOp = SharedMeshResults->LayoutOps[LayoutIndex];
 					}
 					else
 					{
@@ -963,14 +963,14 @@ namespace mu
 							Ptr<ASTOpConstantResource> op = new ASTOpConstantResource();
 							op->Type = OP_TYPE::LA_CONSTANT;
 
-							op->SetValue(pLayout, m_compilerOptions->OptimisationOptions.DiskCacheContext);
+							op->SetValue(pLayout, CompilerOptions->OptimisationOptions.DiskCacheContext);
 							layoutOp = op;
 						}
 
 						// Add children merged meshes layouts
-						for (const FMeshGenerationResult::FExtraLayouts& data : meshResults.extraMeshLayouts)
+						for (const FMeshGenerationResult::FExtraLayouts& data : meshResults.ExtraMeshLayouts)
 						{
-							if (!data.meshFragment)
+							if (!data.MeshFragment)
 							{
 								// No mesh to add, we assume there are no layouts to add either.
 								check(data.GeneratedLayouts.IsEmpty());
@@ -979,7 +979,7 @@ namespace mu
 
 							if (data.GeneratedLayouts.Num() != meshResults.GeneratedLayouts.Num())
 							{
-								m_pErrorLog->GetPrivate()->Add(TEXT("Merged layout has been ignored because the number of layouts is different."), ELMT_ERROR, surfaceNode->GetMessageContext());
+								ErrorLog->GetPrivate()->Add(TEXT("Merged layout has been ignored because the number of layouts is different."), ELMT_ERROR, surfaceNode->GetMessageContext());
 							}
 							else
 							{
@@ -990,7 +990,7 @@ namespace mu
 									op->Type = OP_TYPE::LA_CONSTANT;
 
 									Ptr<const Layout> pCloned = data.GeneratedLayouts[LayoutIndex].Layout;
-									op->SetValue(pCloned, m_compilerOptions->OptimisationOptions.DiskCacheContext);
+									op->SetValue(pCloned, CompilerOptions->OptimisationOptions.DiskCacheContext);
 
 									layoutFragmentAd = op;
 								}
@@ -1001,13 +1001,13 @@ namespace mu
 								mergeAd->Added = layoutFragmentAd;
 
 								// Condition to apply
-								if (data.condition)
+								if (data.Condition)
 								{
 									Ptr<ASTOpConditional> conditionalAd = new ASTOpConditional();
 									conditionalAd->type = OP_TYPE::LA_CONDITIONAL;
 									conditionalAd->no = layoutOp;
 									conditionalAd->yes = mergeAd;
-									conditionalAd->condition = data.condition;
+									conditionalAd->condition = data.Condition;
 									layoutOp = conditionalAd;
 								}
 								else
@@ -1055,11 +1055,11 @@ namespace mu
 
                 }
 
-                meshResults.layoutOps.Add( layoutOp );
+                meshResults.LayoutOps.Add( layoutOp );
             }
 
             // Store in the surface for later use.
-            targetSurface->resultMeshOp = lastMeshOp;
+            targetSurface->ResultMeshOp = lastMeshOp;
         }
 
 
@@ -1067,7 +1067,7 @@ namespace mu
         //------------------------------------------------------------------------
 		if (!bShareSurface)
 		{
-			for (int32 t = 0; t < node.m_images.Num(); ++t)
+			for (int32 t = 0; t < node.Images.Num(); ++t)
 			{
 				MUTABLE_CPUPROFILER_SCOPE(SurfaceTexture);
 
@@ -1077,7 +1077,7 @@ namespace mu
 				Ptr<NodeImageSwizzle> swizzleNode;
 
 				bool bFound = false;
-				NodeImagePtr pImageNode = node.m_images[t].m_pImage;
+				NodeImagePtr pImageNode = node.Images[t].Image;
 
 				while (!bFound && pImageNode)
 				{
@@ -1116,7 +1116,7 @@ namespace mu
 				if (bFound)
 				{
 
-					const int LayoutIndex = node.m_images[t].m_layoutIndex;
+					const int32 LayoutIndex = node.Images[t].LayoutIndex;
 
 					// If the layout index has been set to negative, it means we should ignore the layout for this image.
 					CompilerOptions::TextureLayoutStrategy ImageLayoutStrategy = (LayoutIndex < 0)
@@ -1132,7 +1132,7 @@ namespace mu
 						FImageGenerationOptions ImageOptions;
 						ImageOptions.State = Options.State;
 						ImageOptions.ImageLayoutStrategy = ImageLayoutStrategy;
-						ImageOptions.ActiveTags = node.m_tags;
+						ImageOptions.ActiveTags = node.Tags;
 						ImageOptions.RectSize = { 0, 0 };
 						FImageGenerationResult Result;
 						GenerateImage(ImageOptions, Result, pImageNode);
@@ -1142,11 +1142,11 @@ namespace mu
 						for (int32 editIndex = 0; editIndex < edits.Num(); ++editIndex)
 						{
 							const FirstPassGenerator::FSurface::FEdit& e = edits[editIndex];
-							if (t < e.node->GetPrivate()->m_textures.Num())
+							if (t < e.Node->Textures.Num())
 							{
-								if (const NodePatchImage* pPatch = e.node->GetPrivate()->m_textures[t].m_pPatch.get())
+								if (const NodePatchImage* pPatch = e.Node->Textures[t].Patch.get())
 								{
-									imageAd = GenerateImageBlockPatch(imageAd, pPatch, e.condition, ImageOptions);
+									imageAd = GenerateImageBlockPatch(imageAd, pPatch, e.Condition, ImageOptions);
 								}
 							}
 
@@ -1196,7 +1196,7 @@ namespace mu
 						op->type = OP_TYPE::IN_ADDIMAGE;
 						op->instance = lastSurfOp;
 						op->value = imageAd;
-						op->name = node.m_images[t].m_name;
+						op->name = node.Images[t].Name;
 
 						lastSurfOp = op;
 					}
@@ -1204,9 +1204,9 @@ namespace mu
 					else if (ImageLayoutStrategy == CompilerOptions::TextureLayoutStrategy::Pack) //-V547
 					{						
 						if (LayoutIndex >= meshResults.GeneratedLayouts.Num() ||
-							LayoutIndex >= meshResults.layoutOps.Num())
+							LayoutIndex >= meshResults.LayoutOps.Num())
 						{
-							m_pErrorLog->GetPrivate()->Add("Missing layout in object, or its parent.", ELMT_ERROR, surfaceNode->GetMessageContext());
+							ErrorLog->GetPrivate()->Add("Missing layout in object, or its parent.", ELMT_ERROR, surfaceNode->GetMessageContext());
 						}
 						else
 						{
@@ -1236,7 +1236,7 @@ namespace mu
 							{
 								BlankImageOp = new ASTOpFixed();
 								BlankImageOp->op.type = OP_TYPE::IM_BLANKLAYOUT;
-								BlankImageOp->SetChild(BlankImageOp->op.args.ImageBlankLayout.layout, meshResults.layoutOps[LayoutIndex]);
+								BlankImageOp->SetChild(BlankImageOp->op.args.ImageBlankLayout.layout, meshResults.LayoutOps[LayoutIndex]);
 								// The rest ok the op will be completed below
 								BlankImageOp->op.args.ImageBlankLayout.mipmapCount = 0;
 								imageAd = BlankImageOp;
@@ -1253,13 +1253,13 @@ namespace mu
 										{
 											bImageSizeWarning = true;
 
-											int currentLOD = m_currentParents.Last().m_lod;
+											int32 currentLOD = CurrentParents.Last().Lod;
 											FString Msg = FString::Printf(TEXT("A texture [%s] for material [%s] parameter [%s] in LOD [%d] has been resized because it didn't fit the layout. "),
-												*node.m_images[t].m_name,
-												*node.m_images[t].m_materialName,
-												*node.m_images[t].m_materialParameterName,
+												*node.Images[t].Name,
+												*node.Images[t].MaterialName,
+												*node.Images[t].MaterialParameterName,
 												currentLOD);
-											m_pErrorLog->GetPrivate()->Add(Msg, ELMT_INFO, surfaceNode->GetMessageContext());
+											ErrorLog->GetPrivate()->Add(Msg, ELMT_INFO, surfaceNode->GetMessageContext());
 										}
 									}
 
@@ -1293,7 +1293,7 @@ namespace mu
 								ImageOptions.State = Options.State;
 								ImageOptions.ImageLayoutStrategy = ImageLayoutStrategy;
 								ImageOptions.RectSize = { 0,0 };
-								ImageOptions.ActiveTags = node.m_tags;
+								ImageOptions.ActiveTags = node.Tags;
 								ImageOptions.LayoutToApply = pLayout;
 								ImageOptions.LayoutBlockId = pLayout->Blocks[BlockIndex].Id;
 								FImageGenerationResult Result;
@@ -1326,14 +1326,14 @@ namespace mu
 								for (int32 editIndex = 0; editIndex < edits.Num(); ++editIndex)
 								{
 									const FirstPassGenerator::FSurface::FEdit& e = edits[editIndex];
-									if (t < e.node->GetPrivate()->m_textures.Num())
+									if (t < e.Node->Textures.Num())
 									{
-										if (const NodePatchImage* pPatch = e.node->GetPrivate()->m_textures[t].m_pPatch.get())
+										if (const NodePatchImage* pPatch = e.Node->Textures[t].Patch.get())
 										{
 											// Is the current block to be patched?
 											if (pPatch->GetPrivate()->BlockIndices.Contains(BlockIndex))
 											{
-												blockAd = GenerateImageBlockPatch(blockAd, pPatch, e.condition, ImageOptions);
+												blockAd = GenerateImageBlockPatch(blockAd, pPatch, e.Condition, ImageOptions);
 											}
 										}
 									}
@@ -1351,7 +1351,7 @@ namespace mu
 
 								// Compose layout operation
 								Ptr<ASTOpImageCompose> composeOp = new ASTOpImageCompose();
-								composeOp->Layout = meshResults.layoutOps[LayoutIndex];
+								composeOp->Layout = meshResults.LayoutOps[LayoutIndex];
 								composeOp->Base = imageAd;
 								composeOp->BlockImage = blockAd;
 
@@ -1367,37 +1367,37 @@ namespace mu
 							for (int32 editIndex = 0; editIndex < edits.Num(); ++editIndex)
 							{
 								const FirstPassGenerator::FSurface::FEdit& e = edits[editIndex];
-								if (t < e.node->GetPrivate()->m_textures.Num())
+								if (t < e.Node->Textures.Num())
 								{
-									Ptr<NodeImage> pExtend = e.node->GetPrivate()->m_textures[t].m_pExtend;
+									Ptr<NodeImage> pExtend = e.Node->Textures[t].Extend;
 									if (pExtend)
 									{
-										if (LayoutIndex >= meshResults.extraMeshLayouts[editIndex].GeneratedLayouts.Num()
+										if (LayoutIndex >= meshResults.ExtraMeshLayouts[editIndex].GeneratedLayouts.Num()
 											||
-											!meshResults.extraMeshLayouts[editIndex].GeneratedLayouts[LayoutIndex].Layout)
+											!meshResults.ExtraMeshLayouts[editIndex].GeneratedLayouts[LayoutIndex].Layout)
 										{
 											FString Msg = FString::Printf(TEXT("Trying to extend a layout that doesn't exist in object [%s]."),
-												*m_currentParents.Last().m_pObject->Name
+												*CurrentParents.Last().ObjectNode->Name
 											);
 
-											m_pErrorLog->GetPrivate()->Add(Msg, ELMT_ERROR, surfaceNode->GetMessageContext());
+											ErrorLog->GetPrivate()->Add(Msg, ELMT_ERROR, surfaceNode->GetMessageContext());
 										}
 										else
 										{
-											Ptr<const Layout> pExtendLayout = meshResults.extraMeshLayouts[editIndex].GeneratedLayouts[LayoutIndex].Layout;
+											Ptr<const Layout> pExtendLayout = meshResults.ExtraMeshLayouts[editIndex].GeneratedLayouts[LayoutIndex].Layout;
 
 											// Size of a layout block in pixels
 											FIntPoint extlayout = pExtendLayout->GetGridSize();
 
 											Ptr<ASTOp> lastBase = imageAd;
 
-											for (int b = 0; b < pExtendLayout->GetBlockCount(); ++b)
+											for (int32 b = 0; b < pExtendLayout->GetBlockCount(); ++b)
 											{
 												// Generate the image block
 												FImageGenerationOptions ImageOptions;
 												ImageOptions.State = Options.State;
 												ImageOptions.ImageLayoutStrategy = ImageLayoutStrategy;
-												ImageOptions.ActiveTags = node.m_tags;
+												ImageOptions.ActiveTags = node.Tags;
 												ImageOptions.RectSize = { 0,0 };
 												ImageOptions.LayoutToApply = pExtendLayout;
 												ImageOptions.LayoutBlockId = pExtendLayout->Blocks[b].Id;
@@ -1429,7 +1429,7 @@ namespace mu
 
 												// Compose operation
 												Ptr<ASTOpImageCompose> composeOp = new ASTOpImageCompose();
-												composeOp->Layout = meshResults.layoutOps[LayoutIndex];
+												composeOp->Layout = meshResults.LayoutOps[LayoutIndex];
 												composeOp->Base = lastBase;
 												composeOp->BlockImage = fragmentAd;
 
@@ -1441,14 +1441,14 @@ namespace mu
 											}
 
 											// Condition to enable this image extension
-											if (e.condition)
+											if (e.Condition)
 											{
 												Ptr<ASTOp> conditionalAd;
 												Ptr<ASTOpConditional> cop = new ASTOpConditional();
 												cop->type = OP_TYPE::IM_CONDITIONAL;
 												cop->no = imageAd;
 												cop->yes = lastBase;
-												cop->condition = e.condition;
+												cop->condition = e.Condition;
 												conditionalAd = cop;
 												imageAd = conditionalAd;
 											}
@@ -1525,8 +1525,8 @@ namespace mu
 								// we will devide the layout block by the format block
 								const FImageFormatData& finfo = GetImageFormatData(FinalFormat);
 
-								int mipsX = (int)ceilf(logf((float)BlockPixelsX / finfo.PixelsPerBlockX) / logf(2.0f));
-								int mipsY = (int)ceilf(logf((float)BlockPixelsY / finfo.PixelsPerBlockY) / logf(2.0f));
+								int32 mipsX = (int32)ceilf(logf((float)BlockPixelsX / finfo.PixelsPerBlockX) / logf(2.0f));
+								int32 mipsY = (int32)ceilf(logf((float)BlockPixelsY / finfo.PixelsPerBlockY) / logf(2.0f));
 								mop->BlockLevels = (uint8_t)FMath::Max(mipsX, mipsY);
 
 								if (BlockPixelsX < finfo.PixelsPerBlockX || BlockPixelsY < finfo.PixelsPerBlockY)
@@ -1556,7 +1556,7 @@ namespace mu
 							op->value = imageAd;
 
 							// Name
-							op->name = node.m_images[t].m_name;
+							op->name = node.Images[t].Name;
 
 							lastSurfOp = op;
 						}
@@ -1572,11 +1572,11 @@ namespace mu
 
 			// Create the expression for each vector
 			//------------------------------------------------------------------------
-			for (int32 t = 0; t < node.m_vectors.Num(); ++t)
+			for (int32 t = 0; t < node.Vectors.Num(); ++t)
 			{
 				//MUTABLE_CPUPROFILER_SCOPE(SurfaceVector);
 
-				if (NodeColourPtr pVectorNode = node.m_vectors[t].m_pVector)
+				if (Ptr<NodeColour> pVectorNode = node.Vectors[t].Vector)
 				{
 					Ptr<ASTOpInstanceAdd> op = new ASTOpInstanceAdd();
 					op->type = OP_TYPE::IN_ADDVECTOR;
@@ -1588,7 +1588,7 @@ namespace mu
 					op->value = VectorResult.op;
 
 					// Name
-					op->name = node.m_vectors[t].m_name;
+					op->name = node.Vectors[t].Name;
 
 					lastSurfOp = op;
 				}
@@ -1596,11 +1596,11 @@ namespace mu
 
 			// Create the expression for each scalar
 			//------------------------------------------------------------------------
-			for (int32 t = 0; t < node.m_scalars.Num(); ++t)
+			for (int32 t = 0; t < node.Scalars.Num(); ++t)
 			{
 				// MUTABLE_CPUPROFILER_SCOPE(SurfaceScalar);
 
-				if (NodeScalarPtr pScalarNode = node.m_scalars[t].m_pScalar)
+				if (NodeScalarPtr pScalarNode = node.Scalars[t].Scalar)
 				{
 					Ptr<ASTOpInstanceAdd> op = new ASTOpInstanceAdd();
 					op->type = OP_TYPE::IN_ADDSCALAR;
@@ -1612,7 +1612,7 @@ namespace mu
 					op->value = ScalarResult.op;
 
 					// Name
-					op->name = node.m_scalars[t].m_name;
+					op->name = node.Scalars[t].Name;
 
 					lastSurfOp = op;
 				}
@@ -1620,9 +1620,9 @@ namespace mu
 
 			// Create the expression for each string
 			//------------------------------------------------------------------------
-			for (int32 t = 0; t < node.m_strings.Num(); ++t)
+			for (int32 t = 0; t < node.Strings.Num(); ++t)
 			{
-				if (NodeStringPtr pStringNode = node.m_strings[t].m_pString)
+				if (NodeStringPtr pStringNode = node.Strings[t].String)
 				{
 					Ptr<ASTOpInstanceAdd> op = new ASTOpInstanceAdd();
 					op->type = OP_TYPE::IN_ADDSTRING;
@@ -1633,7 +1633,7 @@ namespace mu
 					op->value = StringResult.op;
 
 					// Name
-					op->name = node.m_strings[t].m_name;
+					op->name = node.Strings[t].Name;
 
 					lastSurfOp = op;
 				}
@@ -1641,7 +1641,7 @@ namespace mu
 		}
 
         result.surfaceOp = lastSurfOp;
-        targetSurface->resultSurfaceOp = lastSurfOp;
+        targetSurface->ResultSurfaceOp = lastSurfOp;
 
 		// If we are going to share this surface properties, remember it.
 		if (bIsBaseForSharedSurface)
@@ -1670,9 +1670,9 @@ namespace mu
         int32 surfaceID=1;
 
         // Look for all surfaces that belong to this component
-		for (int32 i = 0; i<m_firstPass.surfaces.Num(); ++i, ++surfaceID)
+		for (int32 i = 0; i<FirstPass.Surfaces.Num(); ++i, ++surfaceID)
 		{
-			const FirstPassGenerator::FSurface& its = m_firstPass.surfaces[i];
+			const FirstPassGenerator::FSurface& its = FirstPass.Surfaces[i];
 			if (its.Component==Options.Component
 				&&
 				its.LOD==Options.LODIndex)
@@ -1680,12 +1680,12 @@ namespace mu
                 // Apply state conditions: only generate it if it enabled in this state
                 {
                     bool enabledInThisState = true;
-                    if (its.stateCondition.Num() && Options.State >= 0)
+                    if (its.StateCondition.Num() && Options.State >= 0)
                     {
                         enabledInThisState =
-                                (Options.State < its.stateCondition.Num())
+                                (Options.State < its.StateCondition.Num())
                                 &&
-                                ( its.stateCondition[Options.State] );
+                                ( its.StateCondition[Options.State] );
                     }
                     if (!enabledInThisState)
                     {
@@ -1695,24 +1695,24 @@ namespace mu
 
                 Ptr<ASTOpInstanceAdd> sop = new ASTOpInstanceAdd();
                 sop->type = OP_TYPE::IN_ADDSURFACE;
-                sop->name = its.node->GetPrivate()->m_name;
+                sop->name = its.Node->Name;
                 sop->instance = lastCompOp;
 
 				FSurfaceGenerationOptions SurfaceOptions(Options);
 				FSurfaceGenerationResult surfaceGenerationResult;
-                GenerateSurface( surfaceGenerationResult, SurfaceOptions, its.node, its.edits );
+                GenerateSurface( surfaceGenerationResult, SurfaceOptions, its.Node, its.Edits );
                 sop->value = surfaceGenerationResult.surfaceOp;
 
                 sop->id = surfaceID;
-                sop->ExternalId = its.node->GetPrivate()->ExternalId;
-                sop->SharedSurfaceId = its.node->GetPrivate()->SharedSurfaceId;
+                sop->ExternalId = its.Node->ExternalId;
+                sop->SharedSurfaceId = its.Node->SharedSurfaceId;
                 Ptr<ASTOp> surfaceAt = sop;
 
 				// TODO: This could be done earlier?
                 Ptr<ASTOpFixed> surfCondOp = new ASTOpFixed();
                 surfCondOp->op.type = OP_TYPE::BO_AND;
-                surfCondOp->SetChild(surfCondOp->op.args.BoolBinary.a, its.objectCondition );
-                surfCondOp->SetChild(surfCondOp->op.args.BoolBinary.b, its.surfaceCondition );
+                surfCondOp->SetChild(surfCondOp->op.args.BoolBinary.a, its.ObjectCondition );
+                surfCondOp->SetChild(surfCondOp->op.args.BoolBinary.b, its.SurfaceCondition );
                 Ptr<ASTOp> surfaceCondition = surfCondOp;
 
                 if (surfaceCondition)
@@ -1734,7 +1734,7 @@ namespace mu
                 // We add the merge op even for the first mesh, so that we set the surface id.
 				Ptr<ASTOp> mergeAd;
 				{
-                    Ptr<ASTOp> added = its.resultMeshOp;
+                    Ptr<ASTOp> added = its.ResultMeshOp;
  
                     Ptr<ASTOpFixed> mop = new ASTOpFixed();
                     mop->op.type = OP_TYPE::ME_MERGE;
@@ -1787,8 +1787,8 @@ namespace mu
 	{
 		MUTABLE_CPUPROFILER_SCOPE(NodeObjectNew);
 
-        m_currentParents.Add( FParentKey() );
-        m_currentParents.Last().m_pObject = InNode;
+        CurrentParents.Add( FParentKey() );
+        CurrentParents.Last().ObjectNode = InNode;
 
         // Parse the child objects first, which will accumulate operations in the patching lists
         for ( int32 t=0; t< InNode->Children.Num(); ++t )
@@ -1799,9 +1799,9 @@ namespace mu
 
                 // If there are parent objects, the condition of this object depends on the
                 // condition of the parent object
-                if ( m_currentObject.Num() )
+                if ( CurrentObject.Num() )
                 {
-                    paramOp = m_currentObject.Last().m_condition;
+                    paramOp = CurrentObject.Last().Condition;
                 }
                 else
                 {
@@ -1813,13 +1813,13 @@ namespace mu
                 }
 
 				FObjectGenerationData data;
-                data.m_condition = paramOp;
-                m_currentObject.Add( data );
+                data.Condition = paramOp;
+                CurrentObject.Add( data );
 
                 // This op is ignored: everything is stored as patches to apply to the parent when it is compiled.
                 Generate_Generic( pChildNode, Options );
 
-                m_currentObject.Pop();
+                CurrentObject.Pop();
             }
         }
 
@@ -1849,8 +1849,8 @@ namespace mu
 
 		// Add the components from child objects
 		FAdditionalComponentKey thisKey;
-		thisKey.m_lod = m_currentParents.Last().m_lod;
-		thisKey.m_pObject = m_currentParents.Last().m_pObject;
+		thisKey.Lod = CurrentParents.Last().Lod;
+		thisKey.ObjectNode = CurrentParents.Last().ObjectNode;
 		TArray<Ptr<ASTOp>>* addIt = AdditionalComponents.Find(thisKey);
 		if (addIt)
 		{
@@ -1877,12 +1877,12 @@ namespace mu
 
 		// Store for possible parent objects if necessary
 		// 2 is because there must be a parent and there is always a null element as well.
-		if (LastCompOp && m_currentParents.Num() > 2)
+		if (LastCompOp && CurrentParents.Num() > 2)
 		{
-			const auto& parentObjectKey = m_currentParents[m_currentParents.Num() - 2];
+			const auto& parentObjectKey = CurrentParents[CurrentParents.Num() - 2];
 			FAdditionalComponentKey parentKey;
-			parentKey.m_lod = m_currentParents.Last().m_lod;
-			parentKey.m_pObject = parentObjectKey.m_pObject;
+			parentKey.Lod = CurrentParents.Last().Lod;
+			parentKey.ObjectNode = parentObjectKey.ObjectNode;
 			AdditionalComponents.FindOrAdd(parentKey).Add(LastCompOp);
 		}
 
@@ -1911,15 +1911,15 @@ namespace mu
 			}
 
 			FConditionalExtensionDataOp& SavedOp = ConditionalExtensionDataOps.AddDefaulted_GetRef();
-			if (m_currentObject.Num() > 0)
+			if (CurrentObject.Num() > 0)
 			{
-				SavedOp.Condition = m_currentObject.Last().m_condition;
+				SavedOp.Condition = CurrentObject.Last().Condition;
 			}
 			SavedOp.ExtensionDataOp = ChildResult.Op;
 			SavedOp.ExtensionDataName = NamedNode.Name;
 		}
 
-		if (m_currentObject.Num() == 0)
+		if (CurrentObject.Num() == 0)
 		{
 			for (const FConditionalExtensionDataOp& SavedOp : ConditionalExtensionDataOps)
 			{
@@ -1945,7 +1945,7 @@ namespace mu
 			}
 		}
 
-        m_currentParents.Pop();
+        CurrentParents.Pop();
 
         Result.op = rootOp;
     }
@@ -1967,13 +1967,13 @@ namespace mu
 				// Look for the child condition in the first pass
                 Ptr<ASTOp> conditionOp;
 				bool found = false;
-                for( int32 i = 0; !found && i != m_firstPass.objects.Num(); i++ )
+                for( int32 i = 0; !found && i != FirstPass.Objects.Num(); i++ )
 				{
-					FirstPassGenerator::FObject& it = m_firstPass.objects[i];
-					if (it.node == pChildNode)
+					FirstPassGenerator::FObject& it = FirstPass.Objects[i];
+					if (it.Node == pChildNode)
 					{
 						found = true;
-						conditionOp = it.condition;
+						conditionOp = it.Condition;
 					}
 				}
 
@@ -1983,21 +1983,21 @@ namespace mu
                 // check(found);
 
 				FObjectGenerationData data;
-                data.m_condition = conditionOp;
-                m_currentObject.Add( data );
+                data.Condition = conditionOp;
+                CurrentObject.Add( data );
 
                 // This op is ignored: everything is stored as patches to apply to the parent when
                 // it is compiled.
                 Generate_Generic( pChildNode, Options );
 
-                m_currentObject.Pop();
+                CurrentObject.Pop();
 
 				// Check for duplicated child names
 				FString strChildName = pChildNode->GetName();
 				if (usedNames.Contains(strChildName))
 				{
 					FString Msg = FString::Printf(TEXT("Object group has more than one children with the same name [%s]."), *strChildName );
-					m_pErrorLog->GetPrivate()->Add(Msg, ELMT_WARNING, InNode->GetMessageContext());
+					ErrorLog->GetPrivate()->Add(Msg, ELMT_WARNING, InNode->GetMessageContext());
 				}
 				else
 				{
@@ -2013,7 +2013,7 @@ namespace mu
     {
         // Log a warning
 		FString Msg = FString::Printf(TEXT("Required connection not found: %s"), Where);
-        m_pErrorLog->GetPrivate()->Add( Msg, ELMT_ERROR, ErrorContext);
+        ErrorLog->GetPrivate()->Add( Msg, ELMT_ERROR, ErrorContext);
 
         // Create a constant node
         Ptr<NodeBoolConstant> pNode = new NodeBoolConstant();
@@ -2039,7 +2039,7 @@ namespace mu
 			return;
 		}
 
-		for (const FirstPassGenerator::FModifier& m: m_firstPass.modifiers)
+		for (const FirstPassGenerator::FModifier& m: FirstPass.Modifiers)
 		{
 			// Correct LOD?
 			if (m.LOD != LOD)
@@ -2048,14 +2048,14 @@ namespace mu
 			}
 
 			// Correct stage?
-			if (m.node->bApplyBeforeNormalOperations != bModifiersForBeforeOperations)
+			if (m.Node->bApplyBeforeNormalOperations != bModifiersForBeforeOperations)
 			{
 				continue;
 			}
 
 			// Already there?
 			bool bAlreadyAdded = 
-				OutModifiers.FindByPredicate( [&m](const FirstPassGenerator::FModifier& c) {return c.node == m.node; })
+				OutModifiers.FindByPredicate( [&m](const FirstPassGenerator::FModifier& c) {return c.Node == m.Node; })
 				!= 
 				nullptr;
 
@@ -2067,11 +2067,11 @@ namespace mu
 			// Matching tags?
 			bool bApply = false;
 
-			switch (m.node->MultipleTagsPolicy)
+			switch (m.Node->MultipleTagsPolicy)
 			{
 			case EMutableMultipleTagPolicy::OnlyOneRequired:
 			{
-				for (const FString& Tag: m.node->RequiredTags)
+				for (const FString& Tag: m.Node->RequiredTags)
 				{
 					if (SurfaceTags.Contains(Tag))
 					{
@@ -2085,7 +2085,7 @@ namespace mu
 			case EMutableMultipleTagPolicy::AllRequired:
 			{
 				bApply = true;
-				for (const FString& Tag : m.node->RequiredTags)
+				for (const FString& Tag : m.Node->RequiredTags)
 				{
 					if (!SurfaceTags.Contains(Tag))
 					{
@@ -2115,20 +2115,20 @@ namespace mu
 		// Apply mesh modifiers
 		TArray<FirstPassGenerator::FModifier> modifiers;
 
-		int currentLOD = m_currentParents.Last().m_lod;
+		int32 currentLOD = CurrentParents.Last().Lod;
 		GetModifiersFor(Options.ActiveTags, currentLOD, bModifiersForBeforeOperations, modifiers);
 
 		Ptr<ASTOp> preModifiersMesh = lastMeshOp;
 
-		m_activeTags.Add({});
+		ActiveTags.Add({});
 
 		// Process clip-with-mesh modifiers
 		Ptr<ASTOpMeshRemoveMask> removeOp;
 		for (const FirstPassGenerator::FModifier& m : modifiers)
 		{
-			if (m.node->GetType()== NodeModifierMeshClipWithMesh::GetStaticType())
+			if (m.Node->GetType()== NodeModifierMeshClipWithMesh::GetStaticType())
 			{
-				const NodeModifierMeshClipWithMesh* TypedClipNode = static_cast<const NodeModifierMeshClipWithMesh*>(m.node);
+				const NodeModifierMeshClipWithMesh* TypedClipNode = static_cast<const NodeModifierMeshClipWithMesh*>(m.Node);
 				Ptr<ASTOpMeshMaskClipMesh> op = new ASTOpMeshMaskClipMesh();
 				op->source = preModifiersMesh;
 
@@ -2139,11 +2139,11 @@ namespace mu
 
 				FMeshGenerationResult clipResult;
 				GenerateMesh(ClipOptions, clipResult, TypedClipNode->GetPrivate()->ClipMesh);
-				op->clip = clipResult.meshOp;
+				op->clip = clipResult.MeshOp;
 
 				if (!op->clip)
 				{
-					m_pErrorLog->GetPrivate()->Add("Clip mesh has not been generated", ELMT_ERROR, errorContext);
+					ErrorLog->GetPrivate()->Add("Clip mesh has not been generated", ELMT_ERROR, errorContext);
 					continue;
 				}
 
@@ -2158,8 +2158,8 @@ namespace mu
 
 				Ptr<ASTOpFixed> surfCondOp = new ASTOpFixed();
 				surfCondOp->op.type = OP_TYPE::BO_AND;
-				surfCondOp->SetChild(surfCondOp->op.args.BoolBinary.a, m.objectCondition);
-				surfCondOp->SetChild(surfCondOp->op.args.BoolBinary.b, m.surfaceCondition);
+				surfCondOp->SetChild(surfCondOp->op.args.BoolBinary.a, m.ObjectCondition);
+				surfCondOp->SetChild(surfCondOp->op.args.BoolBinary.b, m.SurfaceCondition);
 				Ptr<ASTOp> fullCondition = surfCondOp;
 
 				removeOp->AddRemove(fullCondition, maskAt);
@@ -2169,9 +2169,9 @@ namespace mu
 		// Process clip-with-mask modifiers
 		for (const FirstPassGenerator::FModifier& m : modifiers)
 		{
-			if (m.node->GetType() == NodeModifierMeshClipWithUVMask::GetStaticType())
+			if (m.Node->GetType() == NodeModifierMeshClipWithUVMask::GetStaticType())
 			{
-				const NodeModifierMeshClipWithUVMask* TypedClipNode = static_cast<const NodeModifierMeshClipWithUVMask*>(m.node);
+				const NodeModifierMeshClipWithUVMask* TypedClipNode = static_cast<const NodeModifierMeshClipWithUVMask*>(m.Node);
 				Ptr<ASTOpMeshMaskClipUVMask> op = new ASTOpMeshMaskClipUVMask();
 				op->Source = preModifiersMesh;
 				op->LayoutIndex = TypedClipNode->GetPrivate()->LayoutIndex;
@@ -2190,7 +2190,7 @@ namespace mu
 
 				if (!op->Mask)
 				{
-					m_pErrorLog->GetPrivate()->Add("Clip UV mask has not been generated", ELMT_ERROR, errorContext);
+					ErrorLog->GetPrivate()->Add("Clip UV mask has not been generated", ELMT_ERROR, errorContext);
 					continue;
 				}
 
@@ -2205,8 +2205,8 @@ namespace mu
 
 				Ptr<ASTOpFixed> surfCondOp = new ASTOpFixed();
 				surfCondOp->op.type = OP_TYPE::BO_AND;
-				surfCondOp->SetChild(surfCondOp->op.args.BoolBinary.a, m.objectCondition);
-				surfCondOp->SetChild(surfCondOp->op.args.BoolBinary.b, m.surfaceCondition);
+				surfCondOp->SetChild(surfCondOp->op.args.BoolBinary.a, m.ObjectCondition);
+				surfCondOp->SetChild(surfCondOp->op.args.BoolBinary.b, m.SurfaceCondition);
 				Ptr<ASTOp> fullCondition = surfCondOp;
 
 				removeOp->AddRemove(fullCondition, maskAt);
@@ -2218,9 +2218,9 @@ namespace mu
 		{
 			Ptr<ASTOp> modifiedMeshOp;
 
-			if (m.node->GetType() == NodeModifierMeshClipMorphPlane::GetStaticType())
+			if (m.Node->GetType() == NodeModifierMeshClipMorphPlane::GetStaticType())
 			{
-				const NodeModifierMeshClipMorphPlane::Private* TypedNode = static_cast<const NodeModifierMeshClipMorphPlane*>(m.node)->GetPrivate();
+				const NodeModifierMeshClipMorphPlane::Private* TypedNode = static_cast<const NodeModifierMeshClipMorphPlane*>(m.Node)->GetPrivate();
 				Ptr<ASTOpMeshClipMorphPlane> op = new ASTOpMeshClipMorphPlane();
 				op->source = lastMeshOp;
 
@@ -2281,8 +2281,8 @@ namespace mu
 			{
 				Ptr<ASTOpFixed> surfCondOp = new ASTOpFixed();
 				surfCondOp->op.type = OP_TYPE::BO_AND;
-				surfCondOp->SetChild(surfCondOp->op.args.BoolBinary.a, m.objectCondition);
-				surfCondOp->SetChild(surfCondOp->op.args.BoolBinary.b, m.surfaceCondition);
+				surfCondOp->SetChild(surfCondOp->op.args.BoolBinary.a, m.ObjectCondition);
+				surfCondOp->SetChild(surfCondOp->op.args.BoolBinary.b, m.SurfaceCondition);
 				Ptr<ASTOp> fullCondition = surfCondOp;
 
 				Ptr<ASTOpConditional> op = new ASTOpConditional();
@@ -2299,9 +2299,9 @@ namespace mu
 		{
 			Ptr<ASTOp> ModifiedMeshOp;
 
-			if (M.node->GetType()==NodeModifierMeshClipDeform::GetStaticType())
+			if (M.Node->GetType()==NodeModifierMeshClipDeform::GetStaticType())
 			{
-				const NodeModifierMeshClipDeform::Private* TypedClipNode = static_cast<const NodeModifierMeshClipDeform*>(M.node)->GetPrivate();
+				const NodeModifierMeshClipDeform::Private* TypedClipNode = static_cast<const NodeModifierMeshClipDeform*>(M.Node)->GetPrivate();
 				Ptr<ASTOpMeshBindShape>  BindOp = new ASTOpMeshBindShape();
 				Ptr<ASTOpMeshClipDeform> ClipOp = new ASTOpMeshClipDeform();
 
@@ -2311,17 +2311,17 @@ namespace mu
 
 				FMeshGenerationResult ClipShapeResult;
 				GenerateMesh(ClipOptions, ClipShapeResult, TypedClipNode->ClipMesh);
-				ClipOp->ClipShape = ClipShapeResult.meshOp;
+				ClipOp->ClipShape = ClipShapeResult.MeshOp;
 				
 				BindOp->Mesh = lastMeshOp;
-				BindOp->Shape = ClipShapeResult.meshOp; 
+				BindOp->Shape = ClipShapeResult.MeshOp; 
 				BindOp->BindingMethod = static_cast<uint32>(TypedClipNode->BindingMethod);
 	
 				ClipOp->Mesh = BindOp;
 
 				if (!ClipOp->ClipShape)
 				{
-					m_pErrorLog->GetPrivate()->Add
+					ErrorLog->GetPrivate()->Add
 					("Clip shape mesh has not been generated", ELMT_ERROR, errorContext);
 				}
 				else
@@ -2334,8 +2334,8 @@ namespace mu
 			{
 				Ptr<ASTOpFixed> SurfCondOp = new ASTOpFixed();
 				SurfCondOp->op.type = OP_TYPE::BO_AND;
-				SurfCondOp->SetChild(SurfCondOp->op.args.BoolBinary.a, M.objectCondition);
-				SurfCondOp->SetChild(SurfCondOp->op.args.BoolBinary.b, M.surfaceCondition);
+				SurfCondOp->SetChild(SurfCondOp->op.args.BoolBinary.a, M.ObjectCondition);
+				SurfCondOp->SetChild(SurfCondOp->op.args.BoolBinary.b, M.SurfaceCondition);
 				Ptr<ASTOp> FullCondition = SurfCondOp;
 
 				Ptr<ASTOpConditional> Op = new ASTOpConditional();
@@ -2347,7 +2347,7 @@ namespace mu
 			}
 		}
 			
-		m_activeTags.Pop();
+		ActiveTags.Pop();
 
 		return lastMeshOp;
 	}
diff --git a/Engine/Plugins/Experimental/Mutable/Source/MutableTools/Private/MuT/CodeGenerator.h b/Engine/Plugins/Experimental/Mutable/Source/MutableTools/Private/MuT/CodeGenerator.h
--- a/Engine/Plugins/Experimental/Mutable/Source/MutableTools/Private/MuT/CodeGenerator.h
+++ b/Engine/Plugins/Experimental/Mutable/Source/MutableTools/Private/MuT/CodeGenerator.h
@@ -24,6 +24,7 @@
 #include "MuT/NodeColourSampleImage.h"
 #include "MuT/NodeComponentEdit.h"
 #include "MuT/NodeComponentNew.h"
+#include "MuT/NodeComponentSwitch.h"
 #include "MuT/NodeExtensionData.h"
 #include "MuT/NodeImageProject.h"
 #include "MuT/NodeLOD.h"
@@ -117,7 +118,7 @@ namespace mu
 
         CodeGenerator( CompilerOptions::Private* options );
 
-        //! Data will be stored in m_states
+        //! Data will be stored in States
         void GenerateRoot( const Ptr<const Node> );
 
 	public:
@@ -215,17 +216,17 @@ namespace mu
     public:
 
         //! Settings
-        CompilerOptions::Private* m_compilerOptions = nullptr;
+        CompilerOptions::Private* CompilerOptions = nullptr;
 
 		//!
-		FirstPassGenerator m_firstPass;
+		FirstPassGenerator FirstPass;
 
         //!
-        ErrorLogPtr m_pErrorLog;
+        ErrorLogPtr ErrorLog;
 
         //! After the entire code generation this contains the information about all the states
         typedef TArray< TPair<FObjectState, Ptr<ASTOp>> > StateList;
-        StateList m_states;
+        StateList States;
 
     private:
 
@@ -238,22 +239,22 @@ namespace mu
 		TArray<FGeneratedConstantMesh> GeneratedConstantMeshes;
 
         /** List of image resources for every image formats that have been generated so far as palceholders for missing images. */
-        Ptr<Image> m_missingImage[size_t(EImageFormat::IF_COUNT)];
+        Ptr<Image> MissingImage[size_t(EImageFormat::IF_COUNT)];
 
 
         //! List of already used vertex ID groups that must be unique.
         TSet<uint32> UniqueVertexIDGroups;
 
 		// (top-down) Tags that are active when generating nodes.
-		TArray< TArray<FString> > m_activeTags;
+		TArray< TArray<FString> > ActiveTags;
 
         struct FParentKey
         {
-            const NodeObjectNew* m_pObject = nullptr;
-            int32 m_lod = -1;
+            const NodeObjectNew* ObjectNode = nullptr;
+            int32 Lod = -1;
          };
 
-		TArray< FParentKey > m_currentParents;
+		TArray< FParentKey > CurrentParents;
 
         // List of additional components to add to an object that come from child objects.
         // The index is the object and lod that should receive the components.
@@ -261,25 +262,25 @@ namespace mu
         {
 			FAdditionalComponentKey()
             {
-                m_pObject = nullptr;
-                m_lod = -1;
+				ObjectNode = nullptr;
+				Lod = -1;
             }
 
-            const NodeObjectNew* m_pObject;
-            int32 m_lod;
+            const NodeObjectNew* ObjectNode;
+            int32 Lod;
 
 			FORCEINLINE bool operator==(const FAdditionalComponentKey& Other) const
 			{
-				return m_pObject == Other.m_pObject
+				return ObjectNode == Other.ObjectNode
 					&&
-					m_lod == Other.m_lod;
+					Lod == Other.Lod;
 			}
 
 			friend FORCEINLINE uint32 GetTypeHash(const FAdditionalComponentKey& InKey)
 			{
 				uint32 KeyHash = 0;
-				KeyHash = HashCombineFast(KeyHash, ::GetTypeHash(InKey.m_pObject));
-				KeyHash = HashCombineFast(KeyHash, ::GetTypeHash(InKey.m_lod));
+				KeyHash = HashCombineFast(KeyHash, ::GetTypeHash(InKey.ObjectNode));
+				KeyHash = HashCombineFast(KeyHash, ::GetTypeHash(InKey.Lod));
 				return KeyHash;
 			}
 		};
@@ -289,9 +290,9 @@ namespace mu
         struct FObjectGenerationData
         {
             // Condition that enables a specific object
-            Ptr<ASTOp> m_condition;
+            Ptr<ASTOp> Condition;
         };
-		TArray<FObjectGenerationData> m_currentObject;
+		TArray<FObjectGenerationData> CurrentObject;
 
 		/** The key for generated tables is made of the source table and a parameter name. */
 		struct FTableCacheKey
@@ -442,7 +443,7 @@ namespace mu
 		};
 
 		typedef TMap<FGeneratedImageCacheKey, FImageGenerationResult> GeneratedImagesMap;
-		GeneratedImagesMap m_generatedImages;
+		GeneratedImagesMap GeneratedImages;
 
 		void GenerateImage(const FImageGenerationOptions&, FImageGenerationResult& result, const NodeImagePtrConst& node);
 		void GenerateImage_Constant(const FImageGenerationOptions&, FImageGenerationResult&, const NodeImageConstant*);
@@ -560,24 +561,24 @@ namespace mu
 		struct FMeshGenerationResult
 		{
 			//! Mesh after all code tree is applied
-			Ptr<ASTOp> meshOp;
+			Ptr<ASTOp> MeshOp;
 
 			//! Original base mesh before removes, morphs, etc.
-			Ptr<ASTOp> baseMeshOp;
+			Ptr<ASTOp> BaseMeshOp;
 
 			/** Generated node layouts with their own block ids. */
 			TArray<FGeneratedLayout> GeneratedLayouts;
 
-			TArray<Ptr<ASTOp>> layoutOps;
+			TArray<Ptr<ASTOp>> LayoutOps;
 
 			struct FExtraLayouts
 			{
 				/** Source node layouts to use with these extra mesh. They don't have block ids. */
 				TArray<FGeneratedLayout> GeneratedLayouts;
-				Ptr<ASTOp> condition;
-				Ptr<ASTOp> meshFragment;
+				Ptr<ASTOp> Condition;
+				Ptr<ASTOp> MeshFragment;
 			};
-			TArray< FExtraLayouts > extraMeshLayouts;
+			TArray< FExtraLayouts > ExtraMeshLayouts;
 		};
 		
 		struct FGeneratedMeshCacheKey
@@ -600,7 +601,7 @@ namespace mu
 		};
 
         typedef TMap<FGeneratedMeshCacheKey,FMeshGenerationResult> GeneratedMeshMap;
-        GeneratedMeshMap m_generatedMeshes;
+        GeneratedMeshMap GeneratedMeshes;
 
 		/** Store the mesh generation data for surfaces that we intend to share across LODs. 
 		* The key is the SharedSurfaceId.
@@ -810,18 +811,18 @@ namespace mu
 
 		if (NumRows == 0)
 		{
-			m_pErrorLog->GetPrivate()->Add("The table has no rows.", ELMT_ERROR, node.GetMessageContext());
+			ErrorLog->GetPrivate()->Add("The table has no rows.", ELMT_ERROR, node.GetMessageContext());
 			return nullptr;
 		}
         else if (ColIndex < 0)
         {
-            m_pErrorLog->GetPrivate()->Add("Table column not found.", ELMT_ERROR, node.GetMessageContext());
+            ErrorLog->GetPrivate()->Add("Table column not found.", ELMT_ERROR, node.GetMessageContext());
             return nullptr;
         }
 
         if (NodeTable->GetPrivate()->Columns[ ColIndex ].Type != TYPE )
         {
-            m_pErrorLog->GetPrivate()->Add("Table column type is not the right type.", ELMT_ERROR, node.GetMessageContext());
+            ErrorLog->GetPrivate()->Add("Table column type is not the right type.", ELMT_ERROR, node.GetMessageContext());
             return nullptr;
         }
 
@@ -836,7 +837,7 @@ namespace mu
         {
             check(NodeTable->GetPrivate()->Rows[i].Id <= 0xFFFF);
             auto Condition = (uint16)NodeTable->GetPrivate()->Rows[i].Id;
-            Ptr<ASTOp> Branch = GenerateOption( node, ColIndex, (int)i, m_pErrorLog.get() );
+            Ptr<ASTOp> Branch = GenerateOption( node, ColIndex, (int)i, ErrorLog.get() );
 
 			if (Branch || TYPE != ETableColumnType::Mesh)
 			{
diff --git a/Engine/Plugins/Experimental/Mutable/Source/MutableTools/Private/MuT/CodeGenerator_Bool.cpp b/Engine/Plugins/Experimental/Mutable/Source/MutableTools/Private/MuT/CodeGenerator_Bool.cpp
--- a/Engine/Plugins/Experimental/Mutable/Source/MutableTools/Private/MuT/CodeGenerator_Bool.cpp
+++ b/Engine/Plugins/Experimental/Mutable/Source/MutableTools/Private/MuT/CodeGenerator_Bool.cpp
@@ -87,7 +87,7 @@ namespace mu
 
 		Ptr<ASTOpParameter> op;
 
-		Ptr<ASTOpParameter>* it = m_firstPass.ParameterNodes.Find(node.m_pNode);
+		Ptr<ASTOpParameter>* it = FirstPass.ParameterNodes.Find(node.m_pNode);
 		if (!it)
 		{
 			FParameterDesc param;
@@ -109,7 +109,7 @@ namespace mu
 				op->ranges.Emplace(op.get(), rangeResult.sizeOp, rangeResult.rangeName, rangeResult.rangeUID);
 			}
 
-			m_firstPass.ParameterNodes.Add(node.m_pNode,op);
+			FirstPass.ParameterNodes.Add(node.m_pNode,op);
 		}
 		else
 		{
diff --git a/Engine/Plugins/Experimental/Mutable/Source/MutableTools/Private/MuT/CodeGenerator_Color.cpp b/Engine/Plugins/Experimental/Mutable/Source/MutableTools/Private/MuT/CodeGenerator_Color.cpp
--- a/Engine/Plugins/Experimental/Mutable/Source/MutableTools/Private/MuT/CodeGenerator_Color.cpp
+++ b/Engine/Plugins/Experimental/Mutable/Source/MutableTools/Private/MuT/CodeGenerator_Color.cpp
@@ -139,7 +139,7 @@ namespace mu
 
 		Ptr<ASTOpParameter> op;
 
-		Ptr<ASTOpParameter>* it = m_firstPass.ParameterNodes.Find(Typed);
+		Ptr<ASTOpParameter>* it = FirstPass.ParameterNodes.Find(Typed);
 
 		if (!it)
 		{
@@ -169,7 +169,7 @@ namespace mu
 				op->ranges.Emplace(op.get(), rangeResult.sizeOp, rangeResult.rangeName, rangeResult.rangeUID);
 			}
 
-			m_firstPass.ParameterNodes.Add(Typed, op);
+			FirstPass.ParameterNodes.Add(Typed, op);
 		}
 		else
 		{
@@ -245,13 +245,13 @@ namespace mu
 		}
 
 		// Process variations in reverse order, since conditionals are built bottom-up.
-		for (int t = node.Variations.Num() - 1; t >= 0; --t)
+		for (int32 t = node.Variations.Num() - 1; t >= 0; --t)
 		{
-			int tagIndex = -1;
+			int32 tagIndex = -1;
 			const FString& tag = node.Variations[t].Tag;
-			for (int i = 0; i < m_firstPass.m_tags.Num(); ++i)
+			for (int32 i = 0; i < FirstPass.Tags.Num(); ++i)
 			{
-				if (m_firstPass.m_tags[i].tag == tag)
+				if (FirstPass.Tags[i].Tag == tag)
 				{
 					tagIndex = i;
 				}
@@ -260,7 +260,7 @@ namespace mu
 			if (tagIndex < 0)
 			{
 				FString Msg = FString::Printf(TEXT("Unknown tag found in color variation [%s]."), *tag);
-				m_pErrorLog->GetPrivate()->Add(Msg, ELMT_WARNING, Typed->GetMessageContext());
+				ErrorLog->GetPrivate()->Add(Msg, ELMT_WARNING, Typed->GetMessageContext());
 				continue;
 			}
 
@@ -282,7 +282,7 @@ namespace mu
 			conditional->type = OP_TYPE::CO_CONDITIONAL;
 			conditional->no = currentOp;
 			conditional->yes = variationOp;
-			conditional->condition = m_firstPass.m_tags[tagIndex].genericCondition;
+			conditional->condition = FirstPass.Tags[tagIndex].GenericCondition;
 
 			currentOp = conditional;
 		}
@@ -303,9 +303,9 @@ namespace mu
 		// Source image
 		FImageGenerationOptions ImageOptions;
 		ImageOptions.State = Options.State;
-		if (!m_activeTags.IsEmpty())
+		if (!ActiveTags.IsEmpty())
 		{
-			ImageOptions.ActiveTags = m_activeTags.Last();
+			ImageOptions.ActiveTags = ActiveTags.Last();
 		}
 
 		Ptr<ASTOp> base;
@@ -490,7 +490,7 @@ namespace mu
 		const NodeColourTable& node = *Typed;
 
 		result.op = GenerateTableSwitch<NodeColourTable, ETableColumnType::Color, OP_TYPE::CO_SWITCH>(node,
-			[this, &Options](const NodeColourTable& node, int colIndex, int row, ErrorLog* pErrorLog)
+			[this, &Options](const NodeColourTable& node, int32 colIndex, int32 row, mu::ErrorLog* pErrorLog)
 			{
 				Ptr<NodeColourConstant> CellData = new NodeColourConstant();
 				FVector4f Colour = node.Table->GetPrivate()->Rows[row].Values[colIndex].Color;
@@ -507,7 +507,7 @@ namespace mu
 	{
 		// Log a warning
 		FString Msg = FString::Printf(TEXT("Required connection not found: %s"), strWhere);
-		m_pErrorLog->GetPrivate()->Add(Msg, ELMT_ERROR, errorContext);
+		ErrorLog->GetPrivate()->Add(Msg, ELMT_ERROR, errorContext);
 
 		// Create a constant colour node
 		Ptr<NodeColourConstant> pNode = new NodeColourConstant();
diff --git a/Engine/Plugins/Experimental/Mutable/Source/MutableTools/Private/MuT/CodeGenerator_ExtensionData.cpp b/Engine/Plugins/Experimental/Mutable/Source/MutableTools/Private/MuT/CodeGenerator_ExtensionData.cpp
--- a/Engine/Plugins/Experimental/Mutable/Source/MutableTools/Private/MuT/CodeGenerator_ExtensionData.cpp
+++ b/Engine/Plugins/Experimental/Mutable/Source/MutableTools/Private/MuT/CodeGenerator_ExtensionData.cpp
@@ -75,7 +75,7 @@ class Node;
 			Data = new ExtensionData();
 			
 			// Log an error message
-			m_pErrorLog->GetPrivate()->Add("Constant extension data not set", ELMT_WARNING, Constant->GetMessageContext());
+			ErrorLog->GetPrivate()->Add("Constant extension data not set", ELMT_WARNING, Constant->GetMessageContext());
 		}
 
 		Op->Value = Data;
@@ -149,15 +149,15 @@ class Node;
 		for (int32 VariationIndex = Variation->Variations.Num() - 1; VariationIndex >= 0; --VariationIndex)
 		{
 			const FString& Tag = Variation->Variations[VariationIndex].Tag;
-			const int32 TagIndex = m_firstPass.m_tags.IndexOfByPredicate([Tag](const FirstPassGenerator::FTag& CandidateTag)
+			const int32 TagIndex = FirstPass.Tags.IndexOfByPredicate([Tag](const FirstPassGenerator::FTag& CandidateTag)
 			{
-				return CandidateTag.tag == Tag;
+				return CandidateTag.Tag == Tag;
 			});
 
 			if (TagIndex == INDEX_NONE)
 			{
 				const FString Msg = FString::Printf(TEXT("Unknown tag found in Extension Data variation [%s]"), *Tag);
-				m_pErrorLog->GetPrivate()->Add(Msg, ELMT_WARNING, Variation->GetMessageContext());
+				ErrorLog->GetPrivate()->Add(Msg, ELMT_WARNING, Variation->GetMessageContext());
 				continue;
 			}
 
@@ -179,7 +179,7 @@ class Node;
 			Conditional->type = OP_TYPE::ED_CONDITIONAL;
 			Conditional->no = CurrentOp;
 			Conditional->yes = VariationOp;
-			Conditional->condition = m_firstPass.m_tags[TagIndex].genericCondition;
+			Conditional->condition = FirstPass.Tags[TagIndex].GenericCondition;
 
 			CurrentOp = Conditional;
 		}
@@ -191,7 +191,7 @@ class Node;
 	{
 		// Log a warning
 		const FString Msg = FString::Printf(TEXT("Required connection not found: %s"), StrWhere);
-		m_pErrorLog->GetPrivate()->Add(Msg, ELMT_ERROR, ErrorContext);
+		ErrorLog->GetPrivate()->Add(Msg, ELMT_ERROR, ErrorContext);
 
 		// Create a constant extension data
 		Ptr<const NodeExtensionDataConstant> Node = new NodeExtensionDataConstant;
diff --git a/Engine/Plugins/Experimental/Mutable/Source/MutableTools/Private/MuT/CodeGenerator_FirstPass.cpp b/Engine/Plugins/Experimental/Mutable/Source/MutableTools/Private/MuT/CodeGenerator_FirstPass.cpp
--- a/Engine/Plugins/Experimental/Mutable/Source/MutableTools/Private/MuT/CodeGenerator_FirstPass.cpp
+++ b/Engine/Plugins/Experimental/Mutable/Source/MutableTools/Private/MuT/CodeGenerator_FirstPass.cpp
@@ -26,10 +26,10 @@
 #include "MuT/NodeObjectNew.h"
 #include "MuT/NodePrivate.h"
 #include "MuT/NodeSurface.h"
-#include "MuT/NodeSurfaceEditPrivate.h"
-#include "MuT/NodeSurfaceNewPrivate.h"
-#include "MuT/NodeSurfaceVariationPrivate.h"
-#include "MuT/NodeSurfaceSwitchPrivate.h"
+#include "MuT/NodeSurfaceEdit.h"
+#include "MuT/NodeSurfaceNew.h"
+#include "MuT/NodeSurfaceVariation.h"
+#include "MuT/NodeSurfaceSwitch.h"
 #include "MuT/NodeScalarEnumParameterPrivate.h"
 
 namespace mu
@@ -40,22 +40,21 @@ namespace mu
 	{
 		// Default conditions when there is no restriction accumulated.
 		FConditionContext noCondition;
-        m_currentCondition.Add(noCondition);
-        m_currentStateCondition.Add(StateCondition());
+        CurrentCondition.Add(noCondition);
+        CurrentStateCondition.Add(StateCondition());
 	}
 
 
 	//---------------------------------------------------------------------------------------------
-    void FirstPassGenerator::Generate( Ptr<ErrorLog> InErrorLog,
+    void FirstPassGenerator::Generate( Ptr<mu::ErrorLog> InErrorLog,
                                        const Node* Root,
-                                       bool bIgnoreStates,
+									   bool bIgnoreStates,
 									   CodeGenerator* InGenerator )
 	{
 		MUTABLE_CPUPROFILER_SCOPE(FirstPassGenerate);
 
 		Generator = InGenerator;
-		m_pErrorLog = InErrorLog;
-        m_ignoreStates = bIgnoreStates;
+		ErrorLog = InErrorLog;
 
 		// Step 1: collect all objects, surfaces and object conditions
         if (Root)
@@ -64,16 +63,16 @@ namespace mu
  		}
 
 		// Step 2: Collect all tags and a list of the surfaces that activate them
-		for (int32 s=0; s<surfaces.Num(); ++s)
+		for (int32 s=0; s<Surfaces.Num(); ++s)
 		{
 			// Collect the tags in new surfaces
-			for (int32 t=0; t<surfaces[s].node->GetPrivate()->m_tags.Num(); ++t)
+			for (int32 t=0; t< Surfaces[s].Node->Tags.Num(); ++t)
 			{
 				int32 tag = -1;
-                const FString& tagStr = surfaces[s].node->GetPrivate()->m_tags[t];
-                for (int32 i = 0; i<m_tags.Num() && tag<0; ++i)
+                const FString& tagStr = Surfaces[s].Node->Tags[t];
+                for (int32 i = 0; i<Tags.Num() && tag<0; ++i)
 				{
-                    if (m_tags[i].tag == tagStr)
+                    if (Tags[i].Tag == tagStr)
 					{
 						tag = i;
 					}
@@ -82,30 +81,30 @@ namespace mu
 				// New tag?
 				if (tag < 0)
 				{
-                    tag = m_tags.Num();
+                    tag = Tags.Num();
 					FTag newTag;
-                    newTag.tag = tagStr;
-                    m_tags.Add(newTag);
+                    newTag.Tag = tagStr;
+                    Tags.Add(newTag);
 				}
 
-                if (m_tags[tag].surfaces.Find(s)==INDEX_NONE)
+                if (Tags[tag].Surfaces.Find(s)==INDEX_NONE)
 				{
-                    m_tags[tag].surfaces.Add(s);
+                    Tags[tag].Surfaces.Add(s);
 				}
 			}
 
             // Collect the tags in edit surfaces
-            for (int32 e=0; e<surfaces[s].edits.Num(); ++e)
+            for (int32 e=0; e< Surfaces[s].Edits.Num(); ++e)
             {
-                const FSurface::FEdit& edit = surfaces[s].edits[e];
-                for (int32 t=0; t<edit.node->GetPrivate()->m_tags.Num(); ++t)
+                const FSurface::FEdit& edit = Surfaces[s].Edits[e];
+                for (int32 t=0; t<edit.Node->Tags.Num(); ++t)
                 {
                     int32 tag = -1;
-					const FString& tagStr = edit.node->GetPrivate()->m_tags[t];
+					const FString& tagStr = edit.Node->Tags[t];
 
-                    for (int32 i = 0; i<m_tags.Num() && tag<0; ++i)
+                    for (int32 i = 0; i<Tags.Num() && tag<0; ++i)
                     {
-                        if (m_tags[i].tag == tagStr)
+                        if (Tags[i].Tag == tagStr)
                         {
                             tag = i;
                         }
@@ -114,15 +113,15 @@ namespace mu
                     // New tag?
                     if (tag < 0)
                     {
-                        tag = m_tags.Num();
+                        tag = Tags.Num();
 						FTag newTag;
-                        newTag.tag = tagStr;
-                        m_tags.Add(newTag);
+                        newTag.Tag = tagStr;
+                        Tags.Add(newTag);
                     }
 
-                    if (m_tags[tag].edits.Find({s,e}) == INDEX_NONE)
+                    if (Tags[tag].Edits.Find({s,e}) == INDEX_NONE)
                     {
-                        m_tags[tag].edits.Add({s,e});
+                        Tags[tag].Edits.Add({s,e});
                     }
                 }
             }
@@ -132,14 +131,14 @@ namespace mu
         // Step 3: Create default state if necessary
         if ( bIgnoreStates )
         {
-            m_states.Empty();
+            States.Empty();
         }
 
-        if ( m_states.IsEmpty() )
+        if ( States.IsEmpty() )
         {
             FObjectState data;
             data.Name = "Default";
-            m_states.Emplace( data, Root );
+            States.Emplace( data, Root );
         }
 	}
 
@@ -199,13 +198,13 @@ namespace mu
 
 		// Add the data about this modifier
 		FModifier thisData;
-		thisData.node = InNode;
-		thisData.objectCondition = m_currentCondition.Last().objectCondition;
-		thisData.stateCondition = m_currentStateCondition.Last();
+		thisData.Node = InNode;
+		thisData.ObjectCondition = CurrentCondition.Last().ObjectCondition;
+		thisData.StateCondition = CurrentStateCondition.Last();
 		thisData.LOD = CurrentLOD;
-		thisData.positiveTags = m_currentPositiveTags;
-		thisData.negativeTags = m_currentNegativeTags;
-		modifiers.Add(thisData);
+		thisData.PositiveTags = CurrentPositiveTags;
+		thisData.NegativeTags = CurrentNegativeTags;
+		Modifiers.Add(thisData);
 	}
 
 
@@ -214,33 +213,31 @@ namespace mu
 	{
 		// Add the data about this surface
 		FSurface thisData;
-		thisData.node = InNode;
-		thisData.Component = m_currentComponent;
+		thisData.Node = InNode;
+		thisData.Component = CurrentComponent;
 		thisData.LOD = CurrentLOD;
-		thisData.objectCondition = m_currentCondition.Last().objectCondition;
-		thisData.stateCondition = m_currentStateCondition.Last();
-		thisData.positiveTags = m_currentPositiveTags;
-		thisData.negativeTags = m_currentNegativeTags;
-		surfaces.Add(thisData);
+		thisData.ObjectCondition = CurrentCondition.Last().ObjectCondition;
+		thisData.StateCondition = CurrentStateCondition.Last();
+		thisData.PositiveTags = CurrentPositiveTags;
+		thisData.NegativeTags = CurrentNegativeTags;
+		Surfaces.Add(thisData);
 	}
 
 
 	//---------------------------------------------------------------------------------------------
     void FirstPassGenerator::Generate_SurfaceEdit(const NodeSurfaceEdit* InNode)
 	{
-		const NodeSurfaceEdit::Private* Private = InNode->GetPrivate();
-
 		// Store a reference to this node in the surface data for the surface that this node is
 		// editing.
-		FSurface* Surface = surfaces.FindByPredicate([&Private](const FSurface& s)
+		FSurface* Surface = Surfaces.FindByPredicate([&InNode](const FSurface& s)
         {
             // Are we editing the main surface node of this surface?
-            if (s.node.get() == Private->m_pParent.get()) return true;
+            if (s.Node.get() == InNode->Parent.get()) return true;
 
             // Are we editing an edit node modifying this surface?
-            for (const auto& e: s.edits)
+            for (const FSurface::FEdit& e: s.Edits)
             {
-                if (Private->m_pParent && e.node==Private->m_pParent)
+                if (InNode->Parent && e.Node== InNode->Parent)
                 {
                     return true;
                 }
@@ -254,15 +251,15 @@ namespace mu
 		if (Surface)
 		{
 			FSurface::FEdit edit;
-			edit.PositiveTags = m_currentPositiveTags;
-			edit.NegativeTags = m_currentNegativeTags;
-            edit.node = InNode;
-            edit.condition = m_currentCondition.Last().objectCondition;
-			Surface->edits.Add(edit);
+			edit.PositiveTags = CurrentPositiveTags;
+			edit.NegativeTags = CurrentNegativeTags;
+            edit.Node = InNode;
+            edit.Condition = CurrentCondition.Last().ObjectCondition;
+			Surface->Edits.Add(edit);
 		}
 		else
 		{
-			m_pErrorLog->GetPrivate()->Add("Missing parent object for edit node.", ELMT_WARNING, InNode->GetMessageContext());
+			ErrorLog->GetPrivate()->Add("Missing parent object for edit node.", ELMT_WARNING, InNode->GetMessageContext());
 		}
 	}
 
@@ -270,52 +267,50 @@ namespace mu
 	//---------------------------------------------------------------------------------------------
 	void FirstPassGenerator::Generate_SurfaceVariation(const NodeSurfaceVariation* InNode)
 	{
-		const NodeSurfaceVariation::Private* Private = InNode->GetPrivate();
-
-        switch(Private->m_type)
+        switch(InNode->Type)
         {
 
         case NodeSurfaceVariation::VariationType::Tag:
         {
             // Any of the tags in the variations would prevent the default surface
-            auto oldNegativeTags = m_currentNegativeTags;
-            for (int32 v=0; v< Private->m_variations.Num(); ++v)
+			TArray<FString> OldNegativeTags = CurrentNegativeTags;
+            for (int32 v=0; v< InNode->Variations.Num(); ++v)
             {
-                m_currentNegativeTags.Add(Private->m_variations[v].m_tag);
+                CurrentNegativeTags.Add(InNode->Variations[v].Tag);
             }
 
-            for(const auto& n: Private->m_defaultSurfaces)
+            for(const Ptr<NodeSurface>& n: InNode->DefaultSurfaces)
             {
 				Generate_Generic(n.get());
 			}
-            for(const auto& n: Private->m_defaultModifiers)
+            for(const Ptr<NodeModifier>& n: InNode->DefaultModifiers)
             {
 				Generate_Modifier(n.get());
             }
 
-            m_currentNegativeTags = oldNegativeTags;
+            CurrentNegativeTags = OldNegativeTags;
 
-            for (int32 v=0; v< Private->m_variations.Num(); ++v)
+            for (int32 v=0; v< InNode->Variations.Num(); ++v)
             {
-                m_currentPositiveTags.Add(Private->m_variations[v].m_tag);
-                for (const auto& s : Private->m_variations[v].m_surfaces)
+                CurrentPositiveTags.Add(InNode->Variations[v].Tag);
+                for (const Ptr<NodeSurface>& s : InNode->Variations[v].Surfaces)
                 {
 					Generate_Generic(s.get());
                 }
 
-                for (const auto& s : Private->m_variations[v].m_modifiers)
+                for (const Ptr<NodeModifier>& s : InNode->Variations[v].Modifiers)
                 {
 					Generate_Modifier(s.get());
 				}
 
-                m_currentPositiveTags.Pop();
+                CurrentPositiveTags.Pop();
 
                 // Tags have an order in a variation node: the current tag should prevent any following
                 // variation surface
-                m_currentNegativeTags.Add(Private->m_variations[v].m_tag);
+                CurrentNegativeTags.Add(InNode->Variations[v].Tag);
             }
 
-            m_currentNegativeTags = oldNegativeTags;
+            CurrentNegativeTags = OldNegativeTags;
 
             break;
         }
@@ -323,7 +318,7 @@ namespace mu
 
         case NodeSurfaceVariation::VariationType::State:
         {
-            size_t stateCount = m_states.Num();
+            int32 stateCount = States.Num();
 
             // Default
             {
@@ -332,15 +327,15 @@ namespace mu
                 {
 					StateCondition AllTrue;
 					AllTrue.Init(true,stateCount);
-                    defaultStates = m_currentStateCondition.Last().IsEmpty()
+                    defaultStates = CurrentStateCondition.Last().IsEmpty()
                             ? AllTrue
-                            : m_currentStateCondition.Last();
+                            : CurrentStateCondition.Last();
 
-                    for (const auto& v: Private->m_variations)
+                    for (const NodeSurfaceVariation::FVariation& v: InNode->Variations)
                     {
                         for( size_t s=0; s<stateCount; ++s )
                         {
-                            if (m_states[s].Key.Name==v.m_tag)
+                            if (States[s].Key.Name==v.Tag)
                             {
                                 // Remove this state from the default options, since it has its own variation
                                 defaultStates[s] = false;
@@ -349,47 +344,47 @@ namespace mu
                     }
                 }
 
-                m_currentStateCondition.Add(defaultStates);
+                CurrentStateCondition.Add(defaultStates);
 
-                for (const auto& n : Private->m_defaultSurfaces)
+                for (const Ptr<NodeSurface>& n : InNode->DefaultSurfaces)
                 {
 					Generate_Generic(n.get());
 				}
-                for (const auto& n : Private->m_defaultModifiers)
+                for (const Ptr<NodeModifier>& n : InNode->DefaultModifiers)
                 {
 					Generate_Modifier(n.get());
 				}
 
-                m_currentStateCondition.Pop();
+                CurrentStateCondition.Pop();
             }
 
             // Variation branches
-            for (const auto& v: Private->m_variations)
+            for (const auto& v: InNode->Variations)
             {
                 // Store the states for this variation here
 				StateCondition variationStates;
 				variationStates.Init(false,stateCount);
 
-                for( size_t s=0; s<stateCount; ++s )
+                for( int32 StateIndex=0; StateIndex<stateCount; ++StateIndex)
                 {
-                    if (m_states[s].Key.Name==v.m_tag)
+                    if (States[StateIndex].Key.Name==v.Tag)
                     {
-                        variationStates[s] = true;
+                        variationStates[StateIndex] = true;
                     }
                 }
 
-                m_currentStateCondition.Add(variationStates);
+                CurrentStateCondition.Add(variationStates);
 
-                for (const auto& n : v.m_surfaces)
+                for (const Ptr<NodeSurface>& n : v.Surfaces)
                 {
 					Generate_Generic(n.get());
 				}
-                for (const auto& n : v.m_modifiers)
+                for (const Ptr<NodeModifier>& n : v.Modifiers)
                 {
 					Generate_Modifier(n.get());
 				}
 
-                m_currentStateCondition.Pop();
+                CurrentStateCondition.Pop();
             }
 
             break;
@@ -406,9 +401,7 @@ namespace mu
 	//---------------------------------------------------------------------------------------------
 	void FirstPassGenerator::Generate_SurfaceSwitch(const NodeSurfaceSwitch* InNode)
 	{
-		const NodeSurfaceSwitch::Private* Private = InNode->GetPrivate();
-
-		if (Private->Options.Num() == 0)
+		if (InNode->Options.Num() == 0)
 		{
 			// No options in the switch!
 			return;
@@ -417,9 +410,9 @@ namespace mu
 		// Prepare the enumeration parameter
 		CodeGenerator::FGenericGenerationOptions Options;
 		CodeGenerator::FScalarGenerationResult ScalarResult;
-		if (Private->Parameter)
+		if (InNode->Parameter)
 		{
-			Generator->GenerateScalar( ScalarResult, Options, Private->Parameter );
+			Generator->GenerateScalar( ScalarResult, Options, InNode->Parameter );
 		}
 		else
 		{
@@ -428,7 +421,7 @@ namespace mu
 		}
 
 		// Parse the options
-		for (int32 t = 0; t < Private->Options.Num(); ++t)
+		for (int32 t = 0; t < InNode->Options.Num(); ++t)
 		{
 			// Create a comparison operation as the boolean parameter for the child
 			Ptr<ASTOpFixed> ParamOp = new ASTOpFixed();
@@ -437,25 +430,25 @@ namespace mu
 			ParamOp->op.args.BoolEqualScalarConst.constant = (int16)t;
 
 			// Combine the new condition with previous conditions coming from parent objects
-			if (m_currentCondition.Last().objectCondition)
+			if (CurrentCondition.Last().ObjectCondition)
 			{
 				Ptr<ASTOpFixed> op = new ASTOpFixed();
 				op->op.type = OP_TYPE::BO_AND;
-				op->SetChild(op->op.args.BoolBinary.a, m_currentCondition.Last().objectCondition);
+				op->SetChild(op->op.args.BoolBinary.a, CurrentCondition.Last().ObjectCondition);
 				op->SetChild(op->op.args.BoolBinary.b, ParamOp);
 				ParamOp = op;
 			}
 
 			FConditionContext data;
-			data.objectCondition = ParamOp;
-			m_currentCondition.Push(data);
+			data.ObjectCondition = ParamOp;
+			CurrentCondition.Push(data);
 
-			if (Private->Options[t])
+			if (InNode->Options[t])
 			{
-				Generate_Generic(Private->Options[t].get());
+				Generate_Generic(InNode->Options[t].get());
 			}
 
-			m_currentCondition.Pop();
+			CurrentCondition.Pop();
 		}
 	}
 
@@ -463,7 +456,7 @@ namespace mu
 	//---------------------------------------------------------------------------------------------
 	void FirstPassGenerator::Generate_ComponentNew(const NodeComponentNew* InNode)
 	{
-        m_currentComponent = InNode;
+        CurrentComponent = InNode;
 
 		CurrentLOD = 0;
 		for (const Ptr<NodeLOD>& c : InNode->LODs)
@@ -476,14 +469,14 @@ namespace mu
 		}
 		CurrentLOD = -1;
 
-		m_currentComponent = nullptr;
+		CurrentComponent = nullptr;
 	}
 
 
 	//---------------------------------------------------------------------------------------------
 	void FirstPassGenerator::Generate_ComponentEdit(const NodeComponentEdit* InNode)
 	{
-		m_currentComponent = InNode->GetParentComponentNew();
+		CurrentComponent = InNode->GetParentComponentNew();
 
 		CurrentLOD = 0;
 		for (const Ptr<NodeLOD>& c : InNode->LODs)
@@ -496,7 +489,7 @@ namespace mu
 		}
 		CurrentLOD = -1;
 
-		m_currentComponent = nullptr;
+		CurrentComponent = nullptr;
 	}
 
 
@@ -526,21 +519,21 @@ namespace mu
 	{
 		// Add the data about this object
 		FObject thisData;
-		thisData.node = InNode;
-        thisData.condition = m_currentCondition.Last().objectCondition;
-		objects.Add(thisData);
+		thisData.Node = InNode;
+        thisData.Condition = CurrentCondition.Last().ObjectCondition;
+		Objects.Add(thisData);
 
         // Accumulate the model states
         for ( const FObjectState& s: InNode->States )
         {
-            m_states.Emplace( s, InNode );
+            States.Emplace( s, InNode );
 
             if ( s.RuntimeParams.Num() > MUTABLE_MAX_RUNTIME_PARAMETERS_PER_STATE )
             {
                 FString Msg = FString::Printf( TEXT("State [%s] has more than %d runtime parameters. Their update may fail."), 
 					*s.Name,
                     MUTABLE_MAX_RUNTIME_PARAMETERS_PER_STATE);
-                m_pErrorLog->GetPrivate()->Add(Msg, ELMT_ERROR, InNode->GetMessageContext());
+				ErrorLog->GetPrivate()->Add(Msg, ELMT_ERROR, InNode->GetMessageContext());
             }
         }
 
@@ -612,7 +605,7 @@ namespace mu
 						{
 							FString Msg = FString::Printf(TEXT("The Group Node [%s] has type Toggle and its direct child is a Group node, which is not allowed. Change the type or add a Child Object node in between them."),
 								*Private->Name);
-							m_pErrorLog->GetPrivate()->Add(Msg, ELMT_ERROR, InNode->GetMessageContext());
+							ErrorLog->GetPrivate()->Add(Msg, ELMT_ERROR, InNode->GetMessageContext());
 						}
 						else
 						{
@@ -670,22 +663,22 @@ namespace mu
                 }
 
                 // Combine the new condition with previous conditions coming from parent objects
-                if (m_currentCondition.Last().objectCondition)
+                if (CurrentCondition.Last().ObjectCondition)
                 {
                     Ptr<ASTOpFixed> op = new ASTOpFixed();
                     op->op.type = OP_TYPE::BO_AND;
-                    op->SetChild( op->op.args.BoolBinary.a,m_currentCondition.Last().objectCondition);
+                    op->SetChild( op->op.args.BoolBinary.a,CurrentCondition.Last().ObjectCondition);
                     op->SetChild( op->op.args.BoolBinary.b,paramOp);
                     paramOp = op;
                 }
 
 				FConditionContext data;
-                data.objectCondition = paramOp;
-                m_currentCondition.Add( data );
+                data.ObjectCondition = paramOp;
+                CurrentCondition.Add( data );
 
 				Generate_Generic(pChildNode.get());
 
-                m_currentCondition.Pop();
+                CurrentCondition.Pop();
             }
         }
  	}
diff --git a/Engine/Plugins/Experimental/Mutable/Source/MutableTools/Private/MuT/CodeGenerator_FirstPass.h b/Engine/Plugins/Experimental/Mutable/Source/MutableTools/Private/MuT/CodeGenerator_FirstPass.h
--- a/Engine/Plugins/Experimental/Mutable/Source/MutableTools/Private/MuT/CodeGenerator_FirstPass.h
+++ b/Engine/Plugins/Experimental/Mutable/Source/MutableTools/Private/MuT/CodeGenerator_FirstPass.h
@@ -10,6 +10,8 @@
 #include "MuT/Node.h"
 #include "MuT/NodeComponentEdit.h"
 #include "MuT/NodeComponentNew.h"
+#include "MuT/NodeComponentSwitch.h"
+#include "MuT/NodeComponentVariation.h"
 #include "MuT/NodeLOD.h"
 #include "MuT/NodeModifier.h"
 #include "MuT/NodeModifierMeshClipDeform.h"
@@ -63,10 +65,10 @@ namespace mu
 		//! Store the conditions that will enable or disable every object
 		struct FObject
 		{
-			const NodeObjectNew* node;
-            Ptr<ASTOp> condition;
+			const NodeObjectNew* Node = nullptr;
+            Ptr<ASTOp> Condition;
 		};
-		TArray<FObject> objects;
+		TArray<FObject> Objects;
 
         //! Type used to represent the activation conditions regarding states
         //! This is the state mask for the states in which this surface must be added. If it
@@ -82,7 +84,7 @@ namespace mu
         //! it is reached with.
 		struct FSurface
 		{
-            Ptr<const NodeSurfaceNew> node;
+            Ptr<const NodeSurfaceNew> Node;
 
 			// Parent Component where this surface will be added. It may be different from the 
 			// Component that defined it (if it was an edit component).
@@ -90,95 +92,95 @@ namespace mu
 			int32 LOD = 0;
 
             // List of tags that are required for the presence of this surface
-			TArray<FString> positiveTags;
+			TArray<FString> PositiveTags;
 
             // List of tags that block the presence of this surface
-			TArray<FString> negativeTags;
+			TArray<FString> NegativeTags;
 
 			// This conditions is the condition of the object defining this surface which may not
 			// be the parent object where this surface will be added.
-            Ptr<ASTOp> objectCondition;
+            Ptr<ASTOp> ObjectCondition;
 
             // This is filled in the first pass.
-            StateCondition stateCondition;
+            StateCondition StateCondition;
 
             // Condition for this surface to be enabled when all the object conditions are met.
             // This is filled in CodeGenerator_SecondPass.
-            Ptr<ASTOp> surfaceCondition;
+            Ptr<ASTOp> SurfaceCondition;
 
 			// All surface editing nodes that edit this surface
             struct FEdit
             {
-            	// List of tags that are required for the presence of this surface
+				//! Reference to the edit node, used during compilation.
+				const NodeSurfaceEdit* Node = nullptr;
+				
+				// List of tags that are required for the presence of this surface
             	TArray<FString> PositiveTags;
 
             	// List of tags that block the presence of this surface
             	TArray<FString> NegativeTags;
             	
                 //! Condition that enables the effects of this edit node on the surface
-                Ptr<ASTOp> condition;
-
-                //! Reference to the edit node, used during compilation.
-                const NodeSurfaceEdit* node = nullptr;
+                Ptr<ASTOp> Condition;
             };
-			TArray<FEdit> edits;
+			TArray<FEdit> Edits;
 
             // This is filled in the final code generation pass
-            Ptr<ASTOp> resultSurfaceOp;
-            Ptr<ASTOp> resultMeshOp;
+            Ptr<ASTOp> ResultSurfaceOp;
+            Ptr<ASTOp> ResultMeshOp;
         };
-		TArray<FSurface> surfaces;
+		TArray<FSurface> Surfaces;
 
 		//! Store the conditions that enable every modifier.
 		struct FModifier
 		{
-            const NodeModifier* node = nullptr;
+            const NodeModifier* Node = nullptr;
 
             // List of tags that are required to apply this modifier
-			TArray<FString> positiveTags;
+			TArray<FString> PositiveTags;
 
             // List of tags that block the activation of this modifier
-			TArray<FString> negativeTags;
+			TArray<FString> NegativeTags;
 
             // This conditions is the condition of the object defining this modifier which may not
             // be the parent object where this surface will be added.
-            Ptr<ASTOp> objectCondition;
+            Ptr<ASTOp> ObjectCondition;
 
             // This conditions is the condition for this modifier to be enabled when all the object conditions are met.
             // This is filled in CodeGenerator_SecondPass.
-            Ptr<ASTOp> surfaceCondition;
+            Ptr<ASTOp> SurfaceCondition;
 
             // This is filled in CodeGenerator_SecondPass.
-            StateCondition stateCondition;
+            StateCondition StateCondition;
 
             //
             int32 LOD = 0;
         };
-		TArray<FModifier> modifiers;
+		TArray<FModifier> Modifiers;
 
 		//! Info about all found tags.
 		struct FTag
 		{
-			FString tag;
+			FString Tag;
 
             // Surfaces that activate the tag. These are indices to the FirstPassGenerator::surfaces
             // vector.
-			TArray<int> surfaces;
+			TArray<int32> Surfaces;
 
             // Edit Surfaces that activate the tag. These first element of the pair are indices to
             // the FirstPassGenerator::surfaces vector. The second element are indices to the
             // "edits" in the specific surface.
-			TArray<TPair<int,int>> edits;
+			TArray<TPair<int32,int32>> Edits;
 
             // This conditions is the condition for this tag to be enabled considering no other
             // condition. This is filled in CodeGenerator_SecondPass.
-            Ptr<ASTOp> genericCondition;
+            Ptr<ASTOp> GenericCondition;
         };
-        TArray<FTag> m_tags;
+        TArray<FTag> Tags;
 
-        //! Accumulate the model states found while generating code, with their generated root nodes.
+        /** Accumulate the model states found while generating code, with their generated root nodes. */
         typedef TArray< TPair<FObjectState, const Node*> > StateList;
-        StateList m_states;
+        StateList States;
 
 		/** Parameters added for every node. */
 		TMap< Ptr<const Node>, Ptr<ASTOpParameter> > ParameterNodes;
@@ -187,19 +189,19 @@ namespace mu
 
         struct FConditionContext
         {
-            Ptr<ASTOp> objectCondition;
+            Ptr<ASTOp> ObjectCondition;
         };
-		TArray< FConditionContext > m_currentCondition;
+		TArray< FConditionContext > CurrentCondition;
 
         //!
-		TArray< StateCondition > m_currentStateCondition;
+		TArray< StateCondition > CurrentStateCondition;
 
-		//! When processing surfaces, this is the parent component the surfaces may be added to
-        const NodeComponentNew* m_currentComponent = nullptr;
+		/** When processing surfaces, this is the parent component the surfaces may be added to. */
+        const NodeComponentNew* CurrentComponent = nullptr;
 
         //! Current relevant tags so far. Used during traversal.
-		TArray<FString> m_currentPositiveTags;
-		TArray<FString> m_currentNegativeTags;
+		TArray<FString> CurrentPositiveTags;
+		TArray<FString> CurrentNegativeTags;
 
 		//** Index of the LOD we are processing. */
         int32 CurrentLOD = -1;
@@ -208,10 +210,7 @@ namespace mu
 		CodeGenerator* Generator = nullptr;
 
         //!
-        ErrorLogPtr m_pErrorLog;
-
-        //!
-        bool m_ignoreStates = false;
+        Ptr<ErrorLog> ErrorLog;
 	};
 
 }
diff --git a/Engine/Plugins/Experimental/Mutable/Source/MutableTools/Private/MuT/CodeGenerator_Image.cpp b/Engine/Plugins/Experimental/Mutable/Source/MutableTools/Private/MuT/CodeGenerator_Image.cpp
--- a/Engine/Plugins/Experimental/Mutable/Source/MutableTools/Private/MuT/CodeGenerator_Image.cpp
+++ b/Engine/Plugins/Experimental/Mutable/Source/MutableTools/Private/MuT/CodeGenerator_Image.cpp
@@ -121,7 +121,7 @@ namespace mu
 		FGeneratedImageCacheKey key;
 		key.Options = Options;
 		key.Node = Untyped;
-		GeneratedImagesMap::ValueType* CachedPtr = m_generatedImages.Find(key);
+		GeneratedImagesMap::ValueType* CachedPtr = GeneratedImages.Find(key);
 		if (CachedPtr)
 		{
 			Result = *CachedPtr;
@@ -161,7 +161,7 @@ namespace mu
 			}
 
 			// Cache the Result
-			m_generatedImages.Add(key, Result);
+			GeneratedImages.Add(key, Result);
 		}
 	}
 
@@ -184,7 +184,7 @@ namespace mu
             pImage = GenerateMissingImage(EImageFormat::IF_RGB_UBYTE );
 
             // Log an error message
-            m_pErrorLog->GetPrivate()->Add( "Constant image not set.", ELMT_WARNING, InNode->GetMessageContext());
+            ErrorLog->GetPrivate()->Add( "Constant image not set.", ELMT_WARNING, InNode->GetMessageContext());
         }
 
 
@@ -209,7 +209,7 @@ namespace mu
 		{
 			Ptr<ASTOpConstantResource> op = new ASTOpConstantResource();
 			op->Type = OP_TYPE::IM_CONSTANT;
-			op->SetValue(pImage, m_compilerOptions->OptimisationOptions.DiskCacheContext);
+			op->SetValue(pImage, CompilerOptions->OptimisationOptions.DiskCacheContext);
 			op->SourceDataDescriptor = InNode->SourceDataDescriptor;
 			Result.op = op;
 		}
@@ -261,7 +261,7 @@ namespace mu
 
         Ptr<ASTOpParameter> op;
 
-		Ptr<ASTOpParameter>* it = m_firstPass.ParameterNodes.Find( node.m_pNode );
+		Ptr<ASTOpParameter>* it = FirstPass.ParameterNodes.Find( node.m_pNode );
         if ( !it )
         {
             op = new ASTOpParameter();
@@ -281,7 +281,7 @@ namespace mu
 				op->ranges.Emplace(op.get(), rangeResult.sizeOp, rangeResult.rangeName, rangeResult.rangeUID);
 			}
 
-			m_firstPass.ParameterNodes.Add(node.m_pNode, op);
+			FirstPass.ParameterNodes.Add(node.m_pNode, op);
 		}
         else
         {
@@ -696,7 +696,7 @@ namespace mu
         }
 
         // Target images
-        int numTargets = 0;
+        int32 numTargets = 0;
 
 		UE::Math::TIntVector2<int32> FinalRectSize = Options.RectSize;
 
@@ -1417,11 +1417,11 @@ namespace mu
 			}
 			else
 			{
-				m_pErrorLog->GetPrivate()->Add("Layout or block index error.", ELMT_ERROR, InNode->GetMessageContext());
+				ErrorLog->GetPrivate()->Add("Layout or block index error.", ELMT_ERROR, InNode->GetMessageContext());
 			}
 
 
-			Ptr<ASTOp> CurrentMeshToProjectOp = MeshResult.meshOp;
+			Ptr<ASTOp> CurrentMeshToProjectOp = MeshResult.MeshOp;
 
             if (projectorResult.type == PROJECTOR_TYPE::WRAPPING)
             {
@@ -1432,7 +1432,7 @@ namespace mu
 				Ptr<Mesh> FormatMeshResult = new Mesh();
 				CreateMeshOptimisedForWrappingProjection(FormatMeshResult.get(), node.m_layout);
 
-                cop->SetValue(FormatMeshResult, m_compilerOptions->OptimisationOptions.DiskCacheContext);
+                cop->SetValue(FormatMeshResult, CompilerOptions->OptimisationOptions.DiskCacheContext);
 
                 Ptr<ASTOpMeshFormat> FormatOp = new ASTOpMeshFormat();
 				FormatOp->Flags = OP::MeshFormatArgs::Vertex
@@ -1463,7 +1463,7 @@ namespace mu
 				Ptr<Mesh> FormatMeshResult = new Mesh();
                 CreateMeshOptimisedForProjection(FormatMeshResult.get(), node.m_layout);
 
-                cop->SetValue(FormatMeshResult, m_compilerOptions->OptimisationOptions.DiskCacheContext);
+                cop->SetValue(FormatMeshResult, CompilerOptions->OptimisationOptions.DiskCacheContext);
 
                 Ptr<ASTOpMeshFormat> FormatOp = new ASTOpMeshFormat();
 				FormatOp->Flags = OP::MeshFormatArgs::Vertex
@@ -1483,9 +1483,9 @@ namespace mu
             Ptr<const Mesh> TempMesh = new Mesh();
             Ptr<ASTOpConstantResource> cop = new ASTOpConstantResource();
             cop->Type = OP_TYPE::ME_CONSTANT;
-            cop->SetValue(TempMesh, m_compilerOptions->OptimisationOptions.DiskCacheContext);
+            cop->SetValue(TempMesh, CompilerOptions->OptimisationOptions.DiskCacheContext);
 			ProjectOp->SetChild(ProjectOp->op.args.MeshProject.mesh, cop );
-            m_pErrorLog->GetPrivate()->Add( "Projector mesh not set.", ELMT_ERROR, InNode->GetMessageContext() );
+            ErrorLog->GetPrivate()->Add( "Projector mesh not set.", ELMT_ERROR, InNode->GetMessageContext() );
         }
 
 
@@ -1639,7 +1639,7 @@ namespace mu
 			MakeGrowMapOp->Border = MUTABLE_GROW_BORDER_VALUE;
 
 			// If we want to be able to generate progressive mips efficiently, we need mipmaps for the "displacement map".
-			if (m_compilerOptions->OptimisationOptions.bEnableProgressiveImages)
+			if (CompilerOptions->OptimisationOptions.bEnableProgressiveImages)
 			{
 				Ptr<ASTOpImageMipmap> MipMask = new ASTOpImageMipmap;
 				MipMask->Source = MakeGrowMapOp->Mask.child();
@@ -1706,8 +1706,8 @@ namespace mu
 
         if ( blockX && blockY )
         {
-            int32 mipsX = (int)ceilf( logf( (float)blockX )/logf(2.0f) );
-            int32 mipsY = (int)ceilf( logf( (float)blockY )/logf(2.0f) );
+            int32 mipsX = (int32)ceilf( logf( (float)blockX )/logf(2.0f) );
+            int32 mipsY = (int32)ceilf( logf( (float)blockY )/logf(2.0f) );
             op->BlockLevels = (uint8)FMath::Max( mipsX, mipsY );
         }
         else
@@ -1770,13 +1770,13 @@ namespace mu
         }
         
         // Process variations in reverse order, since conditionals are built bottom-up.
-        for ( int t = int( node.m_variations.Num() ) - 1; t >= 0; --t )
+        for ( int32 t = int32( node.m_variations.Num() ) - 1; t >= 0; --t )
         {
-            int tagIndex = -1;
+            int32 tagIndex = -1;
             const FString& tag = node.m_variations[t].m_tag;
-            for ( int i = 0; i < int( m_firstPass.m_tags.Num() ); ++i )
+            for ( int32 i = 0; i < int32( FirstPass.Tags.Num() ); ++i )
             {
-                if ( m_firstPass.m_tags[i].tag == tag )
+                if ( FirstPass.Tags[i].Tag == tag )
                 {
                     tagIndex = i;
                 }
@@ -1786,7 +1786,7 @@ namespace mu
             {
 				FString Msg = FString::Printf(TEXT("Unknown tag found in image variation [%s]."), *tag );
 
-                m_pErrorLog->GetPrivate()->Add( Msg, ELMT_WARNING, InNode->GetMessageContext() );
+                ErrorLog->GetPrivate()->Add( Msg, ELMT_WARNING, InNode->GetMessageContext() );
                 continue;
             }
 
@@ -1808,13 +1808,13 @@ namespace mu
             conditional->type = OP_TYPE::IM_CONDITIONAL;
             conditional->no = currentOp;
             conditional->yes = variationOp;
-            conditional->condition = m_firstPass.m_tags[tagIndex].genericCondition;
+            conditional->condition = FirstPass.Tags[tagIndex].GenericCondition;
 
             currentOp = conditional;
         }
 
         // Make sure all options are the same format and size
-        auto desc = currentOp->GetImageDesc( true );
+        FImageDesc desc = currentOp->GetImageDesc( true );
         if ( desc.m_format == EImageFormat::IF_NONE )
         {
             // TODO: Look for the most generic of the options?
@@ -1849,7 +1849,7 @@ namespace mu
 		const NodeImageTable& node = *InNode;
 
 		Result.op = GenerateTableSwitch<NodeImageTable, ETableColumnType::Image, OP_TYPE::IM_SWITCH>(node,
-			[this, InNode, Options](const NodeImageTable& node, int colIndex, int row, ErrorLog* pErrorLog)
+			[this, InNode, Options](const NodeImageTable& node, int32 colIndex, int32 row, mu::ErrorLog* pErrorLog)
 			{
 				const FTableValue& CellData = node.Table->GetPrivate()->Rows[row].Values[colIndex];
 				Ptr<const Image> pImage = nullptr;
@@ -1902,7 +1902,7 @@ namespace mu
 	Ptr<Image> CodeGenerator::GenerateMissingImage(EImageFormat Format)
 	{
 		// Create the image node if it hasn't been created yet.
-		if (!m_missingImage[SIZE_T(Format)])
+		if (!MissingImage[SIZE_T(Format)])
 		{
 			// Make a checkered debug image
 			const FImageSize Size(16, 16);
@@ -1983,10 +1983,10 @@ namespace mu
 
 			}
 
-			m_missingImage[(SIZE_T)Format] = GeneratedImage;
+			MissingImage[(SIZE_T)Format] = GeneratedImage;
 		}
 
-		return m_missingImage[(SIZE_T)Format].get();
+		return MissingImage[(SIZE_T)Format].get();
 	}
 
 
@@ -1994,7 +1994,7 @@ namespace mu
 	{
 		// Log an error message
 		FString Msg = FString::Printf(TEXT("Required connection not found: %s"), strWhere );
-		m_pErrorLog->GetPrivate()->Add( Msg, ELMT_ERROR, errorContext );
+		ErrorLog->GetPrivate()->Add( Msg, ELMT_ERROR, errorContext );
 
 		// Make a checkered debug image
 		Ptr<Image> GeneratedImage = GenerateMissingImage( format );
diff --git a/Engine/Plugins/Experimental/Mutable/Source/MutableTools/Private/MuT/CodeGenerator_Mesh.cpp b/Engine/Plugins/Experimental/Mutable/Source/MutableTools/Private/MuT/CodeGenerator_Mesh.cpp
--- a/Engine/Plugins/Experimental/Mutable/Source/MutableTools/Private/MuT/CodeGenerator_Mesh.cpp
+++ b/Engine/Plugins/Experimental/Mutable/Source/MutableTools/Private/MuT/CodeGenerator_Mesh.cpp
@@ -352,33 +352,33 @@ namespace mu
 		if (!OverlappingBlocks.IsEmpty())
 		{
 			FString Msg = FString::Printf(TEXT("Source mesh has %d layout block overlapping in LOD %d"),
-				OverlappingBlocks.Num() + 1, m_currentParents.Last().m_lod
+				OverlappingBlocks.Num() + 1, CurrentParents.Last().Lod
 			);
-			m_pErrorLog->GetPrivate()->Add(Msg, ELMT_WARNING, errorContext);
+			ErrorLog->GetPrivate()->Add(Msg, ELMT_WARNING, errorContext);
 		}
 
 		// Get the information about the texture coordinates channel
-		int buffer = -1;
-		int channel = -1;
+		int32 buffer = -1;
+		int32 channel = -1;
 		currentLayoutMesh->GetVertexBuffers().FindChannel(MBS_TEXCOORDS,
-			(int)currentLayoutChannel,
+			(int32)currentLayoutChannel,
 			&buffer,
 			&channel);
 		check(buffer >= 0);
 		check(channel >= 0);
 
 		EMeshBufferSemantic semantic;
-		int semanticIndex;
+		int32 semanticIndex;
 		EMeshBufferFormat format;
-		int components;
-		int offset;
+		int32 components;
+		int32 offset;
 		currentLayoutMesh->GetVertexBuffers().GetChannel
 		(buffer, channel, &semantic, &semanticIndex, &format, &components, &offset);
 		check(semantic == MBS_TEXCOORDS);
 
 		uint8* pData = currentLayoutMesh->GetVertexBuffers().GetBufferData(buffer);
-		int elemSize = currentLayoutMesh->GetVertexBuffers().GetElementSize(buffer);
-		int channelOffset = currentLayoutMesh->GetVertexBuffers().GetChannelOffset(buffer, channel);
+		int32 elemSize = currentLayoutMesh->GetVertexBuffers().GetElementSize(buffer);
+		int32 channelOffset = currentLayoutMesh->GetVertexBuffers().GetChannelOffset(buffer, channel);
 		pData += channelOffset;
 
 
@@ -420,8 +420,8 @@ namespace mu
 			// Mutable does not support non-normalized UVs
 			if (bNonNormalizedUVs && !bIsOverlayLayout)
 			{
-				FString Msg = FString::Printf(TEXT("Source mesh has non-normalized UVs in LOD %d"), m_currentParents.Last().m_lod );
-				m_pErrorLog->GetPrivate()->Add(Msg, ELMT_WARNING, errorContext);
+				FString Msg = FString::Printf(TEXT("Source mesh has non-normalized UVs in LOD %d"), CurrentParents.Last().Lod );
+				ErrorLog->GetPrivate()->Add(Msg, ELMT_WARNING, errorContext);
 			}
 		}
 
@@ -453,11 +453,11 @@ namespace mu
 		// TODO: We could skip this if there's only one block and it fits the entire grid
 		for (int32 TriangleIndex = 0; TriangleIndex < NumTriangles; ++TriangleIndex)
 		{
-			uint32 Index0 = int(ItIndices.GetAsUINT32());
+			uint32 Index0 = int32(ItIndices.GetAsUINT32());
 			++ItIndices;
-			uint32 Index1 = int(ItIndices.GetAsUINT32());
+			uint32 Index1 = int32(ItIndices.GetAsUINT32());
 			++ItIndices;
-			uint32 Index2 = int(ItIndices.GetAsUINT32());
+			uint32 Index2 = int32(ItIndices.GetAsUINT32());
 			++ItIndices;
 
 			uint16 X, Y;
@@ -593,7 +593,7 @@ namespace mu
 		}
 
 		// Warn about vertices without a block id
-		if (FirstLODToIgnoreWarnings == -1 || m_currentParents.Last().m_lod < FirstLODToIgnoreWarnings)
+		if (FirstLODToIgnoreWarnings == -1 || CurrentParents.Last().Lod < FirstLODToIgnoreWarnings)
 		{
 			TArray<float> UnassignedUVs;
 			UnassignedUVs.Reserve(NumVertices / 100);
@@ -610,13 +610,13 @@ namespace mu
 
 			if (!UnassignedUVs.IsEmpty())
 			{
-				FString Msg = FString::Printf(TEXT("Source mesh has %d vertices not assigned to any layout block in LOD %d"), UnassignedUVs.Num(), m_currentParents.Last().m_lod);
+				FString Msg = FString::Printf(TEXT("Source mesh has %d vertices not assigned to any layout block in LOD %d"), UnassignedUVs.Num(), CurrentParents.Last().Lod);
 
 				ErrorLogMessageAttachedDataView attachedDataView;
 				attachedDataView.m_unassignedUVs = UnassignedUVs.GetData();
 				attachedDataView.m_unassignedUVsSize = (size_t)UnassignedUVs.Num();
 
-				m_pErrorLog->GetPrivate()->Add(Msg, attachedDataView, ELMT_WARNING, errorContext);
+				ErrorLog->GetPrivate()->Add(Msg, attachedDataView, ELMT_WARNING, errorContext);
 			}
 		}
 
@@ -721,7 +721,7 @@ namespace mu
 		FGeneratedMeshCacheKey Key;
 		Key.Node = InUntypedNode;
 		Key.Options = InOptions;
-        GeneratedMeshMap::ValueType* it = m_generatedMeshes.Find(Key);
+        GeneratedMeshMap::ValueType* it = GeneratedMeshes.Find(Key);
         if ( it )
         {
 			OutResult = *it;
@@ -753,7 +753,7 @@ namespace mu
 		}
 
         // Cache the result
-        m_generatedMeshes.Add( Key, OutResult);
+        GeneratedMeshes.Add( Key, OutResult);
     }
 
 
@@ -784,12 +784,12 @@ namespace mu
         if ( node.Base )
         {
             GenerateMesh(InOptions,BaseResult, node.Base );
-            OpMorph->Base = BaseResult.meshOp;
+            OpMorph->Base = BaseResult.MeshOp;
         }
         else
         {
             // This argument is required
-            m_pErrorLog->GetPrivate()->Add( "Mesh morph base node is not set.",
+            ErrorLog->GetPrivate()->Add( "Mesh morph base node is not set.",
                                             ELMT_ERROR, InMorphNode->GetMessageContext());
         }        
 
@@ -805,7 +805,7 @@ namespace mu
             GenerateMesh(TargetOptions, TargetResult, node.Morph);
 
             // TODO: Make sure that the target is a mesh with the morph format
-            Ptr<ASTOp> target = TargetResult.meshOp;
+            Ptr<ASTOp> target = TargetResult.MeshOp;
 
             OpMorph->Target = target;
         }
@@ -829,8 +829,8 @@ namespace mu
 			OpBind->PhysicsToDeform = node.PhysicsToDeform;
 			OpBind->BindingMethod = static_cast<uint32>(EShapeBindingMethod::ReshapeClosestProject);
             
-			OpBind->Mesh = BaseResult.meshOp;
-            OpBind->Shape = BaseResult.meshOp;
+			OpBind->Mesh = BaseResult.MeshOp;
+            OpBind->Shape = BaseResult.MeshOp;
            
 			OpApply->bReshapeVertices = OpBind->bReshapeVertices;
 			OpApply->bRecomputeNormals = OpBind->bRecomputeNormals;
@@ -847,14 +847,14 @@ namespace mu
 
  		if (OpMorphReshape)
 		{
-			OutResult.meshOp = OpMorphReshape;
+			OutResult.MeshOp = OpMorphReshape;
 		}
 		else
 		{
-			OutResult.meshOp = OpMorph;
+			OutResult.MeshOp = OpMorph;
 		}
 
-        OutResult.baseMeshOp = BaseResult.baseMeshOp;
+        OutResult.BaseMeshOp = BaseResult.BaseMeshOp;
 		OutResult.GeneratedLayouts = BaseResult.GeneratedLayouts;
 }
 
@@ -883,12 +883,12 @@ namespace mu
 			BaseOptions.bLayouts = false;
 			GenerateMesh(BaseOptions, BaseResult, node.m_pBase );
 
-            op->Base = BaseResult.meshOp;
+            op->Base = BaseResult.MeshOp;
         }
         else
         {
             // This argument is required
-            m_pErrorLog->GetPrivate()->Add( "Mesh make morph base node is not set.",
+            ErrorLog->GetPrivate()->Add( "Mesh make morph base node is not set.",
                                             ELMT_ERROR, InMakeMorphNode->GetMessageContext());
         }
 
@@ -902,17 +902,17 @@ namespace mu
 			FMeshGenerationResult TargetResult;
             GenerateMesh( TargetOptions, TargetResult, node.m_pTarget );
 
-            op->Target = TargetResult.meshOp;
+            op->Target = TargetResult.MeshOp;
         }
         else
         {
             // This argument is required
-            m_pErrorLog->GetPrivate()->Add( "Mesh make morph target node is not set.",
+            ErrorLog->GetPrivate()->Add( "Mesh make morph target node is not set.",
                                             ELMT_ERROR, InMakeMorphNode->GetMessageContext());
         }
 
-        OutResult.meshOp = op;
-        OutResult.baseMeshOp = BaseResult.baseMeshOp;
+        OutResult.MeshOp = op;
+        OutResult.BaseMeshOp = BaseResult.BaseMeshOp;
 		OutResult.GeneratedLayouts = BaseResult.GeneratedLayouts;
 	}
 
@@ -935,9 +935,9 @@ namespace mu
             if ( node.m_fragmentType==NodeMeshFragment::FT_LAYOUT_BLOCKS )
             {
                 Ptr<ASTOpMeshExtractLayoutBlocks> op = new ASTOpMeshExtractLayoutBlocks();
-                OutResult.meshOp = op;
+                OutResult.MeshOp = op;
 
-                op->Source = BaseResult.meshOp;
+                op->Source = BaseResult.MeshOp;
 
                 if (BaseResult.GeneratedLayouts.Num()>node.LayoutOrGroup )
                 {
@@ -953,7 +953,7 @@ namespace mu
                         }
                         else
                         {
-                            m_pErrorLog->GetPrivate()->Add( "Internal layout block index error.",
+                            ErrorLog->GetPrivate()->Add( "Internal layout block index error.",
                                                             ELMT_ERROR, FragmentNode->GetMessageContext());
                         }
                     }
@@ -961,7 +961,7 @@ namespace mu
                 else
                 {
                     // This argument is required
-                    m_pErrorLog->GetPrivate()->Add( "Missing layout in mesh fragment source.",
+                    ErrorLog->GetPrivate()->Add( "Missing layout in mesh fragment source.",
                                                     ELMT_ERROR, FragmentNode->GetMessageContext());
                 }
             }
@@ -975,11 +975,11 @@ namespace mu
         else
         {
             // This argument is required
-            m_pErrorLog->GetPrivate()->Add( "Mesh fragment source is not set.",
+            ErrorLog->GetPrivate()->Add( "Mesh fragment source is not set.",
                                             ELMT_ERROR, FragmentNode->GetMessageContext());
         }
 
-        OutResult.baseMeshOp = BaseResult.baseMeshOp;
+        OutResult.BaseMeshOp = BaseResult.BaseMeshOp;
 		OutResult.GeneratedLayouts = BaseResult.GeneratedLayouts;
 	}
 
@@ -993,7 +993,7 @@ namespace mu
         // Generate the code
         Ptr<ASTOpFixed> op = new ASTOpFixed();
         op->op.type = OP_TYPE::ME_INTERPOLATE;
-        OutResult.meshOp = op;
+        OutResult.MeshOp = op;
 
         // Factor
         if ( Node* pFactor = node.m_pFactor.get() )
@@ -1009,7 +1009,7 @@ namespace mu
 
         //
         Ptr<ASTOp> base = 0;
-        int count = 0;
+        int32 count = 0;
         for ( int32 t=0
             ; t<node.m_targets.Num() && t<MUTABLE_OP_MAX_INTERPOLATE_COUNT-1
             ; ++t )
@@ -1026,17 +1026,17 @@ namespace mu
                 // The first target is the base
                 if (count==0)
                 {
-                    base = TargetResult.meshOp;
-                    op->SetChild( op->op.args.MeshInterpolate.base, TargetResult.meshOp );
+                    base = TargetResult.MeshOp;
+                    op->SetChild( op->op.args.MeshInterpolate.base, TargetResult.MeshOp );
 
-                    OutResult.baseMeshOp = TargetResult.baseMeshOp;
+                    OutResult.BaseMeshOp = TargetResult.BaseMeshOp;
 					OutResult.GeneratedLayouts = TargetResult.GeneratedLayouts;
 				}
                 else
                 {
                     Ptr<ASTOpMeshDifference> dop = new ASTOpMeshDifference();
                     dop->Base = base;
-                    dop->Target = TargetResult.meshOp;
+                    dop->Target = TargetResult.MeshOp;
 
                     // \todo Texcoords are broken?
                     dop->bIgnoreTextureCoords = true;
@@ -1063,7 +1063,7 @@ namespace mu
         {
             // TODO
             //op.args.MeshInterpolate.target[0] = GenerateMissingImageCode( "First mesh", IF_RGB_UBYTE );
-            m_pErrorLog->GetPrivate()->Add
+            ErrorLog->GetPrivate()->Add
                 ( "Mesh interpolation: at least the first mesh is required.",
                   ELMT_ERROR, InterpolateNode->GetMessageContext());
         }
@@ -1115,7 +1115,7 @@ namespace mu
 				FMeshGenerationResult BranchResults;
                 GenerateMesh(TargetOptions, BranchResults, node.m_options[t] );
 
-                Ptr<ASTOp> branch = BranchResults.meshOp;
+                Ptr<ASTOp> branch = BranchResults.MeshOp;
                 op->cases.Emplace((int16)t,op,branch);
 
 				if (!bFirstValidConnectionFound)
@@ -1126,7 +1126,7 @@ namespace mu
             }
         }
 
-        OutResult.meshOp = op;
+        OutResult.MeshOp = op;
     }
 
 
@@ -1135,11 +1135,11 @@ namespace mu
 	{
 		//
 		FMeshGenerationResult NewResult = OutResult;
-		int t = 0;
+		int32 t = 0;
 		bool bFirstRowGenerated = false;
 
 		Ptr<ASTOp> Op = GenerateTableSwitch<NodeMeshTable, ETableColumnType::Mesh, OP_TYPE::ME_SWITCH>(*TableNode,
-			[this, &NewResult, &bFirstRowGenerated, &InOptions] (const NodeMeshTable& node, int colIndex, int row, ErrorLog* pErrorLog)
+			[this, &NewResult, &bFirstRowGenerated, &InOptions] (const NodeMeshTable& node, int32 colIndex, int32 row, mu::ErrorLog* pErrorLog)
 			{
 				mu::Ptr<mu::Mesh> pMesh = node.Table->GetPrivate()->Rows[row].Values[colIndex].Mesh;
 				FMeshGenerationResult BranchResults;
@@ -1150,9 +1150,9 @@ namespace mu
 					pCell->SetValue(pMesh);
 
 					// TODO Take into account layout strategy
-					int numLayouts = node.Layouts.Num();
+					int32 numLayouts = node.Layouts.Num();
 					pCell->SetLayoutCount(numLayouts);
-					for (int i = 0; i < numLayouts; ++i)
+					for (int32 i = 0; i < numLayouts; ++i)
 					{
 						pCell->SetLayout(i, node.Layouts[i]);
 					}
@@ -1175,10 +1175,10 @@ namespace mu
 					}
 				}
 
-				return BranchResults.meshOp;
+				return BranchResults.MeshOp;
 			});
 
-		NewResult.meshOp = Op;
+		NewResult.MeshOp = Op;
 
 		OutResult = NewResult;
 	}
@@ -1202,7 +1202,7 @@ namespace mu
 			FMeshGenerationOptions DefaultOptions = InOptions;
 
 			GenerateMesh(DefaultOptions, BranchResults, node.m_defaultMesh );
-            currentMeshOp = BranchResults.meshOp;
+            currentMeshOp = BranchResults.MeshOp;
             currentResult = BranchResults;
             firstOptionProcessed = true;
         }
@@ -1210,11 +1210,11 @@ namespace mu
         // Process variations in reverse order, since conditionals are built bottom-up.
         for ( int32 t = node.m_variations.Num()-1; t >= 0; --t )
         {
-            int tagIndex = -1;
+            int32 tagIndex = -1;
             const FString& tag = node.m_variations[t].m_tag;
-            for ( int i = 0; i < m_firstPass.m_tags.Num(); ++i )
+            for ( int32 i = 0; i < FirstPass.Tags.Num(); ++i )
             {
-                if ( m_firstPass.m_tags[i].tag==tag)
+                if ( FirstPass.Tags[i].Tag==tag)
                 {
                     tagIndex = i;
                 }
@@ -1222,7 +1222,7 @@ namespace mu
 
             if ( tagIndex < 0 )
             {
-                m_pErrorLog->GetPrivate()->Add( 
+                ErrorLog->GetPrivate()->Add( 
 					FString::Printf(TEXT("Unknown tag found in mesh variation [%s]."), *tag),
 					ELMT_WARNING,
 					VariationNode->GetMessageContext(),
@@ -1244,7 +1244,7 @@ namespace mu
                 FMeshGenerationResult BranchResults;
 				GenerateMesh(VariationOptions, BranchResults, node.m_variations[t].m_mesh );
 
-                variationMeshOp = BranchResults.meshOp;
+                variationMeshOp = BranchResults.MeshOp;
 
                 if ( !firstOptionProcessed )
                 {
@@ -1257,13 +1257,13 @@ namespace mu
             conditional->type = OP_TYPE::ME_CONDITIONAL;
             conditional->no = currentMeshOp;
             conditional->yes = variationMeshOp;            
-            conditional->condition = m_firstPass.m_tags[tagIndex].genericCondition;
+            conditional->condition = FirstPass.Tags[tagIndex].GenericCondition;
 
             currentMeshOp = conditional;
         }
 
         OutResult = currentResult;
-        OutResult.meshOp = currentMeshOp;
+        OutResult.MeshOp = currentMeshOp;
     }
 
 
@@ -1274,8 +1274,8 @@ namespace mu
 
         Ptr<ASTOpConstantResource> ConstantOp = new ASTOpConstantResource();
 		ConstantOp->Type = OP_TYPE::ME_CONSTANT;
-		OutResult.baseMeshOp = ConstantOp;
-		OutResult.meshOp = ConstantOp;
+		OutResult.BaseMeshOp = ConstantOp;
+		OutResult.MeshOp = ConstantOp;
 		OutResult.GeneratedLayouts.Empty();
 
 		bool bIsOverridingLayouts = !InOptions.OverrideLayouts.IsEmpty();
@@ -1285,7 +1285,7 @@ namespace mu
 		{
 			// This data is required
 			MeshPtr EmptyMesh = new Mesh();
-			ConstantOp->SetValue(EmptyMesh, m_compilerOptions->OptimisationOptions.DiskCacheContext);
+			ConstantOp->SetValue(EmptyMesh, CompilerOptions->OptimisationOptions.DiskCacheContext);
 			EmptyMesh->MeshIDPrefix = ConstantOp->GetValueHash();
 
 			FGeneratedConstantMesh MeshEntry;
@@ -1294,7 +1294,7 @@ namespace mu
 			GeneratedConstantMeshes.Add(MeshEntry);
 
 			// Log an error message
-			m_pErrorLog->GetPrivate()->Add("Constant mesh not set.", ELMT_WARNING, InNode->GetMessageContext());
+			ErrorLog->GetPrivate()->Add("Constant mesh not set.", ELMT_WARNING, InNode->GetMessageContext());
 
 			return;
 		}
@@ -1306,8 +1306,8 @@ namespace mu
 			ReferenceOp->ID = pMesh->GetReferencedMesh();
 			ReferenceOp->bForceLoad = pMesh->IsForceLoad();
 
-			OutResult.baseMeshOp = ReferenceOp;
-			OutResult.meshOp = ReferenceOp;
+			OutResult.BaseMeshOp = ReferenceOp;
+			OutResult.MeshOp = ReferenceOp;
 
 			return;
 		}
@@ -1395,7 +1395,7 @@ namespace mu
 			Ptr<Mesh> Cloned = pMesh->Clone();
 			Cloned->EnsureSurfaceData();
 
-			ConstantOp->SetValue(Cloned, m_compilerOptions->OptimisationOptions.DiskCacheContext);
+			ConstantOp->SetValue(Cloned, CompilerOptions->OptimisationOptions.DiskCacheContext);
 
 			// Add the unique vertex ID prefix in all cases, since it is free memory-wise
 			uint32 MeshIDPrefix = uint32(ConstantOp->GetValueHash());
@@ -1466,8 +1466,8 @@ namespace mu
 			}
 		}
 
-		OutResult.baseMeshOp = LastMeshOp;
-		OutResult.meshOp = LastMeshOp;
+		OutResult.BaseMeshOp = LastMeshOp;
+		OutResult.MeshOp = LastMeshOp;
 
 		// Add the tags operation
 		if (Tags.Num())
@@ -1480,7 +1480,7 @@ namespace mu
 
 		// Apply the modifier for the pre-normal operations stage.
 		bool bModifiersForBeforeOperations = true;
-		OutResult.meshOp = ApplyMeshModifiers(InOptions, LastMeshOp, bModifiersForBeforeOperations, InNode->GetMessageContext());
+		OutResult.MeshOp = ApplyMeshModifiers(InOptions, LastMeshOp, bModifiersForBeforeOperations, InNode->GetMessageContext());
     }
 
 
@@ -1497,7 +1497,7 @@ namespace mu
 			FMeshGenerationResult baseResult;
 			GenerateMesh(Options,baseResult, node.Source);
             Ptr<ASTOpMeshFormat> op = new ASTOpMeshFormat();
-            op->Source = baseResult.meshOp;
+            op->Source = baseResult.MeshOp;
             op->Flags = 0;
 
             Ptr<Mesh> FormatMesh = new Mesh();
@@ -1521,11 +1521,11 @@ namespace mu
 
             Ptr<ASTOpConstantResource> cop = new ASTOpConstantResource();
             cop->Type = OP_TYPE::ME_CONSTANT;
-            cop->SetValue( FormatMesh, m_compilerOptions->OptimisationOptions.DiskCacheContext );
+            cop->SetValue( FormatMesh, CompilerOptions->OptimisationOptions.DiskCacheContext );
             op->Format = cop;
 
-            OutResult.meshOp = op;
-            OutResult.baseMeshOp = baseResult.baseMeshOp;
+            OutResult.MeshOp = op;
+            OutResult.BaseMeshOp = baseResult.BaseMeshOp;
 			OutResult.GeneratedLayouts = baseResult.GeneratedLayouts;
 		}
         else
@@ -1548,17 +1548,17 @@ namespace mu
         if (node.Source)
         {
             GenerateMesh(InOptions, OutResult, node.Source);
-            op->source = OutResult.meshOp;
+            op->source = OutResult.MeshOp;
         }
         else
         {
             // This argument is required
-            m_pErrorLog->GetPrivate()->Add("Mesh transform base node is not set.", ELMT_ERROR, TransformNode->GetMessageContext() );
+            ErrorLog->GetPrivate()->Add("Mesh transform base node is not set.", ELMT_ERROR, TransformNode->GetMessageContext() );
         }
 
         op->matrix = node.Transform;
 
-        OutResult.meshOp = op;
+        OutResult.MeshOp = op;
     }
 
 
@@ -1575,12 +1575,12 @@ namespace mu
         {
 			FMeshGenerationOptions BaseOptions = InOptions;
             GenerateMesh(BaseOptions, OutResult, node.m_pSource);
-            op->source = OutResult.meshOp;
+            op->source = OutResult.MeshOp;
         }
         else
         {
             // This argument is required
-            m_pErrorLog->GetPrivate()->Add("Mesh clip-morph-plane source node is not set.",
+            ErrorLog->GetPrivate()->Add("Mesh clip-morph-plane source node is not set.",
                 ELMT_ERROR, ClipNode->GetMessageContext());
         }
 
@@ -1630,7 +1630,7 @@ namespace mu
         op->dist = node.m_dist;
         op->factor = node.m_factor;
 
-        OutResult.meshOp = op;
+        OutResult.MeshOp = op;
     }
 
 
@@ -1647,12 +1647,12 @@ namespace mu
         if (node.m_pSource)
         {
             GenerateMesh(InOptions, OutResult, node.m_pSource );
-            op->SetChild( op->op.args.MeshClipWithMesh.source, OutResult.meshOp );
+            op->SetChild( op->op.args.MeshClipWithMesh.source, OutResult.MeshOp );
         }
         else
         {
             // This argument is required
-            m_pErrorLog->GetPrivate()->Add("Mesh clip-with-mesh source node is not set.",
+            ErrorLog->GetPrivate()->Add("Mesh clip-with-mesh source node is not set.",
                 ELMT_ERROR, ClipNode->GetMessageContext());
         }
 
@@ -1666,16 +1666,16 @@ namespace mu
 
             FMeshGenerationResult clipResult;
             GenerateMesh(ClipOptions, clipResult, node.m_pClipMesh);
-            op->SetChild( op->op.args.MeshClipWithMesh.clipMesh, clipResult.meshOp );
+            op->SetChild( op->op.args.MeshClipWithMesh.clipMesh, clipResult.MeshOp );
 		}
         else
         {
             // This argument is required
-            m_pErrorLog->GetPrivate()->Add("Mesh clip-with-mesh clipping mesh node is not set.",
+            ErrorLog->GetPrivate()->Add("Mesh clip-with-mesh clipping mesh node is not set.",
                 ELMT_ERROR, ClipNode->GetMessageContext());
         }
 
-        OutResult.meshOp = op;
+        OutResult.MeshOp = op;
     }
 
 	//---------------------------------------------------------------------------------------------
@@ -1690,12 +1690,12 @@ namespace mu
 		if (Node.m_pBaseMesh)
 		{
 			GenerateMesh(InOptions, Result, Node.m_pBaseMesh);
-			OpBind->Mesh = Result.meshOp;
+			OpBind->Mesh = Result.MeshOp;
 		}
 		else
 		{
 			// This argument is required
-			m_pErrorLog->GetPrivate()->Add("Mesh Clip Deform base mesh node is not set.", ELMT_ERROR, ClipDeform->GetMessageContext());
+			ErrorLog->GetPrivate()->Add("Mesh Clip Deform base mesh node is not set.", ELMT_ERROR, ClipDeform->GetMessageContext());
 		}
 
 		// Base Shape
@@ -1708,13 +1708,13 @@ namespace mu
 
 			FMeshGenerationResult baseResult;
 			GenerateMesh(ClipOptions, baseResult, Node.m_pClipShape);
-			OpBind->Shape = baseResult.meshOp;
-			OpClipDeform->ClipShape = baseResult.meshOp;
+			OpBind->Shape = baseResult.MeshOp;
+			OpClipDeform->ClipShape = baseResult.MeshOp;
 		}
 
 		OpClipDeform->Mesh = OpBind;
 
-		Result.meshOp = OpClipDeform;
+		Result.MeshOp = OpClipDeform;
 	}
 
     //---------------------------------------------------------------------------------------------
@@ -1729,12 +1729,12 @@ namespace mu
         if (node.m_pBase)
         {
             GenerateMesh(InOptions, OutResult, node.m_pBase );
-            op->base = OutResult.meshOp;
+            op->base = OutResult.MeshOp;
         }
         else
         {
             // This argument is required
-            m_pErrorLog->GetPrivate()->Add("Mesh apply-pose base node is not set.",
+            ErrorLog->GetPrivate()->Add("Mesh apply-pose base node is not set.",
                 ELMT_ERROR, PoseNode->GetMessageContext());
         }
 
@@ -1748,16 +1748,16 @@ namespace mu
 
             FMeshGenerationResult poseResult;
             GenerateMesh(PoseOptions, poseResult, node.m_pPose );
-            op->pose = poseResult.meshOp;
+            op->pose = poseResult.MeshOp;
 		}
         else
         {
             // This argument is required
-            m_pErrorLog->GetPrivate()->Add("Mesh apply-pose pose node is not set.",
+            ErrorLog->GetPrivate()->Add("Mesh apply-pose pose node is not set.",
                 ELMT_ERROR, PoseNode->GetMessageContext());
         }
 
-        OutResult.meshOp = op;
+        OutResult.MeshOp = op;
     }
 
 
@@ -1772,12 +1772,12 @@ namespace mu
 		if (node.m_pMeshA)
 		{
 			GenerateMesh(InOptions, OutResult, node.m_pMeshA);
-			op->meshA = OutResult.meshOp;
+			op->meshA = OutResult.MeshOp;
 		}
 		else
 		{
 			// This argument is required
-			m_pErrorLog->GetPrivate()->Add("Mesh geometric op mesh-a node is not set.",
+			ErrorLog->GetPrivate()->Add("Mesh geometric op mesh-a node is not set.",
 				ELMT_ERROR, GeomNode->GetMessageContext());
 		}
 
@@ -1791,13 +1791,13 @@ namespace mu
 
 			FMeshGenerationResult bResult;
 			GenerateMesh(OtherOptions, bResult, node.m_pMeshB);
-			op->meshB = bResult.meshOp;
+			op->meshB = bResult.MeshOp;
 		}
 
 		op->scalarA = Generate_Generic(node.m_pScalarA, InOptions);
 		op->scalarB = Generate_Generic(node.m_pScalarB, InOptions);
 
-		OutResult.meshOp = op;
+		OutResult.MeshOp = op;
 	}
 
 
@@ -1833,12 +1833,12 @@ namespace mu
 		if (Node.BaseMesh)
 		{
 			GenerateMesh(InOptions, OutResult, Node.BaseMesh);
-			OpBind->Mesh = OutResult.meshOp;
+			OpBind->Mesh = OutResult.MeshOp;
 		}
 		else
 		{
 			// This argument is required
-			m_pErrorLog->GetPrivate()->Add("Mesh reshape base node is not set.", ELMT_ERROR, Reshape->GetMessageContext());
+			ErrorLog->GetPrivate()->Add("Mesh reshape base node is not set.", ELMT_ERROR, Reshape->GetMessageContext());
 		}
 
 		// Base and target shapes shouldn't have layouts or modifiers.
@@ -1852,7 +1852,7 @@ namespace mu
 		{
 			FMeshGenerationResult baseResult;
 			GenerateMesh(ShapeOptions, baseResult, Node.BaseShape);
-			OpBind->Shape = baseResult.meshOp;
+			OpBind->Shape = baseResult.MeshOp;
 		}
 
 		OpApply->Mesh = OpBind;
@@ -1862,10 +1862,10 @@ namespace mu
 		{
 			FMeshGenerationResult targetResult;
 			GenerateMesh(ShapeOptions, targetResult, Node.TargetShape);
-			OpApply->Shape = targetResult.meshOp;
+			OpApply->Shape = targetResult.MeshOp;
 		}
 
-		OutResult.meshOp = OpApply;
+		OutResult.MeshOp = OpApply;
 	}
 
 }
diff --git a/Engine/Plugins/Experimental/Mutable/Source/MutableTools/Private/MuT/CodeGenerator_Projector.cpp b/Engine/Plugins/Experimental/Mutable/Source/MutableTools/Private/MuT/CodeGenerator_Projector.cpp
--- a/Engine/Plugins/Experimental/Mutable/Source/MutableTools/Private/MuT/CodeGenerator_Projector.cpp
+++ b/Engine/Plugins/Experimental/Mutable/Source/MutableTools/Private/MuT/CodeGenerator_Projector.cpp
@@ -109,7 +109,7 @@ namespace mu
 
 		Ptr<ASTOpParameter> op;
 
-		Ptr<ASTOpParameter>* it = m_firstPass.ParameterNodes.Find(node.m_pNode);
+		Ptr<ASTOpParameter>* it = FirstPass.ParameterNodes.Find(node.m_pNode);
 		if (!it)
 		{
 			FParameterDesc param;
@@ -133,7 +133,7 @@ namespace mu
 				op->ranges.Emplace(op.get(), rangeResult.sizeOp, rangeResult.rangeName, rangeResult.rangeUID);
 			}
 
-			m_firstPass.ParameterNodes.Add(node.m_pNode, op);
+			FirstPass.ParameterNodes.Add(node.m_pNode, op);
 		}
 		else
 		{
@@ -151,7 +151,7 @@ namespace mu
 		const void* errorContext)
 	{
 		// Log an error message
-		m_pErrorLog->GetPrivate()->Add("Required projector connection not found.",
+		ErrorLog->GetPrivate()->Add("Required projector connection not found.",
 			ELMT_ERROR, errorContext);
 
 		FProjector p;
diff --git a/Engine/Plugins/Experimental/Mutable/Source/MutableTools/Private/MuT/CodeGenerator_Scalar.cpp b/Engine/Plugins/Experimental/Mutable/Source/MutableTools/Private/MuT/CodeGenerator_Scalar.cpp
--- a/Engine/Plugins/Experimental/Mutable/Source/MutableTools/Private/MuT/CodeGenerator_Scalar.cpp
+++ b/Engine/Plugins/Experimental/Mutable/Source/MutableTools/Private/MuT/CodeGenerator_Scalar.cpp
@@ -36,7 +36,6 @@
 #include "MuT/NodeScalarParameter.h"
 #include "MuT/NodeScalarParameterPrivate.h"
 #include "MuT/NodeScalarSwitch.h"
-#include "MuT/NodeScalarSwitchPrivate.h"
 #include "MuT/NodeScalarTable.h"
 #include "MuT/NodeScalarVariation.h"
 #include "MuT/NodeScalarVariationPrivate.h"
@@ -138,7 +137,7 @@ namespace mu
 
 		Ptr<ASTOpParameter> op;
 
-		Ptr<ASTOpParameter>* it = m_firstPass.ParameterNodes.Find(node.m_pNode);
+		Ptr<ASTOpParameter>* it = FirstPass.ParameterNodes.Find(node.m_pNode);
 		if (!it)
 		{
 			FParameterDesc param;
@@ -160,7 +159,7 @@ namespace mu
 				op->ranges.Emplace(op.get(), rangeResult.sizeOp, rangeResult.rangeName, rangeResult.rangeUID);
 			}
 
-			m_firstPass.ParameterNodes.Add(node.m_pNode, op);
+			FirstPass.ParameterNodes.Add(node.m_pNode, op);
 		}
 		else
 		{
@@ -178,7 +177,7 @@ namespace mu
 
 		Ptr<ASTOpParameter> op;
 
-		Ptr<ASTOpParameter>* it = m_firstPass.ParameterNodes.Find(node.m_pNode);
+		Ptr<ASTOpParameter>* it = FirstPass.ParameterNodes.Find(node.m_pNode);
 		if (!it)
 		{
 			FParameterDesc param;
@@ -207,7 +206,7 @@ namespace mu
 				op->ranges.Emplace(op.get(), rangeResult.sizeOp, rangeResult.rangeName, rangeResult.rangeUID);
 			}
 
-			m_firstPass.ParameterNodes.Add(node.m_pNode, op);
+			FirstPass.ParameterNodes.Add(node.m_pNode, op);
 		}
 		else
 		{
@@ -221,9 +220,9 @@ namespace mu
 	//-------------------------------------------------------------------------------------------------
 	void CodeGenerator::GenerateScalar_Switch(FScalarGenerationResult& result, const FGenericGenerationOptions& Options, const Ptr<const NodeScalarSwitch>& Typed)
 	{
-		const NodeScalarSwitch::Private& node = *Typed->GetPrivate();
+		const NodeScalarSwitch& node = *Typed;
 
-		if (node.m_options.Num() == 0)
+		if (node.Options.Num() == 0)
 		{
 			// No options in the switch!
 			Ptr<ASTOp> missingOp = GenerateMissingScalarCode(TEXT("Switch option"),
@@ -237,10 +236,10 @@ namespace mu
 		op->type = OP_TYPE::SC_SWITCH;
 
 		// Variable value
-		if (node.m_pParameter)
+		if (node.Parameter)
 		{
 			FScalarGenerationResult ChildResult;
-			GenerateScalar(ChildResult, Options, node.m_pParameter.get());
+			GenerateScalar(ChildResult, Options, node.Parameter.get());
 			op->variable = ChildResult.op;
 		}
 		else
@@ -250,13 +249,13 @@ namespace mu
 		}
 
 		// Options
-		for (int32 t = 0; t < node.m_options.Num(); ++t)
+		for (int32 t = 0; t < node.Options.Num(); ++t)
 		{
 			Ptr<ASTOp> branch;
-			if (node.m_options[t])
+			if (node.Options[t])
 			{
 				FScalarGenerationResult ChildResult;
-				GenerateScalar(ChildResult, Options, node.m_options[t].get());
+				GenerateScalar(ChildResult, Options, node.Options[t].get());
 				branch = ChildResult.op;
 			}
 			else
@@ -289,13 +288,13 @@ namespace mu
 		}
 
 		// Process variations in reverse order, since conditionals are built bottom-up.
-		for (int t = node.m_variations.Num() - 1; t >= 0; --t)
+		for (int32 t = node.m_variations.Num() - 1; t >= 0; --t)
 		{
-			int tagIndex = -1;
+			int32 tagIndex = -1;
 			const FString& tag = node.m_variations[t].m_tag;
-			for (int i = 0; i < m_firstPass.m_tags.Num(); ++i)
+			for (int32 i = 0; i < FirstPass.Tags.Num(); ++i)
 			{
-				if (m_firstPass.m_tags[i].tag == tag)
+				if (FirstPass.Tags[i].Tag == tag)
 				{
 					tagIndex = i;
 				}
@@ -305,7 +304,7 @@ namespace mu
 			{
 				FString Msg = FString::Printf(TEXT("Unknown tag found in image variation [%s]."), *tag);
 
-				m_pErrorLog->GetPrivate()->Add(Msg, ELMT_WARNING, Typed->GetMessageContext());
+				ErrorLog->GetPrivate()->Add(Msg, ELMT_WARNING, Typed->GetMessageContext());
 				continue;
 			}
 
@@ -327,7 +326,7 @@ namespace mu
 			conditional->type = OP_TYPE::SC_CONDITIONAL;
 			conditional->no = op;
 			conditional->yes = variationOp;
-			conditional->condition = m_firstPass.m_tags[tagIndex].genericCondition;
+			conditional->condition = FirstPass.Tags[tagIndex].GenericCondition;
 
 			op = conditional;
 		}
@@ -411,7 +410,7 @@ namespace mu
 		const NodeScalarTable& node = *Typed;
 
 		Ptr<ASTOp> Op = GenerateTableSwitch<NodeScalarTable, ETableColumnType::Scalar, OP_TYPE::SC_SWITCH>(node,
-			[this,&Options](const NodeScalarTable& node, int colIndex, int row, ErrorLog* pErrorLog)
+			[this,&Options](const NodeScalarTable& node, int32 colIndex, int32 row, mu::ErrorLog* pErrorLog)
 			{
 				NodeScalarConstantPtr pCell = new NodeScalarConstant();
 				float scalar = node.Table->GetPrivate()->Rows[row].Values[colIndex].Scalar;
@@ -428,7 +427,7 @@ namespace mu
 	{
 		// Log a warning
 		FString Msg = FString::Printf(TEXT("Required connection not found: %s"), strWhere );
-		m_pErrorLog->GetPrivate()->Add(Msg, ELMT_ERROR, errorContext);
+		ErrorLog->GetPrivate()->Add(Msg, ELMT_ERROR, errorContext);
 
 		// Create a constant node
 		NodeScalarConstantPtr pNode = new NodeScalarConstant();
diff --git a/Engine/Plugins/Experimental/Mutable/Source/MutableTools/Private/MuT/CodeGenerator_SecondPass.cpp b/Engine/Plugins/Experimental/Mutable/Source/MutableTools/Private/MuT/CodeGenerator_SecondPass.cpp
--- a/Engine/Plugins/Experimental/Mutable/Source/MutableTools/Private/MuT/CodeGenerator_SecondPass.cpp
+++ b/Engine/Plugins/Experimental/Mutable/Source/MutableTools/Private/MuT/CodeGenerator_SecondPass.cpp
@@ -21,8 +21,8 @@ namespace mu
 	{
 		check(firstPass);
 		check(options);
-		m_pFirstPass = firstPass;
-		m_pCompilerOptions = options;
+		FirstPass = firstPass;
+		CompilerOptions = options;
 	}
 
 
@@ -33,7 +33,7 @@ namespace mu
 		const set<size_t>& posTag,
 		const set<size_t>& negTag)
 	{
-		auto& t = m_pFirstPass->m_tags[tagIndex];
+		auto& t = FirstPass->Tags[tagIndex];
 
 		// If this tag is already in the list of positive tags, return true as condition
 		if (posTag.find(tagIndex) != posTag.end())
@@ -48,39 +48,39 @@ namespace mu
 		}
 
 		// Cached?
-		CONDITION_GENERATION_KEY key;
+		FConditionGenerationKey key;
 		key.tagOrSurfIndex = tagIndex;
 		//    key.negSurf = negSurf;
 		//    key.posSurf = posSurf;
 		//    key.negTag = negTag;
 		//    key.posTag = posTag;
-		for (auto s : negTag) { if (m_tagsPerTag[tagIndex].count(s) > 0) { key.negTag.insert(s); } }
-		for (auto s : posTag) { if (m_tagsPerTag[tagIndex].count(s) > 0) { key.posTag.insert(s); } }
-		for (auto s : negSurf) { if (m_surfacesPerTag[tagIndex].count(s) > 0) { key.negSurf.insert(s); } }
-		for (auto s : posSurf) { if (m_surfacesPerTag[tagIndex].count(s) > 0) { key.posSurf.insert(s); } }
+		for (auto s : negTag) { if (TagsPerTag[tagIndex].count(s) > 0) { key.negTag.insert(s); } }
+		for (auto s : posTag) { if (TagsPerTag[tagIndex].count(s) > 0) { key.posTag.insert(s); } }
+		for (auto s : negSurf) { if (SurfacesPerTag[tagIndex].count(s) > 0) { key.negSurf.insert(s); } }
+		for (auto s : posSurf) { if (SurfacesPerTag[tagIndex].count(s) > 0) { key.posSurf.insert(s); } }
 
 		{
-			auto it = m_tagConditionGenerationCache.find(key);
-			if (it != m_tagConditionGenerationCache.end())
+			auto it = TagConditionGenerationCache.find(key);
+			if (it != TagConditionGenerationCache.end())
 			{
 				return it->second;
 			}
 		}
 
-		auto PositiveTags = [this](size_t SurfaceIndex) -> TArray<FString>&
+		auto PositiveTagsFunc = [this](size_t SurfaceIndex) -> TArray<FString>&
 		{
-			return m_pFirstPass->surfaces[SurfaceIndex].positiveTags;
+			return FirstPass->Surfaces[SurfaceIndex].PositiveTags;
 		};
 
-		auto NegativeTags = [this](size_t SurfaceIndex) -> TArray<FString>&
+		auto NegativeTagsFunc = [this](size_t SurfaceIndex) -> TArray<FString>&
 		{
-			return m_pFirstPass->surfaces[SurfaceIndex].negativeTags;
+			return FirstPass->Surfaces[SurfaceIndex].NegativeTags;
 		};
 		
 		Ptr<ASTOp> c;
 
 		// Condition expression for all the surfaces that activate the tag
-		for (auto surfIndex : t.surfaces)
+		for (int32 surfIndex : t.Surfaces)
 		{
 			if (posSurf.find(surfIndex) != posSurf.end())
 			{
@@ -96,17 +96,17 @@ namespace mu
 				continue;
 			}
 
-			const auto& surface = m_pFirstPass->surfaces[surfIndex];
+			const auto& surface = FirstPass->Surfaces[surfIndex];
 
-			auto positiveTags = posTag;
-			positiveTags.insert(tagIndex);
+			auto PositiveTags = posTag;
+			PositiveTags.insert(tagIndex);
 
 			Ptr<ASTOp> surfCondition = GenerateSurfaceOrModifierCodition(surfIndex,
-				PositiveTags,
-				NegativeTags,
+				PositiveTagsFunc,
+				NegativeTagsFunc,
 				posSurf,
 				negSurf,
-				positiveTags,
+				PositiveTags,
 				negTag);
 
 			// If the surface is a constant false, we can skip adding it
@@ -124,13 +124,13 @@ namespace mu
 			{
 				Ptr<ASTOpFixed> f = new ASTOpFixed();
 				f->op.type = OP_TYPE::BO_AND;
-				f->SetChild(f->op.args.BoolBinary.a, surface.objectCondition);
+				f->SetChild(f->op.args.BoolBinary.a, surface.ObjectCondition);
 				f->SetChild(f->op.args.BoolBinary.b, surfCondition);
 				fullCondition = m_opPool.Add(f);
 			}
 			else
 			{
-				fullCondition = m_opPool.Add(surface.objectCondition);
+				fullCondition = m_opPool.Add(surface.ObjectCondition);
 			}
 
 
@@ -148,14 +148,14 @@ namespace mu
 			}
 
 			// Optimise the condition now.
-			//PartialOptimise( c, m_pCompilerOptions->OptimisationOptions );
+			//PartialOptimise( c, CompilerOptions->OptimisationOptions );
 		}
 
 
 		// Condition expression for all the edit-surfaces that activate the tag
-		for (auto editKey : t.edits)
+		for (const TPair<int32,int32>& EditKey : t.Edits)
 		{
-			if (negSurf.find(editKey.Key) != negSurf.end())
+			if (negSurf.find(EditKey.Key) != negSurf.end())
 			{
 				// The surface in the edit is a negative requirement higher up in the condition so
 				// this branch never be true.
@@ -163,27 +163,27 @@ namespace mu
 			}
 
 			Ptr<ASTOp> surfCondition;
-			if (posSurf.find(editKey.Key) != posSurf.end())
+			if (posSurf.find(EditKey.Key) != posSurf.end())
 			{
 				// This surface in the edit is already a positive requirement higher up in the condition
 				// so we don't need that part of the condition
 			}
 			else
 			{
-				auto positiveTags = posTag;
-				positiveTags.insert(tagIndex);
+				auto PositiveTags = posTag;
+				PositiveTags.insert(tagIndex);
 
-				surfCondition = GenerateSurfaceOrModifierCodition(editKey.Key,
-					PositiveTags,
-					NegativeTags,
+				surfCondition = GenerateSurfaceOrModifierCodition(EditKey.Key,
+					PositiveTagsFunc,
+					NegativeTagsFunc,
 					posSurf,
 					negSurf,
-					positiveTags,
+					PositiveTags,
 					negTag);
 			}
 
-			const auto& surface = m_pFirstPass->surfaces[editKey.Key];
-			const auto& edit = surface.edits[editKey.Value];
+			const FirstPassGenerator::FSurface& surface = FirstPass->Surfaces[EditKey.Key];
+			const FirstPassGenerator::FSurface::FEdit& edit = surface.Edits[EditKey.Value];
 
 			// Combine object and surface conditions
 			Ptr<ASTOp> fullCondition;
@@ -191,13 +191,13 @@ namespace mu
 			{
 				Ptr<ASTOpFixed> f = new ASTOpFixed();
 				f->op.type = OP_TYPE::BO_AND;
-				f->SetChild(f->op.args.BoolBinary.a, surface.objectCondition);
+				f->SetChild(f->op.args.BoolBinary.a, surface.ObjectCondition);
 				f->SetChild(f->op.args.BoolBinary.b, surfCondition);
 				fullCondition = m_opPool.Add(f);
 			}
 			else
 			{
-				fullCondition = m_opPool.Add(surface.objectCondition);
+				fullCondition = m_opPool.Add(surface.ObjectCondition);
 			}
 
 			// Combine with edit condition
@@ -205,13 +205,13 @@ namespace mu
 			{
 				Ptr<ASTOpFixed> f = new ASTOpFixed();
 				f->op.type = OP_TYPE::BO_AND;
-				f->SetChild(f->op.args.BoolBinary.a, edit.condition);
+				f->SetChild(f->op.args.BoolBinary.a, edit.Condition);
 				f->SetChild(f->op.args.BoolBinary.b, fullCondition);
 				fullCondition = m_opPool.Add(f);
 			}
 			else
 			{
-				fullCondition = m_opPool.Add(edit.condition);
+				fullCondition = m_opPool.Add(edit.Condition);
 			}
 
 
@@ -230,15 +230,15 @@ namespace mu
 		}
 
 
-		m_tagConditionGenerationCache.insert(std::make_pair<>(key, c));
+		TagConditionGenerationCache.insert(std::make_pair<>(key, c));
 
 		return c;
 	}
 
 	//---------------------------------------------------------------------------------------------
 	mu::Ptr<ASTOp> SecondPassGenerator::GenerateSurfaceOrModifierCodition(size_t Index,
-		TFunction<const TArray<FString>&(size_t)> PositiveTags,
-		TFunction<const TArray<FString>&(size_t)> NegativeTags,
+		TFunction<const TArray<FString>&(size_t)> PositiveTagsFunc,
+		TFunction<const TArray<FString>&(size_t)> NegativeTagsFunc,
 		const set<size_t>& posSurf,
 		const set<size_t>& negSurf,
 		const set<size_t>& posTag,
@@ -258,9 +258,9 @@ namespace mu
 
 		Ptr<ASTOp> c;
 
-		for (const auto& t : PositiveTags(Index))
+		for (const FString& t : PositiveTagsFunc(Index))
 		{
-			auto it = m_pFirstPass->m_tags.FindByPredicate([&](const FirstPassGenerator::FTag& e) { return e.tag == t; });
+			const FirstPassGenerator::FTag* it = FirstPass->Tags.FindByPredicate([&](const FirstPassGenerator::FTag& e) { return e.Tag == t; });
 			if (!it)
 			{
 				// This could happen if a tag is in a variation but noone defines it.
@@ -268,7 +268,7 @@ namespace mu
 				return m_opPool.Add(new ASTOpConstantBool(false));
 			}
 
-			size_t tagIndex = it - &m_pFirstPass->m_tags[0];
+			size_t tagIndex = it - &FirstPass->Tags[0];
 
 			set<size_t> positiveSurfacesVisited = posSurf;
 			positiveSurfacesVisited.insert(Index);
@@ -328,16 +328,16 @@ namespace mu
 		}
 
 
-		for (const auto& t : NegativeTags(Index))
+		for (const FString& t : NegativeTagsFunc(Index))
 		{
-			auto it = m_pFirstPass->m_tags.FindByPredicate([&](const FirstPassGenerator::FTag& e) { return e.tag == t; });
+			const FirstPassGenerator::FTag* it = FirstPass->Tags.FindByPredicate([&](const FirstPassGenerator::FTag& e) { return e.Tag == t; });
 			if (!it)
 			{
 				// This could happen if a tag is in a variation but noone defines it.
 				continue;
 			}
 
-			size_t tagIndex = it - &m_pFirstPass->m_tags[0];
+			size_t tagIndex = it - &FirstPass->Tags[0];
 
 			set<size_t> positiveSurfacesVisited = negSurf;
 			set<size_t> negativeSurfacesVisited = posSurf;
@@ -406,22 +406,22 @@ namespace mu
 
 
 	//---------------------------------------------------------------------------------------------
-	bool SecondPassGenerator::Generate( ErrorLogPtr pErrorLog, const Node* root)
+	bool SecondPassGenerator::Generate( Ptr<mu::ErrorLog> InErrorLog, const Node* root)
 	{
 		MUTABLE_CPUPROFILER_SCOPE(SecondPassGenerate);
 
 		check(root);
-		m_pErrorLog = pErrorLog;
+		ErrorLog = InErrorLog;
 
 		// Find the list of surfaces every tag depends on
-		m_surfacesPerTag.clear();
-		m_surfacesPerTag.resize(m_pFirstPass->m_tags.Num());
-		m_tagsPerTag.clear();
-		m_tagsPerTag.resize(m_pFirstPass->m_tags.Num());
-		for (size_t t = 0; t < m_pFirstPass->m_tags.Num(); ++t)
+		SurfacesPerTag.Empty();
+		SurfacesPerTag.SetNum(FirstPass->Tags.Num());
+		TagsPerTag.Empty();
+		TagsPerTag.SetNum(FirstPass->Tags.Num());
+		for (size_t t = 0; t < FirstPass->Tags.Num(); ++t)
 		{
 			set<size_t> pendingSurfs;
-			for (size_t s : m_pFirstPass->m_tags[t].surfaces)
+			for (size_t s : FirstPass->Tags[t].Surfaces)
 			{
 				pendingSurfs.insert(s);
 			}
@@ -440,46 +440,46 @@ namespace mu
 
 				processedSurfs.insert(cs);
 
-				m_surfacesPerTag[t].insert(cs);
+				SurfacesPerTag[t].insert(cs);
 
-				auto& csurf = m_pFirstPass->surfaces[cs];
-				for (auto sct : csurf.positiveTags)
+				auto& csurf = FirstPass->Surfaces[cs];
+				for (const FString& sct : csurf.PositiveTags)
 				{
-					auto it = m_pFirstPass->m_tags.FindByPredicate([&](const FirstPassGenerator::FTag& e) { return e.tag == sct; });
+					const FirstPassGenerator::FTag* it = FirstPass->Tags.FindByPredicate([&](const FirstPassGenerator::FTag& e) { return e.Tag == sct; });
 					if (!it)
 					{
 						// This could happen if a tag is in a variation but noone defines it.
 						continue;
 					}
 
-					size_t ct = it - &m_pFirstPass->m_tags[0];
+					size_t ct = it - &FirstPass->Tags[0];
 
-					m_tagsPerTag[t].insert(ct);
+					TagsPerTag[t].insert(ct);
 
-					for (size_t s : m_pFirstPass->m_tags[ct].surfaces)
+					for (size_t s : FirstPass->Tags[ct].Surfaces)
 					{
-						if (m_surfacesPerTag[t].find(s) != m_surfacesPerTag[t].end())
+						if (SurfacesPerTag[t].find(s) != SurfacesPerTag[t].end())
 						{
 							pendingSurfs.insert(s);
 						}
 					}
 				}
-				for (auto sct : csurf.negativeTags)
+				for (const FString& sct : csurf.NegativeTags)
 				{
-					auto it = m_pFirstPass->m_tags.FindByPredicate([&](const FirstPassGenerator::FTag& e) { return e.tag == sct; });
+					const FirstPassGenerator::FTag* it = FirstPass->Tags.FindByPredicate([&](const FirstPassGenerator::FTag& e) { return e.Tag == sct; });
 					if (!it)
 					{
 						// This could happen if a tag is in a variation but noone defines it.
 						continue;
 					}
 
-					size_t ct = it - &m_pFirstPass->m_tags[0];
+					size_t ct = it - &FirstPass->Tags[0];
 
-					m_tagsPerTag[t].insert(ct);
+					TagsPerTag[t].insert(ct);
 
-					for (size_t s : m_pFirstPass->m_tags[ct].surfaces)
+					for (size_t s : FirstPass->Tags[ct].Surfaces)
 					{
-						if (m_surfacesPerTag[t].find(s) != m_surfacesPerTag[t].end())
+						if (SurfacesPerTag[t].find(s) != SurfacesPerTag[t].end())
 						{
 							pendingSurfs.insert(s);
 						}
@@ -489,85 +489,85 @@ namespace mu
 		}
 
 		// Create the conditions for every surface, modifier and individual tag.
-		m_tagConditionGenerationCache.clear();
+		TagConditionGenerationCache.clear();
 
-		for (int32 SurfaceIndex = 0; SurfaceIndex < m_pFirstPass->surfaces.Num(); ++SurfaceIndex)
+		for (int32 SurfaceIndex = 0; SurfaceIndex < FirstPass->Surfaces.Num(); ++SurfaceIndex)
 		{
-			FirstPassGenerator::FSurface& Surface = m_pFirstPass->surfaces[SurfaceIndex];
+			FirstPassGenerator::FSurface& Surface = FirstPass->Surfaces[SurfaceIndex];
 
 			{
 				set<size_t> Empty;
 
 				auto PositiveTags = [this](size_t SurfaceIndex) -> const TArray<FString>&
 				{
-					return m_pFirstPass->surfaces[SurfaceIndex].positiveTags;
+					return FirstPass->Surfaces[SurfaceIndex].PositiveTags;
 				};
 
 				auto NegativeTags = [this](size_t SurfaceIndex) -> const TArray<FString>&
 				{
-					return m_pFirstPass->surfaces[SurfaceIndex].negativeTags;
+					return FirstPass->Surfaces[SurfaceIndex].NegativeTags;
 				};
 				
 				Ptr<ASTOp> c = GenerateSurfaceOrModifierCodition(SurfaceIndex, PositiveTags, NegativeTags, Empty, Empty, Empty, Empty);
-				m_pFirstPass->surfaces[SurfaceIndex].surfaceCondition = c;
+				FirstPass->Surfaces[SurfaceIndex].SurfaceCondition = c;
 			}
 			
-			for (int32 EditIndex = 0; EditIndex < Surface.edits.Num(); ++EditIndex)
+			for (int32 EditIndex = 0; EditIndex < Surface.Edits.Num(); ++EditIndex)
 			{
-				FirstPassGenerator::FSurface::FEdit& Edit = Surface.edits[EditIndex];
+				FirstPassGenerator::FSurface::FEdit& Edit = Surface.Edits[EditIndex];
 				
 				set<size_t> Empty;
 
 				auto PositiveTags = [&Surface](size_t SurfaceIndex) -> const TArray<FString>&
 				{
-					return Surface.edits[SurfaceIndex].PositiveTags;
+					return Surface.Edits[SurfaceIndex].PositiveTags;
 				};
 
 				auto NegativeTags = [&Surface](size_t SurfaceIndex) -> const TArray<FString>&
 				{
-					return Surface.edits[SurfaceIndex].NegativeTags;
+					return Surface.Edits[SurfaceIndex].NegativeTags;
 				};
 				
 				Ptr<ASTOp> c = GenerateSurfaceOrModifierCodition(EditIndex, PositiveTags, NegativeTags, Empty, Empty, Empty, Empty);
 
 				Ptr<ASTOpFixed> OpAnd = new ASTOpFixed;
 				OpAnd->op.type = OP_TYPE::BO_AND;
-				OpAnd->SetChild(OpAnd->op.args.BoolBinary.a, Edit.condition);
+				OpAnd->SetChild(OpAnd->op.args.BoolBinary.a, Edit.Condition);
 				OpAnd->SetChild(OpAnd->op.args.BoolBinary.b, c);
 				c = m_opPool.Add(OpAnd);
 				
-				Edit.condition = OpAnd;
+				Edit.Condition = OpAnd;
 			}
 		}
 
-		for (int32 ModifierIndex = 0; ModifierIndex < m_pFirstPass->modifiers.Num(); ++ModifierIndex)
+		for (int32 ModifierIndex = 0; ModifierIndex < FirstPass->Modifiers.Num(); ++ModifierIndex)
 		{
 			set<size_t> Empty;
 
 			auto PositiveTags = [this](size_t SurfaceIndex) -> TArray<FString>&
 			{
-				return m_pFirstPass->modifiers[SurfaceIndex].positiveTags;
+				return FirstPass->Modifiers[SurfaceIndex].PositiveTags;
 			};
 
 			auto NegativeTags = [this](size_t SurfaceIndex) -> TArray<FString>&
 			{
-				return m_pFirstPass->modifiers[SurfaceIndex].negativeTags;
+				return FirstPass->Modifiers[SurfaceIndex].NegativeTags;
 			};
 				
 			Ptr<ASTOp> c = GenerateSurfaceOrModifierCodition(ModifierIndex, PositiveTags, NegativeTags, Empty, Empty, Empty, Empty);
 			
-			m_pFirstPass->modifiers[ModifierIndex].surfaceCondition = c;
+			FirstPass->Modifiers[ModifierIndex].SurfaceCondition = c;
 		}
 
-		for (int32 s = 0; s < m_pFirstPass->m_tags.Num(); ++s)
+		for (int32 s = 0; s < FirstPass->Tags.Num(); ++s)
 		{
 			set<size_t> empty;
 			Ptr<ASTOp> c = GenerateTagCondition(s, empty, empty, empty, empty);
-			m_pFirstPass->m_tags[s].genericCondition = c;
+			FirstPass->Tags[s].GenericCondition = c;
 		}
 
 
-		m_pFirstPass = nullptr;
+		FirstPass = nullptr;
 
 		return true;
 	}
diff --git a/Engine/Plugins/Experimental/Mutable/Source/MutableTools/Private/MuT/CodeGenerator_SecondPass.h b/Engine/Plugins/Experimental/Mutable/Source/MutableTools/Private/MuT/CodeGenerator_SecondPass.h
--- a/Engine/Plugins/Experimental/Mutable/Source/MutableTools/Private/MuT/CodeGenerator_SecondPass.h
+++ b/Engine/Plugins/Experimental/Mutable/Source/MutableTools/Private/MuT/CodeGenerator_SecondPass.h
@@ -24,22 +24,22 @@ namespace mu
 	{
 	public:
 
-		SecondPassGenerator( FirstPassGenerator* firstPass, const CompilerOptions::Private* options  );
+		SecondPassGenerator( FirstPassGenerator*, const CompilerOptions::Private*  );
 
 		// Return true on success.
-        bool Generate( ErrorLogPtr pErrorLog, const Node* root );
+        bool Generate( Ptr<ErrorLog>, const Node* Root );
 
 	private:
 
-        FirstPassGenerator* m_pFirstPass = nullptr;
-        const CompilerOptions::Private *m_pCompilerOptions = nullptr;
+        FirstPassGenerator* FirstPass = nullptr;
+        const CompilerOptions::Private *CompilerOptions = nullptr;
 
 
         //!
-        Ptr<ErrorLog> m_pErrorLog;
+        Ptr<ErrorLog> ErrorLog;
 
         //!
-        struct CONDITION_GENERATION_KEY
+        struct FConditionGenerationKey
         {
             size_t tagOrSurfIndex = 0;
             set<size_t> posSurf;
@@ -47,7 +47,7 @@ namespace mu
             set<size_t> posTag;
             set<size_t> negTag;
 
-            inline bool operator<(const CONDITION_GENERATION_KEY& o) const
+            inline bool operator<(const FConditionGenerationKey& o) const
             {
                 if (tagOrSurfIndex<o.tagOrSurfIndex) return true;
                 if (tagOrSurfIndex>o.tagOrSurfIndex) return false;
@@ -65,10 +65,10 @@ namespace mu
 
         // List of surfaces that activate or deactivate every tag, or another surface that activates a tag in this set.
 		// \TODO: Change to UE containers
-		vector< set<size_t> > m_surfacesPerTag;
-		vector< set<size_t> > m_tagsPerTag;
+		TArray< set<size_t> > SurfacesPerTag;
+		TArray< set<size_t> > TagsPerTag;
 
-        std::map<CONDITION_GENERATION_KEY,Ptr<ASTOp>> m_tagConditionGenerationCache;
+        std::map<FConditionGenerationKey,Ptr<ASTOp>> TagConditionGenerationCache;
 
         UniqueOpPool m_opPool;
 
@@ -81,15 +81,15 @@ namespace mu
 
         /** Generate Surface, Edit or Modifier condition.
     	 * @param Index Surface, Edit or Modifier index.
-    	 * @param positiveTags function that given the Surface, Edit or Modifier index, returns its positive tags.
-    	 * @param negativeTags function that given the Surface, Edit or Modifier, returns its negative tags.
+    	 * @param PositiveTags function that given the Surface, Edit or Modifier index, returns its positive tags.
+    	 * @param NegativeTags function that given the Surface, Edit or Modifier, returns its negative tags.
     	 * @param posSurf already visited Surfaces, Edits, or Modifiers that participate positively in the condition.
     	 * @param negSurf already visited Surfaces, Edits, or Modifiers that participate negatively in the condition.
       	 * @param posSurf Tags that already belong to the condition (positively).
 		 * @param negSurf Tags that already belong to the condition (negatively). */
         Ptr<ASTOp> GenerateSurfaceOrModifierCodition(size_t Index,
-									        TFunction<const TArray<FString>&(size_t)> positiveTags,
-											TFunction<const TArray<FString>&(size_t)> negativeTags,
+									        TFunction<const TArray<FString>&(size_t)> PositiveTags,
+											TFunction<const TArray<FString>&(size_t)> NegativeTags,
                                             const set<size_t>& posSurf,
                                             const set<size_t>& negSurf,
                                             const set<size_t>& posTag,
diff --git a/Engine/Plugins/Experimental/Mutable/Source/MutableTools/Private/MuT/CodeGenerator_String.cpp b/Engine/Plugins/Experimental/Mutable/Source/MutableTools/Private/MuT/CodeGenerator_String.cpp
--- a/Engine/Plugins/Experimental/Mutable/Source/MutableTools/Private/MuT/CodeGenerator_String.cpp
+++ b/Engine/Plugins/Experimental/Mutable/Source/MutableTools/Private/MuT/CodeGenerator_String.cpp
@@ -82,7 +82,7 @@ namespace mu
 
 		Ptr<ASTOpParameter> op;
 
-		Ptr<ASTOpParameter>* it = m_firstPass.ParameterNodes.Find(node.m_pNode);
+		Ptr<ASTOpParameter>* it = FirstPass.ParameterNodes.Find(node.m_pNode);
 		if (!it)
 		{
 			FParameterDesc param;
@@ -104,7 +104,7 @@ namespace mu
 				op->ranges.Emplace(op.get(), rangeResult.sizeOp, rangeResult.rangeName, rangeResult.rangeUID);
 			}
 
-			m_firstPass.ParameterNodes.Add(node.m_pNode, op);
+			FirstPass.ParameterNodes.Add(node.m_pNode, op);
 		}
 		else
 		{
diff --git a/Engine/Plugins/Experimental/Mutable/Source/MutableTools/Private/MuT/Compiler.cpp b/Engine/Plugins/Experimental/Mutable/Source/MutableTools/Private/MuT/Compiler.cpp
--- a/Engine/Plugins/Experimental/Mutable/Source/MutableTools/Private/MuT/Compiler.cpp
+++ b/Engine/Plugins/Experimental/Mutable/Source/MutableTools/Private/MuT/Compiler.cpp
@@ -231,9 +231,9 @@ namespace mu
 
             gen.GenerateRoot( pNode );
 
-            check( !gen.m_states.IsEmpty() );
+            check( !gen.States.IsEmpty() );
 
-            for ( const auto& s: gen.m_states )
+            for ( const TPair<FObjectState, Ptr<ASTOp>>& s: gen.States )
             {
                 FStateCompilationData data;
                 data.nodeState = s.Key;
@@ -242,13 +242,13 @@ namespace mu
                 states.Add( data );
             }
 
-            genErrorLog = gen.m_pErrorLog;
+            genErrorLog = gen.ErrorLog;
 
 			// Set the parameter list from the non-optimized data, so that we have them all even if they are optimized out
-			int32 ParameterCount = gen.m_firstPass.ParameterNodes.Num();
+			int32 ParameterCount = gen.FirstPass.ParameterNodes.Num();
 			Parameters.SetNum(ParameterCount);
 			int32 ParameterIndex = 0;
-			for ( const TPair<Ptr<const Node>, Ptr<ASTOpParameter>>& Entry : gen.m_firstPass.ParameterNodes )
+			for ( const TPair<Ptr<const Node>, Ptr<ASTOpParameter>>& Entry : gen.FirstPass.ParameterNodes )
 			{
 				Parameters[ParameterIndex] = Entry.Value->parameter;
 				++ParameterIndex;
@@ -344,7 +344,7 @@ namespace mu
 						"safely dismissed in case of partial compilation."), 
 						*s.nodeState.Name,
 						*s.nodeState.RuntimeParams[p]);
-                    m_pD->m_pErrorLog->GetPrivate()->Add(Temp, ELMT_WARNING, pNode->GetMessageContext() );
+                    m_pD->ErrorLog->GetPrivate()->Add(Temp, ELMT_WARNING, pNode->GetMessageContext() );
                 }
             }
 
@@ -402,8 +402,8 @@ namespace mu
 		UE_LOG(LogMutableCore, Verbose, TEXT("(int) %s : %ld"), TEXT("Program size"), int64(Program.m_opAddress.Num()));
 
         // Merge the log in the right order
-        genErrorLog->Merge( m_pD->m_pErrorLog.get() );
-        m_pD->m_pErrorLog = genErrorLog.get();
+        genErrorLog->Merge( m_pD->ErrorLog.get() );
+        m_pD->ErrorLog = genErrorLog.get();
 
 		// Pack data
 		m_pD->GenerateRoms(pResult.Get(), m_pD->m_options.get(), LinkerOptions.AdditionalData );
@@ -417,7 +417,7 @@ namespace mu
     //---------------------------------------------------------------------------------------------
     ErrorLogPtrConst Compiler::GetLog() const
     {
-        return m_pD->m_pErrorLog;
+        return m_pD->ErrorLog;
     }
 
 
diff --git a/Engine/Plugins/Experimental/Mutable/Source/MutableTools/Private/MuT/CompilerPrivate.h b/Engine/Plugins/Experimental/Mutable/Source/MutableTools/Private/MuT/CompilerPrivate.h
--- a/Engine/Plugins/Experimental/Mutable/Source/MutableTools/Private/MuT/CompilerPrivate.h
+++ b/Engine/Plugins/Experimental/Mutable/Source/MutableTools/Private/MuT/CompilerPrivate.h
@@ -88,10 +88,10 @@ namespace mu
 
         Private()
         {
-            m_pErrorLog = new ErrorLog();
+            ErrorLog = new mu::ErrorLog();
         }
 
-        Ptr<ErrorLog> m_pErrorLog;
+        Ptr<ErrorLog> ErrorLog;
 
         /** */
         Ptr<CompilerOptions> m_options;
diff --git a/Engine/Plugins/Experimental/Mutable/Source/MutableTools/Private/MuT/Node.cpp b/Engine/Plugins/Experimental/Mutable/Source/MutableTools/Private/MuT/Node.cpp
--- a/Engine/Plugins/Experimental/Mutable/Source/MutableTools/Private/MuT/Node.cpp
+++ b/Engine/Plugins/Experimental/Mutable/Source/MutableTools/Private/MuT/Node.cpp
@@ -7,7 +7,14 @@
 
 #include "MuT/NodeComponent.h"
 #include "MuT/NodeComponentNew.h"
+#include "MuT/NodeComponentSwitch.h"
+#include "MuT/NodeComponentVariation.h"
 #include "MuT/NodeLOD.h"
+#include "MuT/NodeSurface.h"
+#include "MuT/NodeSurfaceNew.h"
+#include "MuT/NodeSurfaceEdit.h"
+#include "MuT/NodeSurfaceSwitch.h"
+#include "MuT/NodeSurfaceVariation.h"
 #include "MuT/NodeExtensionData.h"
 #include "MuT/NodeExtensionDataConstant.h"
 #include "MuT/NodeImageTable.h"
@@ -20,6 +27,7 @@
 #include "MuT/NodeColourArithmeticOperation.h"
 #include "MuT/NodeColourFromScalars.h"
 #include "MuT/NodeColourTable.h"
+#include "MuT/NodeScalarSwitch.h"
 
 namespace mu
 {
@@ -29,6 +37,17 @@ namespace mu
 
 	FNodeType NodeComponent::StaticType = FNodeType(Node::EType::Component, Node::GetStaticType());
 	FNodeType NodeComponentNew::StaticType = FNodeType(Node::EType::ComponentNew, NodeComponent::GetStaticType());
+	FNodeType NodeComponentSwitch::StaticType = FNodeType(Node::EType::ComponentSwitch, NodeComponent::GetStaticType());
+	FNodeType NodeComponentVariation::StaticType = FNodeType(Node::EType::ComponentVariation, NodeComponent::GetStaticType());
+
+	FNodeType NodeScalarSwitch::StaticType = FNodeType(Node::EType::ScalarSwitch, NodeScalar::GetStaticType());
+
+	FNodeType NodeSurface::StaticType = FNodeType(Node::EType::Surface, Node::GetStaticType());
+	FNodeType NodeSurfaceNew::StaticType = FNodeType(Node::EType::SurfaceNew, NodeSurface::GetStaticType());
+	FNodeType NodeSurfaceEdit::StaticType = FNodeType(Node::EType::SurfaceEdit, NodeSurface::GetStaticType());
+	FNodeType NodeSurfaceSwitch::StaticType = FNodeType(Node::EType::SurfaceSwitch, NodeSurface::GetStaticType());
+	FNodeType NodeSurfaceVariation::StaticType = FNodeType(Node::EType::SurfaceVariation, NodeSurface::GetStaticType());
+
 	FNodeType NodeLOD::StaticType = FNodeType(Node::EType::LOD, Node::GetStaticType());
 	FNodeType NodeExtensionData::StaticType = FNodeType(Node::EType::ExtensionData, Node::GetStaticType());
 	FNodeType NodeExtensionDataConstant::StaticType = FNodeType(Node::EType::ExtensionDataConstant, NodeExtensionData::GetStaticType());
diff --git a/Engine/Plugins/Experimental/Mutable/Source/MutableTools/Private/MuT/NodeMeshClipMorphPlane.cpp b/Engine/Plugins/Experimental/Mutable/Source/MutableTools/Private/MuT/NodeMeshClipMorphPlane.cpp
--- a/Engine/Plugins/Experimental/Mutable/Source/MutableTools/Private/MuT/NodeMeshClipMorphPlane.cpp
+++ b/Engine/Plugins/Experimental/Mutable/Source/MutableTools/Private/MuT/NodeMeshClipMorphPlane.cpp
@@ -84,6 +84,6 @@ namespace mu
 	//---------------------------------------------------------------------------------------------
 	void NodeMeshClipMorphPlane::AddTag(const char* tagName)
 	{
-		m_pD->m_tags.Add(tagName);
+		m_pD->Tags.Add(tagName);
 	}
 }
diff --git a/Engine/Plugins/Experimental/Mutable/Source/MutableTools/Private/MuT/NodeMeshClipMorphPlanePrivate.h b/Engine/Plugins/Experimental/Mutable/Source/MutableTools/Private/MuT/NodeMeshClipMorphPlanePrivate.h
--- a/Engine/Plugins/Experimental/Mutable/Source/MutableTools/Private/MuT/NodeMeshClipMorphPlanePrivate.h
+++ b/Engine/Plugins/Experimental/Mutable/Source/MutableTools/Private/MuT/NodeMeshClipMorphPlanePrivate.h
@@ -61,7 +61,7 @@ namespace mu
 		FVector3f m_selectionBoxRadius;
 		FBoneName m_vertexSelectionBone;
 
-		TArray<FString> m_tags;
+		TArray<FString> Tags;
 
 		// Max distance a vertex can have to the bone in order to be affected. A negative value
 		// means no limit.
diff --git a/Engine/Plugins/Experimental/Mutable/Source/MutableTools/Private/MuT/NodeMeshClipWithMesh.cpp b/Engine/Plugins/Experimental/Mutable/Source/MutableTools/Private/MuT/NodeMeshClipWithMesh.cpp
--- a/Engine/Plugins/Experimental/Mutable/Source/MutableTools/Private/MuT/NodeMeshClipWithMesh.cpp
+++ b/Engine/Plugins/Experimental/Mutable/Source/MutableTools/Private/MuT/NodeMeshClipWithMesh.cpp
@@ -52,6 +52,6 @@ namespace mu
 	//---------------------------------------------------------------------------------------------
 	void NodeMeshClipWithMesh::AddTag(const char* tagName)
 	{
-		m_pD->m_tags.Add(tagName);
+		m_pD->Tags.Add(tagName);
 	}
 }
diff --git a/Engine/Plugins/Experimental/Mutable/Source/MutableTools/Private/MuT/NodeMeshClipWithMeshPrivate.h b/Engine/Plugins/Experimental/Mutable/Source/MutableTools/Private/MuT/NodeMeshClipWithMeshPrivate.h
--- a/Engine/Plugins/Experimental/Mutable/Source/MutableTools/Private/MuT/NodeMeshClipWithMeshPrivate.h
+++ b/Engine/Plugins/Experimental/Mutable/Source/MutableTools/Private/MuT/NodeMeshClipWithMeshPrivate.h
@@ -21,7 +21,7 @@ namespace mu
 		NodeMeshPtr m_pSource;
 		NodeMeshPtr m_pClipMesh;
 
-		TArray<FString> m_tags;
+		TArray<FString> Tags;
 
 	};
 
diff --git a/Engine/Plugins/Experimental/Mutable/Source/MutableTools/Private/MuT/NodeScalarSwitch.cpp b/Engine/Plugins/Experimental/Mutable/Source/MutableTools/Private/MuT/NodeScalarSwitch.cpp
deleted file mode 100644
diff --git a/Engine/Plugins/Experimental/Mutable/Source/MutableTools/Private/MuT/NodeSurface.cpp b/Engine/Plugins/Experimental/Mutable/Source/MutableTools/Private/MuT/NodeSurface.cpp
deleted file mode 100644
diff --git a/Engine/Plugins/Experimental/Mutable/Source/MutableTools/Private/MuT/NodeSurfaceEdit.cpp b/Engine/Plugins/Experimental/Mutable/Source/MutableTools/Private/MuT/NodeSurfaceEdit.cpp
deleted file mode 100644
diff --git a/Engine/Plugins/Experimental/Mutable/Source/MutableTools/Private/MuT/NodeSurfaceNew.cpp b/Engine/Plugins/Experimental/Mutable/Source/MutableTools/Private/MuT/NodeSurfaceNew.cpp
deleted file mode 100644
diff --git a/Engine/Plugins/Experimental/Mutable/Source/MutableTools/Private/MuT/NodeSurfaceSwitch.cpp b/Engine/Plugins/Experimental/Mutable/Source/MutableTools/Private/MuT/NodeSurfaceSwitch.cpp
deleted file mode 100644
diff --git a/Engine/Plugins/Experimental/Mutable/Source/MutableTools/Private/MuT/NodeSurfaceVariation.cpp b/Engine/Plugins/Experimental/Mutable/Source/MutableTools/Private/MuT/NodeSurfaceVariation.cpp
deleted file mode 100644
diff --git a/Engine/Plugins/Experimental/Mutable/Source/MutableTools/Public/MuT/Node.h b/Engine/Plugins/Experimental/Mutable/Source/MutableTools/Public/MuT/Node.h
--- a/Engine/Plugins/Experimental/Mutable/Source/MutableTools/Public/MuT/Node.h
+++ b/Engine/Plugins/Experimental/Mutable/Source/MutableTools/Public/MuT/Node.h
@@ -129,6 +129,8 @@ namespace mu
 			Component,
 			ComponentNew,
 			ComponentEdit,
+			ComponentSwitch,
+			ComponentVariation,
 
 			Object,
 			ObjectNew,
diff --git a/Engine/Plugins/Experimental/Mutable/Source/MutableTools/Public/MuT/NodeComponent.h b/Engine/Plugins/Experimental/Mutable/Source/MutableTools/Public/MuT/NodeComponent.h
--- a/Engine/Plugins/Experimental/Mutable/Source/MutableTools/Public/MuT/NodeComponent.h
+++ b/Engine/Plugins/Experimental/Mutable/Source/MutableTools/Public/MuT/NodeComponent.h
@@ -21,10 +21,8 @@ namespace mu
 
 	public:
 
-		//-----------------------------------------------------------------------------------------
         // Node interface
-		//-----------------------------------------------------------------------------------------
-		virtual const FNodeType* GetType() const override { return GetStaticType(); }
+		virtual const FNodeType* GetType() const override { return &StaticType; }
 		static const FNodeType* GetStaticType() { return &StaticType; }
 
         //-----------------------------------------------------------------------------------------
diff --git a/Engine/Plugins/Experimental/Mutable/Source/MutableTools/Public/MuT/NodeComponentNew.h b/Engine/Plugins/Experimental/Mutable/Source/MutableTools/Public/MuT/NodeComponentNew.h
--- a/Engine/Plugins/Experimental/Mutable/Source/MutableTools/Public/MuT/NodeComponentNew.h
+++ b/Engine/Plugins/Experimental/Mutable/Source/MutableTools/Public/MuT/NodeComponentNew.h
@@ -16,7 +16,7 @@ namespace mu
 	{
 	public:
 
-		/** Externally managed id assign to this component. */
+		/** Externally managed id assigned to this component. */
 		uint16 Id = 0;
 
 	public:
diff --git a/Engine/Plugins/Experimental/Mutable/Source/MutableTools/Public/MuT/NodeComponentSwitch.h b/Engine/Plugins/Experimental/Mutable/Source/MutableTools/Public/MuT/NodeComponentSwitch.h
new file mode 100644
--- /dev/null
+++ b/Engine/Plugins/Experimental/Mutable/Source/MutableTools/Public/MuT/NodeComponentSwitch.h
@@ -0,0 +1,42 @@
+// Copyright Epic Games, Inc. All Rights Reserved.
+
+#pragma once
+
+#include "MuR/Ptr.h"
+#include "MuR/RefCounted.h"
+#include "MuT/Node.h"
+#include "MuT/NodeComponent.h"
+#include "MuT/NodeScalar.h"
+
+
+namespace mu
+{
+
+    /** This node selects an output component from a set of input components based on a parameter. 
+	*/
+    class MUTABLETOOLS_API NodeComponentSwitch : public NodeComponent
+	{
+	public:
+
+		Ptr<NodeScalar> Parameter;
+		TArray<Ptr<NodeComponent>> Options;
+
+	public:
+
+		/** Node type hierarchy data. */
+		virtual const FNodeType* GetType() const { return &StaticType; }
+		static const FNodeType* GetStaticType() { return &StaticType; }
+
+	protected:
+
+		//! Forbidden. Manage with the Ptr<> template.
+		~NodeComponentSwitch() {}
+
+	private:
+
+		static FNodeType StaticType;
+
+	};
+
+
+}
diff --git a/Engine/Plugins/Experimental/Mutable/Source/MutableTools/Public/MuT/NodeComponentVariation.h b/Engine/Plugins/Experimental/Mutable/Source/MutableTools/Public/MuT/NodeComponentVariation.h
new file mode 100644
--- /dev/null
+++ b/Engine/Plugins/Experimental/Mutable/Source/MutableTools/Public/MuT/NodeComponentVariation.h
@@ -0,0 +1,47 @@
+// Copyright Epic Games, Inc. All Rights Reserved.
+
+#pragma once
+
+#include "MuR/Ptr.h"
+#include "MuR/RefCounted.h"
+#include "MuT/Node.h"
+#include "MuT/NodeComponent.h"
+
+
+namespace mu
+{
+
+	/** Select different component subgraphs based on active tags. */
+    class MUTABLETOOLS_API NodeComponentVariation : public NodeComponent
+    {
+	public:
+
+		Ptr<NodeComponent> DefaultComponent;
+
+		struct FVariation
+		{
+			Ptr<NodeComponent> Component;
+			FString Tag;
+		};
+
+		TArray<FVariation> Variations;
+
+    public:
+
+		// Node Interface
+		virtual const FNodeType* GetType() const override { return &StaticType; }
+		static const FNodeType* GetStaticType() { return &StaticType; }
+
+    protected:
+
+        //! Forbidden. Manage with the Ptr<> template.
+		~NodeComponentVariation() {}
+
+	private:
+
+		static FNodeType StaticType;
+	
+	};
+
+
+} // namespace mu
diff --git a/Engine/Plugins/Experimental/Mutable/Source/MutableTools/Public/MuT/NodeScalarSwitch.h b/Engine/Plugins/Experimental/Mutable/Source/MutableTools/Public/MuT/NodeScalarSwitch.h
--- a/Engine/Plugins/Experimental/Mutable/Source/MutableTools/Public/MuT/NodeScalarSwitch.h
+++ b/Engine/Plugins/Experimental/Mutable/Source/MutableTools/Public/MuT/NodeScalarSwitch.h
@@ -11,61 +11,28 @@
 namespace mu
 {
 
-	// Forward definitions
-	class NodeScalar;
-	typedef Ptr<NodeScalar> NodeScalarPtr;
-	typedef Ptr<const NodeScalar> NodeScalarPtrConst;
-
-    class NodeScalarSwitch;
-    typedef Ptr<NodeScalarSwitch> NodeScalarSwitchPtr;
-    typedef Ptr<const NodeScalarSwitch> NodeScalarSwitchPtrConst;
-
-
-    //! This node selects an output Scalar from a set of input Scalars based on a parameter.
+    /** This node selects an output Scalar from a set of input Scalars based on a parameter. */
     class MUTABLETOOLS_API NodeScalarSwitch : public NodeScalar
 	{
 	public:
 
-        NodeScalarSwitch();
-
-		//-----------------------------------------------------------------------------------------
-		// Node Interface
-		//-----------------------------------------------------------------------------------------
-
-        const FNodeType* GetType() const override;
-		static const FNodeType* GetStaticType();
-
-		//-----------------------------------------------------------------------------------------
-		// Own Interface
-		//-----------------------------------------------------------------------------------------
-
-		//! Get the node generating the parameter used to select the option.
-		NodeScalarPtr GetParameter() const;
-		void SetParameter( NodeScalarPtr );
-
-        //! Set the number of option Scalars. It will keep the currently set targets and initialise
-		//! the new ones as null.
-		void SetOptionCount( int );
-
-        //! Get the node generating the t-th option Scalar.
-        NodeScalarPtr GetOption( int t ) const;
-        void SetOption( int t, NodeScalarPtr );
+		Ptr<NodeScalar> Parameter;
+		TArray<Ptr<NodeScalar>> Options;
 
+	public:
 
-		//-----------------------------------------------------------------------------------------
-		// Interface pattern
-		//-----------------------------------------------------------------------------------------
-		class Private;
-		Private* GetPrivate() const;
+		// Node interface
+		virtual const FNodeType* GetType() const override { return &StaticType; }
+		static const FNodeType* GetStaticType() { return &StaticType; }
 
 	protected:
 
 		//! Forbidden. Manage with the Ptr<> template.
-        ~NodeScalarSwitch();
+		~NodeScalarSwitch() {}
 
 	private:
 
-		Private* m_pD;
+		static FNodeType StaticType;
 
 	};
 
diff --git a/Engine/Plugins/Experimental/Mutable/Source/MutableTools/Public/MuT/NodeSurface.h b/Engine/Plugins/Experimental/Mutable/Source/MutableTools/Public/MuT/NodeSurface.h
--- a/Engine/Plugins/Experimental/Mutable/Source/MutableTools/Public/MuT/NodeSurface.h
+++ b/Engine/Plugins/Experimental/Mutable/Source/MutableTools/Public/MuT/NodeSurface.h
@@ -17,24 +17,24 @@ namespace mu
     typedef Ptr<const NodeSurface> NodeSurfacePtrConst;
 
 
-    //! This class is the parent of all nodes that output a Surface.
+    /** This class is the parent of all nodes that output a Surface. */
     class MUTABLETOOLS_API NodeSurface : public Node
 	{
 	public:
 
-		//-----------------------------------------------------------------------------------------
-		// Node Interface
-		//-----------------------------------------------------------------------------------------
-
-        const FNodeType* GetType() const override;
-		static const FNodeType* GetStaticType();
-
+		// Node interface
+		virtual const FNodeType* GetType() const override { return &StaticType; }
+		static const FNodeType* GetStaticType() { return &StaticType; }
 
 	protected:
 
 		//! Forbidden. Manage with the Ptr<> template.
         inline ~NodeSurface() {}
 
+	private:
+
+		static FNodeType StaticType;
+
 	};
 
 
diff --git a/Engine/Plugins/Experimental/Mutable/Source/MutableTools/Public/MuT/NodeSurfaceEdit.h b/Engine/Plugins/Experimental/Mutable/Source/MutableTools/Public/MuT/NodeSurfaceEdit.h
--- a/Engine/Plugins/Experimental/Mutable/Source/MutableTools/Public/MuT/NodeSurfaceEdit.h
+++ b/Engine/Plugins/Experimental/Mutable/Source/MutableTools/Public/MuT/NodeSurfaceEdit.h
@@ -6,134 +6,54 @@
 #include "MuR/RefCounted.h"
 #include "MuT/Node.h"
 #include "MuT/NodeSurface.h"
+#include "MuT/NodeMesh.h"
+#include "MuT/NodePatchMesh.h"
+#include "MuT/NodeImage.h"
+#include "MuT/NodeScalar.h"
+#include "MuT/NodePatchImage.h"
 
 
 namespace mu
 {
 
-	// Forward definitions
-    class NodeMesh;
-    typedef Ptr<NodeMesh> NodeMeshPtr;
-    typedef Ptr<const NodeMesh> NodeMeshPtrConst;
-
-    class NodeImage;
-	typedef Ptr<NodeImage> NodeImagePtr;
-	typedef Ptr<const NodeImage> NodeImagePtrConst;
-
-	class NodePatchMesh;
-	typedef Ptr<NodePatchMesh> NodePatchMeshPtr;
-	typedef Ptr<const NodePatchMesh> NodePatchMeshPtrConst;
-
-	class NodePatchImage;
-	typedef Ptr<NodePatchImage> NodePatchImagePtr;
-	typedef Ptr<const NodePatchImage> NodePatchImagePtrConst;
-
-    class NodeSurfaceEdit;
-    typedef Ptr<NodeSurfaceEdit> NodeSurfaceEditPtr;
-    typedef Ptr<const NodeSurfaceEdit> NodeSurfaceEditPtrConst;
-
-	class NodeScalar;
-	typedef Ptr<NodeScalar> NodeScalarPtr;
-	typedef Ptr<const NodeScalar> NodeScalarPtrConst;
-
-
-
-	//! This node modifies a node of the parent object of the object that this node belongs to.
-    //! It allows to extend, cut and morph the parent Surface's meshes.
-    //! It also allows to patch the parent Surface's textures.
+	/** This node modifies a node of the parent object of the object that this node belongs to.
+    * It allows to extend, cut and morph the parent Surface's meshes.
+    * It also allows to patch the parent Surface's textures.
+	*/
     class MUTABLETOOLS_API NodeSurfaceEdit : public NodeSurface
 	{
 	public:
 
-        NodeSurfaceEdit();
-
-		//-----------------------------------------------------------------------------------------
-		// Node Interface
-		//-----------------------------------------------------------------------------------------
-        const FNodeType* GetType() const override;
-		static const FNodeType* GetStaticType();
-
-		//-----------------------------------------------------------------------------------------
-		// Own Interface
-		//-----------------------------------------------------------------------------------------
-
-        //! Set the parent Surface to modify. It should be a node from the same LOD of one of
-		//! the parent objects of this node's object.
-        void SetParent( NodeSurface* );
-        NodeSurface* GetParent() const;
-
-        //! \name Tags
-        //! \{
-
-        //! Add a tag to the surface:
-        //! - the surface will be affected by modifier nodes with the same tag
-        //! - the tag will be enabled when the surface is added to an object, and it can activate
-        //! variations for any surface.
-        void AddTag(const FString& tagName);
-
-        //! \}
+		Ptr<NodeSurface> Parent;
+		Ptr<NodePatchMesh> Mesh;
+		Ptr<NodeMesh> Morph;
+		Ptr<NodeScalar> MorphFactor;
 
-		//! \name Mesh modification
-		//! \{
+		struct FTexture
+		{
+			Ptr<NodeImage> Extend;
+			Ptr<NodePatchImage> Patch;
+		};
 
-		//! Mesh a node to extend the parent mesh
-        void SetMesh( NodePatchMesh* );
+		TArray<FTexture> Textures;
 
-		//! Get the node to extend the parent mesh
-        NodePatchMesh* GetMesh() const;
-
-		//! Set a node to morph the parent mesh
-        void SetMorph( NodeMeshPtr );
-
-		//! Get the node to morph the parent mesh
-        NodeMeshPtr GetMorph() const;
-
-		//! Set the factor: weight used to select (and combine) the morphs to apply.
-		void SetFactor( NodeScalarPtr );
-
-		//! Get the factor: weight used to select (and combine) the morphs to apply.
-		NodeScalarPtr GetFactor() const;
-
-		//! \}
-
-		//! \name Image modification
-		//! \{
-
-		//! Get the number of images to modify. Should match the parent image count.
-		int GetImageCount() const;
-
-		//! Set the number of images to modify. Should match the parent image count.
-		void SetImageCount( int );
-
-		//! Get the image node that extends a parent image, adding new blocks
-		NodeImagePtr GetImage( int index ) const;
-
-		//! Set the image node that extends a parent image, adding new blocks
-		void SetImage( int index, NodeImagePtr );
-
-		//! Get the image node that patches a parent image, blending on existing blocks
-		NodePatchImagePtr GetPatch( int index ) const;
-
-		//! Set the image node that patches a parent image, blending on existing blocks
-		void SetPatch( int index, NodePatchImagePtr );
-
-		//!}
+		//! Tags in this surface edit
+		TArray<FString> Tags;
 
+	public:
 
-		//-----------------------------------------------------------------------------------------
-		// Interface pattern
-		//-----------------------------------------------------------------------------------------
-		class Private;
-		Private* GetPrivate() const;
+		// Node interface
+		virtual const FNodeType* GetType() const override { return &StaticType; }
+		static const FNodeType* GetStaticType() { return &StaticType; }
 
 	protected:
 
 		//! Forbidden. Manage with the Ptr<> template.
-        ~NodeSurfaceEdit();
+		~NodeSurfaceEdit() {}
 
 	private:
 
-		Private* m_pD;
+		static FNodeType StaticType;
 
 	};
 
diff --git a/Engine/Plugins/Experimental/Mutable/Source/MutableTools/Public/MuT/NodeSurfaceNew.h b/Engine/Plugins/Experimental/Mutable/Source/MutableTools/Public/MuT/NodeSurfaceNew.h
--- a/Engine/Plugins/Experimental/Mutable/Source/MutableTools/Public/MuT/NodeSurfaceNew.h
+++ b/Engine/Plugins/Experimental/Mutable/Source/MutableTools/Public/MuT/NodeSurfaceNew.h
@@ -6,237 +6,90 @@
 #include "MuR/RefCounted.h"
 #include "MuT/Node.h"
 #include "MuT/NodeSurface.h"
+#include "MuT/NodeMesh.h"
+#include "MuT/NodeImage.h"
+#include "MuT/NodeScalar.h"
+#include "MuT/NodeString.h"
+#include "MuT/NodeColour.h"
 
 namespace mu
 {
 
-	// Forward definitions
-	class NodeMesh;
-	typedef Ptr<NodeMesh> NodeMeshPtr;
-	typedef Ptr<const NodeMesh> NodeMeshPtrConst;
-
-	class NodeImage;
-	typedef Ptr<NodeImage> NodeImagePtr;
-	typedef Ptr<const NodeImage> NodeImagePtrConst;
-
-    class NodeColour;
-    typedef Ptr<NodeColour> NodeColourPtr;
-    typedef Ptr<const NodeColour> NodeColourPtrConst;
-
-    class NodeScalar;
-    typedef Ptr<NodeScalar> NodeScalarPtr;
-    typedef Ptr<const NodeScalar> NodeScalarPtrConst;
-
-    class NodeString;
-    typedef Ptr<NodeString> NodeStringPtr;
-    typedef Ptr<const NodeString> NodeStringPtrConst;
-
-    class NodeSurface;
-    typedef Ptr<NodeSurface> NodeSurfacePtr;
-    typedef Ptr<const NodeSurface> NodeSurfacePtrConst;
-
-    class NodeSurfaceNew;
-    typedef Ptr<NodeSurfaceNew> NodeSurfaceNewPtr;
-    typedef Ptr<const NodeSurfaceNew> NodeSurfaceNewPtrConst;
-
-
-
     //! This node makes a new Surface from several meshes and images.
     class MUTABLETOOLS_API NodeSurfaceNew : public NodeSurface
 	{
 	public:
 
-        NodeSurfaceNew();
-
-		//-----------------------------------------------------------------------------------------
-		// Node Interface
-		//-----------------------------------------------------------------------------------------
-
-        const FNodeType* GetType() const override;
-		static const FNodeType* GetStaticType();
-
-		//-----------------------------------------------------------------------------------------
-		// Own Interface
-		//-----------------------------------------------------------------------------------------
-
-        //! Set the name of the Surface.
-        void SetName( const FString&);
-
-        //! Add an optional, opaque id that will be returned in the surfaces of the created
-        //! instances. Can be useful to identify surfaces on the application side.
-        //! See Instance::GetSurfaceCustomId
-        void SetCustomID( uint32 id );
-
-		//! Add an id that will be used to identify the same surface in other LODs 
-		void SetSharedSurfaceId( int32 SharedSurfaceId );
-
-        //! \name Tags
-        //! \{
-
-        //! Add a tag to the surface:
-        //! - the surface will be affected by modifier nodes with the same tag
-        //! - the tag will be enabled when the surface is added to an object, and it can activate
-        //! variations for any surface.
-		void AddTag(const FString& tagName);
-
-        //! \}
-
-		//! \name Meshes
-        //! \{
-
-		//! Get the node generating one of the meshes in the Surface.
-		//! \param index index of the mesh, from 0 to GetMeshCount()-1
-		Ptr<NodeMesh> GetMesh() const;
-
-        //! Set the node generating one of the meshes in the Surface.
-		//! \param index index of the mesh, from 0 to GetMeshCount()-1
-		void SetMesh( Ptr<NodeMesh> );
-
-		//! \}
-
-
-		//! \name Images
-		//! \{
+		FString Name;
 
-		//! Get the number of images in the Surface.
-		int GetImageCount() const;
+		/** An optional, opaque id that will be returned in the surfaces of the created
+		* instances. Can be useful to identify surfaces on the application side.
+		*/
+		uint32 ExternalId = 0;
 
-		//! Set the number of images in the Surface.
-		void SetImageCount(int);
+		/** Add an id that will be used to identify the same surface in other LODs. */
+		int32 SharedSurfaceId = INDEX_NONE;
 
-		//! Get the node generating one of the images in the Surface.
-		//! \param index index of the image, from 0 to GetImageCount()-1
-		NodeImagePtr GetImage(int index) const;
+		Ptr<NodeMesh> Mesh;
 
-		//! Set the node generating one of the images in the Surface.
-		//! \param index index of the image, from 0 to GetImageCount()-1
-		void SetImage(int index, NodeImagePtr);
+		struct FImageData
+		{
+			FString Name;
+			FString MaterialName;
+			FString MaterialParameterName;
+			NodeImagePtr Image;
 
-		//! Get the name of a image in the Surface.
-		//! \param index index of the image, from 0 to GetImageCount()-1
-		const FString& GetImageName(int index) const;
+			// It could be negative, to indicate no layout.
+			int8 LayoutIndex = 0;
+		};
 
-        //! Set the name of a image in the Surface.
-        //! \param index index of the image, from 0 to GetImageCount()-1
-        //! \param strName name of the image
-        void SetImageName( int index, const FString& strName );
+		TArray<FImageData> Images;
 
-        //! Get the name of a image in the Surface.
-        //! \param index index of the image, from 0 to GetImageCount()-1
-        int GetImageLayoutIndex( int index ) const;
+		struct FVectorData
+		{
+			FString Name;
+			Ptr<NodeColour> Vector;
+		};
 
-        //! Set the name of a image in the Surface.
-        //! \param index index of the image, from 0 to GetImageCount()-1
-        //! \param layoutIndex index of the layout that will be used with this image.
-        void SetImageLayoutIndex( int imageIndex, int layoutIndex );
+		TArray<FVectorData> Vectors;
 
-		//! This can be used to ser additional information for error reporting.
-		//! It is not propagated to compiled objects in any way.
-		void SetImageAdditionalNames(int index, const FString& strMaterialName, const FString& strMaterialParameterName);
+		struct FScalar
+		{
+			FString Name;
+			Ptr<NodeScalar> Scalar;
+		};
 
-        //! \}
+		TArray<FScalar> Scalars;
 
+		struct FStringData
+		{
+			FString Name;
+			Ptr<NodeString> String;
+		};
 
-        //! \name Vectors
-        //! \{
+		TArray<FStringData> Strings;
 
-        //! Get the number of vectors in the Surface.
-        int GetVectorCount() const;
+		/** Tags added to the surface:
+		* - the surface will be affected by modifier nodes with the same tag
+		* - the tag will be enabled when the surface is added to an object, and it can activate
+		* variations for any surface.
+		*/
+		TArray<FString> Tags;
 
-        //! Set the number of vectors in the Surface.
-        void SetVectorCount(int);
-
-        //! Get the node generating one of the vectors in the Surface.
-        //! \param index index of the vector, from 0 to GetVectorCount()-1
-        NodeColourPtr GetVector(int index) const;
-
-        //! Set the node generating one of the vectors in the Surface.
-        //! \param index index of the vector, from 0 to GetVectorCount()-1
-        void SetVector(int index, NodeColourPtr);
-
-        //! Get the name of a vectors in the Surface.
-        //! \param index index of the vector, from 0 to GetVectorCount()-1
-        const FString& GetVectorName(int index) const;
-
-        //! Set the name of a image in the Surface.
-        //! \param index index of the vector, from 0 to GetVectorCount()-1
-        //! \param strName name of the vector
-        void SetVectorName(int index, const FString& strName);
-
-        //! \}
-
-
-        //! \name Scalars
-        //! \{
-
-        //! Get the number of Scalars in the Surface.
-        int GetScalarCount() const;
-
-        //! Set the number of Scalars in the Surface.
-        void SetScalarCount( int );
-
-        //! Get the node generating one of the Scalars in the Surface.
-        //! \param index index of the Scalar, from 0 to GetScalarCount()-1
-        NodeScalarPtr GetScalar( int index ) const;
-
-        //! Set the node generating one of the Scalars in the Surface.
-        //! \param index index of the Scalar, from 0 to GetScalarCount()-1
-        void SetScalar( int index, NodeScalarPtr );
-
-        //! Get the name of a Scalars in the Surface.
-        //! \param index index of the Scalar, from 0 to GetScalarCount()-1
-        const FString& GetScalarName( int index ) const;
-
-        //! Set the name of a image in the Surface.
-        //! \param index index of the Scalar, from 0 to GetScalarCount()-1
-        //! \param strName name of the Scalar
-        void SetScalarName( int index, const FString& strName );
-
-        //! \}
-
-
-        //! \name Strings
-        //! \{
-
-        //! Get the number of strings in the Surface.
-        int GetStringCount() const;
-
-        //! Set the number of strings in the Surface.
-        void SetStringCount( int );
-
-        //! Get the node generating one of the strings in the Surface.
-        //! \param index index of the strings, from 0 to GetStringCount()-1
-        NodeStringPtr GetString( int index ) const;
-
-        //! Set the node generating one of the Strings in the Surface.
-        //! \param index index of the String, from 0 to GetStringCount()-1
-        void SetString( int index, NodeStringPtr );
-
-        //! Get the name of a String in the Surface.
-        //! \param index index of the String, from 0 to GetStringCount()-1
-        const FString& GetStringName( int index ) const;
-
-        //! Set the name of a image in the Surface.
-        //! \param index index of the String, from 0 to GetStringCount()-1
-        //! \param strName name of the String
-        void SetStringName( int index, const FString& strName );
-
-        //! \}
+	public:
 
-        //-----------------------------------------------------------------------------------------
-		// Interface pattern
-		//-----------------------------------------------------------------------------------------
-		class Private;
-		Private* GetPrivate() const;
+		// Node interface
+		virtual const FNodeType* GetType() const override { return &StaticType; }
+		static const FNodeType* GetStaticType() { return &StaticType; }
 
 	protected:
 
 		//! Forbidden. Manage with the Ptr<> template.
-        ~NodeSurfaceNew();
+		~NodeSurfaceNew() {}
 
 	private:
 
-		Private* m_pD;
+		static FNodeType StaticType;
 
 	};
 
diff --git a/Engine/Plugins/Experimental/Mutable/Source/MutableTools/Public/MuT/NodeSurfaceSwitch.h b/Engine/Plugins/Experimental/Mutable/Source/MutableTools/Public/MuT/NodeSurfaceSwitch.h
--- a/Engine/Plugins/Experimental/Mutable/Source/MutableTools/Public/MuT/NodeSurfaceSwitch.h
+++ b/Engine/Plugins/Experimental/Mutable/Source/MutableTools/Public/MuT/NodeSurfaceSwitch.h
@@ -6,67 +6,34 @@
 #include "MuR/RefCounted.h"
 #include "MuT/Node.h"
 #include "MuT/NodeSurface.h"
+#include "MuT/NodeScalar.h"
 
 
 namespace mu
 {
 
-	// Forward definitions
-	class NodeScalar;
-	typedef Ptr<NodeScalar> NodeScalarPtr;
-	typedef Ptr<const NodeScalar> NodeScalarPtrConst;
-
-	class NodeSurfaceSwitch;
-	typedef Ptr<NodeSurfaceSwitch> NodeSurfaceSwitchPtr;
-	typedef Ptr<const NodeSurfaceSwitch> NodeSurfaceSwitchPtrConst;
-
-
 	//! This node selects an output Surface from a set of input Surfaces based on a parameter.
 	class MUTABLETOOLS_API NodeSurfaceSwitch : public NodeSurface
 	{
 	public:
 
-		NodeSurfaceSwitch();
-
-
-		//-----------------------------------------------------------------------------------------
-		// Node Interface
-		//-----------------------------------------------------------------------------------------
-
-		const FNodeType* GetType() const override;
-		static const FNodeType* GetStaticType();
-
-		//-----------------------------------------------------------------------------------------
-		// Own Interface
-		//-----------------------------------------------------------------------------------------
-
-		//! Get the node generating the parameter used to select the option.
-		NodeScalarPtr GetParameter() const;
-		void SetParameter( NodeScalarPtr );
-
-		//! Set the number of option Surfaces. It will keep the currently set targets and initialise
-		//! the new ones as null.
-		void SetOptionCount( int32 );
-
-		//! Get the node generating the t-th option Surface.
-		NodeSurfacePtr GetOption( int32 t ) const;
-		void SetOption( int32 t, NodeSurfacePtr );
+		Ptr<NodeScalar> Parameter;
+		TArray<Ptr<NodeSurface>> Options;
 
+	public:
 
-		//-----------------------------------------------------------------------------------------
-		// Interface pattern
-		//-----------------------------------------------------------------------------------------
-		class Private;
-		Private* GetPrivate() const;
+		// Node interface
+		virtual const FNodeType* GetType() const override { return &StaticType; }
+		static const FNodeType* GetStaticType() { return &StaticType; }
 
 	protected:
 
 		//! Forbidden. Manage with the Ptr<> template.
-		~NodeSurfaceSwitch();
+		~NodeSurfaceSwitch() {}
 
 	private:
 
-		Private* m_pD;
+		static FNodeType StaticType;
 
 	};
 
diff --git a/Engine/Plugins/Experimental/Mutable/Source/MutableTools/Public/MuT/NodeSurfaceVariation.h b/Engine/Plugins/Experimental/Mutable/Source/MutableTools/Public/MuT/NodeSurfaceVariation.h
--- a/Engine/Plugins/Experimental/Mutable/Source/MutableTools/Public/MuT/NodeSurfaceVariation.h
+++ b/Engine/Plugins/Experimental/Mutable/Source/MutableTools/Public/MuT/NodeSurfaceVariation.h
@@ -6,18 +6,11 @@
 #include "MuR/RefCounted.h"
 #include "MuT/Node.h"
 #include "MuT/NodeSurface.h"
+#include "MuT/NodeModifier.h"
 
 namespace mu
 {
 
-	// Forward definitions
-    class NodeSurfaceVariation;
-    typedef Ptr<NodeSurfaceVariation> NodeSurfaceVariationPtr;
-    typedef Ptr<const NodeSurfaceVariation> NodeSurfaceVariationPtrConst;
-
-    class NodeModifier;
-
-
 	//! This node modifies a node of the parent object of the object that this node belongs to.
     //! It allows to extend, cut and morph the parent Surface's meshes.
     //! It also allows to patch the parent Surface's textures.
@@ -25,67 +18,45 @@ namespace mu
 	{
 	public:
 
-        NodeSurfaceVariation();
-
-		//-----------------------------------------------------------------------------------------
-		// Node Interface
-		//-----------------------------------------------------------------------------------------
-
-        const FNodeType* GetType() const override;
-		static const FNodeType* GetStaticType();
-
-		//-----------------------------------------------------------------------------------------
-		// Own Interface
-		//-----------------------------------------------------------------------------------------
-
-		//! 
-        void AddDefaultSurface(NodeSurface* surface);
-        void AddDefaultModifier(NodeModifier* modifier);
-
-		//! Set the number of tags to consider in this variation
-		void SetVariationCount( int count );
+		TArray<Ptr<NodeSurface>> DefaultSurfaces;
+		TArray<Ptr<NodeModifier>> DefaultModifiers;
 
 		//!
-		int GetVariationCount() const;
-
-        //!
-        enum class VariationType : uint8
-        {
-            //! The variation selection is controlled by tags defined in other surfaces.
-            //! Default value.
-            Tag = 0,
+		enum class VariationType : uint8
+		{
+			//! The variation selection is controlled by tags defined in other surfaces.
+			//! Default value.
+			Tag = 0,
 
-            //! The variation selection is controlled by the state the object is in.
-            State
-        };
+			//! The variation selection is controlled by the state the object is in.
+			State
+		};
 
-        //!
-        void SetVariationType(VariationType);
+		NodeSurfaceVariation::VariationType Type = NodeSurfaceVariation::VariationType::Tag;
 
-        //! Set the tag or state name that will enable a specific vartiation
-		void SetVariationTag(int index, const FString& Tag);
+		struct FVariation
+		{
+			TArray<Ptr<NodeSurface>> Surfaces;
+			TArray<Ptr<NodeModifier>> Modifiers;
+			FString Tag;
+		};
 
-		//! 
-        void AddVariationSurface(int index, NodeSurface* surface);
-        void AddVariationModifier(int index, NodeModifier* modifier);
-
-		//!}
+		TArray<FVariation> Variations;
 
+	public:
 
-		//-----------------------------------------------------------------------------------------
-		// Interface pattern
-		//-----------------------------------------------------------------------------------------
-		class Private;
-		Private* GetPrivate() const;
+		// Node interface
+		virtual const FNodeType* GetType() const override { return &StaticType; }
+		static const FNodeType* GetStaticType() { return &StaticType; }
 
 	protected:
 
 		//! Forbidden. Manage with the Ptr<> template.
-        ~NodeSurfaceVariation();
+		~NodeSurfaceVariation() {}
 
 	private:
 
-		Private* m_pD;
+		static FNodeType StaticType;
 
 	};
 
