[mutable] Fixed bad recollection of mutable memory in use at a given time

#rnx

diff --git a/Engine/Plugins/Experimental/Mutable/Source/CustomizableObject/Private/MuCO/BusyWaits_Deprecated.cpp b/Engine/Plugins/Experimental/Mutable/Source/CustomizableObject/Private/MuCO/BusyWaits_Deprecated.cpp
--- a/Engine/Plugins/Experimental/Mutable/Source/CustomizableObject/Private/MuCO/BusyWaits_Deprecated.cpp
+++ b/Engine/Plugins/Experimental/Mutable/Source/CustomizableObject/Private/MuCO/BusyWaits_Deprecated.cpp
@@ -482,10 +482,13 @@ namespace CustomizableObjectMipDataProvider::ImplDeprecated
 		MUTABLE_CPUPROFILER_SCOPE(Task_Mutable_UpdateImage);
 		const double StartTime = FPlatformTime::Seconds();
 		
-		// Cache memory used when starting the update of the image
-		OperationData->ImageUpdateStartBytes = mu::FGlobalMemoryCounter::GetCounter();
-		mu::FGlobalMemoryCounter::Zero();
-		
+		if (CVarEnableBenchmark.GetValueOnAnyThread())
+		{
+			// Cache memory used when starting the update of the image
+			OperationData->ImageUpdateStartBytes = mu::FGlobalMemoryCounter::GetAbsoluteCounter();
+			mu::FGlobalMemoryCounter::Zero();
+		}
+
 		// Any external texture that may be needed for this update will be requested from Mutable Core's GetImage
 		// which will safely access the GlobalExternalImages map, and then just get the cached image or issue a disk read
 
diff --git a/Engine/Plugins/Experimental/Mutable/Source/CustomizableObject/Private/MuCO/CustomizableObjectMipDataProvider.cpp b/Engine/Plugins/Experimental/Mutable/Source/CustomizableObject/Private/MuCO/CustomizableObjectMipDataProvider.cpp
--- a/Engine/Plugins/Experimental/Mutable/Source/CustomizableObject/Private/MuCO/CustomizableObjectMipDataProvider.cpp
+++ b/Engine/Plugins/Experimental/Mutable/Source/CustomizableObject/Private/MuCO/CustomizableObjectMipDataProvider.cpp
@@ -139,7 +139,7 @@ namespace Impl
 		if (CVarEnableBenchmark.GetValueOnAnyThread())
 		{
 			// Cache memory used when starting the update of the image
-			OperationData->ImageUpdateStartBytes = mu::FGlobalMemoryCounter::GetCounter();
+			OperationData->ImageUpdateStartBytes = mu::FGlobalMemoryCounter::GetAbsoluteCounter();
 			mu::FGlobalMemoryCounter::Zero();
 		}
 		
diff --git a/Engine/Plugins/Experimental/Mutable/Source/CustomizableObject/Private/MuCO/CustomizableObjectSystem.cpp b/Engine/Plugins/Experimental/Mutable/Source/CustomizableObject/Private/MuCO/CustomizableObjectSystem.cpp
--- a/Engine/Plugins/Experimental/Mutable/Source/CustomizableObject/Private/MuCO/CustomizableObjectSystem.cpp
+++ b/Engine/Plugins/Experimental/Mutable/Source/CustomizableObject/Private/MuCO/CustomizableObjectSystem.cpp
@@ -1635,8 +1635,10 @@ namespace impl
 		
 		if (CVarEnableBenchmark.GetValueOnAnyThread())
 		{
-			Operation->UpdateStartBytes = mu::FGlobalMemoryCounter::GetCounter();
-			mu::FGlobalMemoryCounter::Zero();
+			// Get the amount of mutable memory in use now
+			Operation->UpdateStartBytes = mu::FGlobalMemoryCounter::GetAbsoluteCounter();
+			// Reset the counter to later get the peak during the updated
+			mu::FGlobalMemoryCounter::Zero();													
 		}
 
 		// Prepare streaming for the current customizable object
@@ -2934,10 +2936,13 @@ namespace impl
 			}
 		} // if (!bInstanceValid)
 
-		// Memory used in the context of this the update of mesh
-		OperationData->UpdateEndPeakBytes = mu::FGlobalMemoryCounter::GetPeak();
-		// Memory used in the context of the mesh update + the baseline memory already in use by mutable
-		OperationData->UpdateEndRealPeakBytes = OperationData->UpdateEndPeakBytes + OperationData->UpdateStartBytes;
+		if (CVarEnableBenchmark.GetValueOnAnyThread())
+		{
+			// Memory used in the context of this the update of mesh
+			OperationData->UpdateEndPeakBytes = mu::FGlobalMemoryCounter::GetPeak();
+			// Memory used in the context of the mesh update + the baseline memory used by mutable when starting the update
+			OperationData->UpdateEndRealPeakBytes = OperationData->UpdateEndPeakBytes + OperationData->UpdateStartBytes;
+		}
 		
 		UCustomizableObjectSystemPrivate* CustomizableObjectSystemPrivateData = System->GetPrivate();
 
diff --git a/Engine/Plugins/Experimental/Mutable/Source/MutableValidation/Private/MuV/CustomizableObjectCompilationUtility.cpp b/Engine/Plugins/Experimental/Mutable/Source/MutableValidation/Private/MuV/CustomizableObjectCompilationUtility.cpp
--- a/Engine/Plugins/Experimental/Mutable/Source/MutableValidation/Private/MuV/CustomizableObjectCompilationUtility.cpp
+++ b/Engine/Plugins/Experimental/Mutable/Source/MutableValidation/Private/MuV/CustomizableObjectCompilationUtility.cpp
@@ -84,7 +84,7 @@ bool FCustomizableObjectCompilationUtility::CompileCustomizableObject(UCustomiza
 #endif
 		
 		// Get the memory being used by mutable before the compilation
-		const int64 CompilationStartMutableBytes = mu::FGlobalMemoryCounter::GetCounter();
+		const int64 CompilationStartMutableBytes = mu::FGlobalMemoryCounter::GetAbsoluteCounter();
 		mu::FGlobalMemoryCounter::Zero();
 		
 		UE_LOG(LogMutable,Display,TEXT("Compiling Customizable Object..."));
