[Mutable] Remove ErrorOptimized update optimization.

Temporally removed due to unexpected consequences in internal projects.

#test packaged build
#rnx

diff --git a/Engine/Plugins/Experimental/Mutable/Source/CustomizableObject/Private/MuCO/CustomizableObjectSystem.cpp b/Engine/Plugins/Experimental/Mutable/Source/CustomizableObject/Private/MuCO/CustomizableObjectSystem.cpp
--- a/Engine/Plugins/Experimental/Mutable/Source/CustomizableObject/Private/MuCO/CustomizableObjectSystem.cpp
+++ b/Engine/Plugins/Experimental/Mutable/Source/CustomizableObject/Private/MuCO/CustomizableObjectSystem.cpp
@@ -119,6 +119,10 @@ TAutoConsoleVariable<bool> CVarEnableNewSplitMutableTask(
 	TEXT("Enables or disables the then new split GetImages and GetMesh tasks that remove BusyWaits."),
 	ECVF_Scalability);
 
+TAutoConsoleVariable<bool> CVarEnableUpdateOptimization(
+	TEXT("mutable.EnableUpdateOptimization"),
+	false,
+	TEXT("Enable or disable update optimization when no changes are made to the parent component."));
 
 #if WITH_EDITOR
 bool bEnableLODManagmentInEditor = false;
@@ -1248,14 +1252,28 @@ void UCustomizableObjectSystemPrivate::EnqueueUpdateSkeletalMesh(const TSharedRe
 			!(CurrentMutableOperation &&
 			Instance == CurrentMutableOperation->Instance)) // This condition is necessary because even if the descriptor is a subset, it will be replaced by the CurrentMutableOperation
 		{
-			Context->bOptimizedUpdate = true;
+			if (CVarEnableUpdateOptimization.GetValueOnGameThread()) // TODO Remove hotfix: UE-218957 
+			{
+				Context->bOptimizedUpdate = true;
 
-			// The user may have changed the AttachParent and we need to recustomize it.
-			// In case nothing need to be recustomized, the update will be considered ErrorOptimized.
-			UpdateSkeletalMesh(Context); 
-			Context->UpdateResult = Context->AttachedParentUpdated.IsEmpty() ? EUpdateResult::ErrorOptimized : EUpdateResult::Success;
-			
-			FinishUpdateGlobal(Context);
+				// The user may have changed the AttachParent and we need to recustomize it.
+				// In case nothing need to be recustomized, the update will be considered ErrorOptimized.
+				UpdateSkeletalMesh(Context);
+				Context->UpdateResult = Context->AttachedParentUpdated.IsEmpty() ? EUpdateResult::ErrorOptimized : EUpdateResult::Success;
+
+				FinishUpdateGlobal(Context);
+			}
+			else 
+			{
+				Context->bOptimizedUpdate = false;
+
+				// The user may have changed the AttachParent and we need to recustomize it.
+				// In case nothing need to be recustomized, the update will be considered ErrorOptimized.
+				UpdateSkeletalMesh(Context);
+				Context->UpdateResult = EUpdateResult::Success;
+
+				FinishUpdateGlobal(Context);
+			}
 		}
 		else
 		{
