[mutable] Add support for texture masks when defining layout blocks.
- Masked blocks get the vertices assigned with more priority than non-masked blocks.
#rnx

diff --git a/Engine/Plugins/Experimental/Mutable/Source/CustomizableObjectEditor/Private/MuCOE/CustomizableObjectLayout.cpp b/Engine/Plugins/Experimental/Mutable/Source/CustomizableObjectEditor/Private/MuCOE/CustomizableObjectLayout.cpp
--- a/Engine/Plugins/Experimental/Mutable/Source/CustomizableObjectEditor/Private/MuCOE/CustomizableObjectLayout.cpp
+++ b/Engine/Plugins/Experimental/Mutable/Source/CustomizableObjectEditor/Private/MuCOE/CustomizableObjectLayout.cpp
@@ -132,10 +132,10 @@ void UCustomizableObjectLayout::GenerateBlocksFromUVs()
 			Blocks.Empty();
 		
 			// Generating the layout blocks with the mutable layout
-			for (int i = 0; i < LayoutNode->Layout->GetBlockCount(); ++i)
+			for (int i = 0; i < LayoutNode->Blocks.Num(); ++i)
 			{
-				mu::FImageSize Min = LayoutNode->Layout->Blocks[i].Min;
-				mu::FImageSize Size = LayoutNode->Layout->Blocks[i].Size;
+				mu::FImageSize Min = LayoutNode->Blocks[i].Min;
+				mu::FImageSize Size = LayoutNode->Blocks[i].Size;
 		
 				FCustomizableObjectLayoutBlock block(FIntPoint(Min.X, Min.Y), FIntPoint(Min.X + Size.X, Min.Y + Size.Y));
 				Blocks.Add(block);
diff --git a/Engine/Plugins/Experimental/Mutable/Source/CustomizableObjectEditor/Private/MuCOE/GenerateMutableSource/GenerateMutableSourceLayout.cpp b/Engine/Plugins/Experimental/Mutable/Source/CustomizableObjectEditor/Private/MuCOE/GenerateMutableSource/GenerateMutableSourceLayout.cpp
--- a/Engine/Plugins/Experimental/Mutable/Source/CustomizableObjectEditor/Private/MuCOE/GenerateMutableSource/GenerateMutableSourceLayout.cpp
+++ b/Engine/Plugins/Experimental/Mutable/Source/CustomizableObjectEditor/Private/MuCOE/GenerateMutableSource/GenerateMutableSourceLayout.cpp
@@ -11,7 +11,7 @@
 #define LOCTEXT_NAMESPACE "CustomizableObjectEditor"
 
 
-mu::Ptr<mu::NodeLayout> GenerateMutableSourceLayout(const UEdGraphPin * Pin, FMutableGraphGenerationContext & GenerationContext, bool bLinkedToExtendMaterial)
+mu::Ptr<mu::NodeLayout> GenerateMutableSourceLayout(const UEdGraphPin * Pin, FMutableGraphGenerationContext& GenerationContext, bool bLinkedToExtendMaterial)
 {
 	check(Pin)
 	RETURN_ON_CYCLE(*Pin, GenerationContext)
@@ -52,45 +52,12 @@ mu::Ptr<mu::NodeLayout> GenerateMutableSourceLayout(const UEdGraphPin * Pin, FMu
 			}
 		}
 
-		mu::Ptr<mu::NodeLayout> LayoutNode = new mu::NodeLayout;
-		Result = LayoutNode;
-
-		LayoutNode->Layout->SetGridSize(TypedNodeBlocks->Layout->GetGridSize().X, TypedNodeBlocks->Layout->GetGridSize().Y);
-		LayoutNode->Layout->SetMaxGridSize(TypedNodeBlocks->Layout->GetMaxGridSize().X, TypedNodeBlocks->Layout->GetMaxGridSize().Y);
-		LayoutNode->Layout->SetBlockCount(TypedNodeBlocks->Layout->Blocks.Num() ? TypedNodeBlocks->Layout->Blocks.Num() : 1);
-
-		mu::EPackStrategy PackStrategy = ConvertLayoutStrategy(TypedNodeBlocks->Layout->GetPackingStrategy());
-		LayoutNode->Layout->SetLayoutPackingStrategy(PackStrategy);
-
-		LayoutNode->Layout->ReductionMethod = (TypedNodeBlocks->Layout->GetBlockReductionMethod() == ECustomizableObjectLayoutBlockReductionMethod::Halve ? mu::EReductionMethod::HALVE_REDUCTION : mu::EReductionMethod::UNITARY_REDUCTION);
-
-		if (bLinkedToExtendMaterial)
-		{
-			// Layout warnings can be safely ignored in this case. Vertices that do not belong to any layout block will be removed (Extend Materials only)
-			LayoutNode->FirstLODToIgnoreWarnings = 0;
-		}
-		else
-		{
-			LayoutNode->FirstLODToIgnoreWarnings = TypedNodeBlocks->Layout->GetIgnoreVertexLayoutWarnings() ? TypedNodeBlocks->Layout->GetFirstLODToIgnoreWarnings() : -1;
-		}
-
-		if (TypedNodeBlocks->Layout->Blocks.Num())
-		{
-			for (int BlockIndex = 0; BlockIndex < TypedNodeBlocks->Layout->Blocks.Num(); ++BlockIndex)
-			{
-				LayoutNode->Layout->Blocks[BlockIndex]= ToMutable(TypedNodeBlocks->Layout->Blocks[BlockIndex]);
-			}
-		}
-		else
+		bool bWasEmpty = false;
+		Result = CreateMutableLayoutNode(GenerationContext, TypedNodeBlocks->Layout, bLinkedToExtendMaterial,bWasEmpty);
+		if (bWasEmpty)
 		{
 			FString msg = "Layout without any block found. A grid sized block will be used instead.";
 			GenerationContext.Compiler->CompilerLog(FText::FromString(msg), Node, EMessageSeverity::Warning);
-
-			LayoutNode->Layout->Blocks[0].Min = { 0,0 };
-			LayoutNode->Layout->Blocks[0].Size = { uint16(TypedNodeBlocks->Layout->GetGridSize().X), uint16(TypedNodeBlocks->Layout->GetGridSize().Y) };
-			LayoutNode->Layout->Blocks[0].Priority = 0;
-			LayoutNode->Layout->Blocks[0].bReduceBothAxes = false;
-			LayoutNode->Layout->Blocks[0].bReduceByTwo = false;
 		}
 	}
 	
@@ -111,5 +78,87 @@ mu::Ptr<mu::NodeLayout> GenerateMutableSourceLayout(const UEdGraphPin * Pin, FMu
 	return Result;
 }
 
+
+mu::Ptr<mu::NodeLayout> CreateMutableLayoutNode(FMutableGraphGenerationContext& GenerationContext, const UCustomizableObjectLayout* UnrealLayout, bool bLinkedToExtendMaterial, bool& bWasEmpty )
+{
+	bWasEmpty = false;
+	mu::Ptr<mu::NodeLayout> LayoutNode = new mu::NodeLayout;
+
+	LayoutNode->Size = UE::Math::TIntVector2<uint16>(UnrealLayout->GetGridSize().X, UnrealLayout->GetGridSize().Y);
+	LayoutNode->MaxSize = UE::Math::TIntVector2<uint16>(UnrealLayout->GetMaxGridSize().X, UnrealLayout->GetMaxGridSize().Y);
+	LayoutNode->Blocks.SetNum(UnrealLayout->Blocks.Num() ? UnrealLayout->Blocks.Num() : 1);
+
+	mu::EPackStrategy PackStrategy = ConvertLayoutStrategy(UnrealLayout->GetPackingStrategy());
+	LayoutNode->Strategy = PackStrategy;
+
+	LayoutNode->ReductionMethod = (UnrealLayout->GetBlockReductionMethod() == ECustomizableObjectLayoutBlockReductionMethod::Halve ? mu::EReductionMethod::HALVE_REDUCTION : mu::EReductionMethod::UNITARY_REDUCTION);
+
+	if (bLinkedToExtendMaterial)
+	{
+		// Layout warnings can be safely ignored in this case. Vertices that do not belong to any layout block will be removed (Extend Materials only)
+		LayoutNode->FirstLODToIgnoreWarnings = 0;
+	}
+	else
+	{
+		LayoutNode->FirstLODToIgnoreWarnings = UnrealLayout->GetIgnoreVertexLayoutWarnings() ? UnrealLayout->GetFirstLODToIgnoreWarnings() : -1;
+	}
+
+	if (UnrealLayout->Blocks.Num())
+	{
+		for (int BlockIndex = 0; BlockIndex < UnrealLayout->Blocks.Num(); ++BlockIndex)
+		{
+			LayoutNode->Blocks[BlockIndex] = ToMutable(GenerationContext, UnrealLayout->Blocks[BlockIndex]);
+		}
+	}
+	else
+	{
+		bWasEmpty = true;
+		LayoutNode->Blocks[0].Min = { 0,0 };
+		LayoutNode->Blocks[0].Size = { uint16(UnrealLayout->GetGridSize().X), uint16(UnrealLayout->GetGridSize().Y) };
+		LayoutNode->Blocks[0].Priority = 0;
+		LayoutNode->Blocks[0].bReduceBothAxes = false;
+		LayoutNode->Blocks[0].bReduceByTwo = false;
+	}
+
+	return LayoutNode;
+}
+
+
+mu::FSourceLayoutBlock ToMutable(FMutableGraphGenerationContext& GenerationContext, const FCustomizableObjectLayoutBlock& UnrealBlock)
+{
+	mu::FSourceLayoutBlock MutableBlock;
+
+	MutableBlock.Min = { uint16(UnrealBlock.Min.X), uint16(UnrealBlock.Min.Y) };
+	FIntPoint Size = UnrealBlock.Max - UnrealBlock.Min;
+	MutableBlock.Size = { uint16(Size.X), uint16(Size.Y) };
+
+	MutableBlock.Priority = UnrealBlock.Priority;
+	MutableBlock.bReduceBothAxes = UnrealBlock.bReduceBothAxes;
+	MutableBlock.bReduceByTwo = UnrealBlock.bReduceByTwo;
+
+	if (UnrealBlock.Mask)
+	{
+		GenerationContext.AddParticipatingObject(*UnrealBlock.Mask);
+
+		// In the editor the src data can be directly accessed
+		mu::Ptr<mu::Image> MaskImage = new mu::Image();
+
+		FMutableSourceTextureData Tex(*UnrealBlock.Mask);
+		EUnrealToMutableConversionError Error = ConvertTextureUnrealSourceToMutable(MaskImage.get(), Tex, 0);
+		if (Error != EUnrealToMutableConversionError::Success)
+		{
+			// This should never happen, so details are not necessary.
+			UE_LOG(LogMutable, Warning, TEXT("Failed to convert layout block mask texture."));
+		}
+		else
+		{
+			MutableBlock.Mask = MaskImage;
+		}
+	}
+
+	return MutableBlock;
+}
+
+
 #undef LOCTEXT_NAMESPACE
 
diff --git a/Engine/Plugins/Experimental/Mutable/Source/CustomizableObjectEditor/Private/MuCOE/GenerateMutableSource/GenerateMutableSourceLayout.h b/Engine/Plugins/Experimental/Mutable/Source/CustomizableObjectEditor/Private/MuCOE/GenerateMutableSource/GenerateMutableSourceLayout.h
--- a/Engine/Plugins/Experimental/Mutable/Source/CustomizableObjectEditor/Private/MuCOE/GenerateMutableSource/GenerateMutableSourceLayout.h
+++ b/Engine/Plugins/Experimental/Mutable/Source/CustomizableObjectEditor/Private/MuCOE/GenerateMutableSource/GenerateMutableSourceLayout.h
@@ -9,21 +9,8 @@ class UEdGraphPin;
 struct FMutableGraphGenerationContext;
 
 
-mu::Ptr<mu::NodeLayout> GenerateMutableSourceLayout(const UEdGraphPin* Pin, FMutableGraphGenerationContext& GenerationContext, bool bLinkedToExtendMaterial = false);
+extern mu::Ptr<mu::NodeLayout> GenerateMutableSourceLayout(const UEdGraphPin*, FMutableGraphGenerationContext&, bool bLinkedToExtendMaterial = false);
 
+extern mu::Ptr<mu::NodeLayout> CreateMutableLayoutNode(FMutableGraphGenerationContext& GenerationContext, const UCustomizableObjectLayout*, bool bLinkedToExtendMaterial, bool& bWasEmpty);
 
-/** */
-inline mu::Layout::FBlock ToMutable(const FCustomizableObjectLayoutBlock& UnrealBlock)
-{
-	mu::Layout::FBlock MutableBlock;
-
-	MutableBlock.Min = { uint16(UnrealBlock.Min.X), uint16(UnrealBlock.Min.Y) };
-	FIntPoint Size = UnrealBlock.Max - UnrealBlock.Min;
-	MutableBlock.Size = { uint16(Size.X), uint16(Size.Y) };
-
-	MutableBlock.Priority = UnrealBlock.Priority;
-	MutableBlock.bReduceBothAxes = UnrealBlock.bReduceBothAxes;
-	MutableBlock.bReduceByTwo = UnrealBlock.bReduceByTwo;
-
-	return MutableBlock;
-}
+extern mu::FSourceLayoutBlock ToMutable(FMutableGraphGenerationContext& GenerationContext, const FCustomizableObjectLayoutBlock&);
diff --git a/Engine/Plugins/Experimental/Mutable/Source/CustomizableObjectEditor/Private/MuCOE/GenerateMutableSource/GenerateMutableSourceMesh.cpp b/Engine/Plugins/Experimental/Mutable/Source/CustomizableObjectEditor/Private/MuCOE/GenerateMutableSource/GenerateMutableSourceMesh.cpp
--- a/Engine/Plugins/Experimental/Mutable/Source/CustomizableObjectEditor/Private/MuCOE/GenerateMutableSource/GenerateMutableSourceMesh.cpp
+++ b/Engine/Plugins/Experimental/Mutable/Source/CustomizableObjectEditor/Private/MuCOE/GenerateMutableSource/GenerateMutableSourceMesh.cpp
@@ -3177,16 +3177,16 @@ mu::Ptr<mu::NodeLayout> CreateDefaultLayout()
 	constexpr int32 GridSize = 4;
 	
 	mu::Ptr<mu::NodeLayout> LayoutNode = new mu::NodeLayout();
-	LayoutNode->Layout->SetGridSize(GridSize, GridSize);
-	LayoutNode->Layout->SetMaxGridSize(GridSize, GridSize);
-	LayoutNode->Layout->SetLayoutPackingStrategy(mu::EPackStrategy::RESIZABLE_LAYOUT);
-	LayoutNode->Layout->ReductionMethod = mu::EReductionMethod::HALVE_REDUCTION;
-	LayoutNode->Layout->SetBlockCount(1);
-	LayoutNode->Layout->Blocks[0].Min = { 0, 0 };
-	LayoutNode->Layout->Blocks[0].Size = { GridSize, GridSize };
-	LayoutNode->Layout->Blocks[0].Priority = 0;
-	LayoutNode->Layout->Blocks[0].bReduceBothAxes = false;
-	LayoutNode->Layout->Blocks[0].bReduceByTwo = false;
+	LayoutNode->Size = { GridSize, GridSize };
+	LayoutNode->MaxSize = { GridSize, GridSize };
+	LayoutNode->Strategy = mu::EPackStrategy::RESIZABLE_LAYOUT;
+	LayoutNode->ReductionMethod = mu::EReductionMethod::HALVE_REDUCTION;
+	LayoutNode->Blocks.SetNum(1);
+	LayoutNode->Blocks[0].Min = { 0, 0 };
+	LayoutNode->Blocks[0].Size = { GridSize, GridSize };
+	LayoutNode->Blocks[0].Priority = 0;
+	LayoutNode->Blocks[0].bReduceBothAxes = false;
+	LayoutNode->Blocks[0].bReduceByTwo = false;
 
 	return LayoutNode;
 }
@@ -4013,11 +4013,11 @@ mu::NodeMeshPtr GenerateMutableSourceMesh(const UEdGraphPin* Pin,
 				}
 				else
 				{
-					if (!TypedNode->PoseAsset) // Check if the slot has a selected pose. Could be left empty by the user
-					{
-						FString msg = FString::Printf(TEXT("Found pose mesh node without a pose asset assigned."));
-						GenerationContext.Compiler->CompilerLog(FText::FromString(msg), TypedNode);
-					}
+						if (!TypedNode->PoseAsset) // Check if the slot has a selected pose. Could be left empty by the user
+						{
+							FString msg = FString::Printf(TEXT("Found pose mesh node without a pose asset assigned."));
+							GenerationContext.Compiler->CompilerLog(FText::FromString(msg), TypedNode);
+						}
 
 					Result = InputMeshNode;
 				}
@@ -4095,7 +4095,7 @@ mu::NodeMeshPtr GenerateMutableSourceMesh(const UEdGraphPin* Pin,
 							MutableColumnName = TypedNodeTable->GenerateStaticMeshMutableColumName(DataTableColumnName, SectionIndexConnected);
 						}
 					}
-
+					
 					// Generating a new Mesh column if not exists
 					if (Table->FindColumn(MutableColumnName) == INDEX_NONE)
 					{
@@ -4154,40 +4154,12 @@ mu::NodeMeshPtr GenerateMutableSourceMesh(const UEdGraphPin* Pin,
 								// Generating node Layouts
 								for (int32 i = 0; i < Layouts.Num(); ++i)
 								{
-									mu::Ptr<mu::NodeLayout> LayoutNode = new mu::NodeLayout;
-
-									LayoutNode->Layout->SetGridSize(Layouts[i]->GetGridSize().X, Layouts[i]->GetGridSize().Y);
-									LayoutNode->Layout->SetMaxGridSize(Layouts[i]->GetMaxGridSize().X, Layouts[i]->GetMaxGridSize().Y);
-									LayoutNode->Layout->SetBlockCount(Layouts[i]->Blocks.Num() ? Layouts[i]->Blocks.Num() : 1);
-
-									mu::EPackStrategy PackStrategy = ConvertLayoutStrategy(Layouts[i]->GetPackingStrategy());
-									LayoutNode->Layout->SetLayoutPackingStrategy(PackStrategy);
-
-									LayoutNode->Layout->ReductionMethod = (Layouts[i]->GetBlockReductionMethod() == ECustomizableObjectLayoutBlockReductionMethod::Halve ? mu::EReductionMethod::HALVE_REDUCTION : mu::EReductionMethod::UNITARY_REDUCTION);
-
-									if (bLinkedToExtendMaterial)
-									{
-										// Layout warnings can be safely ignored in this case. Vertices that do not belong to any layout block will be removed (Extend Materials only)
-										LayoutNode->FirstLODToIgnoreWarnings = 0;
-									}
-
-									if (Layouts[i]->Blocks.Num())
-									{
-										for (int BlockIndex = 0; BlockIndex < Layouts[i]->Blocks.Num(); ++BlockIndex)
-										{
-											LayoutNode->Layout->Blocks[BlockIndex] = ToMutable(Layouts[i]->Blocks[BlockIndex]);
-										}
-									}
-									else
+									bool bWasEmpty = false;
+									mu::Ptr<mu::NodeLayout> LayoutNode = CreateMutableLayoutNode(GenerationContext, Layouts[i], bLinkedToExtendMaterial, bWasEmpty);
+									if (bWasEmpty)
 									{
 										FString msg = "Mesh Column [" + MutableColumnName + "] Layout doesn't has any block. A grid sized block will be used instead.";
 										GenerationContext.Compiler->CompilerLog(FText::FromString(msg), Node, EMessageSeverity::Warning);
-
-										LayoutNode->Layout->Blocks[0].Min = { 0,0 };
-										LayoutNode->Layout->Blocks[0].Size = { uint16(Layouts[i]->GetGridSize().X), uint16(Layouts[i]->GetGridSize().Y) };
-										LayoutNode->Layout->Blocks[0].Priority = 0;
-										LayoutNode->Layout->Blocks[0].bReduceBothAxes = false;
-										LayoutNode->Layout->Blocks[0].bReduceByTwo = false;
 									}
 
 									MeshTableNode->SetLayout(i, LayoutNode);
diff --git a/Engine/Plugins/Experimental/Mutable/Source/CustomizableObjectEditor/Private/MuCOE/SCustomizableObjectLayoutGrid.cpp b/Engine/Plugins/Experimental/Mutable/Source/CustomizableObjectEditor/Private/MuCOE/SCustomizableObjectLayoutGrid.cpp
--- a/Engine/Plugins/Experimental/Mutable/Source/CustomizableObjectEditor/Private/MuCOE/SCustomizableObjectLayoutGrid.cpp
+++ b/Engine/Plugins/Experimental/Mutable/Source/CustomizableObjectEditor/Private/MuCOE/SCustomizableObjectLayoutGrid.cpp
@@ -2,18 +2,21 @@
 
 #include "MuCOE/SCustomizableObjectLayoutGrid.h"
 
+#include "MuCOE/SStandAloneAssetPicker.h"
+#include "MuCOE/UnrealEditorPortabilityHelpers.h"
 #include "BatchedElements.h"
 #include "CanvasTypes.h"
 #include "Framework/MultiBox/MultiBoxBuilder.h"
 #include "Layout/WidgetPath.h"
-// Required for engine branch preprocessor defines.
-#include "MuCO/UnrealPortabilityHelpers.h"
-#include "MuCOE/UnrealEditorPortabilityHelpers.h"
 #include "Rendering/SlateRenderer.h"
 #include "RenderingThread.h"
 #include "Styling/SlateTypes.h"
 #include "UnrealClient.h"
 #include "Widgets/Input/SNumericEntryBox.h"
+#include "ThumbnailRendering/ThumbnailManager.h"
+#include "Engine/Texture2D.h"
+#include "TextureResource.h"
+
 
 class FExtender;
 class FPaintArgs;
@@ -97,7 +100,7 @@ private:
 	virtual void Draw_RenderThread(FRHICommandListImmediate& RHICmdList, const void* InWindowBackBuffer, const FSlateCustomDrawParams& Params) override;
 
 	/** Basic function to draw a block in the canvas */
-	void DrawBlock(FBatchedElements* BatchedElements, const FHitProxyId HitProxyId, const FRect2D& BlockRect, FColor Color);
+	void DrawBlock(FBatchedElements* BatchedElements, const FHitProxyId HitProxyId, const FRect2D& BlockRect, FColor Color, UTexture2D* Mask = nullptr);
 
 	/** SlateElement initialized, can Draw during the DrawRenderThread call. */
 	bool Initialized = false;
@@ -124,6 +127,10 @@ private:
 	ELayoutGridMode LayoutMode = ELGM_Show;
 
 	FSlateCanvasRenderTarget* RenderTarget = new FSlateCanvasRenderTarget();
+
+	/** Default colors. */
+	FColor SelectedBlockColor = FColor(75, 106, 230, 155);
+	FColor UnselectedBlockColor = FColor(230, 199, 75, 155);
 };
 
 
@@ -142,7 +149,8 @@ void SCustomizableObjectLayoutGrid::Construct( const FArguments& InArgs )
 	OnSetBlockPriority = InArgs._OnSetBlockPriority;
 	OnSetReduceBlockSymmetrically = InArgs._OnSetReduceBlockSymmetrically;
 	OnSetReduceBlockByTwo = InArgs._OnSetReduceBlockByTwo;
-	
+	OnSetBlockMask = InArgs._OnSetBlockMask;
+
 	UVCanvasDrawer = TSharedPtr<FUVCanvasDrawer>(new FUVCanvasDrawer());
 	UVCanvasDrawer->SetLayoutMode(Mode);
 }
@@ -396,7 +404,7 @@ FReply SCustomizableObjectLayoutGrid::OnMouseButtonDown( const FGeometry& MyGeom
 			const bool CloseAfterSelection = true;
 			FMenuBuilder MenuBuilder(CloseAfterSelection, NULL, TSharedPtr<FExtender>(), false, &FCoreStyle::Get(), false);
 
-			MenuBuilder.BeginSection("Block Management", LOCTEXT("GridActionsTitle", "Grid Actions"));
+			MenuBuilder.BeginSection("Block Management", LOCTEXT("BlockActionsTitle", "Block Actions"));
 			{
 				if (SelectedBlocks.Num())
 				{
@@ -414,7 +422,7 @@ FReply SCustomizableObjectLayoutGrid::OnMouseButtonDown( const FGeometry& MyGeom
 			}
 			MenuBuilder.EndSection();
 
-			MenuBuilder.BeginSection("Fixed Layout Strategy", LOCTEXT("BlockActionsTitle", "Fixed Layout Actions"));
+			MenuBuilder.BeginSection("Block Properties for Fixed Layout", LOCTEXT("BlockPropertiesFixedTitle", "Block Properties for Fixed Layout"));
 			{
 				if (SelectedBlocks.Num())
 				{
@@ -460,6 +468,28 @@ FReply SCustomizableObjectLayoutGrid::OnMouseButtonDown( const FGeometry& MyGeom
 			}
 			MenuBuilder.EndSection();
 
+
+			MenuBuilder.BeginSection("Block Properties for Masks", LOCTEXT("BlockPropertiesMaskTitle", "Block Mask"));
+			{
+				if (SelectedBlocks.Num())
+				{
+					MenuBuilder.AddWidget(
+						SNew(SBox)
+						.WidthOverride(125.0f)
+						.ToolTipText(LOCTEXT("SetBlockMask_Tooltip", "Sets the UV mask texture for the block."))
+						[
+							SNew(SStandAloneAssetPicker)
+								.OnAssetSelected(this, &SCustomizableObjectLayoutGrid::OnMaskAssetSelected)
+								.OnGetAllowedClasses(FOnGetAllowedClasses::CreateLambda([](TArray<const UClass*>& OutClasses) {OutClasses.Add(UTexture2D::StaticClass()); }))
+								.InitialAsset( GetBlockMaskValue() )
+						]
+						, FText::FromString("Block Mask"), true);
+
+					// TODO: Additional properties: color used for preview?
+				}
+			}
+			MenuBuilder.EndSection();
+
 			FWidgetPath WidgetPath = MouseEvent.GetEventPath() != nullptr ? *MouseEvent.GetEventPath() : FWidgetPath();
 			FSlateApplication::Get().PushMenu(AsShared(), WidgetPath, MenuBuilder.MakeWidget(), FSlateApplication::Get().GetCursorPos(), FPopupTransitionEffect(FPopupTransitionEffect::ContextMenu));
 		}
@@ -477,6 +507,59 @@ FReply SCustomizableObjectLayoutGrid::OnMouseButtonDown( const FGeometry& MyGeom
 }
 
 
+void SCustomizableObjectLayoutGrid::CloseMenu()
+{
+	FSlateApplication::Get().DismissAllMenus();
+}
+
+
+void SCustomizableObjectLayoutGrid::OnMaskAssetSelected(const FAssetData& AssetData)
+{
+	UTexture2D* Mask = Cast<UTexture2D>(AssetData.GetAsset());
+	if (SelectedBlocks.Num())
+	{
+		OnSetBlockMask.ExecuteIfBound(Mask);
+	}
+}
+
+
+
+UTexture2D* SCustomizableObjectLayoutGrid::GetBlockMaskValue() const
+{
+	if (SelectedBlocks.Num())
+	{
+		TArray<FCustomizableObjectLayoutBlock> CurrentSelectedBlocks;
+
+		for (const FCustomizableObjectLayoutBlock& Block : Blocks.Get())
+		{
+			if (SelectedBlocks.Contains(Block.Id))
+			{
+				CurrentSelectedBlocks.Add(Block);
+			}
+		}
+
+		UTexture2D* BlockMask = CurrentSelectedBlocks[0].Mask;
+		bool bSameMask = true;
+
+		for (const FCustomizableObjectLayoutBlock& Block : CurrentSelectedBlocks)
+		{
+			if (Block.Mask != BlockMask)
+			{
+				bSameMask = false;
+				break;
+			}
+		}
+
+		if (bSameMask)
+		{
+			return BlockMask;
+		}
+	}
+
+	return nullptr;
+}
+
+
 FReply SCustomizableObjectLayoutGrid::OnMouseButtonUp( const FGeometry& MyGeometry, const FPointerEvent& MouseEvent )
 {
 	if (Mode == ELGM_Show)
@@ -1299,7 +1382,7 @@ void FUVCanvasDrawer::Draw_RenderThread(class FRHICommandListImmediate& RHICmdLi
 		// Drawing Blocks
 		for (const FCustomizableObjectLayoutBlock& Block : Blocks)
 		{
-			const FColor SelectionBlockColor = SelectedBlocks.Contains(Block.Id) ? FColor(75, 106, 230, 155) : FColor(230, 199, 75, 155);
+			const FColor SelectionBlockColor = SelectedBlocks.Contains(Block.Id) ? SelectedBlockColor : UnselectedBlockColor;
 
 			const FVector2f BlockMin(Block.Min);
 			const FVector2f BlockMax(Block.Max);
@@ -1309,7 +1392,7 @@ void FUVCanvasDrawer::Draw_RenderThread(class FRHICommandListImmediate& RHICmdLi
 			SelectionBlock.Min = FVector2f(Origin) + BlockMin * CellSize + CellSize * 0.1f;
 			SelectionBlock.Size = (BlockMax - BlockMin) * CellSize - CellSize * 0.2f;
 
-			DrawBlock(BatchedElements, HitProxyId, SelectionBlock, SelectionBlockColor);
+			DrawBlock(BatchedElements, HitProxyId, SelectionBlock, SelectionBlockColor, Block.Mask.Get());
 
 			if (LayoutMode == ELayoutGridMode::ELGM_Edit)
 			{
@@ -1331,7 +1414,7 @@ void FUVCanvasDrawer::Draw_RenderThread(class FRHICommandListImmediate& RHICmdLi
 }
 
 
-void FUVCanvasDrawer::DrawBlock(FBatchedElements* BatchedElements, const FHitProxyId HitProxyId, const FRect2D& BlockRect, FColor Color)
+void FUVCanvasDrawer::DrawBlock(FBatchedElements* BatchedElements, const FHitProxyId HitProxyId, const FRect2D& BlockRect, FColor Color, UTexture2D* Mask)
 {
 	// Vertex positions
 	FVector4 Vert0(BlockRect.Min.X, BlockRect.Min.Y, 0, 1);
@@ -1339,24 +1422,34 @@ void FUVCanvasDrawer::DrawBlock(FBatchedElements* BatchedElements, const FHitPro
 	FVector4 Vert2(BlockRect.Min.X + BlockRect.Size.X, BlockRect.Min.Y, 0, 1);
 	FVector4 Vert3(BlockRect.Min.X + BlockRect.Size.X, BlockRect.Min.Y + BlockRect.Size.Y, 0, 1);
 
+	auto VertexToMaskUVs = [this](const FVector4& V)
+		{
+			return ( FVector2D(V.X,V.Y) - Origin ) / (FVector2D(CellSize) * FVector2D(GridSize.X, GridSize.X));
+		};
+
 	// Brush Paint triangle
 	{
-		int32 V0 = BatchedElements->AddVertex(Vert0, FVector2d(0.0f, 0.0f), Color, HitProxyId);
-		int32 V1 = BatchedElements->AddVertex(Vert1, FVector2d(0.0f, 1.0f), Color, HitProxyId);
-		int32 V2 = BatchedElements->AddVertex(Vert2, FVector2d(1.0f, 0.0f), Color, HitProxyId);
-		int32 V3 = BatchedElements->AddVertex(Vert3, FVector2d(1.0f, 1.0f), Color, HitProxyId);
+		int32 V0 = BatchedElements->AddVertex(Vert0, VertexToMaskUVs(Vert0), Color, HitProxyId);
+		int32 V1 = BatchedElements->AddVertex(Vert1, VertexToMaskUVs(Vert1), Color, HitProxyId);
+		int32 V2 = BatchedElements->AddVertex(Vert2, VertexToMaskUVs(Vert2), Color, HitProxyId);
+		int32 V3 = BatchedElements->AddVertex(Vert3, VertexToMaskUVs(Vert3), Color, HitProxyId);
 
-		BatchedElements->AddTriangle(V0, V1, V2, GWhiteTexture, EBlendMode::BLEND_Translucent);
-		BatchedElements->AddTriangle(V1, V3, V2, GWhiteTexture, EBlendMode::BLEND_Translucent);
+		FTexture* Texture = GWhiteTexture;
+		if (Mask)
+		{
+			Texture = Mask->GetResource();
+		}
+		BatchedElements->AddTriangle(V0, V1, V2, Texture, EBlendMode::BLEND_Translucent);
+		BatchedElements->AddTriangle(V1, V3, V2, Texture, EBlendMode::BLEND_Translucent);
 	}
 
 	// Drawing Outline to selected Blocks
-	if (Color == FColor(75, 106, 230, 155))
+	if (Color == SelectedBlockColor)
 	{
-		BatchedElements->AddLine(Vert0, Vert1, FColor(230, 199, 75, 155), HitProxyId, 4.0f);
-		BatchedElements->AddLine(Vert1, Vert3, FColor(230, 199, 75, 155), HitProxyId, 4.0f);
-		BatchedElements->AddLine(Vert3, Vert2, FColor(230, 199, 75, 155), HitProxyId, 4.0f);
-		BatchedElements->AddLine(Vert2, Vert0, FColor(230, 199, 75, 155), HitProxyId, 4.0f);
+		BatchedElements->AddLine(Vert0, Vert1, UnselectedBlockColor, HitProxyId, 4.0f);
+		BatchedElements->AddLine(Vert1, Vert3, UnselectedBlockColor, HitProxyId, 4.0f);
+		BatchedElements->AddLine(Vert3, Vert2, UnselectedBlockColor, HitProxyId, 4.0f);
+		BatchedElements->AddLine(Vert2, Vert0, UnselectedBlockColor, HitProxyId, 4.0f);
 	}
 }
 
diff --git a/Engine/Plugins/Experimental/Mutable/Source/CustomizableObjectEditor/Private/MuCOE/SCustomizableObjectLayoutGrid.h b/Engine/Plugins/Experimental/Mutable/Source/CustomizableObjectEditor/Private/MuCOE/SCustomizableObjectLayoutGrid.h
--- a/Engine/Plugins/Experimental/Mutable/Source/CustomizableObjectEditor/Private/MuCOE/SCustomizableObjectLayoutGrid.h
+++ b/Engine/Plugins/Experimental/Mutable/Source/CustomizableObjectEditor/Private/MuCOE/SCustomizableObjectLayoutGrid.h
@@ -9,6 +9,7 @@ class FPaintArgs;
 class FSlateRect;
 class FSlateWindowElementList;
 class FWidgetStyle;
+class UTexture2D;
 struct FGeometry;
 struct FGuid;
 struct FKeyEvent;
@@ -54,15 +55,16 @@ public:
 	DECLARE_DELEGATE_OneParam(FSetBlockPriority, int32);
 	DECLARE_DELEGATE_OneParam(FSetReduceBlockSymmetrically, bool);
 	DECLARE_DELEGATE_OneParam(FSetReduceBlockByTwo, bool);
+	DECLARE_DELEGATE_OneParam(FSetBlockMask, UTexture2D*);
 
 	SLATE_BEGIN_ARGS( SCustomizableObjectLayoutGrid ){}
 
 		SLATE_ATTRIBUTE( FIntPoint, GridSize )
 		SLATE_ATTRIBUTE( TArray<FCustomizableObjectLayoutBlock>, Blocks )
-		SLATE_ARGUMENT( TArray<FVector2f>, UVLayout  )
+		SLATE_ARGUMENT( TArray<FVector2f>, UVLayout )
 		SLATE_ARGUMENT( TArray<FVector2f>, UnassignedUVLayoutVertices )
-		SLATE_ARGUMENT( ELayoutGridMode, Mode  )
-		SLATE_ARGUMENT( FColor, SelectionColor  )
+		SLATE_ARGUMENT( ELayoutGridMode, Mode )
+		SLATE_ARGUMENT( FColor, SelectionColor )
 		SLATE_EVENT( FBlockChangedDelegate, OnBlockChanged )
 		SLATE_EVENT( FBlockSelectionChangedDelegate, OnSelectionChanged )
 		SLATE_EVENT(FDeleteBlockDelegate, OnDeleteBlocks)
@@ -70,6 +72,7 @@ public:
 		SLATE_EVENT(FSetBlockPriority, OnSetBlockPriority)
 		SLATE_EVENT(FSetReduceBlockSymmetrically, OnSetReduceBlockSymmetrically)
 		SLATE_EVENT(FSetReduceBlockByTwo, OnSetReduceBlockByTwo)
+		SLATE_EVENT(FSetBlockMask, OnSetBlockMask)
 
 	SLATE_END_ARGS()
 
@@ -119,6 +122,9 @@ public:
 	/** Set the grid and blocks to show in the widget. */
 	void SetBlocks( const FIntPoint& GridSize, const TArray<FCustomizableObjectLayoutBlock>& Blocks);
 
+
+private:
+
 	/** Gets the priority value of the selected blocks */
 	TOptional<int32> GetBlockPriorityValue() const;
 
@@ -134,10 +140,16 @@ public:
 	/** Gets block reduction value of the selected blocks */
 	ECheckBoxState GetReductionMethodBoolValue(EFixedReductionOptions Option) const;
 
-private:
-
 	bool MouseOnBlock(FGuid BlockId, FVector2D MousePosition, bool CheckResizeBlock = false) const;
 
+	/** Callback for an asset selection popup menu close. */
+	void CloseMenu();
+
+	/** Callback for an actual mask asset selection in the context menu. */
+	void OnMaskAssetSelected(const FAssetData& AssetData);
+
+	UTexture2D* GetBlockMaskValue() const;
+
 private:
 
 	/** A delegate to report block changes */
@@ -148,6 +160,7 @@ private:
 	FSetBlockPriority OnSetBlockPriority;
 	FSetReduceBlockSymmetrically OnSetReduceBlockSymmetrically;
 	FSetReduceBlockByTwo OnSetReduceBlockByTwo;
+	FSetBlockMask OnSetBlockMask;
 
 	/** Size of the grid in blocks */
 	TAttribute<FIntPoint> GridSize;
diff --git a/Engine/Plugins/Experimental/Mutable/Source/CustomizableObjectEditor/Private/MuCOE/SCustomizableObjectNodeLayoutBlocksEditor.cpp b/Engine/Plugins/Experimental/Mutable/Source/CustomizableObjectEditor/Private/MuCOE/SCustomizableObjectNodeLayoutBlocksEditor.cpp
--- a/Engine/Plugins/Experimental/Mutable/Source/CustomizableObjectEditor/Private/MuCOE/SCustomizableObjectNodeLayoutBlocksEditor.cpp
+++ b/Engine/Plugins/Experimental/Mutable/Source/CustomizableObjectEditor/Private/MuCOE/SCustomizableObjectNodeLayoutBlocksEditor.cpp
@@ -113,6 +113,7 @@ void SCustomizableObjectNodeLayoutBlocksEditor::SetCurrentLayout(UCustomizableOb
 			.OnSetBlockPriority(this, &SCustomizableObjectNodeLayoutBlocksEditor::OnSetBlockPriority)
 			.OnSetReduceBlockSymmetrically(this, &SCustomizableObjectNodeLayoutBlocksEditor::OnSetBlockReductionSymmetry)
 			.OnSetReduceBlockByTwo(this, &SCustomizableObjectNodeLayoutBlocksEditor::OnSetBlockReductionByTwo)
+			.OnSetBlockMask(this, &SCustomizableObjectNodeLayoutBlocksEditor::OnSetBlockMask)
 		]
 	];	
 }
@@ -343,6 +344,29 @@ void SCustomizableObjectNodeLayoutBlocksEditor::OnSetBlockReductionByTwo(bool bI
 }
 
 
+void SCustomizableObjectNodeLayoutBlocksEditor::OnSetBlockMask(UTexture2D* InValue)
+{
+	if (CurrentLayout)
+	{
+		const FScopedTransaction Transaction(LOCTEXT("OnSetBlockMask", "Change Block Mask"));
+		CurrentLayout->Modify();
+
+		if (LayoutGridWidget.IsValid())
+		{
+			const TArray<FGuid>& SelectedBlocks = LayoutGridWidget->GetSelectedBlocks();
+
+			for (FCustomizableObjectLayoutBlock& Block : CurrentLayout->Blocks)
+			{
+				if (SelectedBlocks.Contains(Block.Id))
+				{
+					Block.Mask = InValue;
+				}
+			}
+		}
+	}
+}
+
+
 TArray<FCustomizableObjectLayoutBlock> SCustomizableObjectNodeLayoutBlocksEditor::GetBlocks() const
 {
 	TArray<FCustomizableObjectLayoutBlock> Blocks;
diff --git a/Engine/Plugins/Experimental/Mutable/Source/CustomizableObjectEditor/Private/MuCOE/SCustomizableObjectNodeLayoutBlocksEditor.h b/Engine/Plugins/Experimental/Mutable/Source/CustomizableObjectEditor/Private/MuCOE/SCustomizableObjectNodeLayoutBlocksEditor.h
--- a/Engine/Plugins/Experimental/Mutable/Source/CustomizableObjectEditor/Private/MuCOE/SCustomizableObjectNodeLayoutBlocksEditor.h
+++ b/Engine/Plugins/Experimental/Mutable/Source/CustomizableObjectEditor/Private/MuCOE/SCustomizableObjectNodeLayoutBlocksEditor.h
@@ -87,6 +87,9 @@ private:
 	/** Sets the block reduction ReduceByTwo option. */
 	void OnSetBlockReductionByTwo(bool bInValue);
 
+	/** Callback for block mask change. */
+	void OnSetBlockMask(class UTexture2D* InValue);
+
 	TSharedPtr<IToolTip> GenerateInfoToolTip() const;
 
 };
diff --git a/Engine/Plugins/Experimental/Mutable/Source/CustomizableObjectEditor/Private/MuCOE/SMutableCodeViewer.cpp b/Engine/Plugins/Experimental/Mutable/Source/CustomizableObjectEditor/Private/MuCOE/SMutableCodeViewer.cpp
--- a/Engine/Plugins/Experimental/Mutable/Source/CustomizableObjectEditor/Private/MuCOE/SMutableCodeViewer.cpp
+++ b/Engine/Plugins/Experimental/Mutable/Source/CustomizableObjectEditor/Private/MuCOE/SMutableCodeViewer.cpp
@@ -2613,7 +2613,7 @@ void SMutableCodeViewer::Tick(const FGeometry& AllottedGeometry, const double In
 	case mu::DT_LAYOUT:
 	{
 		check(PreviewLayoutViewer);
-		mu::LayoutPtrConst MutableLayout = System->GetPrivate()->BuildLayout(MutableModel, PreviewParameters.get(), SelectedOperationAddress);
+		mu::Ptr<const mu::Layout> MutableLayout = System->GetPrivate()->BuildLayout(MutableModel, PreviewParameters.get(), SelectedOperationAddress);
 		PreviewLayoutViewer->SetLayout(MutableLayout);
 		break;
 	}
@@ -2780,7 +2780,7 @@ void SMutableCodeViewer::PreviewMutableMesh(mu::MeshPtrConst InMeshPtr)
 }
 
 
-void SMutableCodeViewer::PreviewMutableLayout(mu::LayoutPtrConst Layout)
+void SMutableCodeViewer::PreviewMutableLayout(mu::Ptr<const mu::Layout> Layout)
 {
 	PrepareLayoutViewer();
 	PreviewLayoutViewer->SetLayout(Layout);
diff --git a/Engine/Plugins/Experimental/Mutable/Source/CustomizableObjectEditor/Private/MuCOE/SMutableCodeViewer.h b/Engine/Plugins/Experimental/Mutable/Source/CustomizableObjectEditor/Private/MuCOE/SMutableCodeViewer.h
--- a/Engine/Plugins/Experimental/Mutable/Source/CustomizableObjectEditor/Private/MuCOE/SMutableCodeViewer.h
+++ b/Engine/Plugins/Experimental/Mutable/Source/CustomizableObjectEditor/Private/MuCOE/SMutableCodeViewer.h
@@ -789,7 +789,7 @@ private:
 public:
 	void PreviewMutableImage (mu::ImagePtrConst InImagePtr);
 	void PreviewMutableMesh (mu::MeshPtrConst InMeshPtr);
-	void PreviewMutableLayout(mu::LayoutPtrConst Layout);
+	void PreviewMutableLayout(mu::Ptr<const mu::Layout> Layout);
 	void PreviewMutableSkeleton(mu::SkeletonPtrConst Skeleton);
 	void PreviewMutableString(const FString& InString);
 	void PreviewMutableProjector(const mu::FProjector* Projector);
diff --git a/Engine/Plugins/Experimental/Mutable/Source/CustomizableObjectEditor/Private/MuCOE/SMutableConstantsWidget.h b/Engine/Plugins/Experimental/Mutable/Source/CustomizableObjectEditor/Private/MuCOE/SMutableConstantsWidget.h
--- a/Engine/Plugins/Experimental/Mutable/Source/CustomizableObjectEditor/Private/MuCOE/SMutableConstantsWidget.h
+++ b/Engine/Plugins/Experimental/Mutable/Source/CustomizableObjectEditor/Private/MuCOE/SMutableConstantsWidget.h
@@ -59,7 +59,7 @@ struct FMutableConstantStringElement : public FMutableConstantElement
 */
 struct FMutableConstantLayoutElement : public FMutableConstantElement
 {
-	mu::LayoutPtrConst Layout = nullptr;
+	mu::Ptr<const mu::Layout> Layout = nullptr;
 };
 
 /**
diff --git a/Engine/Plugins/Experimental/Mutable/Source/CustomizableObjectEditor/Private/MuCOE/SMutableLayoutViewer.cpp b/Engine/Plugins/Experimental/Mutable/Source/CustomizableObjectEditor/Private/MuCOE/SMutableLayoutViewer.cpp
--- a/Engine/Plugins/Experimental/Mutable/Source/CustomizableObjectEditor/Private/MuCOE/SMutableLayoutViewer.cpp
+++ b/Engine/Plugins/Experimental/Mutable/Source/CustomizableObjectEditor/Private/MuCOE/SMutableLayoutViewer.cpp
@@ -120,7 +120,7 @@ void SMutableLayoutViewer::Tick(const FGeometry& AllottedGeometry, const double
 }
 
 
-void SMutableLayoutViewer::SetLayout(const mu::LayoutPtrConst& InMutableLayout)
+void SMutableLayoutViewer::SetLayout(const mu::Ptr<const mu::Layout>& InMutableLayout)
 {
 	MutableLayout = InMutableLayout;
 	bIsPendingUpdate = true;
diff --git a/Engine/Plugins/Experimental/Mutable/Source/CustomizableObjectEditor/Private/MuCOE/SMutableLayoutViewer.h b/Engine/Plugins/Experimental/Mutable/Source/CustomizableObjectEditor/Private/MuCOE/SMutableLayoutViewer.h
--- a/Engine/Plugins/Experimental/Mutable/Source/CustomizableObjectEditor/Private/MuCOE/SMutableLayoutViewer.h
+++ b/Engine/Plugins/Experimental/Mutable/Source/CustomizableObjectEditor/Private/MuCOE/SMutableLayoutViewer.h
@@ -17,7 +17,7 @@ public:
 
 	SLATE_BEGIN_ARGS(SMutableLayoutViewer) {}
 		SLATE_ATTRIBUTE(FIntPoint, GridSize)
-		SLATE_ARGUMENT(mu::LayoutPtrConst, Layout)
+		SLATE_ARGUMENT(mu::Ptr<const mu::Layout>, Layout)
 	SLATE_END_ARGS()
 
 	void Construct(const FArguments& InArgs);
@@ -28,7 +28,7 @@ public:
 	// Own interface
 
 	/** Set the image to show in the widget. */
-	void SetLayout(const mu::LayoutPtrConst& Image );
+	void SetLayout(const mu::Ptr<const mu::Layout>& Image );
 
 private:
 
@@ -36,7 +36,7 @@ private:
 	TSharedPtr<SCustomizableObjectLayoutGrid> LayoutViewer;
 
 	/** Mutable Layout object being previewed */
-	mu::LayoutPtrConst MutableLayout;
+	mu::Ptr<const mu::Layout> MutableLayout;
 
 	/** Is true, the image or the visible LOD have changed and we need to update. */
 	bool bIsPendingUpdate = false;
diff --git a/Engine/Plugins/Experimental/Mutable/Source/CustomizableObjectEditor/Private/MuCOE/SStandAloneAssetPicker.cpp b/Engine/Plugins/Experimental/Mutable/Source/CustomizableObjectEditor/Private/MuCOE/SStandAloneAssetPicker.cpp
new file mode 100644
--- /dev/null
+++ b/Engine/Plugins/Experimental/Mutable/Source/CustomizableObjectEditor/Private/MuCOE/SStandAloneAssetPicker.cpp
@@ -0,0 +1,178 @@
+// Copyright Epic Games, Inc. All Rights Reserved.
+
+#include "MuCOE/SStandAloneAssetPicker.h"
+
+#include "Rendering/SlateRenderer.h"
+#include "Styling/SlateTypes.h"
+#include "ThumbnailRendering/ThumbnailManager.h"
+#include "Engine/Texture2D.h"
+#include "Widgets/Layout/SBorder.h"
+#include "Modules/ModuleManager.h"
+#include "Widgets/Images/SImage.h"
+#include "Widgets/Layout/SBox.h"
+#include "Widgets/Input/SButton.h"
+#include "IContentBrowserSingleton.h"
+#include "ContentBrowserModule.h"
+
+
+#define LOCTEXT_NAMESPACE "CustomizableObjectEditor"
+
+
+void SStandAloneAssetPicker::Construct(const FArguments& InArgs)
+{
+	OnAssetSelected = InArgs._OnAssetSelected;
+	OnGetAllowedClasses = InArgs._OnGetAllowedClasses;
+	CurrentAsset = InArgs._InitialAsset;
+
+	ChildSlot
+		[
+			SNew(SHorizontalBox)
+				+ SHorizontalBox::Slot()
+				.Padding(0.0f,3.0f,5.0f,0.0f)
+				.AutoWidth()
+				.VAlign(VAlign_Center)
+				[
+					SNew(SBorder)
+					.Visibility(EVisibility::SelfHitTestInvisible)
+					.Padding(FMargin(0.0f, 0.0f, 4.0f, 4.0f))
+					.BorderImage(FAppStyle::Get().GetBrush("PropertyEditor.AssetTileItem.DropShadow"))
+					[
+						SNew(SOverlay)
+						+SOverlay::Slot()
+						.Padding(1.0f)
+						[
+							SNew(SBorder)
+							.Padding(0)
+							.BorderImage(FStyleDefaults::GetNoBrush())
+							[
+								SAssignNew(ThumbnailContainer, SBox)
+							]
+						]
+					]
+				]
+				+ SHorizontalBox::Slot()
+				.VAlign(VAlign_Center)
+				.HAlign(HAlign_Left)
+				[
+					SNew(SVerticalBox)
+						+ SVerticalBox::Slot()
+						.AutoHeight()
+						.HAlign(HAlign_Right)
+						[
+							SAssignNew(AssetPickerAnchor, SMenuAnchor)
+								.Placement(MenuPlacement_AboveAnchor)
+								.OnGetMenuContent(this, &SStandAloneAssetPicker::OnGenerateAssetPicker)
+						]
+						+ SVerticalBox::Slot()
+						.FillHeight(1)
+						[
+							SNew(SButton)
+								.ButtonStyle(FAppStyle::Get(), "HoverHintOnly")
+								.OnClicked(this, &SStandAloneAssetPicker::OnClicked)
+								.ToolTipText(LOCTEXT("PickButtonLabel", "Pick Asset"))
+								.ContentPadding(0.0f)
+								.ForegroundColor(FSlateColor::UseForeground())
+								.IsFocusable(false)
+								[
+									SNew(SImage)
+										.Image(FAppStyle::GetBrush("PropertyWindow.Button_PickAsset"))
+										.ColorAndOpacity(FSlateColor::UseForeground())
+								]
+						]
+				]
+		];
+
+	RefreshThumbnail();
+}
+
+void SStandAloneAssetPicker::RefreshThumbnail()
+{
+	FIntPoint ThumbnailSize(128, 128);
+	TSharedPtr<FAssetThumbnailPool> Pool = UThumbnailManager::Get().GetSharedThumbnailPool();
+	AssetThumbnail = MakeShareable(new FAssetThumbnail(CurrentAsset, ThumbnailSize.X, ThumbnailSize.Y, Pool));
+
+	FAssetThumbnailConfig AssetThumbnailConfig;
+
+	ThumbnailWidget = AssetThumbnail->MakeThumbnailWidget(AssetThumbnailConfig);
+
+	ThumbnailContainer->SetContent(ThumbnailWidget.ToSharedRef());
+}
+
+FReply SStandAloneAssetPicker::OnClicked()
+{
+	AssetPickerAnchor->SetIsOpen(true);
+	return FReply::Handled();
+}
+
+TSharedRef<SWidget> SStandAloneAssetPicker::OnGenerateAssetPicker()
+{
+	TArray<const UClass*> AllowedClasses;
+	OnGetAllowedClasses.ExecuteIfBound(AllowedClasses);
+
+	if (AllowedClasses.Num() == 0)
+	{
+		// Assume all classes are allowed
+		AllowedClasses.Add(UObject::StaticClass());
+	}
+	FContentBrowserModule& ContentBrowserModule = FModuleManager::Get().LoadModuleChecked<FContentBrowserModule>(TEXT("ContentBrowser"));
+
+	FAssetPickerConfig AssetPickerConfig;
+	for (auto ClassIt = AllowedClasses.CreateConstIterator(); ClassIt; ++ClassIt)
+	{
+		const UClass* Class = (*ClassIt);
+		AssetPickerConfig.Filter.ClassPaths.Add(Class->GetClassPathName());
+	}
+	// Allow child classes
+	AssetPickerConfig.Filter.bRecursiveClasses = true;
+	// Set a delegate for setting the asset from the picker
+	AssetPickerConfig.OnAssetSelected = FOnAssetSelected::CreateSP(this, &SStandAloneAssetPicker::OnAssetSelectedFromPicker);
+	AssetPickerConfig.OnAssetEnterPressed = FOnAssetEnterPressed::CreateSP(this, &SStandAloneAssetPicker::OnAssetEnterPressedFromPicker);
+	AssetPickerConfig.bAllowDragging = false;
+	AssetPickerConfig.bAllowNullSelection = true;
+	// Use the list view by default
+	AssetPickerConfig.InitialAssetViewType = EAssetViewType::List;
+
+	AssetPickerConfig.InitialAssetSelection = FAssetData(CurrentAsset);
+
+	TSharedRef<SWidget> MenuContent =
+		SNew(SBox)
+		.HeightOverride(300.0f)
+		.WidthOverride(300.0f)
+		[
+			SNew(SBorder)
+				.BorderImage(FAppStyle::GetBrush("Menu.Background"))
+				[
+					ContentBrowserModule.Get().CreateAssetPicker(AssetPickerConfig)
+				]
+		];
+
+	return MenuContent;
+}
+
+void SStandAloneAssetPicker::OnAssetSelectedFromPicker(const struct FAssetData& AssetData)
+{
+	AssetPickerAnchor->SetIsOpen(false);
+
+	CurrentAsset = AssetData.GetAsset();
+
+	OnAssetSelected.ExecuteIfBound(CurrentAsset.Get());
+
+	RefreshThumbnail();
+}
+
+void SStandAloneAssetPicker::OnAssetEnterPressedFromPicker(const TArray<struct FAssetData>& AssetData)
+{
+	AssetPickerAnchor->SetIsOpen(false);
+
+	CurrentAsset = nullptr;
+	if (AssetData.Num() > 0)
+	{
+		CurrentAsset = AssetData[0].GetAsset();
+	}
+
+	OnAssetSelected.ExecuteIfBound(CurrentAsset.Get());
+	RefreshThumbnail();
+}
+
+
+#undef LOCTEXT_NAMESPACE
diff --git a/Engine/Plugins/Experimental/Mutable/Source/CustomizableObjectEditor/Private/MuCOE/SStandAloneAssetPicker.h b/Engine/Plugins/Experimental/Mutable/Source/CustomizableObjectEditor/Private/MuCOE/SStandAloneAssetPicker.h
new file mode 100644
--- /dev/null
+++ b/Engine/Plugins/Experimental/Mutable/Source/CustomizableObjectEditor/Private/MuCOE/SStandAloneAssetPicker.h
@@ -0,0 +1,46 @@
+// Copyright Epic Games, Inc. All Rights Reserved.
+
+#pragma once
+
+#include "Widgets/SCompoundWidget.h"
+#include "Widgets/Layout/SBox.h"
+#include "Widgets/Input/SMenuAnchor.h"
+#include "PropertyCustomizationHelpers.h"
+#include "AssetThumbnail.h"
+
+
+class SStandAloneAssetPicker : public SCompoundWidget
+{
+public:
+	SLATE_BEGIN_ARGS(SStandAloneAssetPicker) {}
+		SLATE_EVENT(FOnGetAllowedClasses, OnGetAllowedClasses)
+		SLATE_EVENT(FOnAssetSelected, OnAssetSelected)
+		SLATE_ARGUMENT(TObjectPtr<UObject>, InitialAsset)
+	SLATE_END_ARGS()
+
+	void Construct(const FArguments& InArgs);
+private:
+
+	FReply OnClicked();
+
+	TSharedRef<SWidget> OnGenerateAssetPicker();
+
+	void OnAssetSelectedFromPicker(const struct FAssetData& AssetData);
+	void OnAssetEnterPressedFromPicker(const TArray<struct FAssetData>& AssetData);
+	
+	void RefreshThumbnail();
+
+private:
+
+	/** Menu anchor for opening and closing the asset picker */
+	TSharedPtr<SMenuAnchor> AssetPickerAnchor;
+
+	FOnGetAllowedClasses OnGetAllowedClasses;
+	FOnAssetSelected OnAssetSelected;
+	TObjectPtr<UObject> CurrentAsset;
+
+	TSharedPtr<FAssetThumbnail> AssetThumbnail;
+	TSharedPtr<SWidget> ThumbnailWidget;
+	TSharedPtr<SBox> ThumbnailContainer;
+};
+
diff --git a/Engine/Plugins/Experimental/Mutable/Source/CustomizableObjectEditor/Public/MuCOE/CustomizableObjectLayout.h b/Engine/Plugins/Experimental/Mutable/Source/CustomizableObjectEditor/Public/MuCOE/CustomizableObjectLayout.h
--- a/Engine/Plugins/Experimental/Mutable/Source/CustomizableObjectEditor/Public/MuCOE/CustomizableObjectLayout.h
+++ b/Engine/Plugins/Experimental/Mutable/Source/CustomizableObjectEditor/Public/MuCOE/CustomizableObjectLayout.h
@@ -68,6 +68,10 @@ struct CUSTOMIZABLEOBJECTEDITOR_API FCustomizableObjectLayoutBlock
 	/** Block will be reduced by two in an Unitary Layout reduction. */
 	UPROPERTY(EditAnywhere, Category = CustomizableObject)
 	bool bReduceByTwo = false;
+
+	/** Block mask to use to filter the UVs when assigning them to the block. */
+	UPROPERTY(EditAnywhere, Category = CustomizableObject)
+	TObjectPtr<class UTexture2D> Mask = nullptr;
 };
 
 UCLASS()
@@ -99,7 +103,7 @@ public:
 	FIntPoint GetMaxGridSize() const { return MaxGridSize; }
 	ECustomizableObjectTextureLayoutPackingStrategy GetPackingStrategy() const { return PackingStrategy; }
 	bool GetIgnoreVertexLayoutWarnings() const { return bIgnoreUnassignedVertexWarning; };
-	int32 GetFirstLODToIgnoreWarnings() { return FirstLODToIgnore; };
+	int32 GetFirstLODToIgnoreWarnings() const { return FirstLODToIgnore; };
 	ECustomizableObjectLayoutBlockReductionMethod GetBlockReductionMethod()const { return BlockReductionMethod; }
 
 	void GetUVChannel(TArray<FVector2f>& UVs, int32 UVChannelIndex = 0) const;
diff --git a/Engine/Plugins/Experimental/Mutable/Source/MutableRuntime/Private/MuR/CodeRunner.cpp b/Engine/Plugins/Experimental/Mutable/Source/MutableRuntime/Private/MuR/CodeRunner.cpp
--- a/Engine/Plugins/Experimental/Mutable/Source/MutableRuntime/Private/MuR/CodeRunner.cpp
+++ b/Engine/Plugins/Experimental/Mutable/Source/MutableRuntime/Private/MuR/CodeRunner.cpp
@@ -5866,7 +5866,7 @@ namespace mu
 			OP::ResourceConstantArgs args = Program.GetOpArgs<OP::ResourceConstantArgs>(item.At);
             check( args.value < (uint32)pModel->GetPrivate()->m_program.m_constantLayouts.Num() );
 
-            LayoutPtrConst pResult = Program.m_constantLayouts
+            Ptr<const Layout> pResult = Program.m_constantLayouts
                     [ args.value ];
             StoreLayout( item, pResult );
             break;
@@ -5888,7 +5888,7 @@ namespace mu
                 Ptr<const Layout> pA = LoadLayout( FCacheAddress(args.Base,item) );
                 Ptr<const Layout> pB = LoadLayout( FCacheAddress(args.Added,item) );
 
-                LayoutPtrConst pResult;
+                Ptr<const Layout> pResult;
 
                 if (pA && pB)
                 {
@@ -5928,26 +5928,16 @@ namespace mu
             {
                 Ptr<const Layout> Source = LoadLayout( FCacheAddress(args.Source,item) );
 
-				LayoutPtr pResult;
+				Ptr<Layout> Result;
 
 				if (Source)
 				{
-					pResult = Source->Clone();
+					Result = Source->Clone();
 
-					SCRATCH_LAYOUT_PACK scratch;
-					int32 BlockCount = Source->GetBlockCount();
-					scratch.blocks.SetNum(BlockCount);
-					scratch.sorted.SetNum(BlockCount);
-					scratch.positions.SetNum(BlockCount);
-					scratch.priorities.SetNum(BlockCount);
-					scratch.reductions.SetNum(BlockCount);
-					scratch.ReduceBothAxes.SetNum(BlockCount);
-					scratch.ReduceByTwo.SetNum(BlockCount);
-
-					LayoutPack3(pResult.get(), Source.get(), &scratch);
+					LayoutPack3(Result.get(), Source.get() );
 				}
 
-                StoreLayout( item, pResult );
+                StoreLayout( item, Result );
                 break;
             }
 
diff --git a/Engine/Plugins/Experimental/Mutable/Source/MutableRuntime/Private/MuR/Layout.cpp b/Engine/Plugins/Experimental/Mutable/Source/MutableRuntime/Private/MuR/Layout.cpp
--- a/Engine/Plugins/Experimental/Mutable/Source/MutableRuntime/Private/MuR/Layout.cpp
+++ b/Engine/Plugins/Experimental/Mutable/Source/MutableRuntime/Private/MuR/Layout.cpp
@@ -11,8 +11,8 @@
 
 namespace mu {
 
-	MUTABLE_IMPLEMENT_POD_SERIALISABLE(Layout::FBlock);
-	MUTABLE_IMPLEMENT_POD_VECTOR_SERIALISABLE(Layout::FBlock);
+	MUTABLE_IMPLEMENT_POD_SERIALISABLE(FLayoutBlock);
+	MUTABLE_IMPLEMENT_POD_VECTOR_SERIALISABLE(FLayoutBlock);
 
 	//---------------------------------------------------------------------------------------------
 	Layout::Layout()
diff --git a/Engine/Plugins/Experimental/Mutable/Source/MutableRuntime/Private/MuR/OpImageRasterMesh.h b/Engine/Plugins/Experimental/Mutable/Source/MutableRuntime/Private/MuR/OpImageRasterMesh.h
--- a/Engine/Plugins/Experimental/Mutable/Source/MutableRuntime/Private/MuR/OpImageRasterMesh.h
+++ b/Engine/Plugins/Experimental/Mutable/Source/MutableRuntime/Private/MuR/OpImageRasterMesh.h
@@ -30,7 +30,12 @@ namespace mu
 	inline void ImageRasterMesh( const Mesh* pMesh, Image* pImage, int32 LayoutIndex, uint64 BlockId,
 		UE::Math::TIntVector2<uint16> CropMin, UE::Math::TIntVector2<uint16> UncroppedSize )
 	{
-		MUTABLE_CPUPROFILER_SCOPE(ImageRasterMesh)
+		MUTABLE_CPUPROFILER_SCOPE(ImageRasterMesh);
+
+		if (pMesh->GetVertexCount() == 0)
+		{
+			return;
+		}
 
 		check( pImage->GetFormat()== EImageFormat::IF_L_UBYTE );
 
@@ -86,7 +91,7 @@ namespace mu
 
         UntypedMeshBufferIteratorConst bloIt( pMesh->GetVertexBuffers(), MBS_LAYOUTBLOCK, LayoutIndex );
 
-        if (BlockId==Layout::InvalidBlockId || bloIt.GetElementSize()==0 )
+        if (BlockId== FLayoutBlock::InvalidBlockId || bloIt.GetElementSize()==0 )
 		{
 			// Raster all the faces
             WhitePixelProcessor pixelProc;
diff --git a/Engine/Plugins/Experimental/Mutable/Source/MutableRuntime/Private/MuR/OpLayoutPack.cpp b/Engine/Plugins/Experimental/Mutable/Source/MutableRuntime/Private/MuR/OpLayoutPack.cpp
new file mode 100644
--- /dev/null
+++ b/Engine/Plugins/Experimental/Mutable/Source/MutableRuntime/Private/MuR/OpLayoutPack.cpp
@@ -0,0 +1,652 @@
+// Copyright Epic Games, Inc. All Rights Reserved.
+
+#pragma once
+
+#include "MuR/OpLayoutPack.h"
+#include "MuR/MutableMath.h"
+#include "MuR/MutableTrace.h"
+#include "MuR/Layout.h"
+#include "MuR/Platform.h"
+#include "GenericPlatform/GenericPlatformMath.h"
+
+namespace mu
+{
+
+	struct FPackLayoutBlock
+	{
+		FPackLayoutBlock() {}
+
+		FPackLayoutBlock(int32 InIndex, UE::Math::TIntVector2<uint16> InSize, int32 InPriority = 0, bool bInReduceBothAxes = false, bool bInReduceByTwo = false)
+		{
+			Index = InIndex;
+			size = InSize;
+			priority = InPriority;
+			bReduceBothAxes = bInReduceBothAxes;
+			bReduceByTwo = bInReduceByTwo;
+		}
+
+		int32 Index = -1;
+		UE::Math::TIntVector2<uint16> size = UE::Math::TIntVector2<uint16>(0, 0);
+		int32 priority = 0;
+		bool bReduceBothAxes = false;
+		bool bReduceByTwo = false;
+	};
+
+	inline bool CompareBlocks(const FPackLayoutBlock& a, const FPackLayoutBlock& b)
+	{
+		if (a.size[1] > b.size[1])
+		{
+			return true;
+		}
+		else if (a.size[1] < b.size[1])
+		{
+			return false;
+		}
+		else
+		{
+			int32 areaA = a.size[0] * a.size[1];
+			int32 areaB = b.size[0] * b.size[1];
+			if (areaA > areaB)
+			{
+				return true;
+			}
+			else if (areaA < areaB)
+			{
+				return false;
+			}
+			else
+			{
+				// This has to be deterministic, and indices are supposed to be unique
+				return a.Index < b.Index;
+			}
+		}
+	}
+
+	inline bool CompareBlocksPriority(const FPackLayoutBlock& a, const FPackLayoutBlock& b)
+	{
+		if (a.priority == b.priority)
+		{
+			// TODO(Max): Check if this comparison modifies the block order while reducing them.
+			return CompareBlocks(a, b);
+		}
+
+		return a.priority > b.priority;
+	}
+
+
+	struct FScratchLayoutPack
+	{
+		TArray< UE::Math::TIntVector2<uint16> > blocks;
+		TArray< FPackLayoutBlock > sorted;
+		TArray< FIntVector2 > positions;
+		TArray< int32 > priorities;
+		TArray< FIntVector2 > reductions;
+		TArray< int32 > ReduceBothAxes;
+		TArray< int32 > ReduceByTwo;
+	};
+
+
+    inline char DebugGetBlockAt( const FScratchLayoutPack& scratch,
+                                   const TArray<uint8_t>& packedFlag,
+                                   int32 x, int32 y )
+    {
+        for ( size_t b=0; b<packedFlag.Num(); ++b )
+        {
+            if (packedFlag[b])
+            {
+				int32 i = scratch.sorted[b].Index;
+                if ( x>=scratch.positions[i][0] &&
+                     x<scratch.positions[i][0]+scratch.sorted[b].size[0] &&
+                     y>=scratch.positions[i][1] &&
+                     y<scratch.positions[i][1]+scratch.sorted[b].size[1]
+                     )
+                {
+					return char('a')+char(b);
+                }
+            }
+        }
+        return '.';
+    }
+
+
+	inline void ReductionOperation(uint16& BlockSize, EReductionMethod ReductionMethod, int32 bReduceByTwo)
+	{
+		if (ReductionMethod == EReductionMethod::UNITARY_REDUCTION)
+		{
+			int32 Reduction = (bReduceByTwo && BlockSize > 2) ? 2 : 1;
+			BlockSize -= Reduction;
+			return;
+		}
+		
+		// Reduce size by half
+		BlockSize /= 2;
+	}
+
+
+    /** Updates the area, and the iterator. */
+	inline void ReduceBlock(int32 BlockCount, int32& InOutArea, int32& InOutBlockIt, FScratchLayoutPack& scratch, EReductionMethod ReductionMethod)
+	{
+		int32 r_it = InOutBlockIt;
+		bool pass = false;
+
+		int32 oldBlockArea = scratch.sorted[r_it].size[0] * scratch.sorted[r_it].size[1];
+
+		if (scratch.sorted[r_it].size[0] != 1 || scratch.sorted[r_it].size[1] != 1)
+		{
+			if (scratch.sorted[r_it].bReduceBothAxes)
+			{
+				// We reduce both sides of the block at the same time
+				for (int32 Index = 0; Index <= 1; ++Index)
+				{
+					if (scratch.sorted[r_it].size[Index] > 1)
+					{
+						ReductionOperation(scratch.sorted[r_it].size[Index], ReductionMethod, scratch.sorted[r_it].bReduceByTwo);
+						ReductionOperation(scratch.blocks[scratch.sorted[r_it].Index][Index], ReductionMethod, scratch.sorted[r_it].bReduceByTwo );
+
+						pass = true;
+					}
+				}
+			}
+			else
+			{
+				if (scratch.reductions[scratch.sorted[r_it].Index][0] > scratch.reductions[scratch.sorted[r_it].Index][1])
+				{
+					if (scratch.sorted[r_it].size[1] > 1)
+					{
+						ReductionOperation(scratch.sorted[r_it].size[1], ReductionMethod, scratch.sorted[r_it].bReduceByTwo);
+						ReductionOperation(scratch.blocks[scratch.sorted[r_it].Index][1], ReductionMethod, scratch.sorted[r_it].bReduceByTwo);
+						pass = true;
+					}
+
+					scratch.reductions[scratch.sorted[r_it].Index][1] += 1;
+				}
+				else if (scratch.reductions[scratch.sorted[r_it].Index][0] < scratch.reductions[scratch.sorted[r_it].Index][1])
+				{
+					if (scratch.sorted[r_it].size[0] > 1)
+					{
+						ReductionOperation(scratch.sorted[r_it].size[0], ReductionMethod, scratch.sorted[r_it].bReduceByTwo);
+						ReductionOperation(scratch.blocks[scratch.sorted[r_it].Index][0], ReductionMethod, scratch.sorted[r_it].bReduceByTwo);
+						pass = true;
+					}
+
+					scratch.reductions[scratch.sorted[r_it].Index][0] += 1;
+				}
+				else
+				{
+					// we select the first side to reduce "randomly"
+					int32 Index = r_it % 2;
+
+					// if we can't reduce a dimension then we try to reduce the other one
+					if (scratch.sorted[r_it].size[Index] <= 1)
+					{
+						Index = r_it == 0 ? 1 : 0;
+					}
+
+					if (scratch.sorted[r_it].size[Index] > 1)
+					{
+						ReductionOperation(scratch.sorted[r_it].size[Index], ReductionMethod, scratch.sorted[r_it].bReduceByTwo);
+						ReductionOperation(scratch.blocks[scratch.sorted[r_it].Index][Index], ReductionMethod, scratch.sorted[r_it].bReduceByTwo);
+						pass = true;
+					}
+
+					scratch.reductions[scratch.sorted[r_it].Index][Index] += 1;
+				}
+			}
+		}
+		else
+		{
+			pass = true;
+		}
+		
+		int32 newBlockArea = scratch.sorted[r_it].size[0] * scratch.sorted[r_it].size[1];
+
+		InOutArea = InOutArea - (oldBlockArea - newBlockArea);
+
+		if (pass)
+		{
+			InOutBlockIt = InOutBlockIt + 1;
+
+			if (InOutBlockIt >= BlockCount)
+			{
+				InOutBlockIt = 0;
+			}
+		}
+	}
+
+
+	inline bool SetPositions(int32 bestY,int32 layoutSizeY, uint16* maxX, uint16* maxY, FScratchLayoutPack& scratch, EPackStrategy packStrategy)
+	{
+		bool fits = true;
+
+		// Number of blocks alrady packed
+		size_t packed = 0;
+		TArray<uint8_t> packedFlag;
+		packedFlag.SetNumZeroed(scratch.sorted.Num());
+
+		// Pack with fixed horizontal size
+		check(*maxX < 256);
+		if (*maxX > 256) return true;
+
+		int16_t horizon[256];
+		FMemory::Memzero(horizon, 256 * sizeof(int16_t));
+		*maxY = 0;
+
+		int32 iterations = 0;
+
+		while (packed < scratch.sorted.Num() || iterations > 5000)
+		{
+			++iterations;
+
+			int32 best = -1;
+			int32 bestX = -1;
+			int32 bestLevel = -1;
+			int32 bestWithHole = -1;
+			int32 bestWithHoleX = -1;
+			int32 bestWithHoleLevel = -1;
+			for (size_t candidate = 0; candidate < scratch.sorted.Num(); ++candidate)
+			{
+				// Skip it if we packed it already
+				if (packedFlag[candidate]) continue;
+
+				auto candidateSizeX = scratch.sorted[candidate].size[0];
+				auto candidateSizeY = scratch.sorted[candidate].size[1];
+
+				// Seek for the lowest span where the block fits
+				int32 currentLevel = TNumericLimits<int32>::Max();
+				int32 currentX = 0;
+				int32 currentLevelWithoutHole = TNumericLimits<int32>::Max();
+				int32 currentXWithoutHole = 0;
+				for (int32 x = 0; x <= *maxX - candidateSizeX; ++x)
+				{
+					int32 level = 0;
+					for (int32 xs = x; xs < x + candidateSizeX; ++xs)
+					{
+						level = FMath::Max(level, (int32)horizon[xs]);
+					}
+
+					if (level < currentLevel)
+					{
+						currentLevel = level;
+						currentX = x;
+					}
+
+					// Does it make an unfillable hole with the top or side?
+					uint16 minX = TNumericLimits<uint16>::Max();
+					uint16 minY = TNumericLimits<uint16>::Max();
+					for (size_t b = 0; b < scratch.sorted.Num(); ++b)
+					{
+						if (!packedFlag[b] && b != candidate)
+						{
+							minX = FMath::Min(minX, scratch.sorted[b].size[0]);
+							minY = FMath::Min(minY, scratch.sorted[b].size[1]);
+						}
+					}
+
+					bool hole =
+						// Vertical unfillable gap
+						(
+						(minY != TNumericLimits<int32>::Max())
+							&&
+							(currentLevel + candidateSizeY) < bestY
+							&&
+							(currentLevel + minY + candidateSizeY) > bestY
+							)
+						||
+						// Horizontal unfillable gap
+						(
+						(minX != TNumericLimits<int32>::Max())
+							&&
+							(currentX + candidateSizeX) < *maxX
+							&&
+							(currentX + minX + candidateSizeX) > *maxX
+							);
+
+					// Does it make a hole with the horizon?
+					for (int32 xs = x; !hole && xs < x + scratch.sorted[candidate].size[0]; ++xs)
+					{
+						hole = (level > (int32)horizon[xs]);
+					}
+
+					if (!hole && level < currentLevelWithoutHole)
+					{
+						currentLevelWithoutHole = level;
+						currentXWithoutHole = x;
+					}
+				}
+
+				if (currentLevelWithoutHole <= currentLevel)
+				{
+					best = int32(candidate);
+					bestX = currentXWithoutHole;
+					bestLevel = currentLevelWithoutHole;
+					break;
+				}
+
+				if (bestWithHole < 0)
+				{
+					bestWithHole = int32(candidate);
+					bestWithHoleX = currentX;
+					bestWithHoleLevel = currentLevel;
+				}
+			}
+
+			check(best >= 0 || bestWithHole >= 0);
+
+			// If there is no other option, accept leaving a hole.
+			if (best < 0)
+			{
+				best = bestWithHole;
+				bestX = bestWithHoleX;
+				bestLevel = bestWithHoleLevel;
+			}
+
+			if (bestX >= 0)
+			{
+				// Update horizon
+				for (int32 xs = bestX; xs < bestX + scratch.sorted[best].size[0]; ++xs)
+				{
+					horizon[xs] = (uint16)(bestLevel + scratch.sorted[best].size[1]);
+				}
+			}
+
+			// Store
+			scratch.positions[scratch.sorted[best].Index] = FIntVector2(bestX, bestLevel);
+			*maxY = FMath::Max(*maxY, uint16(bestLevel + scratch.sorted[best].size[1]) );
+
+			if (packStrategy == EPackStrategy::FIXED_LAYOUT && *maxY > layoutSizeY)
+			{
+				fits = false;
+				break;
+			}
+
+			packedFlag[best] = 1;
+			packed++;
+
+			//                for (int32 y=0; y<*maxY; ++y)
+			//                {
+			//                    string line;
+			//                    for (int32 x=0;x<*maxX;++x)
+			//                    {
+			//                        line += DebugGetBlockAt( *scratch, packedFlag, x, y );
+			//                    }
+			//                    AXE_LOG("layout",Warning,line.c_str());
+			//                }
+			//                AXE_LOG("layout",Warning,"--------------------------------------------------");
+		}
+
+		if (packed < scratch.sorted.Num())
+		{
+			fits = false;
+		}
+
+		*maxY = FGenericPlatformMath::RoundUpToPowerOfTwo(*maxY);
+
+		return fits;
+	}
+
+
+    void LayoutPack3( Layout* pResult, const Layout* pSourceLayout )
+    {
+		MUTABLE_CPUPROFILER_SCOPE(MeshLayoutPack3);
+
+        check( pResult->GetBlockCount() == pSourceLayout->GetBlockCount() );
+
+		int32 BlockCount = pSourceLayout->GetBlockCount();
+
+		FScratchLayoutPack scratch;
+		scratch.blocks.SetNum(BlockCount);
+		scratch.sorted.SetNum(BlockCount);
+		scratch.positions.SetNum(BlockCount);
+		scratch.priorities.SetNum(BlockCount);
+		scratch.reductions.SetNum(BlockCount);
+		scratch.ReduceBothAxes.SetNum(BlockCount);
+		scratch.ReduceByTwo.SetNum(BlockCount);
+
+
+        check( scratch.blocks.Num()==BlockCount );
+
+		//Getting maximum layout grid size:
+		int32 layoutSizeX, layoutSizeY;
+		pSourceLayout->GetMaxGridSize(&layoutSizeX, &layoutSizeY);
+
+		bool usePriority = false;
+
+		EPackStrategy LayoutStrategy = pSourceLayout->GetLayoutPackingStrategy();
+		EReductionMethod ReductionMethod = pSourceLayout->ReductionMethod;
+
+        // Look for the maximum block sizes on the layout and the total area
+		uint16 maxX = 0;
+		uint16 maxY = 0;
+        int32 area = 0;
+
+        for ( int32 Index=0; Index<BlockCount; ++Index )
+        {
+            box< UE::Math::TIntVector2<uint16> > b;
+			b.min = pSourceLayout->Blocks[Index].Min;
+			b.size = pSourceLayout->Blocks[Index].Size;
+
+			int32 p = pSourceLayout->Blocks[Index].Priority;
+			bool bReduceBothAxes = pSourceLayout->Blocks[Index].bReduceBothAxes;
+			bool bReduceByTwo = pSourceLayout->Blocks[Index].bReduceByTwo;
+
+			FIntVector2 reductions;
+			reductions[0] = 0;
+			reductions[1] = 0;
+
+			if (p > 0)
+			{
+				usePriority = true;
+			}
+
+            maxX = FMath::Max( maxX, b.size[0] );
+            maxY = FMath::Max( maxY, b.size[1] );
+
+            area += b.size[0] * b.size[1];
+
+            scratch.blocks[Index] = b.size;
+			scratch.priorities[Index] = p;
+			scratch.reductions[Index] = reductions;
+			scratch.ReduceBothAxes[Index] = (int32)bReduceBothAxes;
+			scratch.ReduceByTwo[Index] = (int32)bReduceByTwo;
+        }
+
+
+        // Grow until the area is big enough to fit all blocks. We always grow X first, because
+        // in case we cannot pack everything, we will grow Y with the current horizon algorithm.
+		if (LayoutStrategy == EPackStrategy::RESIZABLE_LAYOUT)
+		{
+			maxX = FGenericPlatformMath::RoundUpToPowerOfTwo(maxX);
+			maxY = FGenericPlatformMath::RoundUpToPowerOfTwo(maxY);
+
+			while (maxX*maxY < area)
+			{
+				if (maxX > maxY)
+				{
+					maxY *= 2;
+				}
+				else
+				{
+					maxX *= 2;
+				}
+			}
+		}
+		else
+		{
+			//Increase the maximum layout size if the grid area is smaller than the number of blocks
+			while (BlockCount > layoutSizeX*layoutSizeY)
+			{
+				if (layoutSizeX > layoutSizeY)
+				{
+					layoutSizeY *= 2;
+
+					if (layoutSizeY == 0)
+					{
+						layoutSizeY = 1;
+					}
+				}
+				else
+				{
+					layoutSizeX *= 2;
+
+					if (layoutSizeX == 0)
+					{
+						layoutSizeX = 1;
+					}
+				}
+			}
+
+			// Reducing blocks that do not fit in the layout grid
+			while (maxX > layoutSizeX)
+			{
+				area = 0;
+
+				for (int32 Index = 0; Index < BlockCount; ++Index)
+				{
+					if (scratch.blocks[Index][0] == maxX)
+					{
+						ReductionOperation(scratch.blocks[Index][0], ReductionMethod, scratch.ReduceByTwo[Index]);
+						scratch.reductions[Index][0]++;
+
+						if (scratch.ReduceBothAxes[Index] == 1)
+						{
+							ReductionOperation(scratch.blocks[Index][1], ReductionMethod, scratch.ReduceByTwo[Index]);
+							scratch.reductions[Index][1]++;
+						}
+					}
+				}
+
+				maxX = maxY = 0;
+				
+				//recalculating area and maximum block sizes
+				for (int32 Index = 0; Index < BlockCount; ++Index)
+				{
+					maxX = FMath::Max(maxX, scratch.blocks[Index][0]);
+
+					// maxY could have been modified if a block had symmetry enabled in the previous reduction
+					maxY = FMath::Max(maxY, scratch.blocks[Index][1]);
+
+					area += scratch.blocks[Index][0] * scratch.blocks[Index][1];
+				}
+			}
+
+			while (maxY > layoutSizeY)
+			{
+				area = 0;
+
+				for (int32 Index = 0; Index < BlockCount; ++Index)
+				{
+					if (scratch.blocks[Index][1] == maxY)
+					{
+						ReductionOperation(scratch.blocks[Index][1], ReductionMethod, scratch.ReduceByTwo[Index]);
+						scratch.reductions[Index][1]++;
+
+						if (scratch.ReduceBothAxes[Index] == 1)
+						{
+							ReductionOperation(scratch.blocks[Index][0], ReductionMethod, scratch.ReduceByTwo[Index]);
+							scratch.reductions[Index][0]++;
+						}
+					}
+				}
+
+				maxX = maxY = 0;
+
+				//recalculating area and maximum block sizes
+				for (int32 Index = 0; Index < BlockCount; ++Index)
+				{
+					maxY = FMath::Max(maxY, scratch.blocks[Index][1]);
+
+					// maxX could have been modified if a block had symmetry enabled in the previous reduction
+					maxX = FMath::Max(maxX, scratch.blocks[Index][0]);
+
+					area += scratch.blocks[Index][0] * scratch.blocks[Index][1];
+				}
+			}
+
+			maxX = FGenericPlatformMath::RoundUpToPowerOfTwo(maxX);
+			maxY = FGenericPlatformMath::RoundUpToPowerOfTwo(maxY);
+
+			// Grow until the area is big enough to fit all blocks or the size is equal to the max layout size.
+			while (maxX*maxY < area && (maxX < layoutSizeX || maxY < layoutSizeY))
+			{
+				if (maxX > maxY)
+				{
+					maxY *= 2;
+				}
+				else
+				{
+					maxX *= 2;
+				}
+			}
+		}
+        
+        int32 bestY = maxY;
+
+		// This is used to iterate through blocks.
+		int32 BlockIterator = 0;
+
+		// Making a copy of the blocks array to sort them
+        check( (int32)scratch.sorted.Num()==BlockCount );
+        for ( int32 Index=0; Index<BlockCount; ++Index )
+        {
+            scratch.sorted[Index] = FPackLayoutBlock( Index, scratch.blocks[Index], scratch.priorities[Index], (bool)scratch.ReduceBothAxes[Index], (bool)scratch.ReduceByTwo[Index]);
+        }
+
+		// Sort blocks by height, area
+		scratch.sorted.Sort(CompareBlocks);
+
+		if(LayoutStrategy == EPackStrategy::FIXED_LAYOUT)
+		{
+			if (usePriority)
+			{
+				//Sort blocks by priority
+				scratch.sorted.Sort(CompareBlocksPriority);
+			}
+
+			// Shrink blocks in case we do not have enough space to pack everything
+			while (maxX*maxY < area)
+			{
+				ReduceBlock(BlockCount, area, BlockIterator, scratch, ReductionMethod);
+			}
+		}
+
+		bool fits = false;
+
+		while (!fits)
+		{
+			// Sort by height&area before packing
+			if (usePriority)
+			{
+				scratch.sorted.Sort(CompareBlocks);
+			}
+
+			// Try to pack everything
+			fits = SetPositions(bestY, layoutSizeY, &maxX, &maxY, scratch, LayoutStrategy);
+
+			if (!fits && LayoutStrategy == EPackStrategy::FIXED_LAYOUT)
+			{
+				// Sort by priority before shrink
+				if (usePriority)
+				{
+					scratch.sorted.Sort(CompareBlocksPriority);
+				}
+
+				ReduceBlock(BlockCount, area, BlockIterator, scratch, ReductionMethod);
+			}
+		}
+
+        // Set data in the result
+        pResult->SetGridSize( maxX, maxY );
+        pResult->SetMaxGridSize(layoutSizeX, layoutSizeY);
+		pResult->SetLayoutPackingStrategy(LayoutStrategy);
+		pResult->ReductionMethod = ReductionMethod;
+
+        for ( int32 Index=0; Index<BlockCount; ++Index )
+        {
+			pResult->Blocks[Index].Min = UE::Math::TIntVector2<uint16>(scratch.positions[Index]);
+			pResult->Blocks[Index].Size = scratch.blocks[Index];
+			pResult->Blocks[Index].Priority = scratch.priorities[Index];
+			pResult->Blocks[Index].bReduceBothAxes = scratch.ReduceBothAxes[Index];
+			pResult->Blocks[Index].bReduceByTwo = scratch.ReduceByTwo[Index];
+        }
+    }
+}
diff --git a/Engine/Plugins/Experimental/Mutable/Source/MutableRuntime/Private/MuR/OpLayoutPack.h b/Engine/Plugins/Experimental/Mutable/Source/MutableRuntime/Private/MuR/OpLayoutPack.h
--- a/Engine/Plugins/Experimental/Mutable/Source/MutableRuntime/Private/MuR/OpLayoutPack.h
+++ b/Engine/Plugins/Experimental/Mutable/Source/MutableRuntime/Private/MuR/OpLayoutPack.h
@@ -2,649 +2,11 @@
 
 #pragma once
 
-#include "MuR/Layout.h"
-#include "MuR/MutableMath.h"
-#include "MuR/MutableTrace.h"
-#include "MuR/Platform.h"
-#include "GenericPlatform/GenericPlatformMath.h"
-#include <string>
 
 namespace mu
 {
+	class Layout;
 
-	//---------------------------------------------------------------------------------------------
-	struct FLayoutBlock
-	{
-		FLayoutBlock()
-		{
-			index = -1;
-		}
+	extern void LayoutPack3(Layout* pResult, const Layout* pSourceLayout);
 
-		FLayoutBlock( int32 i, UE::Math::TIntVector2<uint16> s, int32 p=0, bool bInReduceBothAxes = false, bool bInReduceByTwo = false )
-		{
-			index = i;
-			size = s;
-			priority = p;
-			bReduceBothAxes = bInReduceBothAxes;
-			bReduceByTwo = bInReduceByTwo;
-		}
-
-		int32 index;
-		UE::Math::TIntVector2<uint16> size = UE::Math::TIntVector2<uint16>(0,0);
-		int32 priority;
-		bool bReduceBothAxes;
-		bool bReduceByTwo;
-	};
-
-    inline bool CompareBlocks( const FLayoutBlock& a, const FLayoutBlock& b )
-    {
-        if ( a.size[1]>b.size[1] )
-        {
-            return true;
-        }
-        else if ( a.size[1]<b.size[1] )
-        {
-            return false;
-        }
-        else
-        {
-            int areaA = a.size[0]*a.size[1];
-            int areaB = b.size[0]*b.size[1];
-            if (areaA > areaB)
-            {
-                return true;
-            }
-            else if (areaA < areaB)
-            {
-                return false;
-            }
-            else
-            {
-                // This has to be deterministic, and ids are supposed to be unique
-                return a.index<b.index;
-            }
-        }
-    }
-
-	inline bool CompareBlocksPriority(const FLayoutBlock& a, const FLayoutBlock& b)
-	{
-		if (a.priority == b.priority)
-		{
-			// TODO(Max): Check if this comparison modifies the block order while reducing them.
-			return CompareBlocks(a, b);
-		}
-
-		return a.priority > b.priority;
-	}
-
-
-    //---------------------------------------------------------------------------------------------
-    //! More expensive but precise version
-    //---------------------------------------------------------------------------------------------
-    struct SCRATCH_LAYOUT_PACK
-    {
-        TArray< UE::Math::TIntVector2<uint16> > blocks;
-		TArray< FLayoutBlock > sorted;
-		TArray< FIntVector2 > positions;
-		TArray< int32 > priorities;
-		TArray< FIntVector2 > reductions;
-		TArray< int32 > ReduceBothAxes;
-		TArray< int32 > ReduceByTwo;
-    };
-
-
-    inline char DebugGetBlockAt( const SCRATCH_LAYOUT_PACK& scratch,
-                                   const TArray<uint8_t>& packedFlag,
-                                   int x, int y )
-    {
-        for ( size_t b=0; b<packedFlag.Num(); ++b )
-        {
-            if (packedFlag[b])
-            {
-				int32 i = scratch.sorted[b].index;
-                if ( x>=scratch.positions[i][0] &&
-                     x<scratch.positions[i][0]+scratch.sorted[b].size[0] &&
-                     y>=scratch.positions[i][1] &&
-                     y<scratch.positions[i][1]+scratch.sorted[b].size[1]
-                     )
-                {
-					return char('a')+char(b);
-                }
-            }
-        }
-        return '.';
-    }
-
-
-	inline void ReductionOperation(uint16& BlockSize, EReductionMethod ReductionMethod, int32 bReduceByTwo)
-	{
-		if (ReductionMethod == EReductionMethod::UNITARY_REDUCTION)
-		{
-			int32 Reduction = (bReduceByTwo && BlockSize > 2) ? 2 : 1;
-			BlockSize -= Reduction;
-			return;
-		}
-		
-		// Reduce size by half
-		BlockSize /= 2;
-	}
-
-
-    /** Updates the area, and the iterator. */
-	inline void ReduceBlock(int blockCount, int32& InOutArea, int32& InOutBlockIt, SCRATCH_LAYOUT_PACK* scratch, EReductionMethod ReductionMethod)
-	{
-		int r_it = InOutBlockIt;
-		bool pass = false;
-
-		int oldBlockArea = scratch->sorted[r_it].size[0] * scratch->sorted[r_it].size[1];
-
-		if (scratch->sorted[r_it].size[0] != 1 || scratch->sorted[r_it].size[1] != 1)
-		{
-			if (scratch->sorted[r_it].bReduceBothAxes)
-			{
-				// We reduce both sides of the block at the same time
-				for (int32 Index = 0; Index <= 1; ++Index)
-				{
-					if (scratch->sorted[r_it].size[Index] > 1)
-					{
-						ReductionOperation(scratch->sorted[r_it].size[Index], ReductionMethod, scratch->sorted[r_it].bReduceByTwo);
-						ReductionOperation(scratch->blocks[scratch->sorted[r_it].index][Index], ReductionMethod, scratch->sorted[r_it].bReduceByTwo );
-
-						pass = true;
-					}
-				}
-			}
-			else
-			{
-				if (scratch->reductions[scratch->sorted[r_it].index][0] > scratch->reductions[scratch->sorted[r_it].index][1])
-				{
-					if (scratch->sorted[r_it].size[1] > 1)
-					{
-						ReductionOperation(scratch->sorted[r_it].size[1], ReductionMethod, scratch->sorted[r_it].bReduceByTwo);
-						ReductionOperation(scratch->blocks[scratch->sorted[r_it].index][1], ReductionMethod, scratch->sorted[r_it].bReduceByTwo);
-						pass = true;
-					}
-
-					scratch->reductions[scratch->sorted[r_it].index][1] += 1;
-				}
-				else if (scratch->reductions[scratch->sorted[r_it].index][0] < scratch->reductions[scratch->sorted[r_it].index][1])
-				{
-					if (scratch->sorted[r_it].size[0] > 1)
-					{
-						ReductionOperation(scratch->sorted[r_it].size[0], ReductionMethod, scratch->sorted[r_it].bReduceByTwo);
-						ReductionOperation(scratch->blocks[scratch->sorted[r_it].index][0], ReductionMethod, scratch->sorted[r_it].bReduceByTwo);
-						pass = true;
-					}
-
-					scratch->reductions[scratch->sorted[r_it].index][0] += 1;
-				}
-				else
-				{
-					// we select the first side to reduce "randomly"
-					int32 Index = r_it % 2;
-
-					// if we can't reduce a dimension then we try to reduce the other one
-					if (scratch->sorted[r_it].size[Index] <= 1)
-					{
-						Index = r_it == 0 ? 1 : 0;
-					}
-
-					if (scratch->sorted[r_it].size[Index] > 1)
-					{
-						ReductionOperation(scratch->sorted[r_it].size[Index], ReductionMethod, scratch->sorted[r_it].bReduceByTwo);
-						ReductionOperation(scratch->blocks[scratch->sorted[r_it].index][Index], ReductionMethod, scratch->sorted[r_it].bReduceByTwo);
-						pass = true;
-					}
-
-					scratch->reductions[scratch->sorted[r_it].index][Index] += 1;
-				}
-			}
-		}
-		else
-		{
-			pass = true;
-		}
-		
-		int newBlockArea = scratch->sorted[r_it].size[0] * scratch->sorted[r_it].size[1];
-
-		InOutArea = InOutArea - (oldBlockArea - newBlockArea);
-
-		if (pass)
-		{
-			InOutBlockIt = InOutBlockIt + 1;
-
-			if (InOutBlockIt >= blockCount)
-			{
-				InOutBlockIt = 0;
-			}
-		}
-	}
-
-
-	inline bool SetPositions(int bestY,int layoutSizeY, uint16* maxX, uint16* maxY, SCRATCH_LAYOUT_PACK* scratch, EPackStrategy packStrategy)
-	{
-		bool fits = true;
-
-		// Number of blocks alrady packed
-		size_t packed = 0;
-		TArray<uint8_t> packedFlag;
-		packedFlag.SetNumZeroed(scratch->sorted.Num());
-
-		// Pack with fixed horizontal size
-		check(*maxX < 256);
-		if (*maxX > 256) return true;
-
-		int16_t horizon[256];
-		FMemory::Memzero(horizon, 256 * sizeof(int16_t));
-		*maxY = 0;
-
-		int iterations = 0;
-
-		while (packed < scratch->sorted.Num() || iterations > 5000)
-		{
-			++iterations;
-
-			int best = -1;
-			int bestX = -1;
-			int bestLevel = -1;
-			int bestWithHole = -1;
-			int bestWithHoleX = -1;
-			int bestWithHoleLevel = -1;
-			for (size_t candidate = 0; candidate < scratch->sorted.Num(); ++candidate)
-			{
-				// Skip it if we packed it already
-				if (packedFlag[candidate]) continue;
-
-				auto candidateSizeX = scratch->sorted[candidate].size[0];
-				auto candidateSizeY = scratch->sorted[candidate].size[1];
-
-				// Seek for the lowest span where the block fits
-				int currentLevel = TNumericLimits<int>::Max();
-				int currentX = 0;
-				int currentLevelWithoutHole = TNumericLimits<int>::Max();
-				int currentXWithoutHole = 0;
-				for (int x = 0; x <= *maxX - candidateSizeX; ++x)
-				{
-					int level = 0;
-					for (int xs = x; xs < x + candidateSizeX; ++xs)
-					{
-						level = FMath::Max(level, (int)horizon[xs]);
-					}
-
-					if (level < currentLevel)
-					{
-						currentLevel = level;
-						currentX = x;
-					}
-
-					// Does it make an unfillable hole with the top or side?
-					uint16 minX = TNumericLimits<uint16>::Max();
-					uint16 minY = TNumericLimits<uint16>::Max();
-					for (size_t b = 0; b < scratch->sorted.Num(); ++b)
-					{
-						if (!packedFlag[b] && b != candidate)
-						{
-							minX = FMath::Min(minX, scratch->sorted[b].size[0]);
-							minY = FMath::Min(minY, scratch->sorted[b].size[1]);
-						}
-					}
-
-					bool hole =
-						// Vertical unfillable gap
-						(
-						(minY != TNumericLimits<int>::Max())
-							&&
-							(currentLevel + candidateSizeY) < bestY
-							&&
-							(currentLevel + minY + candidateSizeY) > bestY
-							)
-						||
-						// Horizontal unfillable gap
-						(
-						(minX != TNumericLimits<int>::Max())
-							&&
-							(currentX + candidateSizeX) < *maxX
-							&&
-							(currentX + minX + candidateSizeX) > *maxX
-							);
-
-					// Does it make a hole with the horizon?
-					for (int xs = x; !hole && xs < x + scratch->sorted[candidate].size[0]; ++xs)
-					{
-						hole = (level > (int)horizon[xs]);
-					}
-
-					if (!hole && level < currentLevelWithoutHole)
-					{
-						currentLevelWithoutHole = level;
-						currentXWithoutHole = x;
-					}
-				}
-
-				if (currentLevelWithoutHole <= currentLevel)
-				{
-					best = int(candidate);
-					bestX = currentXWithoutHole;
-					bestLevel = currentLevelWithoutHole;
-					break;
-				}
-
-				if (bestWithHole < 0)
-				{
-					bestWithHole = int(candidate);
-					bestWithHoleX = currentX;
-					bestWithHoleLevel = currentLevel;
-				}
-			}
-
-			check(best >= 0 || bestWithHole >= 0);
-
-			// If there is no other option, accept leaving a hole.
-			if (best < 0)
-			{
-				best = bestWithHole;
-				bestX = bestWithHoleX;
-				bestLevel = bestWithHoleLevel;
-			}
-
-			if (bestX >= 0)
-			{
-				// Update horizon
-				for (int xs = bestX; xs < bestX + scratch->sorted[best].size[0]; ++xs)
-				{
-					horizon[xs] = (uint16)(bestLevel + scratch->sorted[best].size[1]);
-				}
-			}
-
-			// Store
-			scratch->positions[scratch->sorted[best].index] = FIntVector2(bestX, bestLevel);
-			*maxY = FMath::Max(*maxY, uint16(bestLevel + scratch->sorted[best].size[1]) );
-
-			if (packStrategy == EPackStrategy::FIXED_LAYOUT && *maxY > layoutSizeY)
-			{
-				fits = false;
-				break;
-			}
-
-			packedFlag[best] = 1;
-			packed++;
-
-			//                for (int y=0; y<*maxY; ++y)
-			//                {
-			//                    string line;
-			//                    for (int x=0;x<*maxX;++x)
-			//                    {
-			//                        line += DebugGetBlockAt( *scratch, packedFlag, x, y );
-			//                    }
-			//                    AXE_LOG("layout",Warning,line.c_str());
-			//                }
-			//                AXE_LOG("layout",Warning,"--------------------------------------------------");
-		}
-
-		if (packed < scratch->sorted.Num())
-		{
-			fits = false;
-		}
-
-		*maxY = FGenericPlatformMath::RoundUpToPowerOfTwo(*maxY);
-
-		return fits;
-	}
-
-
-    inline void LayoutPack3( Layout* pResult,
-                             const Layout* pSourceLayout,
-                             SCRATCH_LAYOUT_PACK* scratch  )
-    {
-		MUTABLE_CPUPROFILER_SCOPE(MeshLayoutPack3);
-
-        check( pResult->GetBlockCount() == pSourceLayout->GetBlockCount() );
-
-        int blockCount = pSourceLayout->GetBlockCount();
-        check( (int)scratch->blocks.Num()==blockCount );
-
-		//Getting maximum layout grid size:
-		int layoutSizeX, layoutSizeY;
-		pSourceLayout->GetMaxGridSize(&layoutSizeX, &layoutSizeY);
-
-		bool usePriority = false;
-
-		EPackStrategy LayoutStrategy = pSourceLayout->GetLayoutPackingStrategy();
-		EReductionMethod ReductionMethod = pSourceLayout->ReductionMethod;
-
-        // Look for the maximum block sizes on the layout and the total area
-		uint16 maxX = 0;
-		uint16 maxY = 0;
-        int area = 0;
-
-        for ( int index=0; index<blockCount; ++index )
-        {
-            box< UE::Math::TIntVector2<uint16> > b;
-			b.min = pSourceLayout->Blocks[index].Min;
-			b.size = pSourceLayout->Blocks[index].Size;
-
-			int32 p = pSourceLayout->Blocks[index].Priority;
-			bool bReduceBothAxes = pSourceLayout->Blocks[index].bReduceBothAxes;
-			bool bReduceByTwo = pSourceLayout->Blocks[index].bReduceByTwo;
-
-			FIntVector2 reductions;
-			reductions[0] = 0;
-			reductions[1] = 0;
-
-			if (p > 0)
-			{
-				usePriority = true;
-			}
-
-            maxX = FMath::Max( maxX, b.size[0] );
-            maxY = FMath::Max( maxY, b.size[1] );
-
-            area += b.size[0] * b.size[1];
-
-            scratch->blocks[index] = b.size;
-			scratch->priorities[index] = p;
-			scratch->reductions[index] = reductions;
-			scratch->ReduceBothAxes[index] = (int)bReduceBothAxes;
-			scratch->ReduceByTwo[index] = (int)bReduceByTwo;
-        }
-
-
-        // Grow until the area is big enough to fit all blocks. We always grow X first, because
-        // in case we cannot pack everything, we will grow Y with the current horizon algorithm.
-		if (LayoutStrategy == EPackStrategy::RESIZABLE_LAYOUT)
-		{
-			maxX = FGenericPlatformMath::RoundUpToPowerOfTwo(maxX);
-			maxY = FGenericPlatformMath::RoundUpToPowerOfTwo(maxY);
-
-			while (maxX*maxY < area)
-			{
-				if (maxX > maxY)
-				{
-					maxY *= 2;
-				}
-				else
-				{
-					maxX *= 2;
-				}
-			}
-		}
-		else
-		{
-			//Increase the maximum layout size if the grid area is smaller than the number of blocks
-			while (blockCount > layoutSizeX*layoutSizeY)
-			{
-				if (layoutSizeX > layoutSizeY)
-				{
-					layoutSizeY *= 2;
-
-					if (layoutSizeY == 0)
-					{
-						layoutSizeY = 1;
-					}
-				}
-				else
-				{
-					layoutSizeX *= 2;
-
-					if (layoutSizeX == 0)
-					{
-						layoutSizeX = 1;
-					}
-				}
-			}
-
-			// Reducing blocks that do not fit in the layout grid
-			while (maxX > layoutSizeX)
-			{
-				area = 0;
-
-				for (int index = 0; index < blockCount; ++index)
-				{
-					if (scratch->blocks[index][0] == maxX)
-					{
-						ReductionOperation(scratch->blocks[index][0], ReductionMethod, scratch->ReduceByTwo[index]);
-						scratch->reductions[index][0]++;
-
-						if (scratch->ReduceBothAxes[index] == 1)
-						{
-							ReductionOperation(scratch->blocks[index][1], ReductionMethod, scratch->ReduceByTwo[index]);
-							scratch->reductions[index][1]++;
-						}
-					}
-				}
-
-				maxX = maxY = 0;
-				
-				//recalculating area and maximum block sizes
-				for (int index = 0; index < blockCount; ++index)
-				{
-					maxX = FMath::Max(maxX, scratch->blocks[index][0]);
-
-					// maxY could have been modified if a block had symmetry enabled in the previous reduction
-					maxY = FMath::Max(maxY, scratch->blocks[index][1]);
-
-					area += scratch->blocks[index][0] * scratch->blocks[index][1];
-				}
-			}
-
-			while (maxY > layoutSizeY)
-			{
-				area = 0;
-
-				for (int index = 0; index < blockCount; ++index)
-				{
-					if (scratch->blocks[index][1] == maxY)
-					{
-						ReductionOperation(scratch->blocks[index][1], ReductionMethod, scratch->ReduceByTwo[index]);
-						scratch->reductions[index][1]++;
-
-						if (scratch->ReduceBothAxes[index] == 1)
-						{
-							ReductionOperation(scratch->blocks[index][0], ReductionMethod, scratch->ReduceByTwo[index]);
-							scratch->reductions[index][0]++;
-						}
-					}
-				}
-
-				maxX = maxY = 0;
-
-				//recalculating area and maximum block sizes
-				for (int index = 0; index < blockCount; ++index)
-				{
-					maxY = FMath::Max(maxY, scratch->blocks[index][1]);
-
-					// maxX could have been modified if a block had symmetry enabled in the previous reduction
-					maxX = FMath::Max(maxX, scratch->blocks[index][0]);
-
-					area += scratch->blocks[index][0] * scratch->blocks[index][1];
-				}
-			}
-
-			maxX = FGenericPlatformMath::RoundUpToPowerOfTwo(maxX);
-			maxY = FGenericPlatformMath::RoundUpToPowerOfTwo(maxY);
-
-			// Grow until the area is big enough to fit all blocks or the size is equal to the max layout size.
-			while (maxX*maxY < area && (maxX < layoutSizeX || maxY < layoutSizeY))
-			{
-				if (maxX > maxY)
-				{
-					maxY *= 2;
-				}
-				else
-				{
-					maxX *= 2;
-				}
-			}
-		}
-        
-        int bestY = maxY;
-
-		// This is used to iterate through blocks.
-		int32 BlockIterator = 0;
-
-		// Making a copy of the blocks array to sort them
-        check( (int)scratch->sorted.Num()==blockCount );
-        for ( int index=0; index<blockCount; ++index )
-        {
-            scratch->sorted[index] = FLayoutBlock( index, scratch->blocks[index], scratch->priorities[index], (bool)scratch->ReduceBothAxes[index], (bool)scratch->ReduceByTwo[index]);
-        }
-
-		// Sort blocks by height, area
-		scratch->sorted.Sort(CompareBlocks);
-
-		if(LayoutStrategy == EPackStrategy::FIXED_LAYOUT)
-		{
-			if (usePriority)
-			{
-				//Sort blocks by priority
-				scratch->sorted.Sort(CompareBlocksPriority);
-			}
-
-			// Shrink blocks in case we do not have enough space to pack everything
-			while (maxX*maxY < area)
-			{
-				ReduceBlock(blockCount, area, BlockIterator, scratch, ReductionMethod);
-			}
-		}
-
-		bool fits = false;
-
-		while (!fits)
-		{
-			// Sort by height&area before packing
-			if (usePriority)
-			{
-				scratch->sorted.Sort(CompareBlocks);
-			}
-
-			// Try to pack everything
-			fits = SetPositions(bestY, layoutSizeY, &maxX, &maxY, scratch, LayoutStrategy);
-
-			if (!fits && LayoutStrategy == EPackStrategy::FIXED_LAYOUT)
-			{
-				// Sort by priority before shrink
-				if (usePriority)
-				{
-					scratch->sorted.Sort(CompareBlocksPriority);
-				}
-
-				ReduceBlock(blockCount, area, BlockIterator, scratch, ReductionMethod);
-			}
-		}
-
-        // Set data in the result
-        pResult->SetGridSize( maxX, maxY );
-        pResult->SetMaxGridSize(layoutSizeX, layoutSizeY);
-		pResult->SetLayoutPackingStrategy(LayoutStrategy);
-		pResult->ReductionMethod = ReductionMethod;
-
-        for ( int32 index=0; index<blockCount; ++index )
-        {
-			pResult->Blocks[index].Min = UE::Math::TIntVector2<uint16>(scratch->positions[index]);
-			pResult->Blocks[index].Size = scratch->blocks[index];
-			pResult->Blocks[index].Priority = scratch->priorities[index];
-			pResult->Blocks[index].bReduceBothAxes = scratch->ReduceBothAxes[index];
-			pResult->Blocks[index].bReduceByTwo = scratch->ReduceByTwo[index];
-        }
-    }
 }
diff --git a/Engine/Plugins/Experimental/Mutable/Source/MutableRuntime/Private/MuR/OpLayoutRemoveBlocks.h b/Engine/Plugins/Experimental/Mutable/Source/MutableRuntime/Private/MuR/OpLayoutRemoveBlocks.h
--- a/Engine/Plugins/Experimental/Mutable/Source/MutableRuntime/Private/MuR/OpLayoutRemoveBlocks.h
+++ b/Engine/Plugins/Experimental/Mutable/Source/MutableRuntime/Private/MuR/OpLayoutRemoveBlocks.h
@@ -145,11 +145,11 @@ namespace mu
 	{
 		Ptr<Layout> pResult = pA->Clone();
 
-		for ( const Layout::FBlock& block: pB->Blocks )
+		for ( const FLayoutBlock& Block: pB->Blocks )
 		{
-			if ( pResult->FindBlock(block.Id)<0 )
+			if ( pResult->FindBlock(Block.Id)<0 )
 			{
-				pResult->Blocks.Add(block);
+				pResult->Blocks.Add(Block);
 			}
 		}
 
diff --git a/Engine/Plugins/Experimental/Mutable/Source/MutableRuntime/Public/MuR/Layout.h b/Engine/Plugins/Experimental/Mutable/Source/MutableRuntime/Public/MuR/Layout.h
--- a/Engine/Plugins/Experimental/Mutable/Source/MutableRuntime/Public/MuR/Layout.h
+++ b/Engine/Plugins/Experimental/Mutable/Source/MutableRuntime/Public/MuR/Layout.h
@@ -13,12 +13,6 @@
 namespace mu
 {
 
-	// Forward references
-	class Layout;
-
-	typedef Ptr<Layout> LayoutPtr;
-	typedef Ptr<const Layout> LayoutPtrConst;
-
 	//! Types of layout packing strategies 
 	enum class EPackStrategy : uint32
 	{
@@ -34,6 +28,58 @@ namespace mu
 		UNITARY_REDUCTION	// Reduces 1 block the axis 
 	};
 
+	/** */
+	struct FLayoutBlock
+	{
+		static constexpr uint64 InvalidBlockId = TNumericLimits<uint64>::Max();
+
+		UE::Math::TIntVector2<uint16> Min = { 0, 0 };
+		UE::Math::TIntVector2<uint16> Size = { 0, 0 };
+
+		//! Absolute id used to control merging of various layouts
+		uint64 Id;
+
+		//! Priority value to control the shrink texture layout strategy
+		int32 Priority;
+
+		//! Value to control the method to reduce the block
+		uint32 bReduceBothAxes : 1;
+
+		//! Value to control if a block has to be reduced by two in an unitary reduction strategy
+		uint32 bReduceByTwo : 1;
+
+		/** Explicit padding to prevent uninitialized memory in this POD. */
+		uint32 UnusedPadding : 30;
+
+		/** */
+		FLayoutBlock(UE::Math::TIntVector2<uint16> InMin = {}, UE::Math::TIntVector2<uint16> InSize = {})
+		{
+			Min = InMin;
+			Size = InSize;
+			Id = InvalidBlockId;
+			Priority = 0;
+			bReduceBothAxes = false;
+			bReduceByTwo = false;
+			UnusedPadding = 0;
+		}
+
+		//!
+		inline bool operator==(const FLayoutBlock& o) const
+		{
+			return (Id == o.Id) && IsSimilar(o);
+		}
+
+		inline bool IsSimilar(const FLayoutBlock& o) const
+		{
+			// All but ids
+			return (Min == o.Min) &&
+				(Size == o.Size) &&
+				(Priority == o.Priority) &&
+				(bReduceBothAxes == o.bReduceBothAxes) &&
+				(bReduceByTwo == o.bReduceByTwo);
+		}
+	};
+
 
     //! \brief Image block layout class.
     //!
@@ -45,7 +91,21 @@ namespace mu
 	{
 	public:
 
-		static constexpr uint64 InvalidBlockId = TNumericLimits<uint64>::Max();
+		//!
+		UE::Math::TIntVector2<uint16> Size = UE::Math::TIntVector2<uint16>(0, 0);
+
+		/** Maximum size in layout blocks that this layout can grow to. From there on, blocks will shrink to fit. 
+		* If 0,0 then no maximum size applies.
+		*/
+		UE::Math::TIntVector2<uint16> MaxSize = UE::Math::TIntVector2<uint16>(0, 0);
+
+		//!
+		TArray<FLayoutBlock> Blocks;
+
+		//! Packing strategy
+		EPackStrategy Strategy = EPackStrategy::RESIZABLE_LAYOUT;
+
+		EReductionMethod ReductionMethod = EReductionMethod::HALVE_REDUCTION;
 
 	public:
 
@@ -113,74 +173,6 @@ namespace mu
 
 	public:
 
-		struct FBlock
-		{
-			FBlock(UE::Math::TIntVector2<uint16> InMin = UE::Math::TIntVector2<uint16>(), UE::Math::TIntVector2<uint16> InSize = UE::Math::TIntVector2<uint16>())
-			{
-				Min = InMin;
-				Size = InSize;
-				Id = Layout::InvalidBlockId;
-				Priority = 0;
-				bReduceBothAxes = false;
-				bReduceByTwo = false;
-				UnusedPadding = 0;
-			}
-
-			UE::Math::TIntVector2<uint16> Min = UE::Math::TIntVector2<uint16>(0,0);
-			UE::Math::TIntVector2<uint16> Size = UE::Math::TIntVector2<uint16>(0, 0);
-
-			//! Absolute id used to control merging of various layouts
-			uint64 Id;
-
-			//! Priority value to control the shrink texture layout strategy
-			int32 Priority;
-
-			//! Value to control the method to reduce the block
-			uint32 bReduceBothAxes : 1;
-
-			//! Value to control if a block has to be reduced by two in an unitary reduction strategy
-			uint32 bReduceByTwo : 1;
-
-			/** Explicit padding to prevent uninitialized memory in this POD. */
-			uint32 UnusedPadding : 30;
-
-
-			//!
-			inline bool operator==(const FBlock& o) const
-			{
-				return (Min == o.Min) &&
-					(Size == o.Size) &&
-					(Id == o.Id) &&
-					(Priority == o.Priority) &&
-					(bReduceBothAxes == o.bReduceBothAxes) &&
-					(bReduceByTwo == o.bReduceByTwo);
-			}
-
-			inline bool IsSimilar(const FBlock& o) const
-			{
-				// All but ids
-				return (Min == o.Min) &&
-					(Size == o.Size) &&
-					(Priority == o.Priority) &&
-					(bReduceBothAxes == o.bReduceBothAxes) &&
-					(bReduceByTwo == o.bReduceByTwo);
-			}
-		};
-
-
-		//!
-		UE::Math::TIntVector2<uint16> Size = UE::Math::TIntVector2<uint16>(0, 0);
-
-		UE::Math::TIntVector2<uint16> MaxSize = UE::Math::TIntVector2<uint16>(0, 0);
-
-		//!
-		TArray<FBlock> Blocks;
-
-		//! Packing strategy
-		EPackStrategy Strategy = EPackStrategy::RESIZABLE_LAYOUT;
-		 
-		EReductionMethod ReductionMethod = EReductionMethod::HALVE_REDUCTION;
-
 
 		//!
 		void Serialise(OutputArchive& arch) const;
@@ -198,8 +190,8 @@ namespace mu
 		bool IsSingleBlockAndFull() const;
 	};
 
-	MUTABLE_DEFINE_POD_SERIALISABLE(Layout::FBlock);
-	MUTABLE_DEFINE_POD_VECTOR_SERIALISABLE(Layout::FBlock);
+	MUTABLE_DEFINE_POD_SERIALISABLE(FLayoutBlock);
+	MUTABLE_DEFINE_POD_VECTOR_SERIALISABLE(FLayoutBlock);
 
 }
 
diff --git a/Engine/Plugins/Experimental/Mutable/Source/MutableTools/Private/MuT/ASTOpImageRasterMesh.cpp b/Engine/Plugins/Experimental/Mutable/Source/MutableTools/Private/MuT/ASTOpImageRasterMesh.cpp
--- a/Engine/Plugins/Experimental/Mutable/Source/MutableTools/Private/MuT/ASTOpImageRasterMesh.cpp
+++ b/Engine/Plugins/Experimental/Mutable/Source/MutableTools/Private/MuT/ASTOpImageRasterMesh.cpp
@@ -26,7 +26,7 @@ namespace mu
 		, mask(this)
 		, projector(this)
 	{
-		BlockId = Layout::InvalidBlockId;
+		BlockId = FLayoutBlock::InvalidBlockId;
 		LayoutIndex = -1;
 		SizeX = SizeY = 0;
 		SourceSizeX = SourceSizeY = 0;
diff --git a/Engine/Plugins/Experimental/Mutable/Source/MutableTools/Private/MuT/CodeGenerator.cpp b/Engine/Plugins/Experimental/Mutable/Source/MutableTools/Private/MuT/CodeGenerator.cpp
--- a/Engine/Plugins/Experimental/Mutable/Source/MutableTools/Private/MuT/CodeGenerator.cpp
+++ b/Engine/Plugins/Experimental/Mutable/Source/MutableTools/Private/MuT/CodeGenerator.cpp
@@ -384,11 +384,8 @@ namespace mu
 	}
 
 
-	Ptr<const Layout> CodeGenerator::AddLayout(Ptr<const Layout> SourceLayout, uint32 MeshIDPrefix)
+	Ptr<const Layout> CodeGenerator::AddLayout(Ptr<const NodeLayout> SourceLayout, uint32 MeshIDPrefix)
 	{
-		// The layout we are adding must be a source layout, without block ids yet.
-		check(SourceLayout->Blocks.IsEmpty() || SourceLayout->Blocks[0].Id == Layout::InvalidBlockId);
-
 		Ptr<const Layout>* it = GeneratedLayouts.Find(SourceLayout.get());
 
 		if (it)
@@ -396,18 +393,33 @@ namespace mu
 			return *it;
 		}
 
-		// Assign unique ids to each layout block
-		Ptr<Layout> ClonedLayout = SourceLayout->Clone();
-		for (int32 b = 0; b < ClonedLayout->Blocks.Num(); ++b)
+		Ptr<Layout> GeneratedLayout = new Layout;
+		GeneratedLayout->Size = SourceLayout->Size;
+		GeneratedLayout->MaxSize = SourceLayout->MaxSize;
+		GeneratedLayout->Strategy = SourceLayout->Strategy;
+		GeneratedLayout->ReductionMethod = SourceLayout->ReductionMethod;
+
+		const int32 BlockCount = SourceLayout->Blocks.Num();
+		GeneratedLayout->Blocks.SetNum(BlockCount);
+		for (int32 BlockIndex = 0; BlockIndex < BlockCount; ++BlockIndex)
 		{
-			uint64 Id = uint64(MeshIDPrefix) << 32 | uint64(b);
-			ClonedLayout->Blocks[b].Id = Id;
+			const FSourceLayoutBlock& From = SourceLayout->Blocks[BlockIndex];
+			FLayoutBlock& To = GeneratedLayout->Blocks[BlockIndex];
+			To.Min = From.Min;
+			To.Size = From.Size;
+			To.Priority = From.Priority;
+			To.bReduceBothAxes = From.bReduceBothAxes;
+			To.bReduceByTwo = From.bReduceByTwo;
+
+			// Assign unique ids to each layout block
+			uint64 Id = uint64(MeshIDPrefix) << 32 | uint64(BlockIndex);
+			To.Id = Id;
 		}
-		check(SourceLayout->Blocks.Num() == ClonedLayout->Blocks.Num());
-		check(ClonedLayout->Blocks.IsEmpty() || ClonedLayout->Blocks[0].Id != Layout::InvalidBlockId);
-		GeneratedLayouts.Add(SourceLayout.get(), ClonedLayout);
 
-		return ClonedLayout;
+		check(GeneratedLayout->Blocks.IsEmpty() || GeneratedLayout->Blocks[0].Id != FLayoutBlock::InvalidBlockId);
+		GeneratedLayouts.Add(SourceLayout.get(), GeneratedLayout);
+
+		return GeneratedLayout;
 	}
 
 
@@ -1356,7 +1368,7 @@ namespace mu
 								composeOp->BlockImage = blockAd;
 
 								// Set the absolute block index.
-								check(pLayout->Blocks[BlockIndex].Id != Layout::InvalidBlockId);
+								check(pLayout->Blocks[BlockIndex].Id != FLayoutBlock::InvalidBlockId);
 								composeOp->BlockId = pLayout->Blocks[BlockIndex].Id;
 
 								imageAd = composeOp;
@@ -1434,7 +1446,7 @@ namespace mu
 												composeOp->BlockImage = fragmentAd;
 
 												// Set the absolute block index.
-												check(pExtendLayout->Blocks[b].Id != Layout::InvalidBlockId);
+												check(pExtendLayout->Blocks[b].Id != FLayoutBlock::InvalidBlockId);
 												composeOp->BlockId = pExtendLayout->Blocks[b].Id;
 
 												lastBase = composeOp;
@@ -2179,7 +2191,7 @@ namespace mu
 				// Parameters
 				FImageGenerationOptions ClipOptions;
 				ClipOptions.ImageLayoutStrategy = CompilerOptions::TextureLayoutStrategy::None;
-				ClipOptions.LayoutBlockId = Layout::InvalidBlockId;
+				ClipOptions.LayoutBlockId = FLayoutBlock::InvalidBlockId;
 				ClipOptions.State = Options.State;
 
 				FImageGenerationResult ClipMaskResult;
diff --git a/Engine/Plugins/Experimental/Mutable/Source/MutableTools/Private/MuT/CodeGenerator.h b/Engine/Plugins/Experimental/Mutable/Source/MutableTools/Private/MuT/CodeGenerator.h
--- a/Engine/Plugins/Experimental/Mutable/Source/MutableTools/Private/MuT/CodeGenerator.h
+++ b/Engine/Plugins/Experimental/Mutable/Source/MutableTools/Private/MuT/CodeGenerator.h
@@ -385,7 +385,7 @@ namespace mu
 			UE::Math::TIntVector2<int32> RectSize = {0, 0};
 
 			/** Layout block that we are trying to generate if any. */
-			uint64 LayoutBlockId = Layout::InvalidBlockId;
+			uint64 LayoutBlockId = FLayoutBlock::InvalidBlockId;
 			Ptr<const Layout> LayoutToApply;
 
 			friend FORCEINLINE uint32 GetTypeHash(const FImageGenerationOptions& InKey)
@@ -502,12 +502,12 @@ namespace mu
 		struct FGeneratedLayout
 		{
 			Ptr<const Layout> Layout;
-			int32 FirstLODToIgnoreWarnings=0;
+			Ptr<const NodeLayout> Source;
 
 			FORCEINLINE bool operator==(const FGeneratedLayout& Other) const
 			{
 				return Layout == Other.Layout
-					&& FirstLODToIgnoreWarnings == Other.FirstLODToIgnoreWarnings;
+					&& Source == Other.Source;
 			}
 		};
 
@@ -609,9 +609,8 @@ namespace mu
 		TMap<int32, FMeshGenerationResult> SharedMeshOptionsMap;
 
 		//! Map of layouts found in the code already generated. The map is from the source layout
-		//! pointer of the layouts in the node meshes to the cloned and modified layout. 
-		//! The cloned layout will have absolute block ids assigned.
-		TMap<Ptr<const Layout>, Ptr<const Layout>> GeneratedLayouts;
+		//! node to the generated layout.
+		TMap<Ptr<const NodeLayout>, Ptr<const Layout>> GeneratedLayouts;
 
         void GenerateMesh(const FMeshGenerationOptions&, FMeshGenerationResult& result, const NodeMeshPtrConst&);
         void GenerateMesh_Constant(const FMeshGenerationOptions&, FMeshGenerationResult&, const NodeMeshConstant* );
@@ -632,16 +631,14 @@ namespace mu
 		void GenerateMesh_ClipDeform(const FMeshGenerationOptions&, FMeshGenerationResult&, const NodeMeshClipDeform*);
 
 		//-----------------------------------------------------------------------------------------
-		void PrepareForLayout( Ptr<const Layout> GeneratedLayout,
-			Ptr<Mesh> currentLayoutMesh,
+		void PrepareMeshForLayout(const FGeneratedLayout&, Ptr<Mesh>,
 			int32 currentLayoutChannel,
 			const void* errorContext,
-			const FMeshGenerationOptions& MeshOptions,
-			bool bUseAbsoluteBlockIds, 
-			int32 FirstLODToIgnoreWarnings);
+			const FMeshGenerationOptions&,
+			bool bUseAbsoluteBlockIds);
 
 		//!
-		Ptr<const Layout> AddLayout(Ptr<const Layout> SourceLayout, uint32 MeshIDPrefix);
+		Ptr<const Layout> AddLayout(Ptr<const NodeLayout> SourceLayout, uint32 MeshIDPrefix);
 
 		struct FExtensionDataGenerationResult
 		{
diff --git a/Engine/Plugins/Experimental/Mutable/Source/MutableTools/Private/MuT/CodeGenerator_Image.cpp b/Engine/Plugins/Experimental/Mutable/Source/MutableTools/Private/MuT/CodeGenerator_Image.cpp
--- a/Engine/Plugins/Experimental/Mutable/Source/MutableTools/Private/MuT/CodeGenerator_Image.cpp
+++ b/Engine/Plugins/Experimental/Mutable/Source/MutableTools/Private/MuT/CodeGenerator_Image.cpp
@@ -1386,9 +1386,9 @@ namespace mu
 		int32 LayoutBlockIndex = -1;
 		if (Options.LayoutToApply)
 		{
-			LayoutBlockIndex = Options.LayoutToApply->Blocks.IndexOfByPredicate([&](const Layout::FBlock& Block) { return Block.Id == Options.LayoutBlockId; });
+			LayoutBlockIndex = Options.LayoutToApply->Blocks.IndexOfByPredicate([&](const FLayoutBlock& Block) { return Block.Id == Options.LayoutBlockId; });
 		}
-		uint64 GeneratedLayoutBlockId = Layout::InvalidBlockId;
+		uint64 GeneratedLayoutBlockId = FLayoutBlock::InvalidBlockId;
 
         // Mesh
         if ( node.m_pMesh )
@@ -1445,7 +1445,7 @@ namespace mu
             else
             {
                 // Extract the mesh layout block
-                if ( GeneratedLayoutBlockId!= Layout::InvalidBlockId)
+                if ( GeneratedLayoutBlockId!= FLayoutBlock::InvalidBlockId)
                 {
                     Ptr<ASTOpMeshExtractLayoutBlocks> eop = new ASTOpMeshExtractLayoutBlocks();
                     eop->Source = CurrentMeshToProjectOp;
@@ -1538,7 +1538,7 @@ namespace mu
 			FImageGenerationOptions NewOptions = Options;
 			NewOptions.ImageLayoutStrategy = CompilerOptions::TextureLayoutStrategy::None;
 			NewOptions.LayoutToApply = nullptr;
-			NewOptions.LayoutBlockId = Layout::InvalidBlockId;
+			NewOptions.LayoutBlockId = FLayoutBlock::InvalidBlockId;
 			NewOptions.RectSize = { 0,0 };
 
 			FImageGenerationResult ImageResult;
diff --git a/Engine/Plugins/Experimental/Mutable/Source/MutableTools/Private/MuT/CodeGenerator_Mesh.cpp b/Engine/Plugins/Experimental/Mutable/Source/MutableTools/Private/MuT/CodeGenerator_Mesh.cpp
--- a/Engine/Plugins/Experimental/Mutable/Source/MutableTools/Private/MuT/CodeGenerator_Mesh.cpp
+++ b/Engine/Plugins/Experimental/Mutable/Source/MutableTools/Private/MuT/CodeGenerator_Mesh.cpp
@@ -264,60 +264,44 @@ namespace mu
 	}
 
     //---------------------------------------------------------------------------------------------
-	void CodeGenerator::PrepareForLayout(Ptr<const Layout> GeneratedLayout,
-		Ptr<Mesh> currentLayoutMesh,
-		int32 currentLayoutChannel,
+	void CodeGenerator::PrepareMeshForLayout( const FGeneratedLayout& GeneratedLayout,
+		Ptr<Mesh> Mesh,
+		int32 LayoutChannel,
 		const void* errorContext,
 		const FMeshGenerationOptions& MeshOptions,
-		bool bUseAbsoluteBlockIds,
-		int32 FirstLODToIgnoreWarnings
+		bool bUseAbsoluteBlockIds
 		)
 	{
-		MUTABLE_CPUPROFILER_SCOPE(LayoutUV_PrepareForLayout);
+		MUTABLE_CPUPROFILER_SCOPE(PrepareMeshForLayout);
 
-		if (currentLayoutMesh->GetVertexCount() == 0)
+		if (Mesh->GetVertexCount() == 0)
 		{
 			return;
 		}
 
 		// The layout must have block ids.
-		check(GeneratedLayout->Blocks.IsEmpty() || GeneratedLayout->Blocks[0].Id != Layout::InvalidBlockId);
+		check(GeneratedLayout.Layout->Blocks.IsEmpty() || GeneratedLayout.Layout->Blocks[0].Id != FLayoutBlock::InvalidBlockId);
 
 		// 
-		Ptr<const Layout> Layout = GeneratedLayout;
-		currentLayoutMesh->AddLayout(Layout);
+		Ptr<const Layout> Layout = GeneratedLayout.Layout;
+		Mesh->AddLayout(Layout);
 
-		const int32 NumVertices = currentLayoutMesh->GetVertexCount();
+		const int32 NumVertices = Mesh->GetVertexCount();
 		const int32 NumBlocks = Layout->GetBlockCount();
 
-		// Clear block data
-		TArray<uint16> LayoutData;
-		LayoutData.SetNumUninitialized(NumVertices);
-		constexpr uint16 NullBlockId = MAX_uint16 - 1;
-		for (int32 VertexIndex = 0; VertexIndex < NumVertices; ++VertexIndex)
-		{
-			LayoutData[VertexIndex] = NullBlockId;
-		}
-
 		// Find block ids for each block in the grid
 		const FIntPoint Grid = Layout->GetGridSize();
 		TArray<int32> GridBlockBlockId;
 		GridBlockBlockId.Init(MAX_uint16, Grid.X * Grid.Y);
 
 		// 
-		TArray<uint64> BlockIds;
 		TArray<box<FVector2f>> BlockRects;
-
-		BlockIds.SetNumUninitialized(NumBlocks);
 		BlockRects.SetNumUninitialized(NumBlocks);
 
 		// Create an array of block index per cell
 		TArray<int32> OverlappingBlocks;
 		for (int32 BlockIndex = 0; BlockIndex < NumBlocks; ++BlockIndex)
 		{
-			// Get the block id
-			BlockIds[BlockIndex] = Layout->Blocks[BlockIndex].Id;
-
 			// Get the block rect
 			FImageSize Min = Layout->Blocks[BlockIndex].Min;
 			FImageSize Size = Layout->Blocks[BlockIndex].Size;
@@ -328,20 +312,26 @@ namespace mu
 			BlockRect.size[0] = ((float)Size.X) / (float)Grid.X;
 			BlockRect.size[1] = ((float)Size.Y) / (float)Grid.Y;
 
-			// Create block index per cell array
-			for (uint16 Y = Min.Y; Y < Min.Y + Size.Y; ++Y)
-			{
-				const uint16 PositionY = Y * Grid.X;
+			bool bBlockHasMask = GeneratedLayout.Source->Blocks[BlockIndex].Mask.get() != nullptr;
 
-				for (uint16 X = Min.X; X < Min.X + Size.X; ++X)
+			// Fill the block index per cell array
+			// Ignore the block in this stage if it has a mask, because blocks with masks will very likely overlap other blocks
+			if (!bBlockHasMask)
+			{
+				for (uint16 Y = Min.Y; Y < Min.Y + Size.Y; ++Y)
 				{
-					if (GridBlockBlockId[PositionY + X] == MAX_uint16)
-					{
-						GridBlockBlockId[PositionY + X] = BlockIndex;
-					}
-					else
+					const uint16 PositionY = Y * Grid.X;
+
+					for (uint16 X = Min.X; X < Min.X + Size.X; ++X)
 					{
-						OverlappingBlocks.AddUnique(BlockIndex);
+						if (GridBlockBlockId[PositionY + X] == MAX_uint16)
+						{
+							GridBlockBlockId[PositionY + X] = BlockIndex;
+						}
+						else
+						{
+							OverlappingBlocks.AddUnique(BlockIndex);
+						}
 					}
 				}
 			}
@@ -358,48 +348,37 @@ namespace mu
 		}
 
 		// Get the information about the texture coordinates channel
-		int32 buffer = -1;
-		int32 channel = -1;
-		currentLayoutMesh->GetVertexBuffers().FindChannel(MBS_TEXCOORDS,
-			(int32)currentLayoutChannel,
-			&buffer,
-			&channel);
-		check(buffer >= 0);
-		check(channel >= 0);
-
-		EMeshBufferSemantic semantic;
-		int32 semanticIndex;
-		EMeshBufferFormat format;
-		int32 components;
-		int32 offset;
-		currentLayoutMesh->GetVertexBuffers().GetChannel
-		(buffer, channel, &semantic, &semanticIndex, &format, &components, &offset);
-		check(semantic == MBS_TEXCOORDS);
-
-		uint8* pData = currentLayoutMesh->GetVertexBuffers().GetBufferData(buffer);
-		int32 elemSize = currentLayoutMesh->GetVertexBuffers().GetElementSize(buffer);
-		int32 channelOffset = currentLayoutMesh->GetVertexBuffers().GetChannelOffset(buffer, channel);
-		pData += channelOffset;
-
-
-		// Temp copy of the UVs
-		TArray<FVector2f> TempUVs;
-		TempUVs.SetNumUninitialized(NumVertices);
+		int32 TexCoordsBufferIndex = -1;
+		int32 TexCoordsChannelIndex = -1;
+		Mesh->GetVertexBuffers().FindChannel(MBS_TEXCOORDS, LayoutChannel, &TexCoordsBufferIndex, &TexCoordsChannelIndex);
+		check(TexCoordsBufferIndex >= 0);
+		check(TexCoordsChannelIndex >= 0);
+
+		const FMeshBufferChannel& TexCoordsChannel = Mesh->VertexBuffers.m_buffers[TexCoordsBufferIndex].m_channels[TexCoordsChannelIndex];
+		check(TexCoordsChannel.m_semantic== MBS_TEXCOORDS);
+
+		uint8* TexCoordData = Mesh->GetVertexBuffers().GetBufferData(TexCoordsBufferIndex);
+		int32 elemSize = Mesh->GetVertexBuffers().GetElementSize(TexCoordsBufferIndex);
+		int32 channelOffset = TexCoordsChannel.m_offset;
+		TexCoordData += channelOffset;
 
 		// Get a copy of the UVs as FVector2f to work with them. 
+		TArray<FVector2f> TexCoords;
 		{
+			TexCoords.SetNumUninitialized(NumVertices);
+
 			bool bNonNormalizedUVs = false;
-			const bool bIsOverlayLayout = GeneratedLayout->GetLayoutPackingStrategy() == mu::EPackStrategy::OVERLAY_LAYOUT;
+			const bool bIsOverlayLayout = GeneratedLayout.Layout->GetLayoutPackingStrategy() == mu::EPackStrategy::OVERLAY_LAYOUT;
 
-			const uint8* pVertices = pData;
+			const uint8* pVertices = TexCoordData;
 			for (int32 VertexIndex = 0; VertexIndex < NumVertices; ++VertexIndex)
 			{
-				FVector2f& UV = TempUVs[VertexIndex];
-				if (format == MBF_FLOAT32)
+				FVector2f& UV = TexCoords[VertexIndex];
+				if (TexCoordsChannel.m_format == MBF_FLOAT32)
 				{
 					UV = *((FVector2f*)pVertices);
 				}
-				else if (format == MBF_FLOAT16)
+				else if (TexCoordsChannel.m_format == MBF_FLOAT16)
 				{
 					const FFloat16* pUV = reinterpret_cast<const FFloat16*>(pVertices);
 					UV = FVector2f(float(pUV[0]), float(pUV[1]));
@@ -426,7 +405,7 @@ namespace mu
 		}
 
 
-		const int32 NumTriangles = currentLayoutMesh->GetIndexCount() / 3;
+		const int32 NumTriangles = Mesh->GetIndexCount() / 3;
 		TArray<FTriangle> Triangles;
 
 		// Vertices mapped to unique vertex index
@@ -441,50 +420,79 @@ namespace mu
 			VertexToFaceMap.Reserve(NumVertices);
 			Triangles.SetNumUninitialized(NumTriangles);
 
-			MeshCreateCollapsedVertexMap(currentLayoutMesh.get(), CollapsedVertices);
+			MeshCreateCollapsedVertexMap(Mesh.get(), CollapsedVertices);
 		}
 
-		UntypedMeshBufferIteratorConst ItIndices(currentLayoutMesh->GetIndexBuffers(), MBS_VERTEXINDEX);
 		TArray<int32> ConflictiveTriangles;
 
 		const uint32 MaxGridX = MeshOptions.bNormalizeUVs ? MAX_uint32 : Grid.X - 1;
 		const uint32 MaxGridY = MeshOptions.bNormalizeUVs ? MAX_uint32 : Grid.Y - 1;
 
-		// TODO: We could skip this if there's only one block and it fits the entire grid
+		// Allocate the per-vertex layout block data
+		TArray<uint16> LayoutData;
+		constexpr uint16 NullBlockId = MAX_uint16 - 1;
+		LayoutData.Init(NullBlockId, NumVertices);
+
+		UntypedMeshBufferIteratorConst ItIndices(Mesh->GetIndexBuffers(), MBS_VERTEXINDEX);
 		for (int32 TriangleIndex = 0; TriangleIndex < NumTriangles; ++TriangleIndex)
 		{
-			uint32 Index0 = int32(ItIndices.GetAsUINT32());
+			uint32 Index0 = ItIndices.GetAsUINT32();
 			++ItIndices;
-			uint32 Index1 = int32(ItIndices.GetAsUINT32());
+			uint32 Index1 = ItIndices.GetAsUINT32();
 			++ItIndices;
-			uint32 Index2 = int32(ItIndices.GetAsUINT32());
+			uint32 Index2 = ItIndices.GetAsUINT32();
 			++ItIndices;
 
-			uint16 X, Y;
+			auto AssignOneVertex = 
+				[Grid, MaxGridX, MaxGridY, NumBlocks, Layout, &GeneratedLayout, &GridBlockBlockId, &TexCoords, &LayoutData]
+				(int32 VertexIndex)
+				{
+					uint16& BlockIndex = LayoutData[VertexIndex];
 
-			uint16& BlockIndexV0 = LayoutData[Index0];
-			if (BlockIndexV0 == NullBlockId)
-			{
-				X = (uint16)FMath::Min<uint32>(MaxGridX, FMath::Max<uint32>(0, (uint32)(Grid.X * TempUVs[Index0][0])));
-				Y = (uint16)FMath::Min<uint32>(MaxGridY, FMath::Max<uint32>(0, (uint32)(Grid.Y * TempUVs[Index0][1])));
-				BlockIndexV0 = (X < Grid.X && Y < Grid.Y) ? GridBlockBlockId[Y * Grid.X + X] : 0;
-			}
+					FVector2f UV = TexCoords[VertexIndex];
 
-			uint16& BlockIndexV1 = LayoutData[Index1];
-			if (BlockIndexV1 == NullBlockId)
-			{
-				X = (uint16)FMath::Min<uint32>(MaxGridX, FMath::Max<uint32>(0, (uint32)(Grid.X * TempUVs[Index1][0])));
-				Y = (uint16)FMath::Min<uint32>(MaxGridY, FMath::Max<uint32>(0, (uint32)(Grid.Y * TempUVs[Index1][1])));
-				BlockIndexV1 = (X < Grid.X && Y < Grid.Y) ? GridBlockBlockId[Y * Grid.X + X] : 0;
-			}
+					// First: Assign the vertices to masked blocks in order
+					for (int32 CandidateBlockIndex = 0; CandidateBlockIndex < NumBlocks; ++CandidateBlockIndex)
+					{
+						const mu::Image* Mask = GeneratedLayout.Source->Blocks[CandidateBlockIndex].Mask.get();
+						if (Mask)
+						{
+							// First discard with block limits.
+							FImageSize Min = Layout->Blocks[CandidateBlockIndex].Min;
+							FImageSize Size = Layout->Blocks[CandidateBlockIndex].Size;
+
+							uint16 X = (uint16)FMath::Min<uint32>(MaxGridX, FMath::Max<uint32>(0, (uint32)(Grid.X * UV[0])));
+							uint16 Y = (uint16)FMath::Min<uint32>(MaxGridY, FMath::Max<uint32>(0, (uint32)(Grid.Y * UV[1])));
+							bool bInBlock = 
+								(X >= Min.X && X < Min.X + Size.X)
+								&&
+								(Y >= Min.Y && Y < Min.Y + Size.Y);
+
+							if (bInBlock)
+							{
+								FVector4f MaskValue = Mask->Sample(UV);
+								if (MaskValue.X > 0.5f)
+								{
+									BlockIndex = CandidateBlockIndex;
+									break;
+								}
+							}
+						}
+					}
 
-			uint16& BlockIndexV2 = LayoutData[Index2];
-			if (BlockIndexV2 == NullBlockId)
-			{
-				X = (uint16)FMath::Min<uint32>(MaxGridX, FMath::Max<uint32>(0, (uint32)(Grid.X * TempUVs[Index2][0])));
-				Y = (uint16)FMath::Min<uint32>(MaxGridY, FMath::Max<uint32>(0, (uint32)(Grid.Y * TempUVs[Index2][1])));
-				BlockIndexV2 = (X < Grid.X && Y < Grid.Y) ? GridBlockBlockId[Y * Grid.X + X] : 0;
-			}
+					// Second: Assign to non-masked blocks if not assigned yet
+					if (BlockIndex == NullBlockId)
+					{
+						uint16 X = (uint16)FMath::Min<uint32>(MaxGridX, FMath::Max<uint32>(0, (uint32)(Grid.X * UV[0])));
+						uint16 Y = (uint16)FMath::Min<uint32>(MaxGridY, FMath::Max<uint32>(0, (uint32)(Grid.Y * UV[1])));
+						BlockIndex = (X < Grid.X && Y < Grid.Y) ? GridBlockBlockId[Y * Grid.X + X] : 0;
+					}
+					return BlockIndex;
+				};
+
+			const uint16 BlockIndexV0 = AssignOneVertex(Index0);
+			const uint16 BlockIndexV1 = AssignOneVertex(Index1);
+			const uint16 BlockIndexV2 = AssignOneVertex(Index2);
 
 			if (MeshOptions.bClampUVIslands)
 			{
@@ -526,7 +534,7 @@ namespace mu
 
 			// Find triangles from the same UV Island
 			TArray<uint32> TriangleIndices;
-			GetUVIsland(Triangles, ConflictiveTriangleIndex, TriangleIndices, TempUVs, VertexToFaceMap);
+			GetUVIsland(Triangles, ConflictiveTriangleIndex, TriangleIndices, TexCoords, VertexToFaceMap);
 
 			// Get predominant BlockId != MAX_uint16
 			TArray<uint32> NumVerticesPerBlock;
@@ -557,7 +565,7 @@ namespace mu
 			}
 
 			// Get the limits of the predominant block rect
-			const Layout::FBlock& LayoutBlock = Layout->Blocks[BlockIndex];
+			const FLayoutBlock& LayoutBlock = Layout->Blocks[BlockIndex];
 
 			const float SmallNumber = 0.000001;
 			const float MinX = ((float)LayoutBlock.Min.X) / (float)Grid.X + SmallNumber;
@@ -566,7 +574,6 @@ namespace mu
 			const float MaxY = (((float)LayoutBlock.Size.Y + LayoutBlock.Min.Y) / (float)Grid.Y) - 2 * SmallNumber;
 
 			// Iterate triangles and clamp the UVs
-			FVector2f* TempUVsData = TempUVs.GetData();
 			for (int32 TriangleIndex : TriangleIndices)
 			{
 				FTriangle& OtherTriangle = Triangles[TriangleIndex];
@@ -582,7 +589,7 @@ namespace mu
 
 					// Clamp UVs
 					const int32 UVIndex = OtherTriangle.Indices[VertexIndex];
-					FVector2f& UV = TempUVs[UVIndex];
+					FVector2f& UV = TexCoords[UVIndex];
 					UV[0] = FMath::Clamp(UV[0], MinX, MaxX);
 					UV[1] = FMath::Clamp(UV[1], MinY, MaxY);
 					LayoutData[UVIndex] = BlockIndex;
@@ -593,12 +600,13 @@ namespace mu
 		}
 
 		// Warn about vertices without a block id
+		int32 FirstLODToIgnoreWarnings = GeneratedLayout.Source->FirstLODToIgnoreWarnings;
 		if (FirstLODToIgnoreWarnings == -1 || CurrentParents.Last().Lod < FirstLODToIgnoreWarnings)
 		{
 			TArray<float> UnassignedUVs;
 			UnassignedUVs.Reserve(NumVertices / 100);
 
-			const FVector2f* UVs = TempUVs.GetData();
+			const FVector2f* UVs = TexCoords.GetData();
 			for (int32 VertexIndex = 0; VertexIndex < NumVertices; ++VertexIndex)
 			{
 				if (LayoutData[VertexIndex] == MAX_uint16)
@@ -623,18 +631,18 @@ namespace mu
 		// Create the layout block vertex buffer
 		uint8* LayoutBufferPtr = nullptr;
 		{
-			const int32 LayoutBufferIndex = currentLayoutMesh->GetVertexBuffers().GetBufferCount();
-			currentLayoutMesh->GetVertexBuffers().SetBufferCount(LayoutBufferIndex + 1);
+			const int32 LayoutBufferIndex = Mesh->GetVertexBuffers().GetBufferCount();
+			Mesh->GetVertexBuffers().SetBufferCount(LayoutBufferIndex + 1);
 
 			// TODO
 			check(Layout->GetBlockCount() < MAX_uint16);
 			const EMeshBufferSemantic LayoutSemantic = MBS_LAYOUTBLOCK;
-			const int32 LayoutSemanticIndex = int32(currentLayoutChannel);
+			const int32 LayoutSemanticIndex = int32(LayoutChannel);
 			const EMeshBufferFormat LayoutFormat = bUseAbsoluteBlockIds ? MBF_UINT64 : MBF_UINT16;
 			const int32 LayoutComponents = 1;
 			const int32 LayoutOffset = 0;
 			int32 ElementSize = bUseAbsoluteBlockIds ? sizeof(uint64) : sizeof(uint16);
-			currentLayoutMesh->GetVertexBuffers().SetBuffer
+			Mesh->GetVertexBuffers().SetBuffer
 			(
 				LayoutBufferIndex,
 				ElementSize,
@@ -643,59 +651,59 @@ namespace mu
 				&LayoutFormat, &LayoutComponents,
 				&LayoutOffset
 			);
-			LayoutBufferPtr = currentLayoutMesh->GetVertexBuffers().GetBufferData(LayoutBufferIndex);
+			LayoutBufferPtr = Mesh->GetVertexBuffers().GetBufferData(LayoutBufferIndex);
 		}
 
-		// Format and copy UVs
+		// Copy UVs back to the mesh
 		{
-			uint8* pVertices = pData;
-			FVector2f* UVs = TempUVs.GetData();
+			uint8* pVertices = TexCoordData;
+			FVector2f* UVs = TexCoords.GetData();
 
 			for (int32 VertexIndex = 0; VertexIndex < NumVertices; ++VertexIndex)
 			{
-				FVector2f* UV = &TempUVs[VertexIndex];
+				FVector2f* UV = &TexCoords[VertexIndex];
 
 				uint16 LayoutBlockIndex = LayoutData[VertexIndex];
-
-				// Replace block index by the actual id of the block
-				if (bUseAbsoluteBlockIds)
+				if (Layout->Blocks.IsValidIndex(LayoutBlockIndex))
 				{
-					uint64* Ptr = reinterpret_cast<uint64*>(LayoutBufferPtr) + VertexIndex;
+					uint64 LayoutBlockId = Layout->Blocks[LayoutBlockIndex].Id;
+
+					*UV = BlockRects[LayoutBlockIndex].Homogenize(*UV);
 
-					if (BlockIds.IsValidIndex(LayoutBlockIndex))
+					// Replace block index by the actual id of the block
+					if (bUseAbsoluteBlockIds)
 					{
-						*UV = BlockRects[LayoutBlockIndex].Homogenize(*UV);
-						*Ptr = BlockIds[LayoutBlockIndex];
+						uint64* Ptr = reinterpret_cast<uint64*>(LayoutBufferPtr) + VertexIndex;
+						*Ptr = LayoutBlockId;
 					}
 					else
 					{
-						// Map vertices without block to the first block.
-						*Ptr = 0;
+						uint16* Ptr = reinterpret_cast<uint16*>(LayoutBufferPtr) + VertexIndex;
+						*Ptr = uint16(LayoutBlockId & 0xffff);
 					}
 				}
 				else
 				{
-					uint16* Ptr = reinterpret_cast<uint16*>(LayoutBufferPtr) + VertexIndex;
-
-					if (BlockIds.IsValidIndex(LayoutBlockIndex))
+					// Map vertices without block to the first block.
+					if (bUseAbsoluteBlockIds)
 					{
-						*UV = BlockRects[LayoutBlockIndex].Homogenize(*UV);
-						*Ptr = uint16(BlockIds[LayoutBlockIndex] & 0xffff);
+						uint64* Ptr = reinterpret_cast<uint64*>(LayoutBufferPtr) + VertexIndex;
+						*Ptr = 0;
 					}
 					else
 					{
-						// Map vertices without block to the first block.
+						uint16* Ptr = reinterpret_cast<uint16*>(LayoutBufferPtr) + VertexIndex;
 						*Ptr = 0;
 					}
 				}
 
 				// Copy UVs
-				if (format == MBF_FLOAT32)
+				if (TexCoordsChannel.m_format == MBF_FLOAT32)
 				{
 					FVector2f* pUV = reinterpret_cast<FVector2f*>(pVertices);
 					*pUV = *UV;
 				}
-				else if (format == MBF_FLOAT16)
+				else if (TexCoordsChannel.m_format == MBF_FLOAT16)
 				{
 					FFloat16* pUV = reinterpret_cast<FFloat16*>(pVertices);
 					pUV[0] = FFloat16((*UV)[0]);
@@ -1392,10 +1400,10 @@ namespace mu
 		else
 		{
 			// We need to clone the mesh in the node because we will modify it.
-			Ptr<Mesh> Cloned = pMesh->Clone();
-			Cloned->EnsureSurfaceData();
+			Ptr<Mesh> ClonedMesh = pMesh->Clone();
+			ClonedMesh->EnsureSurfaceData();
 
-			ConstantOp->SetValue(Cloned, CompilerOptions->OptimisationOptions.DiskCacheContext);
+			ConstantOp->SetValue(ClonedMesh, CompilerOptions->OptimisationOptions.DiskCacheContext);
 
 			// Add the unique vertex ID prefix in all cases, since it is free memory-wise
 			uint32 MeshIDPrefix = uint32(ConstantOp->GetValueHash());
@@ -1413,12 +1421,12 @@ namespace mu
 					}
 				} while (bValid);
 
-				Cloned->MeshIDPrefix = MeshIDPrefix;
+				ClonedMesh->MeshIDPrefix = MeshIDPrefix;
 			}
 
 			// Add the constant data
 			FGeneratedConstantMesh MeshEntry;
-			MeshEntry.Mesh = Cloned;
+			MeshEntry.Mesh = ClonedMesh;
 			MeshEntry.LastMeshOp = LastMeshOp;
 			GeneratedConstantMeshes.Add(MeshEntry);
 
@@ -1429,21 +1437,21 @@ namespace mu
 					// Apply whatever transform is necessary for every layout
 					for (int32 LayoutIndex = 0; LayoutIndex < Node.Layouts.Num(); ++LayoutIndex)
 					{
-						Ptr<NodeLayout> pLayoutNode = Node.Layouts[LayoutIndex];
-						if (!pLayoutNode)
+						Ptr<NodeLayout> LayoutNode = Node.Layouts[LayoutIndex];
+						if (!LayoutNode)
 						{
 							continue;
 						}
 
-						Ptr<const Layout> SourceLayout = pLayoutNode->Layout;
-						Ptr<const Layout> GeneratedLayout = AddLayout(SourceLayout, MeshIDPrefix);
-						int32 FirstLODToIgnoreWarnings = pLayoutNode->FirstLODToIgnoreWarnings;
+						FGeneratedLayout GeneratedData;
+						GeneratedData.Source = LayoutNode;
+						GeneratedData.Layout = AddLayout(LayoutNode, MeshIDPrefix);
 						const void* Context = InOptions.OverrideContext.Get(InNode->GetMessageContext());
 
 						bool bUseAbsoluteBlockIds = false;
-						PrepareForLayout(GeneratedLayout, Cloned, LayoutIndex, Context, InOptions, bUseAbsoluteBlockIds, FirstLODToIgnoreWarnings);
+						PrepareMeshForLayout(GeneratedData, ClonedMesh, LayoutIndex, Context, InOptions, bUseAbsoluteBlockIds);
 
-						OutResult.GeneratedLayouts.Add({ GeneratedLayout, FirstLODToIgnoreWarnings });
+						OutResult.GeneratedLayouts.Add(GeneratedData);
 					}
 				}
 				else
@@ -1451,16 +1459,16 @@ namespace mu
 					// We need to apply the transform of the layouts used to override
 					for (int32 LayoutIndex = 0; LayoutIndex < InOptions.OverrideLayouts.Num(); ++LayoutIndex)
 					{
-						Ptr<const Layout> GeneratedLayout = InOptions.OverrideLayouts[LayoutIndex].Layout;
-						int32 FirstLODToIgnoreWarnings = InOptions.OverrideLayouts[LayoutIndex].FirstLODToIgnoreWarnings;
+						const FGeneratedLayout& OverrideData = InOptions.OverrideLayouts[LayoutIndex];
+						Ptr<const Layout> GeneratedLayout = OverrideData.Layout;
 						const void* Context = InOptions.OverrideContext.Get(InNode->GetMessageContext());
 
 						// In this case we need the layout block ids to use the ids in the parent layout, and not be prefixed with
 						// the current mesh id prefix. For this reason we need them to be absolute.
 						bool bUseAbsoluteBlockIds = true;
-						PrepareForLayout(GeneratedLayout, Cloned, LayoutIndex, Context, InOptions, bUseAbsoluteBlockIds, FirstLODToIgnoreWarnings);
+						PrepareMeshForLayout(OverrideData, ClonedMesh, LayoutIndex, Context, InOptions, bUseAbsoluteBlockIds);
 
-						OutResult.GeneratedLayouts.Add({ GeneratedLayout,FirstLODToIgnoreWarnings });
+						OutResult.GeneratedLayouts.Add(OverrideData);
 					}
 				}
 			}
diff --git a/Engine/Plugins/Experimental/Mutable/Source/MutableTools/Private/MuT/CodeOptimiser.cpp b/Engine/Plugins/Experimental/Mutable/Source/MutableTools/Private/MuT/CodeOptimiser.cpp
--- a/Engine/Plugins/Experimental/Mutable/Source/MutableTools/Private/MuT/CodeOptimiser.cpp
+++ b/Engine/Plugins/Experimental/Mutable/Source/MutableTools/Private/MuT/CodeOptimiser.cpp
@@ -84,7 +84,7 @@ namespace mu
 
 		struct custom_layout_equal
 		{
-			bool operator()( const LayoutPtrConst& a, const LayoutPtrConst& b ) const
+			bool operator()( const Ptr<const Layout>& a, const Ptr<const Layout>& b ) const
 			{
 				return a==b || *a==*b;
 			}
diff --git a/Engine/Plugins/Experimental/Mutable/Source/MutableTools/Private/MuT/NodeLayout.cpp b/Engine/Plugins/Experimental/Mutable/Source/MutableTools/Private/MuT/NodeLayout.cpp
--- a/Engine/Plugins/Experimental/Mutable/Source/MutableTools/Private/MuT/NodeLayout.cpp
+++ b/Engine/Plugins/Experimental/Mutable/Source/MutableTools/Private/MuT/NodeLayout.cpp
@@ -16,13 +16,7 @@ namespace mu
 	FNodeType NodeLayout::StaticType = FNodeType(Node::EType::Layout, Node::GetStaticType() );
 
 
-	NodeLayout::NodeLayout()
-	{
-		Layout = new mu::Layout();
-	}
-
-
-	Ptr<NodeLayout> NodeLayout::GenerateLayoutBlocks(const Ptr<Mesh> pMesh, int32 layoutIndex, int32 gridSizeX, int32 gridSizeY)
+	Ptr<NodeLayout> NodeLayout::GenerateLayoutBlocks(const Ptr<Mesh>& pMesh, int32 layoutIndex, int32 gridSizeX, int32 gridSizeY)
 	{
 		Ptr<NodeLayout> LayoutNode = nullptr;
 
@@ -75,9 +69,9 @@ namespace mu
 			}
 
 			LayoutNode = new NodeLayout;
-			LayoutNode->Layout->SetGridSize(gridSizeX, gridSizeY);
-			LayoutNode->Layout->SetMaxGridSize(gridSizeX, gridSizeY);
-			LayoutNode->Layout->SetLayoutPackingStrategy(EPackStrategy::RESIZABLE_LAYOUT);
+			LayoutNode->Size = { uint16(gridSizeX), uint16(gridSizeY) };
+			LayoutNode->MaxSize = { uint16(gridSizeX), uint16(gridSizeY) };
+			LayoutNode->Strategy = EPackStrategy::RESIZABLE_LAYOUT;
 			
 			TArray<box<FIntVector2>> blocks;
 			
@@ -209,12 +203,12 @@ namespace mu
 			//Generating layout blocks
 			if (NumBlocks > 0)
 			{
-				LayoutNode->Layout->SetBlockCount(NumBlocks);
+				LayoutNode->Blocks.SetNum(NumBlocks);
 			
 				for (int32 BlockIndex = 0; BlockIndex < NumBlocks; ++BlockIndex)
 				{
-					LayoutNode->Layout->Blocks[BlockIndex].Min = UE::Math::TIntVector2<uint16>(blocks[BlockIndex].min);
-					LayoutNode->Layout->Blocks[BlockIndex].Size = UE::Math::TIntVector2<uint16>(blocks[BlockIndex].size);
+					LayoutNode->Blocks[BlockIndex].Min = UE::Math::TIntVector2<uint16>(blocks[BlockIndex].min);
+					LayoutNode->Blocks[BlockIndex].Size = UE::Math::TIntVector2<uint16>(blocks[BlockIndex].size);
 				}
 			}
 		}
diff --git a/Engine/Plugins/Experimental/Mutable/Source/MutableTools/Private/MuT/NodeMeshConstant.cpp b/Engine/Plugins/Experimental/Mutable/Source/MutableTools/Private/MuT/NodeMeshConstant.cpp
--- a/Engine/Plugins/Experimental/Mutable/Source/MutableTools/Private/MuT/NodeMeshConstant.cpp
+++ b/Engine/Plugins/Experimental/Mutable/Source/MutableTools/Private/MuT/NodeMeshConstant.cpp
@@ -7,6 +7,7 @@
 #include "MuT/NodeLayout.h"
 #include "MuT/NodeMeshConstantPrivate.h"
 #include "MuT/NodePrivate.h"
+#include "MuR/Mesh.h"
 
 
 namespace mu
diff --git a/Engine/Plugins/Experimental/Mutable/Source/MutableTools/Public/MuT/NodeLayout.h b/Engine/Plugins/Experimental/Mutable/Source/MutableTools/Public/MuT/NodeLayout.h
--- a/Engine/Plugins/Experimental/Mutable/Source/MutableTools/Public/MuT/NodeLayout.h
+++ b/Engine/Plugins/Experimental/Mutable/Source/MutableTools/Public/MuT/NodeLayout.h
@@ -2,30 +2,65 @@
 
 #pragma once
 
-#include "HAL/Platform.h"
+#include "MuT/Node.h"
 #include "MuR/Layout.h"
-#include "MuR/Mesh.h"
+#include "MuR/Image.h"
 #include "MuR/Ptr.h"
 #include "MuR/RefCounted.h"
-#include "MuT/Node.h"
+#include "HAL/Platform.h"
 
 
 namespace mu
 {
+	class Mesh;
+
+	/** Data for a layout block before it is compiled. */
+	struct FSourceLayoutBlock
+	{
+		/** Optional mask image that selects the vertices to include in the block. */
+		Ptr<Image> Mask;
+
+		UE::Math::TIntVector2<uint16> Min = { 0, 0 };
+		UE::Math::TIntVector2<uint16> Size = { 0, 0 };
+
+		//! Priority value to control the shrink texture layout strategy
+		int32 Priority;
+
+		//! Value to control the method to reduce the block
+		uint32 bReduceBothAxes : 1;
+
+		//! Value to control if a block has to be reduced by two in an unitary reduction strategy
+		uint32 bReduceByTwo : 1;
+	};
+
 
 	/** This node is used to define the texture layout for a texture coordinates channel of a mesh. */
 	class MUTABLETOOLS_API NodeLayout : public Node
 	{
 	public:
 
-		Ptr<Layout> Layout;
+		//!
+		UE::Math::TIntVector2<uint16> Size = UE::Math::TIntVector2<uint16>(0, 0);
 
+		/** Maximum size in layout blocks that this layout can grow to. From there on, blocks will shrink to fit.
+		* If 0,0 then no maximum size applies.
+		*/
+		UE::Math::TIntVector2<uint16> MaxSize = UE::Math::TIntVector2<uint16>(0, 0);
+
+		//!
+		TArray<FSourceLayoutBlock> Blocks;
+
+		//! Packing strategy
+		EPackStrategy Strategy = EPackStrategy::RESIZABLE_LAYOUT;
+		EReductionMethod ReductionMethod = EReductionMethod::HALVE_REDUCTION;
+
+		/** When compiling, ignore generated warnings from this LOD on.
+		* -1 means all warnings are generated.
+		*/
 		int32 FirstLODToIgnoreWarnings = 0;
 
 	public:
 
-		NodeLayout();
-
 		// Node Interface
 		virtual const FNodeType* GetType() const override { return GetStaticType(); }
 		static const FNodeType* GetStaticType() { return &StaticType; }
@@ -33,7 +68,7 @@ namespace mu
 		// Own interface
 
 		//! Generate the blocks of a layout using the UV of the meshes
-		static Ptr<NodeLayout> GenerateLayoutBlocks(const Ptr<Mesh> pMesh, int32 layoutIndex, int32 gridSizeX, int32 gridSizeY);
+		static Ptr<NodeLayout> GenerateLayoutBlocks(const Ptr<Mesh>& pMesh, int32 layoutIndex, int32 gridSizeX, int32 gridSizeY);
 
 	protected:
 
