[mutable] Compilation optimization: remove unnecessary memory work to generate resource hashes.
#rnx

diff --git a/Engine/Plugins/Experimental/Mutable/Source/CustomizableObjectEditor/Private/MuCOE/GenerateMutableSource/GenerateMutableSourceTable.cpp b/Engine/Plugins/Experimental/Mutable/Source/CustomizableObjectEditor/Private/MuCOE/GenerateMutableSource/GenerateMutableSourceTable.cpp
--- a/Engine/Plugins/Experimental/Mutable/Source/CustomizableObjectEditor/Private/MuCOE/GenerateMutableSource/GenerateMutableSourceTable.cpp
+++ b/Engine/Plugins/Experimental/Mutable/Source/CustomizableObjectEditor/Private/MuCOE/GenerateMutableSource/GenerateMutableSourceTable.cpp
@@ -120,6 +120,7 @@ bool FillTableColumn(const UCustomizableObjectNodeTable* TableNode, mu::TablePtr
 
 											if (!AnimInstance.IsNull())
 											{
+												// TODO: We only need it to add it to participating objects. We should probably skip fully loading it.
 												if (UClass* Anim = AnimInstance.LoadSynchronous())
 												{
 													GenerationContext.AddParticipatingObject(*Anim);													
@@ -881,6 +882,8 @@ TArray<FName> GetRowsToCompile(const UDataTable& DataTable, const UCustomizableO
 bool GenerateTableColumn(const UCustomizableObjectNodeTable* TableNode, const UEdGraphPin* Pin, mu::TablePtr MutableTable, const FString& DataTableColumnName, const FProperty* ColumnProperty,
 	const int32 LODIndexConnected, const int32 SectionIndexConnected, const int32 LODIndex, const int32 SectionIndex, const bool bOnlyConnectedLOD, FMutableGraphGenerationContext& GenerationContext)
 {
+	MUTABLE_CPUPROFILER_SCOPE(GenerateTableColumn);
+
 	SCOPED_PIN_DATA(GenerationContext, Pin)
 
 	if (!TableNode)
@@ -900,6 +903,31 @@ bool GenerateTableColumn(const UCustomizableObjectNodeTable* TableNode, const UE
 	// Getting names of the rows to access the information
 	TArray<FName> RowNames = GetRowsToCompile(*DataTable, *TableNode, GenerationContext);
 
+	// Pre-pass to request async loading of all data. This seems to be slightly faster because it avoids the sync after every separate load.
+	{
+		TArray<int32> LoadRequests;
+		LoadRequests.Reserve(RowNames.Num());
+
+		MUTABLE_CPUPROFILER_SCOPE(Preload);
+		for (int32 RowIndex = 0; RowIndex < RowNames.Num(); ++RowIndex)
+		{
+			if (uint8* CellData = GetCellData(RowNames[RowIndex], *DataTable, *ColumnProperty))
+			{
+				// Getting property type
+				if (const FSoftObjectProperty* SoftObjectProperty = CastField<FSoftObjectProperty>(ColumnProperty))
+				{
+					const FSoftObjectPtr& Path = SoftObjectProperty->GetPropertyValue(CellData);
+					LoadRequests.Add( LoadPackageAsync(Path.GetLongPackageName()) );
+				}
+			}
+		}
+
+		{
+			MUTABLE_CPUPROFILER_SCOPE(Flush);
+			FlushAsyncLoading(LoadRequests);
+		}
+	}
+
 	for (int32 RowIndex = 0; RowIndex < RowNames.Num(); ++RowIndex)
 	{
 		if (uint8* CellData = GetCellData(RowNames[RowIndex], *DataTable, *ColumnProperty))
diff --git a/Engine/Plugins/Experimental/Mutable/Source/MutableRuntime/Private/MuR/Serialisation.cpp b/Engine/Plugins/Experimental/Mutable/Source/MutableRuntime/Private/MuR/Serialisation.cpp
--- a/Engine/Plugins/Experimental/Mutable/Source/MutableRuntime/Private/MuR/Serialisation.cpp
+++ b/Engine/Plugins/Experimental/Mutable/Source/MutableRuntime/Private/MuR/Serialisation.cpp
@@ -7,6 +7,7 @@
 #include "HAL/PlatformCrt.h"
 #include "HAL/PlatformMath.h"
 #include "Misc/AssertionMacros.h"
+#include "Hash/CityHash.h"
 #include "MuR/Image.h"
 #include "MuR/SerialisationPrivate.h"
 
@@ -185,27 +186,31 @@ namespace mu
 
     //-------------------------------------------------------------------------------------------------
     //-------------------------------------------------------------------------------------------------
-    //-------------------------------------------------------------------------------------------------
-    OutputSizeStream::OutputSizeStream()
-    {
-        WrittenBytes = 0;
-    }
-
-
     //-------------------------------------------------------------------------------------------------
     void OutputSizeStream::Write( const void*, uint64 size )
     {
         WrittenBytes += size;
     }
 
-
-    //-------------------------------------------------------------------------------------------------
 	uint64 OutputSizeStream::GetBufferSize() const
     {
         return WrittenBytes;
     }
 
 
+	//-------------------------------------------------------------------------------------------------
+	void OutputHashStream::Write(const void* Data, uint64 size)
+	{
+		Hash = CityHash64WithSeed(reinterpret_cast<const char*>(Data), size, Hash);
+	}
+
+
+	uint64 OutputHashStream::GetHash() const
+	{
+		return Hash;
+	}
+
+
     //---------------------------------------------------------------------------------------------
 	//---------------------------------------------------------------------------------------------
 	//---------------------------------------------------------------------------------------------
diff --git a/Engine/Plugins/Experimental/Mutable/Source/MutableRuntime/Public/MuR/Serialisation.h b/Engine/Plugins/Experimental/Mutable/Source/MutableRuntime/Public/MuR/Serialisation.h
--- a/Engine/Plugins/Experimental/Mutable/Source/MutableRuntime/Public/MuR/Serialisation.h
+++ b/Engine/Plugins/Experimental/Mutable/Source/MutableRuntime/Public/MuR/Serialisation.h
@@ -305,33 +305,42 @@ namespace mu
     };
 
 
-	//! This stream doesn't store any data, it just counts the amount of data serialised.
+	/** This stream doesn't store any data: it just counts the amount of data serialised. */
 	class MUTABLERUNTIME_API OutputSizeStream : public OutputStream
 	{
 	public:
 
-		//-----------------------------------------------------------------------------------------
-		// Life cycle
-		//-----------------------------------------------------------------------------------------
+		// OutputStream interface
+		void Write(const void* pData, uint64 size) override;
 
-		//!
-		OutputSizeStream();
+		// Own interface
+
+		/** Get the amount of data serialised, in bytes. */
+		uint64 GetBufferSize() const;
+
+	private:
+
+		uint64 WrittenBytes = 0;
+
+	};
+
+
+	/** This stream doesn't store any data: it just calculates of a hash of the data as it receives it. */
+	class MUTABLERUNTIME_API OutputHashStream : public OutputStream
+	{
+	public:
 
-		//-----------------------------------------------------------------------------------------
 		// OutputStream interface
-		//-----------------------------------------------------------------------------------------
 		void Write(const void* pData, uint64 size) override;
 
-		//-----------------------------------------------------------------------------------------
 		// Own interface
-		//-----------------------------------------------------------------------------------------
 
-		//! Get the amount of data serialised, in bytes.
-		uint64 GetBufferSize() const;
+		/** Return the hash of the data written so far. */
+		uint64 GetHash() const;
 
 	private:
 
-		uint64 WrittenBytes;
+		uint64 Hash = 0;
 
 	};
 
diff --git a/Engine/Plugins/Experimental/Mutable/Source/MutableTools/Private/MuT/ASTOpConstantResource.cpp b/Engine/Plugins/Experimental/Mutable/Source/MutableTools/Private/MuT/ASTOpConstantResource.cpp
--- a/Engine/Plugins/Experimental/Mutable/Source/MutableTools/Private/MuT/ASTOpConstantResource.cpp
+++ b/Engine/Plugins/Experimental/Mutable/Source/MutableTools/Private/MuT/ASTOpConstantResource.cpp
@@ -691,17 +691,22 @@ namespace mu
 	//-------------------------------------------------------------------------------------------------
 	void ASTOpConstantResource::SetValue(const Ptr<const RefCounted>& v, FProxyFileContext* DiskCacheContext)
 	{
+		MUTABLE_CPUPROFILER_SCOPE(ASTOpConstantResource_SetValue);
+
 		switch (Type)
 		{
 		case OP_TYPE::IM_CONSTANT:
 		{
 			Ptr<const Image> r = static_cast<const Image*>(v.get());
 
-			OutputMemoryStream stream(r->GetDataSize() + 1024);
-			OutputArchive arch(&stream);
-			Image::Serialise(r.get(), arch);
+			OutputHashStream stream;
+			{
+				MUTABLE_CPUPROFILER_SCOPE(Serialize);
+				OutputArchive arch(&stream);
+				Image::Serialise(r.get(), arch);
+			}
 
-			ValueHash = CityHash64(static_cast<const char*>(stream.GetBuffer()), stream.GetBufferSize());
+			ValueHash = stream.GetHash();
 
 			if (DiskCacheContext)
 			{
@@ -716,13 +721,16 @@ namespace mu
 
 		case OP_TYPE::ME_CONSTANT:
 		{
-			OutputMemoryStream stream;
-			OutputArchive arch(&stream);
-
 			Ptr<const Mesh> r = static_cast<const Mesh*>(v.get());
-			Mesh::Serialise(r.get(), arch);
 
-			ValueHash = CityHash64(static_cast<const char*>(stream.GetBuffer()), stream.GetBufferSize());
+			OutputHashStream stream;
+			{
+				MUTABLE_CPUPROFILER_SCOPE(Serialize);
+				OutputArchive arch(&stream);
+				Mesh::Serialise(r.get(), arch);
+			}
+
+			ValueHash = stream.GetHash();
 
 			LoadedValue = v;
 			break;
@@ -730,13 +738,16 @@ namespace mu
 
 		case OP_TYPE::LA_CONSTANT:
 		{
-			OutputMemoryStream stream;
-			OutputArchive arch(&stream);
-
 			Ptr<const Layout> r = static_cast<const Layout*>(v.get());
-			Layout::Serialise(r.get(), arch);
 
-			ValueHash = CityHash64(static_cast<const char*>(stream.GetBuffer()), stream.GetBufferSize());
+			OutputHashStream stream;
+			{
+				MUTABLE_CPUPROFILER_SCOPE(Serialize);
+				OutputArchive arch(&stream);
+				Layout::Serialise(r.get(), arch);
+			}
+
+			ValueHash = stream.GetHash();
 
 			LoadedValue = v;
 			break;
