#Mutable Do not attempt to load soft object pointers are that are not assigned

diff --git a/Engine/Plugins/Experimental/Mutable/Source/CustomizableObjectEditor/Private/MuCOE/GenerateMutableSource/GenerateMutableSourceTable.cpp b/Engine/Plugins/Experimental/Mutable/Source/CustomizableObjectEditor/Private/MuCOE/GenerateMutableSource/GenerateMutableSourceTable.cpp
--- a/Engine/Plugins/Experimental/Mutable/Source/CustomizableObjectEditor/Private/MuCOE/GenerateMutableSource/GenerateMutableSourceTable.cpp
+++ b/Engine/Plugins/Experimental/Mutable/Source/CustomizableObjectEditor/Private/MuCOE/GenerateMutableSource/GenerateMutableSourceTable.cpp
@@ -917,7 +917,11 @@ bool GenerateTableColumn(const UCustomizableObjectNodeTable* TableNode, const UE
 				if (const FSoftObjectProperty* SoftObjectProperty = CastField<FSoftObjectProperty>(ColumnProperty))
 				{
 					const FSoftObjectPtr& Path = SoftObjectProperty->GetPropertyValue(CellData);
-					LoadRequests.Add( LoadPackageAsync(Path.GetLongPackageName()) );
+					const FString PackageName = Path.GetLongPackageName();
+					if (!PackageName.IsEmpty())
+					{
+						LoadRequests.Add( LoadPackageAsync(PackageName) );
+					}
 				}
 			}
 		}
