[mutable] Optimize the code generation stage for meshes.
- Generate mutable meshes already in separate buffers.
- Optimize similarity check.
#rnx

diff --git a/Engine/Plugins/Experimental/Mutable/Source/CustomizableObject/Internal/MuCO/CustomizableObjectPrivate.h b/Engine/Plugins/Experimental/Mutable/Source/CustomizableObject/Internal/MuCO/CustomizableObjectPrivate.h
--- a/Engine/Plugins/Experimental/Mutable/Source/CustomizableObject/Internal/MuCO/CustomizableObjectPrivate.h
+++ b/Engine/Plugins/Experimental/Mutable/Source/CustomizableObject/Internal/MuCO/CustomizableObjectPrivate.h
@@ -1170,6 +1170,8 @@ public:
 
 		DataTablesParamTrackingForCompileOnlySelected,
 
+		CompilationOptimizationsMeshFormat,
+
 		// -----<new versions can be added above this line>--------
 		LastCustomizableObjectVersion
 	};
diff --git a/Engine/Plugins/Experimental/Mutable/Source/CustomizableObjectEditor/Private/MuCOE/GenerateMutableSource/GenerateMutableSource.cpp b/Engine/Plugins/Experimental/Mutable/Source/CustomizableObjectEditor/Private/MuCOE/GenerateMutableSource/GenerateMutableSource.cpp
--- a/Engine/Plugins/Experimental/Mutable/Source/CustomizableObjectEditor/Private/MuCOE/GenerateMutableSource/GenerateMutableSource.cpp
+++ b/Engine/Plugins/Experimental/Mutable/Source/CustomizableObjectEditor/Private/MuCOE/GenerateMutableSource/GenerateMutableSource.cpp
@@ -815,6 +815,8 @@ mu::NodeMeshApplyPosePtr CreateNodeMeshApplyPose(FMutableGraphGenerationContext&
 // Convert a CustomizableObject Source Graph into a mutable source graph  
 mu::Ptr<mu::NodeObject> GenerateMutableSource(const UEdGraphPin * Pin, FMutableGraphGenerationContext & GenerationContext, bool bPartialCompilation)
 {
+	MUTABLE_CPUPROFILER_SCOPE(GenerateMutableSource);
+
 	check(Pin)
 	RETURN_ON_CYCLE(*Pin, GenerationContext)
 
@@ -1924,6 +1926,8 @@ mu::FImageDesc GenerateImageDescriptor(UTexture* Texture)
 
 mu::Ptr<mu::Image> GenerateImageConstant(UTexture* Texture, FMutableGraphGenerationContext& GenerationContext, bool bIsReference)
 {
+	MUTABLE_CPUPROFILER_SCOPE(GenerateImageConstant);
+
 	if (!Texture)
 	{
 		return nullptr;
diff --git a/Engine/Plugins/Experimental/Mutable/Source/CustomizableObjectEditor/Private/MuCOE/GenerateMutableSource/GenerateMutableSourceMesh.cpp b/Engine/Plugins/Experimental/Mutable/Source/CustomizableObjectEditor/Private/MuCOE/GenerateMutableSource/GenerateMutableSourceMesh.cpp
--- a/Engine/Plugins/Experimental/Mutable/Source/CustomizableObjectEditor/Private/MuCOE/GenerateMutableSource/GenerateMutableSourceMesh.cpp
+++ b/Engine/Plugins/Experimental/Mutable/Source/CustomizableObjectEditor/Private/MuCOE/GenerateMutableSource/GenerateMutableSourceMesh.cpp
@@ -609,50 +609,71 @@ mu::Ptr<mu::PhysicsBody> MakePhysicsBodyFromAsset(FMutableGraphGenerationContext
 namespace MutablePrivate
 {
 
-	// \TODO: Since we are no longer bulk-copying, we could generate buffers without padding, and only with the relevant information.
-	void CopyBufferClearingPadding(mu::FMeshBufferSet& BufferSet, int32 BufferIndex, const void* InSourceData)
+	/** Copy mesh data from a single source buffer into a set of destination buffers assuming the format and components is the same. */
+	void CopyBufferClearingPadding(mu::FMeshBufferSet& BufferSet, const mu::FMeshBufferSet& SourceBufferFormat, const void* InSourceData)
 	{
-		check(BufferIndex < BufferSet.m_buffers.Num());
+		MUTABLE_CPUPROFILER_SCOPE(CopyBufferClearingPadding);
 
 		int32 ElementCount = BufferSet.m_elementCount;
-		int32 ChannelCount = BufferSet.m_buffers[BufferIndex].m_channels.Num();
-		int32 ElementSize = BufferSet.m_buffers[BufferIndex].m_elementSize;
 
-		uint8* TargetData = BufferSet.GetBufferData(BufferIndex);
-		const uint8* SourceData = reinterpret_cast<const uint8*>(InSourceData);
-
-		for (int32 Element = 0; Element < ElementCount; ++Element)
+		for (int32 BufferIndex = 0; BufferIndex < BufferSet.GetBufferCount(); ++BufferIndex)
 		{
-			int32 CurrentOffset = 0;
+			int32 ChannelCount = BufferSet.m_buffers[BufferIndex].m_channels.Num();
+			int32 ElementSize = BufferSet.m_buffers[BufferIndex].m_elementSize;
+
+			check(SourceBufferFormat.m_buffers.Num()==1);
+			const mu::FMeshBuffer& SourceBuffer = SourceBufferFormat.m_buffers[0];
+
+			uint8* TargetData = BufferSet.GetBufferData(BufferIndex);
+			TArray<const uint8*, TInlineAllocator<8>> SourceDataPerChannel;
+			SourceDataPerChannel.SetNumUninitialized(ChannelCount);
 			for (int32 ChannelIndex = 0; ChannelIndex < ChannelCount; ++ChannelIndex)
 			{
-				const mu::FMeshBufferChannel& Channel = BufferSet.m_buffers[BufferIndex].m_channels[ChannelIndex];
+				const mu::FMeshBufferChannel& DestinationChannel = BufferSet.m_buffers[BufferIndex].m_channels[ChannelIndex];
 
-				int32 ChannelOffset = Channel.m_offset;
+				int32 SourceBufferIndex = -1;
+				int32 SourceChannelIndex = -1;
+				SourceBufferFormat.FindChannel(DestinationChannel.m_semantic, DestinationChannel.m_semanticIndex, &SourceBufferIndex, &SourceChannelIndex);
+				check(SourceBufferIndex==0 && SourceChannelIndex >=0);
 
-				int32 PreviousPadding = ChannelOffset - CurrentOffset;
-				if (PreviousPadding > 0)
-				{
-					FMemory::Memzero(TargetData, PreviousPadding);
-					TargetData += PreviousPadding;
-					SourceData += PreviousPadding;
-					CurrentOffset += PreviousPadding;
-				}
+				const mu::FMeshBufferChannel& SourceChannel = SourceBuffer.m_channels[SourceChannelIndex];
+				check(SourceChannel.m_format == DestinationChannel.m_format);
+				check(SourceChannel.m_componentCount == DestinationChannel.m_componentCount);
 
-				int32 ChannelSize = Channel.m_componentCount * GetMeshFormatData(Channel.m_format).SizeInBytes;
-				FMemory::Memcpy(TargetData, SourceData, ChannelSize);
-				TargetData += ChannelSize;
-				SourceData += ChannelSize;
-				CurrentOffset += ChannelSize;
+				SourceDataPerChannel[ChannelIndex] = reinterpret_cast<const uint8*>(InSourceData) + SourceChannel.m_offset;
 			}
 
-			// Padding at the end?
-			int32 FinalPadding = ElementSize - CurrentOffset;
-			if (FinalPadding > 0)
+			for (int32 Element = 0; Element < ElementCount; ++Element)
 			{
-				FMemory::Memzero(TargetData, FinalPadding);
-				TargetData += FinalPadding;
-				SourceData += FinalPadding;
+				int32 CurrentOffset = 0;
+				for (int32 ChannelIndex = 0; ChannelIndex < ChannelCount; ++ChannelIndex)
+				{
+					const mu::FMeshBufferChannel& Channel = BufferSet.m_buffers[BufferIndex].m_channels[ChannelIndex];
+
+					int32 ChannelOffset = Channel.m_offset;
+
+					int32 PreviousPadding = ChannelOffset - CurrentOffset;
+					if (PreviousPadding > 0)
+					{
+						FMemory::Memzero(TargetData, PreviousPadding);
+						TargetData += PreviousPadding;
+						CurrentOffset += PreviousPadding;
+					}
+
+					int32 ChannelSize = Channel.m_componentCount * GetMeshFormatData(Channel.m_format).SizeInBytes;
+					FMemory::Memcpy(TargetData, SourceDataPerChannel[ChannelIndex], ChannelSize);
+					TargetData += ChannelSize;
+					SourceDataPerChannel[ChannelIndex] += SourceBuffer.m_elementSize;
+					CurrentOffset += ChannelSize;
+				}
+
+				// Padding at the end?
+				int32 FinalPadding = ElementSize - CurrentOffset;
+				if (FinalPadding > 0)
+				{
+					FMemory::Memzero(TargetData, FinalPadding);
+					TargetData += FinalPadding;
+				}
 			}
 		}
 	}
@@ -662,6 +683,8 @@ namespace MutablePrivate
 
 mu::MeshPtr ConvertSkeletalMeshToMutable(const USkeletalMesh* InSkeletalMesh, const TSoftClassPtr<UAnimInstance>& AnimBp, int32 LODIndexConnected, int32 SectionIndexConnected, int32 LODIndex, int32 SectionIndex, FMutableGraphGenerationContext& GenerationContext, const UCustomizableObjectNode* CurrentNode)
 {
+	MUTABLE_CPUPROFILER_SCOPE(ConvertSkeletalMeshToMutable);
+
 	if(!InSkeletalMesh)
 	{
 		return nullptr;
@@ -907,24 +930,45 @@ mu::MeshPtr ConvertSkeletalMeshToMutable(const USkeletalMesh* InSkeletalMesh, co
 
 	MutableMesh->GetVertexBuffers().SetElementCount(VertexCount);
 
-	MutableMesh->GetVertexBuffers().SetBufferCount(1);
+	MutableMesh->GetVertexBuffers().SetBufferCount(2);
 
 	const int32 MaxSectionInfluences = MeshSection.MaxBoneInfluences;
 	const bool bUseUnlimitedInfluences = FGPUBaseSkinVertexFactory::UseUnlimitedBoneInfluences(MaxSectionInfluences, GenerationContext.Options.TargetPlatform);
 
-	if (bIgnoreSkeleton)
+	// Create a mutable vertex buffer definition that matches the unreal soft vertex format.
+	mu::FMeshBufferSet UnrealSourceVertexFormat;
 	{
-		// Create the mesh with the same data, but skinning is considered padding.
 		using namespace mu;
-		const int ElementSize = sizeof(FSoftSkinVertex);
-		const int ChannelCount = 9;
-		const EMeshBufferSemantic Semantics[ChannelCount] = { MBS_POSITION, MBS_TANGENT, MBS_BINORMAL, MBS_NORMAL, MBS_TEXCOORDS, MBS_TEXCOORDS, MBS_TEXCOORDS, MBS_TEXCOORDS, MBS_COLOUR };
-		const int SemanticIndices[ChannelCount] = { 0, 0, 0, 0, 0, 1, 2, 3, 0 };
-		const EMeshBufferFormat Formats[ChannelCount] = { MBF_FLOAT32, MBF_FLOAT32, MBF_FLOAT32, MBF_FLOAT32, MBF_FLOAT32, MBF_FLOAT32, MBF_FLOAT32, MBF_FLOAT32, MBF_NUINT8 };
-		int Components[ChannelCount] = { 3, 3, 3, 4, 2, 2, 2, 2, 4 };
+		const int32 ElementSize = sizeof(FSoftSkinVertex);
+		const int32 ChannelCount = 11;
+		const EMeshBufferSemantic Semantics[ChannelCount] = { MBS_POSITION, MBS_TANGENT, MBS_BINORMAL, MBS_NORMAL, MBS_TEXCOORDS, MBS_TEXCOORDS, MBS_TEXCOORDS, MBS_TEXCOORDS, MBS_COLOUR, MBS_BONEINDICES, MBS_BONEWEIGHTS };
+		const int32 SemanticIndices[ChannelCount] = { 0, 0, 0, 0, 0, 1, 2, 3, 0, 0, 0 };
+
+		// TODO: Remove BoneWeightFormat after merge
+		EMeshBufferFormat BoneWeightFormat = sizeof(TDecay<decltype(DeclVal<FSoftSkinVertex>().InfluenceWeights[0])>::Type) == 1 ? MBF_NUINT8 : MBF_NUINT16;
+		const EMeshBufferFormat Formats[ChannelCount] = { MBF_FLOAT32, MBF_FLOAT32, MBF_FLOAT32, MBF_FLOAT32, MBF_FLOAT32, MBF_FLOAT32, MBF_FLOAT32, MBF_FLOAT32, MBF_NUINT8, MBF_UINT16, BoneWeightFormat };
+
+		int32 Components[ChannelCount] = { 3, 3, 3, 4, 2, 2, 2, 2, 4, 4, 4 };
+		if (GenerationContext.Options.CustomizableObjectNumBoneInfluences != ECustomizableObjectNumBoneInfluences::Four &&
+			MaxSectionInfluences > 4)
+		{
+			int32 NewBoneInfluencesNum = (int32)GenerationContext.Options.CustomizableObjectNumBoneInfluences;
+
+			if (bUseUnlimitedInfluences &&
+				MaxSectionInfluences < NewBoneInfluencesNum)
+			{
+				Components[9] = MaxSectionInfluences;
+				Components[10] = MaxSectionInfluences;
+			}
+			else
+			{
+				Components[9] = NewBoneInfluencesNum;
+				Components[10] = NewBoneInfluencesNum;
+			}
+		}
 
 		constexpr size_t SoftSkinVertexUVsElemSize = sizeof(TDecay<decltype(DeclVal<FSoftSkinVertex>().UVs[0])>::Type);
-		const int Offsets[ChannelCount] =
+		const int32 Offsets[ChannelCount] =
 		{
 			STRUCT_OFFSET(FSoftSkinVertex, Position),
 			STRUCT_OFFSET(FSoftSkinVertex, TangentX),
@@ -935,24 +979,75 @@ mu::MeshPtr ConvertSkeletalMeshToMutable(const USkeletalMesh* InSkeletalMesh, co
 			STRUCT_OFFSET(FSoftSkinVertex, UVs) + 2 * SoftSkinVertexUVsElemSize,
 			STRUCT_OFFSET(FSoftSkinVertex, UVs) + 3 * SoftSkinVertexUVsElemSize,
 			STRUCT_OFFSET(FSoftSkinVertex, Color),
+			STRUCT_OFFSET(FSoftSkinVertex, InfluenceBones),
+			STRUCT_OFFSET(FSoftSkinVertex, InfluenceWeights),
 		};
 
-		MutableMesh->GetVertexBuffers().SetBuffer(0, ElementSize, ChannelCount, Semantics, SemanticIndices, Formats, Components, Offsets);
-		MutablePrivate::CopyBufferClearingPadding(MutableMesh->GetVertexBuffers(), 0, Vertices.GetData() + VertexStart);
+		UnrealSourceVertexFormat.SetBufferCount(1);
+		UnrealSourceVertexFormat.SetBuffer(0, ElementSize, ChannelCount, Semantics, SemanticIndices, Formats, Components, Offsets);
 	}
-	else
+
+	// Create the mutable mesh with the same data without padding and separating the UVs, also add skinning only optionally.
 	{
 		using namespace mu;
-		const int ElementSize = sizeof(FSoftSkinVertex);
-		const int ChannelCount = 11;
-		const EMeshBufferSemantic Semantics[ChannelCount] = { MBS_POSITION, MBS_TANGENT, MBS_BINORMAL, MBS_NORMAL, MBS_TEXCOORDS, MBS_TEXCOORDS, MBS_TEXCOORDS, MBS_TEXCOORDS, MBS_COLOUR, MBS_BONEINDICES, MBS_BONEWEIGHTS };
-		const int SemanticIndices[ChannelCount] = { 0, 0, 0, 0, 0, 1, 2, 3, 0, 0, 0 };
+
+		// Base channels
+		{
+			const int32 ChannelCount = 5;
+			const EMeshBufferSemantic Semantics[ChannelCount] = { MBS_POSITION, MBS_TANGENT, MBS_BINORMAL, MBS_NORMAL, MBS_COLOUR };
+			const int32 SemanticIndices[ChannelCount] = { 0, 0, 0, 0, 0 };
+			const EMeshBufferFormat Formats[ChannelCount] = { MBF_FLOAT32, MBF_FLOAT32, MBF_FLOAT32, MBF_FLOAT32, MBF_NUINT8 };
+			int32 Components[ChannelCount] = { 3, 3, 3, 4, 4 };
+			const int32 Offsets[ChannelCount] =
+			{
+				0,
+				sizeof(float) * 3,
+				sizeof(float) * 6,
+				sizeof(float) * 9,
+				sizeof(float) * 13,
+			};
+			const int32 ElementSize = sizeof(float) * 13 + 4;
+
+			MutableMesh->GetVertexBuffers().SetBuffer(0, ElementSize, ChannelCount, Semantics, SemanticIndices, Formats, Components, Offsets);
+			check(!MutableMesh->VertexBuffers.m_buffers[0].HasPadding());
+		}
+
+		// Texture coordinates
+		{
+			const int32 ElementSize = 4 * 2 * sizeof(float);
+			const int32 ChannelCount = 4;
+			const EMeshBufferSemantic Semantics[ChannelCount] = { MBS_TEXCOORDS, MBS_TEXCOORDS, MBS_TEXCOORDS, MBS_TEXCOORDS };
+			const int32 SemanticIndices[ChannelCount] = { 0, 1, 2, 3 };
+			const EMeshBufferFormat Formats[ChannelCount] = { MBF_FLOAT32, MBF_FLOAT32, MBF_FLOAT32, MBF_FLOAT32 };
+			int32 Components[ChannelCount] = { 2, 2, 2, 2 };
+			const int32 Offsets[ChannelCount] =
+			{
+				0,
+				sizeof(float) * 2,
+				sizeof(float) * 4,
+				sizeof(float) * 6,
+			};
+
+			MutableMesh->GetVertexBuffers().SetBuffer(1, ElementSize, ChannelCount, Semantics, SemanticIndices, Formats, Components, Offsets);
+			check(!MutableMesh->VertexBuffers.m_buffers[1].HasPadding());
+		}
+	}
+
+	if (!bIgnoreSkeleton)
+	{
+		MutableMesh->GetVertexBuffers().SetBufferCount(3);
+
+		// Skinning buffer
+		using namespace mu;
+		const int32 ChannelCount = 2;
+		const EMeshBufferSemantic Semantics[ChannelCount] = { MBS_BONEINDICES, MBS_BONEWEIGHTS };
+		const int32 SemanticIndices[ChannelCount] = { 0, 0 };
 
 		// TODO: Remove BoneWeightFormat after merge
 		EMeshBufferFormat BoneWeightFormat = sizeof(TDecay<decltype(DeclVal<FSoftSkinVertex>().InfluenceWeights[0])>::Type) == 1 ? MBF_NUINT8 : MBF_NUINT16;
-		const EMeshBufferFormat Formats[ChannelCount] = { MBF_FLOAT32, MBF_FLOAT32, MBF_FLOAT32, MBF_FLOAT32, MBF_FLOAT32, MBF_FLOAT32, MBF_FLOAT32, MBF_FLOAT32, MBF_NUINT8, MBF_UINT16, BoneWeightFormat };
+		const EMeshBufferFormat Formats[ChannelCount] = { MBF_UINT16, BoneWeightFormat };
 
-		int Components[ChannelCount] = { 3, 3, 3, 4, 2, 2, 2, 2, 4, 4, 4 };
+		int32 Components[ChannelCount] = { 4, 4 };
 		if (GenerationContext.Options.CustomizableObjectNumBoneInfluences != ECustomizableObjectNumBoneInfluences::Four && 
 			MaxSectionInfluences > 4)
 		{
@@ -961,32 +1056,27 @@ mu::MeshPtr ConvertSkeletalMeshToMutable(const USkeletalMesh* InSkeletalMesh, co
 			if (bUseUnlimitedInfluences &&
 				MaxSectionInfluences < NewBoneInfluencesNum)
 			{
-				Components[9] = MaxSectionInfluences;
-				Components[10] = MaxSectionInfluences;
+				Components[0] = MaxSectionInfluences;
+				Components[1] = MaxSectionInfluences;
 			}
 			else
 			{
-				Components[9] = NewBoneInfluencesNum;
-				Components[10] = NewBoneInfluencesNum;
+				Components[0] = NewBoneInfluencesNum;
+				Components[1] = NewBoneInfluencesNum;
 			}
 		}
 
 		constexpr size_t SoftSkinVertexUVsElemSize = sizeof(TDecay<decltype(DeclVal<FSoftSkinVertex>().UVs[0])>::Type);
-		const int Offsets[ChannelCount] =
+		const int32 Offsets[ChannelCount] =
 		{
-			STRUCT_OFFSET(FSoftSkinVertex, Position),
-			STRUCT_OFFSET(FSoftSkinVertex, TangentX),
-			STRUCT_OFFSET(FSoftSkinVertex, TangentY),
-			STRUCT_OFFSET(FSoftSkinVertex, TangentZ),
-			STRUCT_OFFSET(FSoftSkinVertex, UVs) + 0 * SoftSkinVertexUVsElemSize,
-			STRUCT_OFFSET(FSoftSkinVertex, UVs) + 1 * SoftSkinVertexUVsElemSize,
-			STRUCT_OFFSET(FSoftSkinVertex, UVs) + 2 * SoftSkinVertexUVsElemSize,
-			STRUCT_OFFSET(FSoftSkinVertex, UVs) + 3 * SoftSkinVertexUVsElemSize,
-			STRUCT_OFFSET(FSoftSkinVertex, Color),
-			STRUCT_OFFSET(FSoftSkinVertex, InfluenceBones),
-			STRUCT_OFFSET(FSoftSkinVertex, InfluenceWeights),
+			0,
+			Components[0] * sizeof(uint16),
 		};
 
+		int32 ElementSize = Components[0] * sizeof(uint16) + Components[1] * ((BoneWeightFormat == MBF_NUINT8) ? 1 : 2);
+
+		MutableMesh->GetVertexBuffers().SetBuffer(2, ElementSize, ChannelCount, Semantics, SemanticIndices, Formats, Components, Offsets);
+
 		// Fix bone weights if required (uint8 -> uint16)
 		if (BoneWeightFormat == MBF_NUINT16 && Vertices.IsValidIndex(VertexStart))
 		{
@@ -1051,13 +1141,13 @@ mu::MeshPtr ConvertSkeletalMeshToMutable(const USkeletalMesh* InSkeletalMesh, co
 			}
 		}
 
-		MutableMesh->GetVertexBuffers().SetBuffer(0, ElementSize, ChannelCount, Semantics, SemanticIndices, Formats, Components, Offsets);
-		MutablePrivate::CopyBufferClearingPadding(MutableMesh->GetVertexBuffers(), 0, Vertices.GetData() + VertexStart);
 	}
 
+	MutablePrivate::CopyBufferClearingPadding(MutableMesh->GetVertexBuffers(), UnrealSourceVertexFormat, Vertices.GetData() + VertexStart);
+
 
 	// TODO: Add Mesh generation flags to not include RT Morph and clothing if not needed.
-	int32 NextBufferIndex = 1;
+	int32 NextBufferIndex = MutableMesh->VertexBuffers.m_buffers.Num();
 	if (GenerationContext.Options.bRealTimeMorphTargetsEnabled)
 	{
 
@@ -3011,6 +3101,8 @@ mu::NodeMeshPtr GenerateMorphMesh(const UEdGraphPin* Pin,
 	const bool bOnlyConnectedLOD,
 	const FString& TableColumnName = "")
 {
+	MUTABLE_CPUPROFILER_SCOPE(GenerateMorphMesh);
+
 	SCOPED_PIN_DATA(GenerationContext, Pin)
 	
 	// SkeletalMesh node
@@ -3198,6 +3290,8 @@ mu::NodeMeshPtr GenerateMutableSourceMesh(const UEdGraphPin* Pin,
 	const bool bLinkedToExtendMaterial,
 	const bool bOnlyConnectedLOD)
 {
+	MUTABLE_CPUPROFILER_SCOPE(GenerateMutableSourceMesh);
+
 	check(Pin)
 	RETURN_ON_CYCLE(*Pin, GenerationContext)
 	SCOPED_PIN_DATA(GenerationContext, Pin)
diff --git a/Engine/Plugins/Experimental/Mutable/Source/CustomizableObjectEditor/Private/MuCOE/GenerateMutableSource/GenerateMutableSourceSurface.cpp b/Engine/Plugins/Experimental/Mutable/Source/CustomizableObjectEditor/Private/MuCOE/GenerateMutableSource/GenerateMutableSourceSurface.cpp
--- a/Engine/Plugins/Experimental/Mutable/Source/CustomizableObjectEditor/Private/MuCOE/GenerateMutableSource/GenerateMutableSourceSurface.cpp
+++ b/Engine/Plugins/Experimental/Mutable/Source/CustomizableObjectEditor/Private/MuCOE/GenerateMutableSource/GenerateMutableSourceSurface.cpp
@@ -231,6 +231,8 @@ void AddModifierToSharedSurface(FMutableGraphGenerationContext& GenerationContex
 
 mu::Ptr<mu::NodeSurface> GenerateMutableSourceSurface(const UEdGraphPin * Pin, FMutableGraphGenerationContext & GenerationContext)
 {
+	MUTABLE_CPUPROFILER_SCOPE(GenerateMutableSourceSurface);
+
 	check(Pin)
 	RETURN_ON_CYCLE(*Pin, GenerationContext)
 
diff --git a/Engine/Plugins/Experimental/Mutable/Source/MutableRuntime/Private/MuR/MeshBufferSet.cpp b/Engine/Plugins/Experimental/Mutable/Source/MutableRuntime/Private/MuR/MeshBufferSet.cpp
--- a/Engine/Plugins/Experimental/Mutable/Source/MutableRuntime/Private/MuR/MeshBufferSet.cpp
+++ b/Engine/Plugins/Experimental/Mutable/Source/MutableRuntime/Private/MuR/MeshBufferSet.cpp
@@ -502,29 +502,42 @@ namespace mu
 				return false;
 			}
 
-			for (uint32 Elem = 0; Elem < m_elementCount; ++Elem)
+			bool bBufferHasTexCoords = ThisBuffer.HasSemantic(MBS_TEXCOORDS);
+			bool bCanBeFullCompared = !ThisBuffer.HasPadding() && (!bBufferHasTexCoords || bCompareUVs);
+			if (bCanBeFullCompared)
 			{
-				for (int32 C = 0; C < ThisNumChannels; ++C)
-				{
-					if (!bCompareUVs && ThisBuffer.m_channels[C].m_semantic==MBS_TEXCOORDS)
-					{						
-						continue;
-					}
+				MUTABLE_CPUPROFILER_SCOPE(FastCompare);
 
-					const SIZE_T SizeA = GetMeshFormatData(ThisBuffer.m_channels[C].m_format).SizeInBytes * ThisBuffer.m_channels[C].m_componentCount;
-					const SIZE_T SizeB = GetMeshFormatData(OtherBuffer.m_channels[C].m_format).SizeInBytes * OtherBuffer.m_channels[C].m_componentCount;
-
-					if (SizeA != SizeB)
-					{
-						return false;
-					}
-
-					const uint8* BuffA = ThisBuffer.m_data.GetData() + Elem * ThisBuffer.m_elementSize + ThisBuffer.m_channels[C].m_offset;
-					const uint8* BuffB = OtherBuffer.m_data.GetData() + Elem * OtherBuffer.m_elementSize + OtherBuffer.m_channels[C].m_offset;
+				// This buffer can be directly compared
+				if (FMemory::Memcmp(ThisBuffer.m_data.GetData(), OtherBuffer.m_data.GetData(), ThisBuffer.m_data.Num()) != 0)
+				{
+					return false;
+				}
+			}
+			else
+			{
+				MUTABLE_CPUPROFILER_SCOPE(SlowCompare);
 
-					if (FMemory::Memcmp(BuffA, BuffB, SizeA) != 0)
+				for (uint32 Elem = 0; Elem < m_elementCount; ++Elem)
+				{
+					for (int32 C = 0; C < ThisNumChannels; ++C)
 					{
-						return false;
+						if (!bCompareUVs && ThisBuffer.m_channels[C].m_semantic == MBS_TEXCOORDS)
+						{
+							continue;
+						}
+
+						const SIZE_T SizeA = GetMeshFormatData(ThisBuffer.m_channels[C].m_format).SizeInBytes * ThisBuffer.m_channels[C].m_componentCount;
+						const SIZE_T SizeB = GetMeshFormatData(OtherBuffer.m_channels[C].m_format).SizeInBytes * OtherBuffer.m_channels[C].m_componentCount;
+						check(SizeA == SizeB);
+
+						const uint8* BuffA = ThisBuffer.m_data.GetData() + Elem * ThisBuffer.m_elementSize + ThisBuffer.m_channels[C].m_offset;
+						const uint8* BuffB = OtherBuffer.m_data.GetData() + Elem * OtherBuffer.m_elementSize + OtherBuffer.m_channels[C].m_offset;
+
+						if (FMemory::Memcmp(BuffA, BuffB, SizeA) != 0)
+						{
+							return false;
+						}
 					}
 				}
 			}
@@ -695,4 +708,3 @@ namespace mu
 	}
 	
 }
-
diff --git a/Engine/Plugins/Experimental/Mutable/Source/MutableRuntime/Public/MuR/MeshBufferSet.h b/Engine/Plugins/Experimental/Mutable/Source/MutableRuntime/Public/MuR/MeshBufferSet.h
--- a/Engine/Plugins/Experimental/Mutable/Source/MutableRuntime/Public/MuR/MeshBufferSet.h
+++ b/Engine/Plugins/Experimental/Mutable/Source/MutableRuntime/Public/MuR/MeshBufferSet.h
@@ -221,6 +221,16 @@ namespace mu
 			return (m_channels == Other.m_channels && m_elementSize == Other.m_elementSize);
 		}
 
+		inline bool HasPadding() const
+		{
+			uint32 ActualElementSize = 0;
+			for (const FMeshBufferChannel& Channel : m_channels)
+			{
+				ActualElementSize += Channel.m_componentCount * GetMeshFormatData(Channel.m_format).SizeInBytes;
+			}
+			check(ActualElementSize<=m_elementSize);
+			return ActualElementSize < m_elementSize;
+		}
 	};
 
 
