[mutable] Optimize the code generation stage for meshes.
- Optimize duplication check container.
#rnx

diff --git a/Engine/Plugins/Experimental/Mutable/Source/MutableTools/Private/MuT/CodeGenerator.h b/Engine/Plugins/Experimental/Mutable/Source/MutableTools/Private/MuT/CodeGenerator.h
--- a/Engine/Plugins/Experimental/Mutable/Source/MutableTools/Private/MuT/CodeGenerator.h
+++ b/Engine/Plugins/Experimental/Mutable/Source/MutableTools/Private/MuT/CodeGenerator.h
@@ -230,13 +230,15 @@ namespace mu
 
     private:
 
-        /** List of meshes generated to be able to reuse them. */
+        /** Container of meshes generated to be able to reuse them. 
+		* They are sorted by a cheap hash to speed up searches.
+		*/
 		struct FGeneratedConstantMesh
 		{
 			Ptr<Mesh> Mesh;
 			Ptr<ASTOp> LastMeshOp;
 		};
-		TArray<FGeneratedConstantMesh> GeneratedConstantMeshes;
+		TMap<uint64,TArray<FGeneratedConstantMesh>> GeneratedConstantMeshes;
 
         /** List of image resources for every image formats that have been generated so far as palceholders for missing images. */
         Ptr<Image> MissingImage[size_t(EImageFormat::IF_COUNT)];
diff --git a/Engine/Plugins/Experimental/Mutable/Source/MutableTools/Private/MuT/CodeGenerator_Mesh.cpp b/Engine/Plugins/Experimental/Mutable/Source/MutableTools/Private/MuT/CodeGenerator_Mesh.cpp
--- a/Engine/Plugins/Experimental/Mutable/Source/MutableTools/Private/MuT/CodeGenerator_Mesh.cpp
+++ b/Engine/Plugins/Experimental/Mutable/Source/MutableTools/Private/MuT/CodeGenerator_Mesh.cpp
@@ -88,8 +88,7 @@ namespace mu
     //! Create a map from vertices into vertices, collapsing vertices that have the same position. 
 	//! This version uses UE Containers to return.	
     //---------------------------------------------------------------------------------------------
-	void MeshCreateCollapsedVertexMap(const mu::Mesh* Mesh,
-		TArray<int32>& CollapsedVertices)
+	void MeshCreateCollapsedVertexMap(const mu::Mesh* Mesh, TArray<int32>& CollapsedVertices)
 	{
 		MUTABLE_CPUPROFILER_SCOPE(LayoutUV_CreateCollapsedVertexMap);
 	
@@ -299,6 +298,8 @@ namespace mu
 			LayoutData[VertexIndex] = NullBlockId;
 		}
 
+		bool bIsSingleFullBlock = (NumBlocks == 1) && (Layout->Blocks[0].Min == FImageSize(0, 0) && Layout->Blocks[0].Size == Layout->Size);
+
 		// Find block ids for each block in the grid
 		const FIntPoint Grid = Layout->GetGridSize();
 		TArray<int32> GridBlockBlockId;
@@ -436,7 +437,7 @@ namespace mu
 		TMultiMap<int32, uint32> VertexToFaceMap;
 
 		// Find Unique Vertices
-		if (MeshOptions.bClampUVIslands)
+		if (!bIsSingleFullBlock && MeshOptions.bClampUVIslands)
 		{
 			VertexToFaceMap.Reserve(NumVertices);
 			Triangles.SetNumUninitialized(NumTriangles);
@@ -451,7 +452,7 @@ namespace mu
 		const uint32 MaxGridY = MeshOptions.bNormalizeUVs ? MAX_uint32 : Grid.Y - 1;
 
 		// TODO: We could skip this if there's only one block and it fits the entire grid
-		for (int32 TriangleIndex = 0; TriangleIndex < NumTriangles; ++TriangleIndex)
+		for (int32 TriangleIndex = 0;  TriangleIndex < NumTriangles; ++TriangleIndex)
 		{
 			uint32 Index0 = int32(ItIndices.GetAsUINT32());
 			++ItIndices;
@@ -486,7 +487,7 @@ namespace mu
 				BlockIndexV2 = (X < Grid.X && Y < Grid.Y) ? GridBlockBlockId[Y * Grid.X + X] : 0;
 			}
 
-			if (MeshOptions.bClampUVIslands)
+			if (!bIsSingleFullBlock && MeshOptions.bClampUVIslands)
 			{
 				if (BlockIndexV0 != BlockIndexV1 || BlockIndexV0 != BlockIndexV2)
 				{
@@ -1270,6 +1271,8 @@ namespace mu
     //---------------------------------------------------------------------------------------------
     void CodeGenerator::GenerateMesh_Constant(const FMeshGenerationOptions& InOptions, FMeshGenerationResult& OutResult, const NodeMeshConstant* InNode )
     {
+		MUTABLE_CPUPROFILER_SCOPE(GenerateMesh_Constant);
+
         NodeMeshConstant::Private& Node = *InNode->GetPrivate();
 
         Ptr<ASTOpConstantResource> ConstantOp = new ASTOpConstantResource();
@@ -1288,11 +1291,6 @@ namespace mu
 			ConstantOp->SetValue(EmptyMesh, CompilerOptions->OptimisationOptions.DiskCacheContext);
 			EmptyMesh->MeshIDPrefix = ConstantOp->GetValueHash();
 
-			FGeneratedConstantMesh MeshEntry;
-			MeshEntry.Mesh = EmptyMesh;
-			MeshEntry.LastMeshOp = ConstantOp;
-			GeneratedConstantMeshes.Add(MeshEntry);
-
 			// Log an error message
 			ErrorLog->GetPrivate()->Add("Constant mesh not set.", ELMT_WARNING, InNode->GetMessageContext());
 
@@ -1323,10 +1321,10 @@ namespace mu
 
 		// Find out if we can (or have to) reuse a mesh that we have already generated.
 		FGeneratedConstantMesh DuplicateOf;
-		for (int32 i = 0; i < GeneratedConstantMeshes.Num(); ++i)
+		uint32 ThisMeshHash = HashCombineFast( GetTypeHash(pMesh->GetVertexCount()), GetTypeHash(pMesh->GetIndexCount()) );
+		TArray<FGeneratedConstantMesh>& CachedCandidates = GeneratedConstantMeshes.FindOrAdd(ThisMeshHash,{});
+		for (const FGeneratedConstantMesh& Candidate : CachedCandidates)
 		{
-			FGeneratedConstantMesh Candidate = GeneratedConstantMeshes[i];
-			
 			bool bCompareLayouts = InOptions.bLayouts && !bIsOverridingLayouts;
 
 			if (Candidate.Mesh->IsSimilar(*pMesh, bCompareLayouts))
@@ -1420,7 +1418,7 @@ namespace mu
 			FGeneratedConstantMesh MeshEntry;
 			MeshEntry.Mesh = Cloned;
 			MeshEntry.LastMeshOp = LastMeshOp;
-			GeneratedConstantMeshes.Add(MeshEntry);
+			CachedCandidates.Add(MeshEntry);
 
 			if (InOptions.bLayouts)
 			{
