[mutable] Fixed skeletal meshes losing vertex color information after baking them.
#tests Editor
#rnx

diff --git a/Engine/Plugins/Experimental/Mutable/Source/CustomizableObject/Private/MuCO/CustomizableObjectInstance.cpp b/Engine/Plugins/Experimental/Mutable/Source/CustomizableObject/Private/MuCO/CustomizableObjectInstance.cpp
--- a/Engine/Plugins/Experimental/Mutable/Source/CustomizableObject/Private/MuCO/CustomizableObjectInstance.cpp
+++ b/Engine/Plugins/Experimental/Mutable/Source/CustomizableObject/Private/MuCO/CustomizableObjectInstance.cpp
@@ -7319,6 +7319,9 @@ void UCustomizableInstancePrivate::RegenerateImportedModels()
 				ImportedSection.SoftVertices.AddUninitialized(RenderSection.NumVertices);
 				ImportedSection.bUse16BitBoneIndex = LODRenderData.DoesVertexBufferUse16BitBoneIndex();
 
+				TArray<FColor> VertexColors;
+				LODRenderData.StaticVertexBuffers.ColorVertexBuffer.GetVertexColors(VertexColors);
+
 				for (uint32 i = 0; i < RenderSection.NumVertices; ++i)
 				{
 					const FPositionVertex* PosPtr = static_cast<const FPositionVertex*>(LODRenderData.StaticVertexBuffers.PositionVertexBuffer.GetVertexData());
@@ -7349,8 +7352,9 @@ void UCustomizableInstancePrivate::RegenerateImportedModels()
 						Vertex.InfluenceWeights[j] = 0;
 					}
 
-
-					Vertex.Color = FColor::White;
+					Vertex.Color = VertexColors.IsValidIndex(i) ?
+						VertexColors[i] :
+						FColor::White;
 
 					Vertex.Position = PosPtr->Position;
 
