[mutable] Change pointers to UTexture2D to the more generic UTexture to prepare for incoming engine changes related to pointers.
#rnx

diff --git a/Engine/Plugins/Experimental/Mutable/Source/CustomizableObject/Internal/MuCO/CustomizableObjectInstancePrivate.h b/Engine/Plugins/Experimental/Mutable/Source/CustomizableObject/Internal/MuCO/CustomizableObjectInstancePrivate.h
--- a/Engine/Plugins/Experimental/Mutable/Source/CustomizableObject/Internal/MuCO/CustomizableObjectInstancePrivate.h
+++ b/Engine/Plugins/Experimental/Mutable/Source/CustomizableObject/Internal/MuCO/CustomizableObjectInstancePrivate.h
@@ -350,18 +350,18 @@ public:
 	TMap<TSubclassOf<UAnimInstance>, FAnimBpGeneratedPhysicsAssets> AnimBpPhysicsAssets;
 
 	// The pass-through assets that will be loaded during an update
-	TArray<TSoftObjectPtr<UTexture>> PassThroughTexturesToLoad;
-	TArray<TSoftObjectPtr<UStreamableRenderAsset>> PassThroughMeshesToLoad;
+	TArray<TSoftObjectPtr<const UTexture>> PassThroughTexturesToLoad;
+	TArray<TSoftObjectPtr<const UStreamableRenderAsset>> PassThroughMeshesToLoad;
 
 	// Used during an update to prevent the pass-through textures loaded by LoadAdditionalAssetsAsync() from being unloaded by GC
 	// between AdditionalAssetsAsyncLoaded() and their setting into the generated materials in BuildMaterials()
 	UPROPERTY(Transient)
-	TArray<TObjectPtr<UTexture>> LoadedPassThroughTexturesPendingSetMaterial;
+	TArray<TObjectPtr<const UTexture>> LoadedPassThroughTexturesPendingSetMaterial;
 
 	// Used during an update to prevent the pass-through meshes loaded by LoadAdditionalAssetsAsync() from being unloaded by GC
 	// between AdditionalAssetsAsyncLoaded() and their setting into the generated materials in BuildMaterials()
 	UPROPERTY(Transient)
-	TArray<TObjectPtr<UStreamableRenderAsset>> LoadedPassThroughMeshesPendingSetMaterial;
+	TArray<TObjectPtr<const UStreamableRenderAsset>> LoadedPassThroughMeshesPendingSetMaterial;
 
 private:
 	ECOInstanceFlags InstanceFlagsPrivate = ECOInstanceFlags::ECONone;
diff --git a/Engine/Plugins/Experimental/Mutable/Source/CustomizableObject/Internal/MuCO/CustomizableObjectPrivate.h b/Engine/Plugins/Experimental/Mutable/Source/CustomizableObject/Internal/MuCO/CustomizableObjectPrivate.h
--- a/Engine/Plugins/Experimental/Mutable/Source/CustomizableObject/Internal/MuCO/CustomizableObjectPrivate.h
+++ b/Engine/Plugins/Experimental/Mutable/Source/CustomizableObject/Internal/MuCO/CustomizableObjectPrivate.h
@@ -727,7 +727,7 @@ struct FModelResources
 #if WITH_EDITORONLY_DATA
 	/** Runtime referenced textures used by the mu::Model. */
 	UPROPERTY()
-	TArray<TSoftObjectPtr<UTexture2D>> RuntimeReferencedTextures;
+	TArray<TSoftObjectPtr<const UTexture>> RuntimeReferencedTextures;
 #endif
 	
 	/** Physics assets gathered from the SkeletalMeshes, to be used in mesh generation in-game */
diff --git a/Engine/Plugins/Experimental/Mutable/Source/CustomizableObject/Private/MuCO/BusyWaits_Deprecated.cpp b/Engine/Plugins/Experimental/Mutable/Source/CustomizableObject/Private/MuCO/BusyWaits_Deprecated.cpp
--- a/Engine/Plugins/Experimental/Mutable/Source/CustomizableObject/Private/MuCO/BusyWaits_Deprecated.cpp
+++ b/Engine/Plugins/Experimental/Mutable/Source/CustomizableObject/Private/MuCO/BusyWaits_Deprecated.cpp
@@ -246,7 +246,7 @@ namespace CustomizableObjectSystem::ImplDeprecated
 
 							if (ModelResources.PassThroughMeshes.IsValidIndex(ReferenceID))
 							{
-								TSoftObjectPtr<USkeletalMesh> Ref = ModelResources.PassThroughMeshes[ReferenceID];
+								TSoftObjectPtr<const USkeletalMesh> Ref = ModelResources.PassThroughMeshes[ReferenceID];
 								CustomizableObjectInstancePrivateData->PassThroughMeshesToLoad.Add(Ref);
 							}
 							else
@@ -337,7 +337,7 @@ namespace CustomizableObjectSystem::ImplDeprecated
 
 									if (ModelResources.PassThroughTextures.IsValidIndex(ReferenceID))
 									{
-										TSoftObjectPtr<UTexture> Ref = ModelResources.PassThroughTextures[ReferenceID];
+										TSoftObjectPtr<const UTexture> Ref = ModelResources.PassThroughTextures[ReferenceID];
 										CustomizableObjectInstancePrivateData->PassThroughTexturesToLoad.Add(Ref);
 									}
 									else
diff --git a/Engine/Plugins/Experimental/Mutable/Source/CustomizableObject/Private/MuCO/CustomizableObject.cpp b/Engine/Plugins/Experimental/Mutable/Source/CustomizableObject/Private/MuCO/CustomizableObject.cpp
--- a/Engine/Plugins/Experimental/Mutable/Source/CustomizableObject/Private/MuCO/CustomizableObject.cpp
+++ b/Engine/Plugins/Experimental/Mutable/Source/CustomizableObject/Private/MuCO/CustomizableObject.cpp
@@ -558,7 +558,7 @@ void UCustomizableObjectPrivate::SaveCompiledData(FArchive& MemoryWriter, bool b
 	int32 NumReferencedMaterials = LocalModelResources.Materials.Num();
 	MemoryWriter << NumReferencedMaterials;
 
-	for (const TSoftObjectPtr<UMaterialInterface>& Material : LocalModelResources.Materials)
+	for (const TSoftObjectPtr<const UMaterialInterface>& Material : LocalModelResources.Materials)
 	{
 		FString StringRef = Material.ToString();
 		MemoryWriter << StringRef;
@@ -567,7 +567,7 @@ void UCustomizableObjectPrivate::SaveCompiledData(FArchive& MemoryWriter, bool b
 	int32 NumReferencedSkeletons = LocalModelResources.Skeletons.Num();
 	MemoryWriter << NumReferencedSkeletons;
 
-	for (const TSoftObjectPtr<USkeleton>& Skeleton : LocalModelResources.Skeletons)
+	for (const TSoftObjectPtr<const USkeleton>& Skeleton : LocalModelResources.Skeletons)
 	{
 		FString StringRef = Skeleton.ToString();
 		MemoryWriter << StringRef;
@@ -576,7 +576,7 @@ void UCustomizableObjectPrivate::SaveCompiledData(FArchive& MemoryWriter, bool b
 	int32 NumPassthroughTextures = LocalModelResources.PassThroughTextures.Num();
 	MemoryWriter << NumPassthroughTextures;
 
-	for (const TSoftObjectPtr<UTexture>& PassthroughTexture : LocalModelResources.PassThroughTextures)
+	for (const TSoftObjectPtr<const UTexture>& PassthroughTexture : LocalModelResources.PassThroughTextures)
 	{
 		FString StringRef = PassthroughTexture.ToString();
 		MemoryWriter << StringRef;
@@ -585,7 +585,7 @@ void UCustomizableObjectPrivate::SaveCompiledData(FArchive& MemoryWriter, bool b
 	int32 NumPassthroughMeshes = LocalModelResources.PassThroughMeshes.Num();
 	MemoryWriter << NumPassthroughMeshes;
 
-	for (const TSoftObjectPtr<USkeletalMesh>& PassthroughMesh : LocalModelResources.PassThroughMeshes)
+	for (const TSoftObjectPtr<const USkeletalMesh>& PassthroughMesh : LocalModelResources.PassThroughMeshes)
 	{
 		FString StringRef = PassthroughMesh.ToString();
 		MemoryWriter << StringRef;
@@ -595,7 +595,7 @@ void UCustomizableObjectPrivate::SaveCompiledData(FArchive& MemoryWriter, bool b
 	int32 NumRuntimeReferencedTextures = LocalModelResources.RuntimeReferencedTextures.Num();
 	MemoryWriter << NumRuntimeReferencedTextures;
 	
-	for (const TSoftObjectPtr<UTexture2D>& RuntimeReferencedTexture : LocalModelResources.RuntimeReferencedTextures)
+	for (const TSoftObjectPtr<const UTexture>& RuntimeReferencedTexture : LocalModelResources.RuntimeReferencedTextures)
 	{
 		FString StringRef = RuntimeReferencedTexture.ToString();
 		MemoryWriter << StringRef;
@@ -605,7 +605,7 @@ void UCustomizableObjectPrivate::SaveCompiledData(FArchive& MemoryWriter, bool b
 	int32 NumPhysicsAssets = LocalModelResources.PhysicsAssets.Num();
 	MemoryWriter << NumPhysicsAssets;
 
-	for (const TSoftObjectPtr<UPhysicsAsset>& PhysicsAsset : LocalModelResources.PhysicsAssets)
+	for (const TSoftObjectPtr<const UPhysicsAsset>& PhysicsAsset : LocalModelResources.PhysicsAssets)
 	{
 		FString StringRef = PhysicsAsset.ToString();
 		MemoryWriter << StringRef;
@@ -614,7 +614,7 @@ void UCustomizableObjectPrivate::SaveCompiledData(FArchive& MemoryWriter, bool b
 	int32 NumAnimBps = LocalModelResources.AnimBPs.Num();
 	MemoryWriter << NumAnimBps;
 
-	for (const TSoftClassPtr<UAnimInstance>& AnimBp : LocalModelResources.AnimBPs)
+	for (const TSoftClassPtr<const UAnimInstance>& AnimBp : LocalModelResources.AnimBPs)
 	{
 		FString StringRef = AnimBp.ToString();
 		MemoryWriter << StringRef;
diff --git a/Engine/Plugins/Experimental/Mutable/Source/CustomizableObject/Private/MuCO/CustomizableObjectInstance.cpp b/Engine/Plugins/Experimental/Mutable/Source/CustomizableObject/Private/MuCO/CustomizableObjectInstance.cpp
--- a/Engine/Plugins/Experimental/Mutable/Source/CustomizableObject/Private/MuCO/CustomizableObjectInstance.cpp
+++ b/Engine/Plugins/Experimental/Mutable/Source/CustomizableObject/Private/MuCO/CustomizableObjectInstance.cpp
@@ -5385,12 +5385,12 @@ UE::Tasks::FTask UCustomizableInstancePrivate::LoadAdditionalAssetsAndData(
 		}
 	}
 
-	for (TSoftObjectPtr<UTexture>& TextureRef : PassThroughTexturesToLoad)
+	for (TSoftObjectPtr<const UTexture>& TextureRef : PassThroughTexturesToLoad)
 	{
 		AssetsToStream.Add(TextureRef.ToSoftObjectPath());
 	}
 
-	for (TSoftObjectPtr<UStreamableRenderAsset>& MeshRef : PassThroughMeshesToLoad)
+	for (TSoftObjectPtr<const UStreamableRenderAsset>& MeshRef : PassThroughMeshesToLoad)
 	{
 		AssetsToStream.Add(MeshRef.ToSoftObjectPath());
 	}
@@ -5843,7 +5843,7 @@ void UCustomizableInstancePrivate::AdditionalAssetsAsyncLoaded(UCustomizableObje
 
 	LoadedPassThroughTexturesPendingSetMaterial.Empty(PassThroughTexturesToLoad.Num());
 
-	for (TSoftObjectPtr<UTexture>& TextureRef : PassThroughTexturesToLoad)
+	for (TSoftObjectPtr<const UTexture>& TextureRef : PassThroughTexturesToLoad)
 	{
 		ensure(TextureRef.IsValid());
 		LoadedPassThroughTexturesPendingSetMaterial.Add(TextureRef.Get());
@@ -5853,7 +5853,7 @@ void UCustomizableInstancePrivate::AdditionalAssetsAsyncLoaded(UCustomizableObje
 
 	LoadedPassThroughMeshesPendingSetMaterial.Empty(PassThroughMeshesToLoad.Num());
 
-	for (TSoftObjectPtr<UStreamableRenderAsset>& MeshRef : PassThroughMeshesToLoad)
+	for (TSoftObjectPtr<const UStreamableRenderAsset>& MeshRef : PassThroughMeshesToLoad)
 	{
 		ensure(MeshRef.IsValid());
 		LoadedPassThroughMeshesPendingSetMaterial.Add(MeshRef.Get());
diff --git a/Engine/Plugins/Experimental/Mutable/Source/CustomizableObject/Private/MuCO/CustomizableObjectSystem.cpp b/Engine/Plugins/Experimental/Mutable/Source/CustomizableObject/Private/MuCO/CustomizableObjectSystem.cpp
--- a/Engine/Plugins/Experimental/Mutable/Source/CustomizableObject/Private/MuCO/CustomizableObjectSystem.cpp
+++ b/Engine/Plugins/Experimental/Mutable/Source/CustomizableObject/Private/MuCO/CustomizableObjectSystem.cpp
@@ -2136,7 +2136,7 @@ namespace impl
 
 		if (ModelResources.PassThroughTextures.IsValidIndex(ReferenceID))
 		{
-			const TSoftObjectPtr<UTexture> Ref = ModelResources.PassThroughTextures[ReferenceID];
+			const TSoftObjectPtr<const UTexture> Ref = ModelResources.PassThroughTextures[ReferenceID];
 			Instance->GetPrivate()->PassThroughTexturesToLoad.Add(Ref);
 		}
 		else
@@ -2394,7 +2394,7 @@ namespace impl
 
 			if (ModelResources.PassThroughMeshes.IsValidIndex(ReferenceID))
 			{
-				TSoftObjectPtr<USkeletalMesh> Ref = ModelResources.PassThroughMeshes[ReferenceID];
+				TSoftObjectPtr<const USkeletalMesh> Ref = ModelResources.PassThroughMeshes[ReferenceID];
 				Instance->GetPrivate()->PassThroughMeshesToLoad.Add(Ref);
 			}
 			else
diff --git a/Engine/Plugins/Experimental/Mutable/Source/CustomizableObject/Private/MuCO/UnrealMutableImageProvider.cpp b/Engine/Plugins/Experimental/Mutable/Source/CustomizableObject/Private/MuCO/UnrealMutableImageProvider.cpp
--- a/Engine/Plugins/Experimental/Mutable/Source/CustomizableObject/Private/MuCO/UnrealMutableImageProvider.cpp
+++ b/Engine/Plugins/Experimental/Mutable/Source/CustomizableObject/Private/MuCO/UnrealMutableImageProvider.cpp
@@ -729,7 +729,7 @@ void FUnrealMutableResourceProvider::UnCacheImages(const mu::Parameters& Paramet
 
 
 #if WITH_EDITOR
-void FUnrealMutableResourceProvider::CacheRuntimeReferencedImages(const TSharedRef<const mu::Model>& Model, const TArray<TSoftObjectPtr<UTexture2D>>& RuntimeReferencedTextures)
+void FUnrealMutableResourceProvider::CacheRuntimeReferencedImages(const TSharedRef<const mu::Model>& Model, const TArray<TSoftObjectPtr<const UTexture>>& RuntimeReferencedTextures)
 {
 	check(IsInGameThread());
 	
@@ -741,9 +741,9 @@ void FUnrealMutableResourceProvider::CacheRuntimeReferencedImages(const TSharedR
 	ModelImages.Model = Model.ToWeakPtr();
 
 	ModelImages.SourceTextures.Reset();
-	for (const TSoftObjectPtr<UTexture2D>& RuntimeReferencedTexture : RuntimeReferencedTextures)
+	for (const TSoftObjectPtr<const UTexture>& RuntimeReferencedTexture : RuntimeReferencedTextures)
 	{
-		UTexture2D* Texture = RuntimeReferencedTexture.LoadSynchronous();
+		const UTexture* Texture = RuntimeReferencedTexture.LoadSynchronous();
 		if (!Texture)
 		{
 			UE_LOG(LogMutable, Warning, TEXT("Failed to load texture [%s]."), *RuntimeReferencedTexture->GetPathName());
diff --git a/Engine/Plugins/Experimental/Mutable/Source/CustomizableObject/Private/MuCO/UnrealMutableImageProvider.h b/Engine/Plugins/Experimental/Mutable/Source/CustomizableObject/Private/MuCO/UnrealMutableImageProvider.h
--- a/Engine/Plugins/Experimental/Mutable/Source/CustomizableObject/Private/MuCO/UnrealMutableImageProvider.h
+++ b/Engine/Plugins/Experimental/Mutable/Source/CustomizableObject/Private/MuCO/UnrealMutableImageProvider.h
@@ -49,7 +49,7 @@ public:
 	void UnCacheImages(const mu::Parameters& Parameters);
 
 #if WITH_EDITOR
-	void CacheRuntimeReferencedImages(const TSharedRef<const mu::Model>& Model, const TArray<TSoftObjectPtr<UTexture2D>>& RuntimeReferencedTextures);
+	void CacheRuntimeReferencedImages(const TSharedRef<const mu::Model>& Model, const TArray<TSoftObjectPtr<const UTexture>>& RuntimeReferencedTextures);
 #endif
 	
 	/** List of actual image providers that have been registered to the CustomizableObjectSystem. */
diff --git a/Engine/Plugins/Experimental/Mutable/Source/CustomizableObjectEditor/Private/MuCOE/CustomizableObjectCompiler.cpp b/Engine/Plugins/Experimental/Mutable/Source/CustomizableObjectEditor/Private/MuCOE/CustomizableObjectCompiler.cpp
--- a/Engine/Plugins/Experimental/Mutable/Source/CustomizableObjectEditor/Private/MuCOE/CustomizableObjectCompiler.cpp
+++ b/Engine/Plugins/Experimental/Mutable/Source/CustomizableObjectEditor/Private/MuCOE/CustomizableObjectCompiler.cpp
@@ -953,7 +953,7 @@ void FCustomizableObjectCompiler::CompileInternal(bool bAsync)
 			ModelResources.PassThroughTextures.Add(Pair.Key);
 		}
 
-		for (const TPair<TSoftObjectPtr<UTexture2D>, FMutableGraphGenerationContext::FGeneratedReferencedTexture>& Pair : GenerationContext.RuntimeReferencedTextureMap)
+		for (const TPair<TSoftObjectPtr<const UTexture>, FMutableGraphGenerationContext::FGeneratedReferencedTexture>& Pair : GenerationContext.RuntimeReferencedTextureMap)
 		{
 			check(Pair.Value.ID == ModelResources.RuntimeReferencedTextures.Num());
 			ModelResources.RuntimeReferencedTextures.Add(Pair.Key);
@@ -1158,7 +1158,7 @@ void FCustomizableObjectCompiler::CompileInternal(bool bAsync)
 
 		// Pass-through textures
 		TArray<FMutableSourceTextureData> NewCompileTimeReferencedTextures;
-		for (const TPair<TSoftObjectPtr<UTexture>, FMutableGraphGenerationContext::FGeneratedReferencedTexture>& Pair : GenerationContext.CompileTimeTextureMap)
+		for (const TPair<TSoftObjectPtr<const UTexture>, FMutableGraphGenerationContext::FGeneratedReferencedTexture>& Pair : GenerationContext.CompileTimeTextureMap)
 		{
 			check(Pair.Value.ID == NewCompileTimeReferencedTextures.Num());
 
@@ -1328,7 +1328,7 @@ bool FCustomizableObjectCompiler::TryPopCompileRequest()
 
 
 mu::NodePtr FCustomizableObjectCompiler::Export(UCustomizableObject* Object, const FCompilationOptions& InCompilerOptions, 
-	TArray<TSoftObjectPtr<UTexture>>& OutRuntimeReferencedTextures,
+	TArray<TSoftObjectPtr<const UTexture>>& OutRuntimeReferencedTextures,
 	TArray<FMutableSourceTextureData>& OutCompilerReferencedTextures )
 {
 	UE_LOG(LogMutable, Log, TEXT("Started Customizable Object Export %s."), *Object->GetName());
@@ -1374,14 +1374,14 @@ mu::NodePtr FCustomizableObjectCompiler::Export(UCustomizableObject* Object, con
 
 	// Pass out the references textures
 	OutRuntimeReferencedTextures.Empty();
-	for (const TPair<TSoftObjectPtr<UTexture2D>, FMutableGraphGenerationContext::FGeneratedReferencedTexture>& Pair : GenerationContext.RuntimeReferencedTextureMap)
+	for (const TPair<TSoftObjectPtr<const UTexture>, FMutableGraphGenerationContext::FGeneratedReferencedTexture>& Pair : GenerationContext.RuntimeReferencedTextureMap)
 	{
 		check(Pair.Value.ID == OutRuntimeReferencedTextures.Num());
 		OutRuntimeReferencedTextures.Add(Pair.Key);
 	}
 
 	OutCompilerReferencedTextures.Empty();
-	for (const TPair<TSoftObjectPtr<UTexture>, FMutableGraphGenerationContext::FGeneratedReferencedTexture>& Pair : GenerationContext.CompileTimeTextureMap)
+	for (const TPair<TSoftObjectPtr<const UTexture>, FMutableGraphGenerationContext::FGeneratedReferencedTexture>& Pair : GenerationContext.CompileTimeTextureMap)
 	{
 		check(Pair.Value.ID == OutCompilerReferencedTextures.Num());
 
diff --git a/Engine/Plugins/Experimental/Mutable/Source/CustomizableObjectEditor/Private/MuCOE/CustomizableObjectCompiler.h b/Engine/Plugins/Experimental/Mutable/Source/CustomizableObjectEditor/Private/MuCOE/CustomizableObjectCompiler.h
--- a/Engine/Plugins/Experimental/Mutable/Source/CustomizableObjectEditor/Private/MuCOE/CustomizableObjectCompiler.h
+++ b/Engine/Plugins/Experimental/Mutable/Source/CustomizableObjectEditor/Private/MuCOE/CustomizableObjectCompiler.h
@@ -42,7 +42,7 @@ public:
 
 
 	/** Generate the Mutable Graph from the Unreal Graph. */
-	mu::NodePtr Export(UCustomizableObject* Object, const FCompilationOptions& Options, TArray<TSoftObjectPtr<UTexture>>& OutRuntimeReferencedTextures, TArray<FMutableSourceTextureData>& OutCompilerReferencedTextures);
+	mu::NodePtr Export(UCustomizableObject* Object, const FCompilationOptions& Options, TArray<TSoftObjectPtr<const UTexture>>& OutRuntimeReferencedTextures, TArray<FMutableSourceTextureData>& OutCompilerReferencedTextures);
 
 	void CompilerLog(const FText& Message, const TArray<const UObject*>& UObject, const EMessageSeverity::Type MessageSeverity = EMessageSeverity::Warning, const bool bAddBaseObjectInfo = true, const ELoggerSpamBin SpamBin = ELoggerSpamBin::ShowAll);
 	void CompilerLog(const FText& Message, const UObject* Context = nullptr, const EMessageSeverity::Type MessageSeverity = EMessageSeverity::Warning, const bool bAddBaseObjectInfo = true, const ELoggerSpamBin SpamBin = ELoggerSpamBin::ShowAll);
diff --git a/Engine/Plugins/Experimental/Mutable/Source/CustomizableObjectEditor/Private/MuCOE/GenerateMutableSource/GenerateMutableSource.h b/Engine/Plugins/Experimental/Mutable/Source/CustomizableObjectEditor/Private/MuCOE/GenerateMutableSource/GenerateMutableSource.h
--- a/Engine/Plugins/Experimental/Mutable/Source/CustomizableObjectEditor/Private/MuCOE/GenerateMutableSource/GenerateMutableSource.h
+++ b/Engine/Plugins/Experimental/Mutable/Source/CustomizableObjectEditor/Private/MuCOE/GenerateMutableSource/GenerateMutableSource.h
@@ -630,11 +630,11 @@ struct FMutableGraphGenerationContext
 
 	// Cache of runtime images and their IDs used in the core to indentify them.
 	// These textures will remain as external references even in optimized models.
-	TMap<TSoftObjectPtr<UTexture2D>, FGeneratedReferencedTexture> RuntimeReferencedTextureMap;
+	TMap<TSoftObjectPtr<const UTexture>, FGeneratedReferencedTexture> RuntimeReferencedTextureMap;
 	
 	// Cache of runtime pass-through images and their IDs used in the core to indentify them
 	// These textures will become mutable images in the compiled model.
-	TMap<TSoftObjectPtr<UTexture>, FGeneratedReferencedTexture> CompileTimeTextureMap;
+	TMap<TSoftObjectPtr<const UTexture>, FGeneratedReferencedTexture> CompileTimeTextureMap;
 
     // Global morph selection overrides.
     TArray<FRealTimeMorphSelectionOverride> RealTimeMorphTargetsOverrides;
diff --git a/Engine/Plugins/Experimental/Mutable/Source/CustomizableObjectEditor/Private/MuCOE/SMutableCodeViewer.cpp b/Engine/Plugins/Experimental/Mutable/Source/CustomizableObjectEditor/Private/MuCOE/SMutableCodeViewer.cpp
--- a/Engine/Plugins/Experimental/Mutable/Source/CustomizableObjectEditor/Private/MuCOE/SMutableCodeViewer.cpp
+++ b/Engine/Plugins/Experimental/Mutable/Source/CustomizableObjectEditor/Private/MuCOE/SMutableCodeViewer.cpp
@@ -291,7 +291,7 @@ void SMutableCodeViewer::ClearSelectedTreeRow() const
 	TreeView->ClearSelection();
 }
 
-void SMutableCodeViewer::SetCurrentModel(const TSharedPtr<mu::Model, ESPMode::ThreadSafe>& InMutableModel, const TArray<TSoftObjectPtr<UTexture>>& InReferencedTextures)
+void SMutableCodeViewer::SetCurrentModel(const TSharedPtr<mu::Model, ESPMode::ThreadSafe>& InMutableModel, const TArray<TSoftObjectPtr<const UTexture>>& InReferencedTextures)
 {
 	MutableModel = InMutableModel;
 	ReferencedTextures = InReferencedTextures;
@@ -340,7 +340,7 @@ void SMutableCodeViewer::SetCurrentModel(const TSharedPtr<mu::Model, ESPMode::Th
 
 
 void SMutableCodeViewer::Construct(const FArguments& InArgs, const TSharedPtr<mu::Model, ESPMode::ThreadSafe>& InMutableModel,
-	const TArray<TSoftObjectPtr<UTexture>>& InReferencedTextures )
+	const TArray<TSoftObjectPtr<const UTexture>>& InReferencedTextures )
 {
 	// Min width allowed for the column. Needed to avoid having issues with the constants space being to small
 	// and then getting too tall on the y axis crashing the UI drawer.
@@ -2437,7 +2437,7 @@ namespace
 	public:
 
 		/** */
-		TArray<TSoftObjectPtr<UTexture>> ReferencedTextures;
+		TArray<TSoftObjectPtr<const UTexture>> ReferencedTextures;
 
 	public:
 
diff --git a/Engine/Plugins/Experimental/Mutable/Source/CustomizableObjectEditor/Private/MuCOE/SMutableCodeViewer.h b/Engine/Plugins/Experimental/Mutable/Source/CustomizableObjectEditor/Private/MuCOE/SMutableCodeViewer.h
--- a/Engine/Plugins/Experimental/Mutable/Source/CustomizableObjectEditor/Private/MuCOE/SMutableCodeViewer.h
+++ b/Engine/Plugins/Experimental/Mutable/Source/CustomizableObjectEditor/Private/MuCOE/SMutableCodeViewer.h
@@ -59,7 +59,7 @@ public:
 	SLATE_END_ARGS()
 
 	void Construct(const FArguments& InArgs, const TSharedPtr<mu::Model, ESPMode::ThreadSafe>& InMutableModel,
-		const TArray<TSoftObjectPtr<UTexture>>& ReferencedTextures );
+		const TArray<TSoftObjectPtr<const UTexture>>& ReferencedTextures );
 
 	// SWidget interface
 	virtual void Tick(const FGeometry& AllottedGeometry, const double InCurrentTime, const float InDeltaTime) override;
@@ -77,7 +77,7 @@ private:
 	TSharedPtr<mu::Model, ESPMode::ThreadSafe> MutableModel;
 	
 	/** Array of external referenced textures in MutableModel, indexed by id. */
-	TArray<TSoftObjectPtr<UTexture>> ReferencedTextures;
+	TArray<TSoftObjectPtr<const UTexture>> ReferencedTextures;
 
 	/** Selected model operation for preview. */
 	mu::OP::ADDRESS SelectedOperationAddress = 0;
@@ -182,7 +182,7 @@ private:
 	TMap< mu::OP::ADDRESS, TSharedPtr<FMutableCodeTreeElement>> ExpandedElements;
 
 	/** Prepare the widget for the given model. */
-	void SetCurrentModel(const TSharedPtr<mu::Model, ESPMode::ThreadSafe>&, const TArray<TSoftObjectPtr<UTexture>>& ReferencedTextures);
+	void SetCurrentModel(const TSharedPtr<mu::Model, ESPMode::ThreadSafe>&, const TArray<TSoftObjectPtr<const UTexture>>& ReferencedTextures);
 
 	/** Before any UI operation generate all the elements that may be navigable over the tree. No children of duplicated 
 	 * addresses will be generated.
diff --git a/Engine/Plugins/Experimental/Mutable/Source/CustomizableObjectEditor/Private/MuCOE/SMutableGraphViewer.h b/Engine/Plugins/Experimental/Mutable/Source/CustomizableObjectEditor/Private/MuCOE/SMutableGraphViewer.h
--- a/Engine/Plugins/Experimental/Mutable/Source/CustomizableObjectEditor/Private/MuCOE/SMutableGraphViewer.h
+++ b/Engine/Plugins/Experimental/Mutable/Source/CustomizableObjectEditor/Private/MuCOE/SMutableGraphViewer.h
@@ -35,7 +35,7 @@ public:
 	/** User-visible tag to indentify the source of the data shown. */
 	SLATE_ARGUMENT(FString, DataTag)
 
-	SLATE_ARGUMENT(TArray<TSoftObjectPtr<UTexture>>, ReferencedRuntimeTextures)
+	SLATE_ARGUMENT(TArray<TSoftObjectPtr<const UTexture>>, ReferencedRuntimeTextures)
 	SLATE_ARGUMENT(TArray<FMutableSourceTextureData>, ReferencedCompileTextures)
 
 	SLATE_END_ARGS()
@@ -55,7 +55,7 @@ private:
 	mu::NodePtr RootNode;
 
 	/** Array of external referenced textures in MutableModel, indexed by id. */
-	TArray<TSoftObjectPtr<UTexture>> ReferencedRuntimeTextures;
+	TArray<TSoftObjectPtr<const UTexture>> ReferencedRuntimeTextures;
 	TArray<FMutableSourceTextureData> ReferencedCompileTextures;
 
 	/** Object compiler. */
diff --git a/Engine/Plugins/Experimental/Mutable/Source/CustomizableObjectEditor/Private/MuCOE/SMutableObjectViewer.cpp b/Engine/Plugins/Experimental/Mutable/Source/CustomizableObjectEditor/Private/MuCOE/SMutableObjectViewer.cpp
--- a/Engine/Plugins/Experimental/Mutable/Source/CustomizableObjectEditor/Private/MuCOE/SMutableObjectViewer.cpp
+++ b/Engine/Plugins/Experimental/Mutable/Source/CustomizableObjectEditor/Private/MuCOE/SMutableObjectViewer.cpp
@@ -108,7 +108,7 @@ void SMutableObjectViewer::GenerateMutableGraphPressed()
 {
 	// Convert from Unreal graph to Mutable graph.
 	FCustomizableObjectCompiler Compiler;
-	TArray<TSoftObjectPtr<UTexture>> RuntimeTextures;
+	TArray<TSoftObjectPtr<const UTexture>> RuntimeTextures;
 	TArray<FMutableSourceTextureData> CompilerTextures;
 	mu::Ptr<mu::Node> RootNode = Compiler.Export(CustomizableObject, CompileOptions, RuntimeTextures, CompilerTextures);
 	if (!RootNode)
@@ -131,7 +131,7 @@ void SMutableObjectViewer::GenerateMutableGraphPressed()
 
 void SMutableObjectViewer::CompileMutableCodePressed()
 {
-	TArray<TSoftObjectPtr<UTexture>> RuntimeTextures;
+	TArray<TSoftObjectPtr<const UTexture>> RuntimeTextures;
 	TArray<FMutableSourceTextureData> CompilerTextures;
 	if (CompileOptions.bForceLargeLODBias)
 	{
