CIS fix

diff --git a/Engine/Plugins/Experimental/Mutable/Source/CustomizableObject/Private/MuCO/CustomizableObject.cpp b/Engine/Plugins/Experimental/Mutable/Source/CustomizableObject/Private/MuCO/CustomizableObject.cpp
--- a/Engine/Plugins/Experimental/Mutable/Source/CustomizableObject/Private/MuCO/CustomizableObject.cpp
+++ b/Engine/Plugins/Experimental/Mutable/Source/CustomizableObject/Private/MuCO/CustomizableObject.cpp
@@ -558,7 +558,7 @@ void UCustomizableObjectPrivate::SaveCompiledData(FArchive& MemoryWriter, bool b
 	int32 NumReferencedMaterials = LocalModelResources.Materials.Num();
 	MemoryWriter << NumReferencedMaterials;
 
-	for (const TSoftObjectPtr<const UMaterialInterface>& Material : LocalModelResources.Materials)
+	for (const TSoftObjectPtr<UMaterialInterface>& Material : LocalModelResources.Materials)
 	{
 		FString StringRef = Material.ToString();
 		MemoryWriter << StringRef;
@@ -567,7 +567,7 @@ void UCustomizableObjectPrivate::SaveCompiledData(FArchive& MemoryWriter, bool b
 	int32 NumReferencedSkeletons = LocalModelResources.Skeletons.Num();
 	MemoryWriter << NumReferencedSkeletons;
 
-	for (const TSoftObjectPtr<const USkeleton>& Skeleton : LocalModelResources.Skeletons)
+	for (const TSoftObjectPtr<USkeleton>& Skeleton : LocalModelResources.Skeletons)
 	{
 		FString StringRef = Skeleton.ToString();
 		MemoryWriter << StringRef;
@@ -576,7 +576,7 @@ void UCustomizableObjectPrivate::SaveCompiledData(FArchive& MemoryWriter, bool b
 	int32 NumPassthroughTextures = LocalModelResources.PassThroughTextures.Num();
 	MemoryWriter << NumPassthroughTextures;
 
-	for (const TSoftObjectPtr<const UTexture>& PassthroughTexture : LocalModelResources.PassThroughTextures)
+	for (const TSoftObjectPtr<UTexture>& PassthroughTexture : LocalModelResources.PassThroughTextures)
 	{
 		FString StringRef = PassthroughTexture.ToString();
 		MemoryWriter << StringRef;
@@ -585,7 +585,7 @@ void UCustomizableObjectPrivate::SaveCompiledData(FArchive& MemoryWriter, bool b
 	int32 NumPassthroughMeshes = LocalModelResources.PassThroughMeshes.Num();
 	MemoryWriter << NumPassthroughMeshes;
 
-	for (const TSoftObjectPtr<const USkeletalMesh>& PassthroughMesh : LocalModelResources.PassThroughMeshes)
+	for (const TSoftObjectPtr<USkeletalMesh>& PassthroughMesh : LocalModelResources.PassThroughMeshes)
 	{
 		FString StringRef = PassthroughMesh.ToString();
 		MemoryWriter << StringRef;
@@ -605,7 +605,7 @@ void UCustomizableObjectPrivate::SaveCompiledData(FArchive& MemoryWriter, bool b
 	int32 NumPhysicsAssets = LocalModelResources.PhysicsAssets.Num();
 	MemoryWriter << NumPhysicsAssets;
 
-	for (const TSoftObjectPtr<const UPhysicsAsset>& PhysicsAsset : LocalModelResources.PhysicsAssets)
+	for (const TSoftObjectPtr<UPhysicsAsset>& PhysicsAsset : LocalModelResources.PhysicsAssets)
 	{
 		FString StringRef = PhysicsAsset.ToString();
 		MemoryWriter << StringRef;
@@ -614,7 +614,7 @@ void UCustomizableObjectPrivate::SaveCompiledData(FArchive& MemoryWriter, bool b
 	int32 NumAnimBps = LocalModelResources.AnimBPs.Num();
 	MemoryWriter << NumAnimBps;
 
-	for (const TSoftClassPtr<const UAnimInstance>& AnimBp : LocalModelResources.AnimBPs)
+	for (const TSoftClassPtr<UAnimInstance>& AnimBp : LocalModelResources.AnimBPs)
 	{
 		FString StringRef = AnimBp.ToString();
 		MemoryWriter << StringRef;
