[mutable] More const fixes for the incoming engine pointer changes.
#rnx

diff --git a/Engine/Plugins/Experimental/Mutable/Source/CustomizableObjectEditor/Private/MuCOE/GenerateMutableSource/GenerateMutableSource.cpp b/Engine/Plugins/Experimental/Mutable/Source/CustomizableObjectEditor/Private/MuCOE/GenerateMutableSource/GenerateMutableSource.cpp
--- a/Engine/Plugins/Experimental/Mutable/Source/CustomizableObjectEditor/Private/MuCOE/GenerateMutableSource/GenerateMutableSource.cpp
+++ b/Engine/Plugins/Experimental/Mutable/Source/CustomizableObjectEditor/Private/MuCOE/GenerateMutableSource/GenerateMutableSource.cpp
@@ -1986,7 +1986,7 @@ mu::Ptr<mu::Image> GenerateImageConstant(UTexture* Texture, FMutableGraphGenerat
 }
 
 
-mu::Ptr<mu::Mesh> GenerateMeshConstant(const USkeletalMesh* Mesh, FMutableGraphGenerationContext& GenerationContext, bool bIsReference)
+mu::Ptr<mu::Mesh> GenerateMeshConstant(USkeletalMesh* Mesh, FMutableGraphGenerationContext& GenerationContext, bool bIsReference)
 {
 	if (!Mesh)
 	{
diff --git a/Engine/Plugins/Experimental/Mutable/Source/CustomizableObjectEditor/Private/MuCOE/GenerateMutableSource/GenerateMutableSource.h b/Engine/Plugins/Experimental/Mutable/Source/CustomizableObjectEditor/Private/MuCOE/GenerateMutableSource/GenerateMutableSource.h
--- a/Engine/Plugins/Experimental/Mutable/Source/CustomizableObjectEditor/Private/MuCOE/GenerateMutableSource/GenerateMutableSource.h
+++ b/Engine/Plugins/Experimental/Mutable/Source/CustomizableObjectEditor/Private/MuCOE/GenerateMutableSource/GenerateMutableSource.h
@@ -1014,7 +1014,7 @@ int32 GetMaxTextureSize(const UTexture2D& ReferenceTexture, const UTextureLODSet
 int32 GetTextureSizeInGame(const UTexture2D& Texture, const UTextureLODSettings& LODSettings, uint8 SurfaceLODBias = 0);
 
 mu::Ptr<mu::Image> GenerateImageConstant(UTexture*, FMutableGraphGenerationContext&, bool bIsReference);
-mu::Ptr<mu::Mesh> GenerateMeshConstant(const USkeletalMesh*, FMutableGraphGenerationContext&, bool bIsReference);
+mu::Ptr<mu::Mesh> GenerateMeshConstant(USkeletalMesh*, FMutableGraphGenerationContext&, bool bIsReference);
 
 /** Generates a mutable image descriptor from an unreal engine texture */
 mu::FImageDesc GenerateImageDescriptor(UTexture* Texture);
diff --git a/Engine/Plugins/Experimental/Mutable/Source/CustomizableObjectEditor/Private/MuCOE/GenerateMutableSource/GenerateMutableSourceMesh.cpp b/Engine/Plugins/Experimental/Mutable/Source/CustomizableObjectEditor/Private/MuCOE/GenerateMutableSource/GenerateMutableSourceMesh.cpp
--- a/Engine/Plugins/Experimental/Mutable/Source/CustomizableObjectEditor/Private/MuCOE/GenerateMutableSource/GenerateMutableSourceMesh.cpp
+++ b/Engine/Plugins/Experimental/Mutable/Source/CustomizableObjectEditor/Private/MuCOE/GenerateMutableSource/GenerateMutableSourceMesh.cpp
@@ -2427,7 +2427,7 @@ mu::MeshPtr ConvertStaticMeshToMutable(const UStaticMesh* StaticMesh, int32 LODI
 
 // Convert a Mesh constant to a mutable format. UniqueTags are the tags that make this Mesh unique that cannot be merged in the cache 
 // with the exact same Mesh with other tags
-mu::Ptr<mu::Mesh> GenerateMutableMesh(const UObject * Mesh, const TSoftClassPtr<UAnimInstance>& AnimInstance, int32 LODIndexConnected, int32 SectionIndexConnected, int32 LODIndex, int32 SectionIndex, const FString& UniqueTags, FMutableGraphGenerationContext & GenerationContext, const UCustomizableObjectNode* CurrentNode, bool bIsReference)
+mu::Ptr<mu::Mesh> GenerateMutableMesh(UObject * Mesh, const TSoftClassPtr<UAnimInstance>& AnimInstance, int32 LODIndexConnected, int32 SectionIndexConnected, int32 LODIndex, int32 SectionIndex, const FString& UniqueTags, FMutableGraphGenerationContext & GenerationContext, const UCustomizableObjectNode* CurrentNode, bool bIsReference)
 {
 	// Get the mesh generation flags to use
 	EMutableMeshConversionFlags CurrentFlags = GenerationContext.MeshGenerationFlags.Last();
@@ -2437,12 +2437,12 @@ mu::Ptr<mu::Mesh> GenerateMutableMesh(const UObject * Mesh, const TSoftClassPtr<
 	
 	if (!MutableMesh)
 	{
-		if (const USkeletalMesh* SkeletalMesh = Cast<USkeletalMesh>(Mesh))
+		if (USkeletalMesh* SkeletalMesh = Cast<USkeletalMesh>(Mesh))
 		{
 			// At some point we will want all meshes to be references at compile-time. For now, just create the actual pass-through meshes.
 			if (bIsReference)
 			{
-				MutableMesh = GenerateMeshConstant(SkeletalMesh,GenerationContext, bIsReference);
+				MutableMesh = GenerateMeshConstant(SkeletalMesh, GenerationContext, bIsReference);
 			}
 			else
 			{
diff --git a/Engine/Plugins/Experimental/Mutable/Source/CustomizableObjectEditor/Private/MuCOE/GenerateMutableSource/GenerateMutableSourceMesh.h b/Engine/Plugins/Experimental/Mutable/Source/CustomizableObjectEditor/Private/MuCOE/GenerateMutableSource/GenerateMutableSourceMesh.h
--- a/Engine/Plugins/Experimental/Mutable/Source/CustomizableObjectEditor/Private/MuCOE/GenerateMutableSource/GenerateMutableSourceMesh.h
+++ b/Engine/Plugins/Experimental/Mutable/Source/CustomizableObjectEditor/Private/MuCOE/GenerateMutableSource/GenerateMutableSourceMesh.h
@@ -71,7 +71,7 @@ mu::Ptr<mu::Mesh> ConvertStaticMeshToMutable(const UStaticMesh* StaticMesh, int3
  * @param GenerationContext 
  * @param CurrentNode 
  * @return Mutable Mesh. Nullptr if there has been an error. Empty mesh if the Skeletal Mesh does not contain the requested LOD + Section. */
-mu::Ptr<mu::Mesh> GenerateMutableMesh(const UObject* Mesh, const TSoftClassPtr<UAnimInstance>& AnimBp,
+mu::Ptr<mu::Mesh> GenerateMutableMesh(UObject* Mesh, const TSoftClassPtr<UAnimInstance>& AnimBp,
                                 int32 LODIndexConnected, int32 SectionIndexConnected,
                                 int32 LODIndex, int32 SectionIndex, const FString& MeshUniqueTags,
                                 FMutableGraphGenerationContext& GenerationContext, const UCustomizableObjectNode* CurrentNode, 
