[Mutable] Fix material OverrideMaterials not being set.

#tests pie
#rnx

diff --git a/Engine/Plugins/Experimental/Mutable/Source/CustomizableObject/Private/MuCO/CustomizableObjectInstanceUsage.cpp b/Engine/Plugins/Experimental/Mutable/Source/CustomizableObject/Private/MuCO/CustomizableObjectInstanceUsage.cpp
--- a/Engine/Plugins/Experimental/Mutable/Source/CustomizableObject/Private/MuCO/CustomizableObjectInstanceUsage.cpp
+++ b/Engine/Plugins/Experimental/Mutable/Source/CustomizableObject/Private/MuCO/CustomizableObjectInstanceUsage.cpp
@@ -233,7 +233,7 @@ bool RequiresReinitPose(USkeletalMesh* CurrentSkeletalMesh, USkeletalMesh* Skele
 
 
 void UCustomizableObjectInstanceUsage::SetSkeletalMeshAndOverrideMaterials(USkeletalMeshComponent& Parent, USkeletalMesh* SkeletalMesh, 
-	const UCustomizableObjectInstance& CustomizableObjectInstance, const bool bInstanceGenerated, bool* bOutSkeletalMeshUpdated, bool* bOutMaterialsUpdated)
+	const UCustomizableObjectInstance& CustomizableObjectInstance, bool* bOutSkeletalMeshUpdated, bool* bOutMaterialsUpdated)
 {
 	if (SkeletalMesh != Parent.GetSkeletalMeshAsset())
 	{
@@ -268,12 +268,10 @@ void UCustomizableObjectInstanceUsage::SetSkeletalMeshAndOverrideMaterials(USkel
 		return;
 	}
 	
-	const bool bIsTransientMesh = SkeletalMesh ? static_cast<bool>(SkeletalMesh->HasAllFlags(EObjectFlags::RF_Transient)) : false;
-	const bool bUseOverrideMaterials = !bIsTransientMesh
-		||
-		(CustomizableObject->bEnableMeshCache && CVarEnableMeshCache.GetValueOnAnyThread());
+	const bool bIsTransientMesh = SkeletalMesh ? SkeletalMesh->HasAllFlags(EObjectFlags::RF_Transient) : false;
+	const bool bUseOverrideMaterials = !bIsTransientMesh || (CustomizableObject->bEnableMeshCache && CVarEnableMeshCache.GetValueOnAnyThread());
 
-	if (bUseOverrideMaterials && bInstanceGenerated)
+	if (bUseOverrideMaterials)
 	{
 		if (FCustomizableInstanceComponentData* ComponentData = CustomizableObjectInstance.GetPrivate()->GetComponentData(GetComponentIndex()))
 		{
@@ -298,9 +296,7 @@ void UCustomizableObjectInstanceUsage::SetSkeletalMesh(USkeletalMesh* SkeletalMe
 
 	if (Parent && CustomizableObjectInstance)
 	{
-		const bool bInstanceGenerated = CustomizableObjectInstance->GetPrivate()->SkeletalMeshStatus == ESkeletalMeshStatus::Success;
-
-		SetSkeletalMeshAndOverrideMaterials(*Parent, SkeletalMesh, *CustomizableObjectInstance, bInstanceGenerated, bOutSkeletalMeshUpdated, bOutMaterialsUpdated);
+		SetSkeletalMeshAndOverrideMaterials(*Parent, SkeletalMesh, *CustomizableObjectInstance, bOutSkeletalMeshUpdated, bOutMaterialsUpdated);
 	}
 }
 
@@ -586,7 +582,7 @@ void UCustomizableObjectInstanceUsage::Tick(float DeltaTime)
 		// Set SkeletalMesh
 		if (bInstanceGenerated || SkeletalMesh)
 		{
-			SetSkeletalMeshAndOverrideMaterials(*Parent, SkeletalMesh, *CustomizableObjectInstance, bInstanceGenerated, nullptr, nullptr);
+			SetSkeletalMeshAndOverrideMaterials(*Parent, SkeletalMesh, *CustomizableObjectInstance, nullptr, nullptr);
 		}
 	}
 }
diff --git a/Engine/Plugins/Experimental/Mutable/Source/CustomizableObject/Public/MuCO/CustomizableObjectInstanceUsage.h b/Engine/Plugins/Experimental/Mutable/Source/CustomizableObject/Public/MuCO/CustomizableObjectInstanceUsage.h
--- a/Engine/Plugins/Experimental/Mutable/Source/CustomizableObject/Public/MuCO/CustomizableObjectInstanceUsage.h
+++ b/Engine/Plugins/Experimental/Mutable/Source/CustomizableObject/Public/MuCO/CustomizableObjectInstanceUsage.h
@@ -115,8 +115,7 @@ private:
 	virtual bool IsTickable() const override;
 	// End FTickableGameObject
 
-	void SetSkeletalMeshAndOverrideMaterials(USkeletalMeshComponent& Parent, USkeletalMesh* SkeletalMesh, const UCustomizableObjectInstance& CustomizableObjectInstance, 
-											 const bool bInstanceGenerated, bool* bOutSkeletalMeshUpdated, bool* bOutMaterialsUpdated);
+	void SetSkeletalMeshAndOverrideMaterials(USkeletalMeshComponent& Parent, USkeletalMesh* SkeletalMesh, const UCustomizableObjectInstance& CustomizableObjectInstance, bool* bOutSkeletalMeshUpdated, bool* bOutMaterialsUpdated);
 
 	friend UCustomizableSkeletalComponent;
 };
