[mutable] Support for optional components (in child objects).
- Add core support for component switching and variation nodes.
- Fix problems with pass-through mesh components being ignored when nothing is overridden.
- Fix problems with pass-through mesh components in child objects.
- Rename all ComponentIndex variables to make them explicitly refer to object component index or instance component index.
- Add mutable debugger support for instance data viewing.
#rnx

diff --git a/Engine/Plugins/Experimental/Mutable/Source/CustomizableObject/Internal/MuCO/CustomizableObjectInstancePrivate.h b/Engine/Plugins/Experimental/Mutable/Source/CustomizableObject/Internal/MuCO/CustomizableObjectInstancePrivate.h
--- a/Engine/Plugins/Experimental/Mutable/Source/CustomizableObject/Internal/MuCO/CustomizableObjectInstancePrivate.h
+++ b/Engine/Plugins/Experimental/Mutable/Source/CustomizableObject/Internal/MuCO/CustomizableObjectInstancePrivate.h
@@ -173,7 +173,7 @@ class CUSTOMIZABLEOBJECT_API UCustomizableInstancePrivate : public UObject
 public:
 	GENERATED_BODY()
 
-	/** The generated skeletal meshes for this Instance, one for each component */
+	/** The generated skeletal meshes for this Instance, one for each instance component */
 	UPROPERTY(Transient, VisibleAnywhere, Category = NoCategory, meta=(DisplayName = "Meshes"))
 	TArray<TObjectPtr<USkeletalMesh>> SkeletalMeshes;
 
@@ -269,10 +269,10 @@ private:
 
 	bool BuildSkeletonData(const TSharedRef<FUpdateContextPrivate>& OperationData, USkeletalMesh& SkeletalMesh, const FMutableRefSkeletalMeshData& RefSkeletalMeshData, UCustomizableObject& CustomizableObject, int32 ComponentIndex);
 	void BuildMeshSockets(const TSharedRef<FUpdateContextPrivate>& OperationData, USkeletalMesh* SkeletalMesh, const FModelResources& ModelResources, const FMutableRefSkeletalMeshData& RefSkeletalMeshData, mu::Ptr<const mu::Mesh> MutableMesh);
-	void BuildOrCopyElementData(const TSharedRef<FUpdateContextPrivate>& OperationData, USkeletalMesh* SkeletalMesh, UCustomizableObjectInstance* CustomizableObjectInstance, int32 ComponentIndex);
-	void BuildOrCopyMorphTargetsData(const TSharedRef<FUpdateContextPrivate>& OperationData, USkeletalMesh* SkeletalMesh, const USkeletalMesh* SrcSkeletalMesh, UCustomizableObjectInstance* CustomizableObjectInstance, int32 ComponentIndex);
-	bool BuildOrCopyRenderData(const TSharedRef<FUpdateContextPrivate>& OperationData, USkeletalMesh* SkeletalMesh, const USkeletalMesh* SrcSkeletalMesh, UCustomizableObjectInstance* CustomizableObjectInstance, int32 ComponentIndex);
-	void BuildOrCopyClothingData(const TSharedRef<FUpdateContextPrivate>& OperationData, USkeletalMesh* SkeletalMesh, const USkeletalMesh* SrcSkeletalMesh, UCustomizableObjectInstance* CustomizableObjectInstance, int32 ComponentIndex);
+	void BuildOrCopyElementData(const TSharedRef<FUpdateContextPrivate>& OperationData, USkeletalMesh* SkeletalMesh, UCustomizableObjectInstance* CustomizableObjectInstance, int32 InstanceComponentIndex);
+	void BuildOrCopyMorphTargetsData(const TSharedRef<FUpdateContextPrivate>& OperationData, USkeletalMesh* SkeletalMesh, const USkeletalMesh* SrcSkeletalMesh, UCustomizableObjectInstance* CustomizableObjectInstance, int32 InstanceComponentIndex);
+	bool BuildOrCopyRenderData(const TSharedRef<FUpdateContextPrivate>& OperationData, USkeletalMesh* SkeletalMesh, const USkeletalMesh* SrcSkeletalMesh, UCustomizableObjectInstance* CustomizableObjectInstance, int32 InstanceComponentIndex);
+	void BuildOrCopyClothingData(const TSharedRef<FUpdateContextPrivate>& OperationData, USkeletalMesh* SkeletalMesh, const USkeletalMesh* SrcSkeletalMesh, UCustomizableObjectInstance* CustomizableObjectInstance, int32 InstanceComponentIndex);
 	
 	//
 	USkeleton* MergeSkeletons(UCustomizableObject& CustomizableObject, const FMutableRefSkeletalMeshData& RefSkeletalMeshData, int32 ComponentIndex);
@@ -287,7 +287,7 @@ private:
 
 	bool DoComponentsNeedUpdate(UCustomizableObjectInstance* CustomizableObjectInstance, const TSharedRef<FUpdateContextPrivate>& OperationData, bool& bOutEmptyMesh);
 
-	void SetLastMeshId(int32 ComponentIndex, int32 LODIndex, mu::FResourceID MeshId);
+	void SetLastMeshId(int32 ObjectComponentIndex, int32 LODIndex, mu::FResourceID MeshId);
 
 public:
 	bool LoadParametersFromProfile(int32 ProfileIndex);
@@ -325,6 +325,7 @@ public:
 	// First Non-streamable SkeletalMesh LOD
 	uint8 FirstResidentLOD = MAX_MESH_LOD_COUNT;
 	
+	// To be indexed with object component index
 	UPROPERTY(Transient)
 	TArray<FCustomizableInstanceComponentData> ComponentsData;
 
@@ -332,7 +333,7 @@ public:
 	TArray< TObjectPtr<UMaterialInterface> > ReferencedMaterials;
 
 	// Converts a ReferencedMaterials index from the CustomizableObject to an index in the ReferencedMaterials in the Instance
-	TMap<uint32, uint32> ObjectToInstanceIndexMap;
+	TMap<int32, uint32> ObjectToInstanceIndexMap;
 
 	TArray<FGeneratedTexture> TexturesToRelease;
 
diff --git a/Engine/Plugins/Experimental/Mutable/Source/CustomizableObject/Internal/MuCO/CustomizableObjectSystemPrivate.h b/Engine/Plugins/Experimental/Mutable/Source/CustomizableObject/Internal/MuCO/CustomizableObjectSystemPrivate.h
--- a/Engine/Plugins/Experimental/Mutable/Source/CustomizableObject/Internal/MuCO/CustomizableObjectSystemPrivate.h
+++ b/Engine/Plugins/Experimental/Mutable/Source/CustomizableObject/Internal/MuCO/CustomizableObjectSystemPrivate.h
@@ -350,7 +350,11 @@ struct FInstanceUpdateData
 		uint16 FirstScalar = 0;
 		uint16 ScalarCount = 0;
 
-		uint32 MaterialIndex = 0;
+		/** Index of the material in the referenced materials array of the CO.
+		* A negative value means that the material of this surface slot of the mesh doesn't need to be changed.
+		* This is valid for pass-through meshes.
+		*/
+		int32 MaterialIndex = -1;
 
 		/** Id of the surface in the mutable core instance. */
 		uint32 SurfaceId = 0;
@@ -416,7 +420,7 @@ struct FInstanceUpdateData
 
 	struct FRealTimeMorphsComponentData
 	{
-		int32 ComponentIndex = INDEX_NONE;
+		int32 ObjectComponentIndex = INDEX_NONE;
 
 		TArray<FName> RealTimeMorphTargetNames;
 		TArray<TArray<FMorphTargetLODModel>> RealTimeMorphsLODData; 
@@ -443,7 +447,8 @@ struct FInstanceUpdateData
 		TMap<mu::FBoneName, TPair<FName, uint16>> BoneInfoMap;
 	};
 
-	// Access by component index
+	// Access by object component index
+	// \TODO: somewhat wasteful, since instances may have less components
 	TArray<FSkeletonData> Skeletons;
 
 	struct FNamedExtensionData
@@ -484,6 +489,7 @@ public:
 
 	void SetMinLOD(int32 MinLOD);
 
+	/** Return an array of LODs per object component. */
 	const TArray<uint16>& GetRequestedLODs() const;
 
 	void SetRequestedLODs(TArray<uint16>& RequestedLODs);
@@ -543,10 +549,16 @@ public:
 
 	TSharedPtr<mu::Model> Model;
 
-	// TODO: This is also available in InstanceUpdateData in MutableInstance and in NumLODsAvailablePerComponent
-	uint8 NumComponents = 0;
+	// Number of possible components in the entire CO
+	// TODO: Redundant because it is in the CO
+	uint8 NumObjectComponents = 0;
 
-	// TODO: Why is this cached?
+	// Number of components in the instance being generated
+	// TODO: Redundant to store it here because it is implicit in the size of some arrays or in the MutableInstance while it is valid.
+	uint8 NumInstanceComponents = 0;
+
+	// Index with ObjectComponent index
+	// \TODO: Wasteful: it only needs to be as big as NumInstanceComponents
 	TArray<uint8> NumLODsAvailablePerComponent;
 
 	uint8 FirstLODAvailable = 0;
@@ -558,7 +570,7 @@ public:
 
 	mu::FImageOperator::FImagePixelFormatFunc PixelFormatOverride;
 
-	/** Mutable Meshes required for each component. Outermost index is the component index, inner index is the LOD. */
+	/** Mutable Meshes required for each component. Outermost index is the object component index, inner index is the LOD. */
 	TArray<TArray<mu::FResourceID>> MeshDescriptors;
 
 	/** Used to know if the updated instances' meshes are different from the previous ones. 
diff --git a/Engine/Plugins/Experimental/Mutable/Source/CustomizableObject/Private/MuCO/BusyWaits_Deprecated.cpp b/Engine/Plugins/Experimental/Mutable/Source/CustomizableObject/Private/MuCO/BusyWaits_Deprecated.cpp
--- a/Engine/Plugins/Experimental/Mutable/Source/CustomizableObject/Private/MuCO/BusyWaits_Deprecated.cpp
+++ b/Engine/Plugins/Experimental/Mutable/Source/CustomizableObject/Private/MuCO/BusyWaits_Deprecated.cpp
@@ -204,14 +204,15 @@ namespace CustomizableObjectSystem::ImplDeprecated
 		TArray<int32> SurfacesSharedId;
 
 		// Generate the mesh and gather all the required resource Ids
-		int32 NumComponents = OperationData->NumComponents;
-		OperationData->InstanceUpdateData.Components.SetNum(NumComponents);
-		for (int32 ComponentIndex = 0; ComponentIndex < NumComponents; ++ComponentIndex)
+		OperationData->InstanceUpdateData.Components.SetNum(OperationData->NumInstanceComponents);
+		for (int32 InstanceComponentIndex = 0; InstanceComponentIndex < OperationData->NumInstanceComponents; ++InstanceComponentIndex)
 		{
-			FInstanceUpdateData::FComponent& Component = OperationData->InstanceUpdateData.Components[ComponentIndex];
+			FInstanceUpdateData::FComponent& Component = OperationData->InstanceUpdateData.Components[InstanceComponentIndex];
 			Component.FirstLOD = OperationData->InstanceUpdateData.LODs.Num();
-			Component.LODCount = Instance->GetLODCount(ComponentIndex);
-			Component.Id = Instance->GetComponentId(ComponentIndex);
+			Component.LODCount = Instance->GetLODCount(InstanceComponentIndex);
+			Component.Id = Instance->GetComponentId(InstanceComponentIndex);
+
+			int32 ObjectComponentIndex = Component.Id;
 
 			for (int32 MutableLODIndex = 0; MutableLODIndex < Component.LODCount; ++MutableLODIndex)
 			{
@@ -228,13 +229,13 @@ namespace CustomizableObjectSystem::ImplDeprecated
 				LOD.FirstSurface = OperationData->InstanceUpdateData.Surfaces.Num();
 				LOD.SurfaceCount = 0;
 
-				const bool bGenerateLOD = RequestedLODs.IsValidIndex(ComponentIndex) ? RequestedLODs[ComponentIndex] <= MutableLODIndex : true;
+				const bool bGenerateLOD = RequestedLODs.IsValidIndex(ObjectComponentIndex) ? RequestedLODs[ObjectComponentIndex] <= MutableLODIndex : true;
 
 				// Mesh
 				{
 					MUTABLE_CPUPROFILER_SCOPE(GetMesh);
 
-					LOD.MeshID = Instance->GetMeshId(ComponentIndex, MutableLODIndex);
+					LOD.MeshID = Instance->GetMeshId(InstanceComponentIndex, MutableLODIndex);
 
 					if (bGenerateLOD)
 					{
@@ -267,7 +268,7 @@ namespace CustomizableObjectSystem::ImplDeprecated
 
 
 				// This lambda does all the work to fill up the surface data
-				auto AddSurface = [&LOD, &SurfacesSharedId, &ModelResources, OperationData, Instance, System, CustomizableObject, CustomizableObjectInstancePrivateData, ComponentIndex, MutableLODIndex ]( uint32 SurfaceId, int32 InstanceSurfaceIndex)
+				auto AddSurface = [&LOD, &SurfacesSharedId, &ModelResources, OperationData, Instance, System, CustomizableObject, CustomizableObjectInstancePrivateData, InstanceComponentIndex, MutableLODIndex ]( uint32 SurfaceId, int32 InstanceSurfaceIndex)
 					{
 						int32 BaseSurfaceIndex = InstanceSurfaceIndex;
 						int32 BaseLODIndex = MutableLODIndex;
@@ -279,7 +280,7 @@ namespace CustomizableObjectSystem::ImplDeprecated
 						// Now Surface.MaterialIndex is decoded from a parameter at the end of this if()
 						Surface.SurfaceId = SurfaceId;
 
-						const int32 SharedSurfaceId = Instance->GetSharedSurfaceId(ComponentIndex, MutableLODIndex, InstanceSurfaceIndex);
+						const int32 SharedSurfaceId = Instance->GetSharedSurfaceId(InstanceComponentIndex, MutableLODIndex, InstanceSurfaceIndex);
 						const int32 SharedSurfaceIndex = SurfacesSharedId.Find(SharedSurfaceId);
 
 						SurfacesSharedId.Add(SharedSurfaceId);
@@ -293,22 +294,22 @@ namespace CustomizableObjectSystem::ImplDeprecated
 							}
 
 							// Find the first LOD where this surface can be found
-							Instance->FindBaseSurfaceBySharedId(ComponentIndex, SharedSurfaceId, BaseSurfaceIndex, BaseLODIndex);
+							Instance->FindBaseSurfaceBySharedId(InstanceComponentIndex, SharedSurfaceId, BaseSurfaceIndex, BaseLODIndex);
 
-							Surface.SurfaceId = Instance->GetSurfaceId(ComponentIndex, BaseLODIndex, BaseSurfaceIndex);
+							Surface.SurfaceId = Instance->GetSurfaceId(InstanceComponentIndex, BaseLODIndex, BaseSurfaceIndex);
 						}
 
 						// Images
 						Surface.FirstImage = OperationData->InstanceUpdateData.Images.Num();
-						Surface.ImageCount = Instance->GetImageCount(ComponentIndex, BaseLODIndex, BaseSurfaceIndex);
+						Surface.ImageCount = Instance->GetImageCount(InstanceComponentIndex, BaseLODIndex, BaseSurfaceIndex);
 						for (int32 ImageIndex = 0; ImageIndex < Surface.ImageCount; ++ImageIndex)
 						{
 							MUTABLE_CPUPROFILER_SCOPE(GetImageId);
 
 							OperationData->InstanceUpdateData.Images.Push({});
 							FInstanceUpdateData::FImage& Image = OperationData->InstanceUpdateData.Images.Last();
-							Image.Name = Instance->GetImageName(ComponentIndex, BaseLODIndex, BaseSurfaceIndex, ImageIndex);
-							Image.ImageID = Instance->GetImageId(ComponentIndex, BaseLODIndex, BaseSurfaceIndex, ImageIndex);
+							Image.Name = Instance->GetImageName(InstanceComponentIndex, BaseLODIndex, BaseSurfaceIndex, ImageIndex);
+							Image.ImageID = Instance->GetImageId(InstanceComponentIndex, BaseLODIndex, BaseSurfaceIndex, ImageIndex);
 							Image.FullImageSizeX = 0;
 							Image.FullImageSizeY = 0;
 							Image.BaseLOD = BaseLODIndex;
@@ -356,25 +357,25 @@ namespace CustomizableObjectSystem::ImplDeprecated
 
 						// Vectors
 						Surface.FirstVector = OperationData->InstanceUpdateData.Vectors.Num();
-						Surface.VectorCount = Instance->GetVectorCount(ComponentIndex, BaseLODIndex, BaseSurfaceIndex);
+						Surface.VectorCount = Instance->GetVectorCount(InstanceComponentIndex, BaseLODIndex, BaseSurfaceIndex);
 						for (int32 VectorIndex = 0; VectorIndex < Surface.VectorCount; ++VectorIndex)
 						{
 							MUTABLE_CPUPROFILER_SCOPE(GetVector);
 							OperationData->InstanceUpdateData.Vectors.Push({});
 							FInstanceUpdateData::FVector& Vector = OperationData->InstanceUpdateData.Vectors.Last();
-							Vector.Name = Instance->GetVectorName(ComponentIndex, BaseLODIndex, BaseSurfaceIndex, VectorIndex);
-							Vector.Vector = Instance->GetVector(ComponentIndex, BaseLODIndex, BaseSurfaceIndex, VectorIndex);
+							Vector.Name = Instance->GetVectorName(InstanceComponentIndex, BaseLODIndex, BaseSurfaceIndex, VectorIndex);
+							Vector.Vector = Instance->GetVector(InstanceComponentIndex, BaseLODIndex, BaseSurfaceIndex, VectorIndex);
 						}
 
 						// Scalars
 						Surface.FirstScalar = OperationData->InstanceUpdateData.Scalars.Num();
-						Surface.ScalarCount = Instance->GetScalarCount(ComponentIndex, BaseLODIndex, BaseSurfaceIndex);
+						Surface.ScalarCount = Instance->GetScalarCount(InstanceComponentIndex, BaseLODIndex, BaseSurfaceIndex);
 						for (int32 ScalarIndex = 0; ScalarIndex < Surface.ScalarCount; ++ScalarIndex)
 						{
 							MUTABLE_CPUPROFILER_SCOPE(GetScalar)
 
-								const FName ScalarName = Instance->GetScalarName(ComponentIndex, BaseLODIndex, BaseSurfaceIndex, ScalarIndex);
-							const float ScalarValue = Instance->GetScalar(ComponentIndex, BaseLODIndex, BaseSurfaceIndex, ScalarIndex);
+								const FName ScalarName = Instance->GetScalarName(InstanceComponentIndex, BaseLODIndex, BaseSurfaceIndex, ScalarIndex);
+							const float ScalarValue = Instance->GetScalar(InstanceComponentIndex, BaseLODIndex, BaseSurfaceIndex, ScalarIndex);
 
 							FString EncodingMaterialIdString = "__MutableMaterialId";
 
@@ -400,10 +401,10 @@ namespace CustomizableObjectSystem::ImplDeprecated
 				// and all defined surfaces will be present. 
 				if (LOD.Mesh->IsReference())
 				{
-					const int32 SurfaceCount = Instance->GetSurfaceCount(ComponentIndex, MutableLODIndex);
+					const int32 SurfaceCount = Instance->GetSurfaceCount(InstanceComponentIndex, MutableLODIndex);
 					for (int32 SurfaceIndex = 0; SurfaceIndex < SurfaceCount; ++SurfaceIndex)
 					{
-						uint32 SurfaceId = Instance->GetSurfaceId(ComponentIndex, MutableLODIndex, SurfaceIndex);
+						uint32 SurfaceId = Instance->GetSurfaceId(InstanceComponentIndex, MutableLODIndex, SurfaceIndex);
 						AddSurface( SurfaceId, SurfaceIndex);
 					}
 				}
@@ -416,7 +417,7 @@ namespace CustomizableObjectSystem::ImplDeprecated
 					{
 						const uint32 SurfaceId = LOD.Mesh->GetSurfaceId(MeshSurfaceIndex);
 
-						const int32 InstanceSurfaceIndex = Instance->FindSurfaceById(ComponentIndex, MutableLODIndex, SurfaceId);
+						const int32 InstanceSurfaceIndex = Instance->FindSurfaceById(InstanceComponentIndex, MutableLODIndex, SurfaceId);
 						check(LOD.Mesh->GetVertexCount() > 0 || InstanceSurfaceIndex >= 0);
 
 						if (InstanceSurfaceIndex >= 0)
diff --git a/Engine/Plugins/Experimental/Mutable/Source/CustomizableObject/Private/MuCO/CustomizableObjectInstance.cpp b/Engine/Plugins/Experimental/Mutable/Source/CustomizableObject/Private/MuCO/CustomizableObjectInstance.cpp
--- a/Engine/Plugins/Experimental/Mutable/Source/CustomizableObject/Private/MuCO/CustomizableObjectInstance.cpp
+++ b/Engine/Plugins/Experimental/Mutable/Source/CustomizableObject/Private/MuCO/CustomizableObjectInstance.cpp
@@ -184,9 +184,9 @@ UTexture2D* UCustomizableInstancePrivate::CreateTexture()
 }
 
 
-void UCustomizableInstancePrivate::SetLastMeshId(int32 ComponentIndex, int32 LODIndex, mu::FResourceID MeshId)
+void UCustomizableInstancePrivate::SetLastMeshId(int32 ObjectComponentIndex, int32 LODIndex, mu::FResourceID MeshId)
 {
-	FCustomizableInstanceComponentData* ComponentData = GetComponentData(ComponentIndex);
+	FCustomizableInstanceComponentData* ComponentData = GetComponentData(ObjectComponentIndex);
 	if (ComponentData && ComponentData->LastMeshIdPerLOD.IsValidIndex(LODIndex))
 	{
 		ComponentData->LastMeshIdPerLOD[LODIndex] = MeshId;
@@ -1933,6 +1933,16 @@ bool UCustomizableInstancePrivate::DoComponentsNeedUpdate(UCustomizableObjectIns
 
 		// Components with mesh must have valid geometry at CurrentMaxLOD
 		const FInstanceUpdateData::FComponent& Component = OperationData->InstanceUpdateData.Components[ComponentIndex];
+
+		const bool bHasSkeletalMesh = SkeletalMeshes.IsValidIndex(ComponentIndex) && SkeletalMeshes[ComponentIndex];
+
+		if (Component.LODCount == 0)
+		{
+			// We don't have a mesh in the component, so it has changed if we had one before.
+			OperationData->MeshChanged[ComponentIndex] = bHasSkeletalMesh;
+			continue;
+		}
+
 		const FInstanceUpdateData::FLOD& LOD = OperationData->InstanceUpdateData.LODs[Component.FirstLOD];
 		bool bIsReferenced = LOD.Mesh && LOD.Mesh->IsReference();
 		if (!bIsReferenced)
@@ -1948,7 +1958,6 @@ bool UCustomizableInstancePrivate::DoComponentsNeedUpdate(UCustomizableObjectIns
 		}
 
 		// Update the component if there is a mesh and it shouldn't, or the other way around.
-		const bool bHasSkeletalMesh = SkeletalMeshes.IsValidIndex(ComponentIndex) && SkeletalMeshes[ComponentIndex];
 		OperationData->MeshChanged[ComponentIndex] = (ComponentWithMesh[ComponentIndex] != bHasSkeletalMesh);
 
 		const FCustomizableInstanceComponentData* ComponentData = GetComponentData(ComponentIndex);
@@ -2284,26 +2293,27 @@ bool UCustomizableInstancePrivate::UpdateSkeletalMesh_PostBeginUpdate0(UCustomiz
 	{
 		MUTABLE_CPUPROFILER_SCOPE(BuildSkeletalMeshData);
 
-		const int32 NumComponents = SkeletalMeshes.Num();
-		for (int32 ComponentIndex = 0; bSuccess && ComponentIndex < NumComponents; ++ComponentIndex)
+		for (int32 InstanceComponentIndex = 0; bSuccess && InstanceComponentIndex < OperationData->NumInstanceComponents; ++InstanceComponentIndex)
 		{
-			const FInstanceUpdateData::FComponent& Component = OperationData->InstanceUpdateData.Components[ComponentIndex];
+			const FInstanceUpdateData::FComponent& Component = OperationData->InstanceUpdateData.Components[InstanceComponentIndex];
+
+			int32 ObjectComponentIndex = Component.Id;
 
 			if (OperationData->bUseMeshCache)
 			{
-				const TArray<mu::FResourceID>& MeshId = OperationData->MeshDescriptors[ComponentIndex];
+				const TArray<mu::FResourceID>& MeshId = OperationData->MeshDescriptors[ObjectComponentIndex];
 				if (CustomizableObject->GetPrivate()->MeshCache.Get(MeshId))
 				{
 					continue;
 				}
 			}
 			
-			if (!OperationData->MeshChanged[ComponentIndex])
+			if (!OperationData->MeshChanged[ObjectComponentIndex])
 			{
 				continue;
 			}
 
-			USkeletalMesh* SkeletalMesh = SkeletalMeshes[ComponentIndex];
+			USkeletalMesh* SkeletalMesh = SkeletalMeshes[ObjectComponentIndex];
 
 			if (!SkeletalMesh)
 			{
@@ -2324,17 +2334,17 @@ bool UCustomizableInstancePrivate::UpdateSkeletalMesh_PostBeginUpdate0(UCustomiz
 			}
 
 			// Mesh to copy data from if possible. 
-			USkeletalMesh* OldSkeletalMesh = OldSkeletalMeshes.IsValidIndex(ComponentIndex) ? OldSkeletalMeshes[ComponentIndex] : nullptr;
+			USkeletalMesh* OldSkeletalMesh = OldSkeletalMeshes.IsValidIndex(ObjectComponentIndex) ? OldSkeletalMeshes[ObjectComponentIndex] : nullptr;
 
 			{
-				BuildOrCopyElementData(OperationData, SkeletalMesh, Public, ComponentIndex);
-				bSuccess = BuildOrCopyRenderData(OperationData, SkeletalMesh, OldSkeletalMesh, Public, ComponentIndex);
+				BuildOrCopyElementData(OperationData, SkeletalMesh, Public, InstanceComponentIndex);
+				bSuccess = BuildOrCopyRenderData(OperationData, SkeletalMesh, OldSkeletalMesh, Public, InstanceComponentIndex);
 			}
 
 			if (bSuccess)
 			{
-				BuildOrCopyMorphTargetsData(OperationData, SkeletalMesh, OldSkeletalMesh, Public, ComponentIndex);
-				BuildOrCopyClothingData(OperationData, SkeletalMesh, OldSkeletalMesh, Public, ComponentIndex);
+				BuildOrCopyMorphTargetsData(OperationData, SkeletalMesh, OldSkeletalMesh, Public, InstanceComponentIndex);
+				BuildOrCopyClothingData(OperationData, SkeletalMesh, OldSkeletalMesh, Public, InstanceComponentIndex);
 
 				FSkeletalMeshRenderData* RenderData = SkeletalMesh->GetResourceForRendering();
 				ensure(RenderData && RenderData->LODRenderData.Num() > 0);
@@ -2348,7 +2358,7 @@ bool UCustomizableInstancePrivate::UpdateSkeletalMesh_PostBeginUpdate0(UCustomiz
 
 			if (OperationData->bUseMeshCache)
 			{
-				const TArray<mu::FResourceID>& MeshId = OperationData->MeshDescriptors[ComponentIndex];
+				const TArray<mu::FResourceID>& MeshId = OperationData->MeshDescriptors[ObjectComponentIndex];
 				CustomizableObject->GetPrivate()->MeshCache.Add(MeshId, SkeletalMesh);
 			}
 		}
@@ -3814,11 +3824,11 @@ void UCustomizableInstancePrivate::BuildMeshSockets(const TSharedRef<FUpdateCont
 }
 
 
-void UCustomizableInstancePrivate::BuildOrCopyElementData(const TSharedRef<FUpdateContextPrivate>& OperationData, USkeletalMesh* SkeletalMesh, UCustomizableObjectInstance* CustomizableObjectInstance, int32 ComponentIndex)
+void UCustomizableInstancePrivate::BuildOrCopyElementData(const TSharedRef<FUpdateContextPrivate>& OperationData, USkeletalMesh* SkeletalMesh, UCustomizableObjectInstance* CustomizableObjectInstance, int32 InstanceComponentIndex)
 {
 	MUTABLE_CPUPROFILER_SCOPE(UCustomizableInstancePrivate::BuildOrCopyElementData);
 
-	const FInstanceUpdateData::FComponent& Component = OperationData->InstanceUpdateData.Components[ComponentIndex];
+	const FInstanceUpdateData::FComponent& Component = OperationData->InstanceUpdateData.Components[InstanceComponentIndex];
 
 	for (int32 LODIndex = FirstLODAvailable; LODIndex < Component.LODCount; ++LODIndex)
 	{
@@ -3841,7 +3851,7 @@ void UCustomizableInstancePrivate::BuildOrCopyElementData(const TSharedRef<FUpda
 	}
 }
 
-void UCustomizableInstancePrivate::BuildOrCopyMorphTargetsData(const TSharedRef<FUpdateContextPrivate>& OperationData, USkeletalMesh* SkeletalMesh, const USkeletalMesh* LastUpdateSkeletalMesh, UCustomizableObjectInstance* CustomizableObjectInstance, int32 ComponentIndex)
+void UCustomizableInstancePrivate::BuildOrCopyMorphTargetsData(const TSharedRef<FUpdateContextPrivate>& OperationData, USkeletalMesh* SkeletalMesh, const USkeletalMesh* LastUpdateSkeletalMesh, UCustomizableObjectInstance* CustomizableObjectInstance, int32 InstanceComponentIndex)
 {
 	MUTABLE_CPUPROFILER_SCOPE(UCustomizableInstancePrivate::BuildOrCopyMorphTargetsData);
 
@@ -3852,11 +3862,14 @@ void UCustomizableInstancePrivate::BuildOrCopyMorphTargetsData(const TSharedRef<
 	
 	FInstanceUpdateData& UpdateData = OperationData->InstanceUpdateData;
 
+	FInstanceUpdateData::FComponent& Component = UpdateData.Components[InstanceComponentIndex];
+	int32 ObjectComponentIndex = Component.Id;
+
 	FInstanceUpdateData::FRealTimeMorphsComponentData* ComponentMorphTargets = 
 		UpdateData.RealTimeMorphTargets.FindByPredicate(
-				[ComponentIndex](const FInstanceUpdateData::FRealTimeMorphsComponentData& Elem) 
+				[ObjectComponentIndex](const FInstanceUpdateData::FRealTimeMorphsComponentData& Elem)
 				{ 
-					return Elem.ComponentIndex == ComponentIndex; 
+					return Elem.ObjectComponentIndex == ObjectComponentIndex;
 				});
 
 	if (!ComponentMorphTargets)
@@ -3883,7 +3896,7 @@ void UCustomizableInstancePrivate::BuildOrCopyMorphTargetsData(const TSharedRef<
 	}
 
 	// Copy MorphTargets from the FirstGeneratedLOD to the LODs below
-	const int32 FirstGeneratedLOD = FMath::Max((int32)OperationData->GetRequestedLODs()[ComponentIndex], OperationData->GetMinLOD());
+	const int32 FirstGeneratedLOD = FMath::Max((int32)OperationData->GetRequestedLODs()[ObjectComponentIndex], OperationData->GetMinLOD());
 	for (int32 LODIndex = OperationData->FirstLODAvailable; LODIndex < FirstGeneratedLOD; ++LODIndex)
 	{
 		MUTABLE_CPUPROFILER_SCOPE(CopyMorphTargetsData);
@@ -3913,7 +3926,7 @@ namespace
 
 }
 
-void UCustomizableInstancePrivate::BuildOrCopyClothingData(const TSharedRef<FUpdateContextPrivate>& OperationData, USkeletalMesh * SkeletalMesh, const USkeletalMesh* LastUpdateSkeletalMesh, UCustomizableObjectInstance* CustomizableObjectInstance, int32 ComponentIndex)
+void UCustomizableInstancePrivate::BuildOrCopyClothingData(const TSharedRef<FUpdateContextPrivate>& OperationData, USkeletalMesh * SkeletalMesh, const USkeletalMesh* LastUpdateSkeletalMesh, UCustomizableObjectInstance* CustomizableObjectInstance, int32 InstanceComponentIndex)
 {
 	MUTABLE_CPUPROFILER_SCOPE(UCustomizableInstancePrivate::BuildOrCopyClothingData);
 
@@ -3954,11 +3967,12 @@ void UCustomizableInstancePrivate::BuildOrCopyClothingData(const TSharedRef<FUpd
 
 	const int32 LODCount = OperationData->InstanceUpdateData.LODs.Num();
 
+	FInstanceUpdateData::FComponent& Component = OperationData->InstanceUpdateData.Components[InstanceComponentIndex];
+	int32 ObjectComponentIndex = Component.Id;
+
 	{
 		MUTABLE_CPUPROFILER_SCOPE(DiscoverSectionsWithCloth);
 
-		FInstanceUpdateData::FComponent& Component = OperationData->InstanceUpdateData.Components[ComponentIndex];
-
 		for (int32 LODIndex = OperationData->FirstLODAvailable; LODIndex < Component.LODCount; ++LODIndex)
 		{
 			const FInstanceUpdateData::FLOD& LOD = OperationData->InstanceUpdateData.LODs[Component.FirstLOD+LODIndex];
@@ -4096,7 +4110,7 @@ void UCustomizableInstancePrivate::BuildOrCopyClothingData(const TSharedRef<FUpd
 			// Only initilialize once, multiple sections with cloth could point to the same cloth asset.
 			if (!DstAssetData.LodMap.Num())
 			{
-				DstAssetData.LodMap.Init(INDEX_NONE, OperationData->NumLODsAvailablePerComponent[ComponentIndex]);
+				DstAssetData.LodMap.Init(INDEX_NONE, OperationData->NumLODsAvailablePerComponent[ObjectComponentIndex]);
 
 				DstAssetData.UsedBoneNames = SrcAssetData.UsedBoneNames;
 				DstAssetData.UsedBoneIndices = SrcAssetData.UsedBoneIndices;
@@ -4684,7 +4698,6 @@ void UCustomizableInstancePrivate::BuildOrCopyClothingData(const TSharedRef<FUpd
 	{
 		int32 NumSectionsWithClothProcessed = 0;
 
-		FInstanceUpdateData::FComponent& Component = OperationData->InstanceUpdateData.Components[ComponentIndex];
 		for (int32 LODIndex = OperationData->GetMinLOD(); LODIndex < Component.LODCount; ++LODIndex)
 		{
 			const FInstanceUpdateData::FLOD& LOD = OperationData->InstanceUpdateData.LODs[Component.FirstLOD+LODIndex];
@@ -4734,7 +4747,7 @@ void UCustomizableInstancePrivate::BuildOrCopyClothingData(const TSharedRef<FUpd
 	{
 		MUTABLE_CPUPROFILER_SCOPE(InitClothRenderData)
 		// Based on FSkeletalMeshLODModel::GetClothMappingData().
-		for (int32 LODIndex = OperationData->GetMinLOD(); LODIndex < OperationData->NumLODsAvailablePerComponent[ComponentIndex]; ++LODIndex)
+		for (int32 LODIndex = OperationData->GetMinLOD(); LODIndex < OperationData->NumLODsAvailablePerComponent[ObjectComponentIndex]; ++LODIndex)
 		{
 			FSkeletalMeshLODRenderData& LODModel = RenderResource->LODRenderData[LODIndex];
 	
@@ -4870,7 +4883,7 @@ void UCustomizableInstancePrivate::BuildOrCopyClothingData(const TSharedRef<FUpd
 	SkeletalMesh->SetHasActiveClothingAssets( static_cast<bool>( SectionsWithCloth.Num() ) );
 }
 
-bool UCustomizableInstancePrivate::BuildOrCopyRenderData(const TSharedRef<FUpdateContextPrivate>& OperationData, USkeletalMesh* SkeletalMesh, const USkeletalMesh* LastUpdateSkeletalMesh, UCustomizableObjectInstance* Public, int32 ComponentIndex)
+bool UCustomizableInstancePrivate::BuildOrCopyRenderData(const TSharedRef<FUpdateContextPrivate>& OperationData, USkeletalMesh* SkeletalMesh, const USkeletalMesh* LastUpdateSkeletalMesh, UCustomizableObjectInstance* Public, int32 InstanceComponentIndex)
 {
 	MUTABLE_CPUPROFILER_SCOPE(UCustomizableInstancePrivate::BuildOrCopyRenderData);
 
@@ -4882,11 +4895,13 @@ bool UCustomizableInstancePrivate::BuildOrCopyRenderData(const TSharedRef<FUpdat
 	// It must be not null as it's checked in the calling function
 	check(CustomizableObject);
 
+	FInstanceUpdateData::FComponent& Component = OperationData->InstanceUpdateData.Components[InstanceComponentIndex];
+	int32 ObjectComponentIndex = Component.Id;
+
 	const FModelResources& ModelResources = CustomizableObject->GetPrivate()->GetModelResources();
 
-	const int32 FirstGeneratedLOD = FMath::Max((int32)OperationData->GetRequestedLODs()[ComponentIndex], OperationData->GetMinLOD());
+	const int32 FirstGeneratedLOD = FMath::Max((int32)OperationData->GetRequestedLODs()[ObjectComponentIndex], OperationData->GetMinLOD());
 
-	FInstanceUpdateData::FComponent& Component = OperationData->InstanceUpdateData.Components[ComponentIndex];
 
 	for (int32 LODIndex = FirstGeneratedLOD; LODIndex < Component.LODCount; ++LODIndex)
 	{
@@ -4897,22 +4912,22 @@ bool UCustomizableInstancePrivate::BuildOrCopyRenderData(const TSharedRef<FUpdat
 		// There could be components without a mesh in LODs
 		if (!LOD.bGenerated || !LOD.Mesh || LOD.SurfaceCount == 0)
 		{
-			UE_LOG(LogMutable, Warning, TEXT("Building instance: generated mesh [%s] has LOD [%d] of component index [%d] with no mesh.")
+			UE_LOG(LogMutable, Warning, TEXT("Building instance: generated mesh [%s] has LOD [%d] of object component index [%d] with no mesh.")
 				, *SkeletalMesh->GetName()
 				, LODIndex
-				, ComponentIndex);
+				, ObjectComponentIndex);
 
 			// End with failure
 			return false;
 		}
 
-		TRACE_CPUPROFILER_EVENT_SCOPE_TEXT(*FString::Printf(TEXT("BuildRenderData: Component index %d, LOD %d"), ComponentIndex, LODIndex));
+		TRACE_CPUPROFILER_EVENT_SCOPE_TEXT(*FString::Printf(TEXT("BuildRenderData: Component index %d, LOD %d"), ObjectComponentIndex, LODIndex));
 
-		SetLastMeshId(ComponentIndex, LODIndex, LOD.MeshID);
+		SetLastMeshId(ObjectComponentIndex, LODIndex, LOD.MeshID);
 
 		FSkeletalMeshLODRenderData& LODResource = RenderData->LODRenderData[LODIndex];
 
-		const TMap<mu::FBoneName, TPair<FName, uint16>>& BoneInfoMap = OperationData->InstanceUpdateData.Skeletons[ComponentIndex].BoneInfoMap;
+		const TMap<mu::FBoneName, TPair<FName, uint16>>& BoneInfoMap = OperationData->InstanceUpdateData.Skeletons[InstanceComponentIndex].BoneInfoMap;
 		
 		// Set active and required bones
 		{
@@ -4938,14 +4953,14 @@ bool UCustomizableInstancePrivate::BuildOrCopyRenderData(const TSharedRef<FUpdat
 		for (int32 MeshSectionIndex = 0; MeshSectionIndex < MeshNumSurfaces; ++MeshSectionIndex)
 		{
 			uint32 MeshSurfaceId = LOD.Mesh->GetSurfaceId(MeshSectionIndex);
-			int32 InstanceSurfaceIndex = OperationData->MutableInstance->FindSurfaceById(ComponentIndex, LODIndex, MeshSurfaceId);
+			int32 InstanceSurfaceIndex = OperationData->MutableInstance->FindSurfaceById(InstanceComponentIndex, LODIndex, MeshSurfaceId);
 			
 			if (InstanceSurfaceIndex < 0)
 			{
 				continue;
 			}
 			
-			uint32 SurfaceMetadataId = OperationData->MutableInstance->GetSurfaceCustomId(ComponentIndex, LODIndex, InstanceSurfaceIndex);
+			uint32 SurfaceMetadataId = OperationData->MutableInstance->GetSurfaceCustomId(InstanceComponentIndex, LODIndex, InstanceSurfaceIndex);
 			MeshSurfacesMetadata[MeshSectionIndex] = ModelResources.SurfaceMetadata.Find(SurfaceMetadataId);
 		}
 
@@ -5119,8 +5134,8 @@ UE::Tasks::FTask UCustomizableInstancePrivate::LoadAdditionalAssetsAndData(
 	
 	for (const FInstanceUpdateData::FSurface& Surface : OperationData->InstanceUpdateData.Surfaces)
 	{
-		const uint32 MaterialIndex = Surface.MaterialIndex;
-		if (ObjectToInstanceIndexMap.Contains(MaterialIndex))
+		const int32 MaterialIndex = Surface.MaterialIndex;
+		if (MaterialIndex<0 || ObjectToInstanceIndexMap.Contains(MaterialIndex))
 		{
 			continue;
 		}
@@ -5139,11 +5154,13 @@ UE::Tasks::FTask UCustomizableInstancePrivate::LoadAdditionalAssetsAndData(
 
 
 	// Load Skeletons required by the SubMeshes of the newly generated Mesh, will be merged later
-	for (int32 ComponentIndex = 0; ComponentIndex < OperationData->NumComponents; ++ ComponentIndex)
+	for (int32 InstanceComponentIndex = 0; InstanceComponentIndex < OperationData->NumInstanceComponents; ++InstanceComponentIndex)
 	{
-		const FInstanceUpdateData::FSkeletonData& SkeletonData = OperationData->InstanceUpdateData.Skeletons[ComponentIndex];
+		int32 ObjectComponentIndex = OperationData->InstanceUpdateData.Components[InstanceComponentIndex].Id;
+
+		const FInstanceUpdateData::FSkeletonData& SkeletonData = OperationData->InstanceUpdateData.Skeletons[InstanceComponentIndex];
 		
-		FCustomizableInstanceComponentData* ComponentData = GetComponentData(ComponentIndex);
+		FCustomizableInstanceComponentData* ComponentData = GetComponentData(ObjectComponentIndex);
 		if (!ComponentData)
 		{
 			check(false);
@@ -6000,25 +6017,29 @@ void UCustomizableInstancePrivate::BuildMaterials(const TSharedRef<FUpdateContex
 	
 	const bool bReuseTextures = OperationData->bReuseInstanceTextures;
 
-	TArray<bool> RecreateRenderStateOnComponent;
-	RecreateRenderStateOnComponent.Init(false, OperationData->NumComponents);
+	TArray<bool> RecreateRenderStateOnInstanceComponent;
+	RecreateRenderStateOnInstanceComponent.Init(false, OperationData->NumInstanceComponents);
 
-	for (int32 ComponentIndex = 0; ComponentIndex < OperationData->NumComponents; ++ComponentIndex)
+	for (int32 InstanceComponentIndex = 0; InstanceComponentIndex < OperationData->NumInstanceComponents; ++InstanceComponentIndex)
 	{
-		if (!SkeletalMeshes.IsValidIndex(ComponentIndex))
+		const FInstanceUpdateData::FComponent& Component = OperationData->InstanceUpdateData.Components[InstanceComponentIndex];
+
+		int32 ObjectComponentIndex = Component.Id;
+
+		if (!SkeletalMeshes.IsValidIndex(ObjectComponentIndex))
 		{
 			// This could happen if there are no meshes at all.
 			continue;
 		}
 
-		USkeletalMesh* SkeletalMesh = SkeletalMeshes[ComponentIndex];
+		USkeletalMesh* SkeletalMesh = SkeletalMeshes[ObjectComponentIndex];
 
 		if (!SkeletalMesh)
 		{
 			continue;
 		}
 
-		const bool bReuseMaterials = !OperationData->MeshChanged[ComponentIndex];
+		const bool bReuseMaterials = !OperationData->MeshChanged[InstanceComponentIndex];
 
 		// If the mesh is not transient, it means it's pass-through so it should use material overrides and not be modified in any way
 		const bool bIsTransientMesh = static_cast<bool>(SkeletalMesh->HasAllFlags(EObjectFlags::RF_Transient));
@@ -6026,7 +6047,7 @@ void UCustomizableInstancePrivate::BuildMaterials(const TSharedRef<FUpdateContex
 		// It is not safe to replace the materials of a SkeletalMesh whose resources are initialized. Use overrides instead.
 		const bool bUseOverrideMaterialsOnly = !bIsTransientMesh || (OperationData->bUseMeshCache && SkeletalMesh->GetResourceForRendering()->IsInitialized());
 
-		ComponentsData[ComponentIndex].OverrideMaterials.Reset();
+		ComponentsData[ObjectComponentIndex].OverrideMaterials.Reset();
 
 		TArray<FSkeletalMaterial> Materials;
 
@@ -6038,9 +6059,7 @@ void UCustomizableInstancePrivate::BuildMaterials(const TSharedRef<FUpdateContex
 
 		MUTABLE_CPUPROFILER_SCOPE(BuildMaterials_LODLoop);
 
-		const int32 FirstGeneratedLOD = FMath::Max((int32)OperationData->GetRequestedLODs()[ComponentIndex], OperationData->GetMinLOD());
-
-		const FInstanceUpdateData::FComponent& Component = OperationData->InstanceUpdateData.Components[ComponentIndex];
+		const int32 FirstGeneratedLOD = FMath::Max((int32)OperationData->GetRequestedLODs()[ObjectComponentIndex], OperationData->GetMinLOD());
 
 		for (int32 LODIndex = FirstGeneratedLOD; LODIndex < Component.LODCount; LODIndex++)
 		{
@@ -6058,15 +6077,26 @@ void UCustomizableInstancePrivate::BuildMaterials(const TSharedRef<FUpdateContex
 
 			// Pass-through components will not have a reference mesh.
 			const FMutableRefSkeletalMeshData* RefSkeletalMeshData = nullptr;
-			if (ModelResources.ReferenceSkeletalMeshesData.IsValidIndex(ComponentIndex))
+			if (ModelResources.ReferenceSkeletalMeshesData.IsValidIndex(ObjectComponentIndex))
 			{
-				RefSkeletalMeshData = &ModelResources.ReferenceSkeletalMeshesData[ComponentIndex];
+				RefSkeletalMeshData = &ModelResources.ReferenceSkeletalMeshesData[ObjectComponentIndex];
 			}
 
 			for (int32 SurfaceIndex = 0; SurfaceIndex < LOD.SurfaceCount; ++SurfaceIndex)
 			{
 				const FInstanceUpdateData::FSurface& Surface = OperationData->InstanceUpdateData.Surfaces[LOD.FirstSurface + SurfaceIndex];
 
+				// Is this a material in a passthrough mesh that we don't modify?
+				if (Surface.MaterialIndex<0)
+				{
+					Materials.Emplace();
+#if WITH_EDITOR
+					// Without this, a change of a referenced material and recompilation doesn't show up in the preview.
+					RecreateRenderStateOnInstanceComponent[InstanceComponentIndex] = true;
+#endif
+					continue;
+				}
+
 				// Reuse MaterialSlot from the previous LOD.
 				if (const int32 MaterialIndex = SurfaceIdToMaterialIndex.Find(Surface.SurfaceId); MaterialIndex != INDEX_NONE)
 				{
@@ -6207,7 +6237,7 @@ void UCustomizableInstancePrivate::BuildMaterials(const TSharedRef<FUpdateContex
 						UTexture* PassThroughTexture = nullptr; // Texture not generated by mutable
 
 						// \TODO: Change this key to a struct.
-						FString TextureReuseCacheRef = bReuseTextures ? FString::Printf(TEXT("%d-%d-%d-%d"), Image.BaseLOD, ComponentIndex, Surface.SurfaceId, ImageIndex) : FString();
+						FString TextureReuseCacheRef = bReuseTextures ? FString::Printf(TEXT("%d-%d-%d-%d"), Image.BaseLOD, ObjectComponentIndex, Surface.SurfaceId, ImageIndex) : FString();
 
 						// If the mutable image is null, it must be in the cache
 						FMutableImageCacheKey ImageCacheKey = { Image.ImageID, OperationData->MipsToSkip };
@@ -6326,10 +6356,10 @@ void UCustomizableInstancePrivate::BuildMaterials(const TSharedRef<FUpdateContex
 									}
 									else
 									{
-										UE_LOG(LogMutable, Error, TEXT("Required image [%s] with ID [%lld] was not generated in the mutable thread, and it is not cached. LOD [%d]. Component [%d]"),
+										UE_LOG(LogMutable, Error, TEXT("Required image [%s] with ID [%lld] was not generated in the mutable thread, and it is not cached. LOD [%d]. Object Component [%d]"),
 											*Props.TextureParameterName,
 											Image.ImageID,
-											LODIndex, ComponentIndex);
+											LODIndex, ObjectComponentIndex);
 										continue;
 									}
 
@@ -6385,11 +6415,11 @@ void UCustomizableInstancePrivate::BuildMaterials(const TSharedRef<FUpdateContex
 													if (MutableMipDataProviderFactory)
 													{
 														MutableMipDataProviderFactory->CustomizableObjectInstance = Public;
-														check(LODIndex < 256 && ComponentIndex < 256 && ImageIndex < 256);
+														check(LODIndex < 256 && ObjectComponentIndex < 256 && ImageIndex < 256);
 														MutableMipDataProviderFactory->ImageRef.ImageID = Image.ImageID;
 														MutableMipDataProviderFactory->ImageRef.SurfaceId = Surface.SurfaceId;
 														MutableMipDataProviderFactory->ImageRef.LOD = uint8(Image.BaseLOD);
-														MutableMipDataProviderFactory->ImageRef.Component = uint8(ComponentIndex);
+														MutableMipDataProviderFactory->ImageRef.Component = uint8(ObjectComponentIndex);
 														MutableMipDataProviderFactory->ImageRef.Image = uint8(ImageIndex);
 														MutableMipDataProviderFactory->ImageRef.BaseMip = uint8(Image.BaseMip);
 														MutableMipDataProviderFactory->UpdateContext = UpdateContext;
@@ -6560,7 +6590,7 @@ void UCustomizableInstancePrivate::BuildMaterials(const TSharedRef<FUpdateContex
 
 				// Assign the material to the slot, and add it to the  OverrideMaterials
 				MaterialSlot.MaterialInterface = MaterialInterface;
-				ComponentsData[ComponentIndex].OverrideMaterials.Add(MaterialInterface);
+				ComponentsData[ObjectComponentIndex].OverrideMaterials.Add(MaterialInterface);
 			}
 		}
 
@@ -6587,12 +6617,12 @@ void UCustomizableInstancePrivate::BuildMaterials(const TSharedRef<FUpdateContex
 			}
 
 			// Force recreate render state after replacing the materials to avoid a crash in the render pipeline if the old materials are GCed while in use.
-			RecreateRenderStateOnComponent[ComponentIndex] = SkeletalMesh->GetResourceForRendering()->IsInitialized() && SkeletalMesh->GetMaterials() != Materials;
+			RecreateRenderStateOnInstanceComponent[InstanceComponentIndex] = SkeletalMesh->GetResourceForRendering()->IsInitialized() && SkeletalMesh->GetMaterials() != Materials;
 
 			SkeletalMesh->SetMaterials(Materials);
 
 #if WITH_EDITOR
-			if (RecreateRenderStateOnComponent[ComponentIndex])
+			if (RecreateRenderStateOnInstanceComponent[InstanceComponentIndex])
 			{
 				// Close all open editors for this mesh to invalidate viewports.
 				GEditor->GetEditorSubsystem<UAssetEditorSubsystem>()->CloseAllEditorsForAsset(SkeletalMesh);
@@ -6606,7 +6636,7 @@ void UCustomizableInstancePrivate::BuildMaterials(const TSharedRef<FUpdateContex
 
 	// Force recreate render state if the mesh is reused and the materials have changed.
 	// TODO: MTBL-1697 Remove after merging ConvertResources and Callbacks.
-	if (RecreateRenderStateOnComponent.Find(true) != INDEX_NONE)
+	if (RecreateRenderStateOnInstanceComponent.Find(true) != INDEX_NONE)
 	{
 		MUTABLE_CPUPROFILER_SCOPE(BuildMaterials_RecreateRenderState);
 
@@ -6626,14 +6656,18 @@ void UCustomizableInstancePrivate::BuildMaterials(const TSharedRef<FUpdateContex
 			}
 #endif
 
-			const int32 ComponentIndex = CustomizableObjectInstanceUsage->GetComponentIndex();
-			if (!RecreateRenderStateOnComponent.IsValidIndex(ComponentIndex) || !RecreateRenderStateOnComponent[ComponentIndex])
+			// TODO: What is the proper way to refer to dynamic components?
+			const int32 InstanceComponentIndex = CustomizableObjectInstanceUsage->GetComponentIndex();
+			const FInstanceUpdateData::FComponent& Component = OperationData->InstanceUpdateData.Components[InstanceComponentIndex];
+			int32 ObjectComponentIndex = Component.Id;
+
+			if (!RecreateRenderStateOnInstanceComponent.IsValidIndex(InstanceComponentIndex) || !RecreateRenderStateOnInstanceComponent[InstanceComponentIndex])
 			{
 				continue;
 			}
 
 			USkeletalMeshComponent* AttachedParent = CustomizableObjectInstanceUsage->GetAttachParent();
-			if (!AttachedParent || AttachedParent->GetSkeletalMeshAsset() != SkeletalMeshes[ComponentIndex])
+			if (!AttachedParent || AttachedParent->GetSkeletalMeshAsset() != SkeletalMeshes[ObjectComponentIndex])
 			{
 				continue;
 			}
diff --git a/Engine/Plugins/Experimental/Mutable/Source/CustomizableObject/Private/MuCO/CustomizableObjectSystem.cpp b/Engine/Plugins/Experimental/Mutable/Source/CustomizableObject/Private/MuCO/CustomizableObjectSystem.cpp
--- a/Engine/Plugins/Experimental/Mutable/Source/CustomizableObject/Private/MuCO/CustomizableObjectSystem.cpp
+++ b/Engine/Plugins/Experimental/Mutable/Source/CustomizableObject/Private/MuCO/CustomizableObjectSystem.cpp
@@ -161,7 +161,7 @@ FUpdateContextPrivate::FUpdateContextPrivate(UCustomizableObjectInstance& InInst
 	CapturedDescriptor = Descriptor;
 	CapturedDescriptorHash = FDescriptorHash(Descriptor);
 	Parameters = Descriptor.GetParameters();
-	NumComponents = InInstance.GetCustomizableObject()->GetComponentCount();
+	NumObjectComponents = InInstance.GetCustomizableObject()->GetComponentCount();
 	FirstLODAvailable = InInstance.GetCustomizableObject()->GetPrivate()->GetMinLODIndex();
 	FirstResidentLOD = InInstance.GetPrivate()->FirstResidentLOD;
 
@@ -1689,26 +1689,20 @@ namespace impl
 		}
 
 		Operation->MutableInstance = MutableSystem->BeginUpdate(Operation->InstanceID, Operation->Parameters, Operation->GetCapturedDescriptor().GetState(), mu::System::AllLODs);
+		Operation->NumInstanceComponents = Operation->MutableInstance->GetComponentCount();
 	}
 
 	
 	void FixLODs(const TSharedRef<FUpdateContextPrivate>& Operation)
 	{
-		// This is set at context creation because for now the component count is a CO constant that includes all components
-		// regardles of what is generated.
-		//Operation->NumComponents = Operation->MutableInstance->GetComponentCount();
-		if (!Operation->NumComponents)
+		if (!Operation->NumObjectComponents)
 		{
 			return;
 		}
 
-		// For now this should always match, since components are not optional.
-		check(Operation->NumComponents== Operation->MutableInstance->GetComponentCount());
+		Operation->NumLODsAvailablePerComponent.SetNumZeroed(Operation->NumObjectComponents);
 
-		Operation->NumLODsAvailablePerComponent.SetNumZeroed(Operation->NumComponents);
-
-		int32 NumComponentsInGeneratedInstance = Operation->MutableInstance->GetComponentCount();
-		for (int32 ComponentIndex = 0; ComponentIndex < NumComponentsInGeneratedInstance; ++ComponentIndex)
+		for (int32 ComponentIndex = 0; ComponentIndex < Operation->NumInstanceComponents; ++ComponentIndex)
 		{
 			Operation->NumLODsAvailablePerComponent[ComponentIndex] = Operation->MutableInstance->GetLODCount(ComponentIndex);
 		}
@@ -1728,9 +1722,9 @@ namespace impl
 		
 		// Initialize RequestedLODs to zero if not set
 		TArray<uint16> RequestedLODs = Operation->GetRequestedLODs();
-		RequestedLODs.SetNumZeroed(Operation->NumComponents);
+		RequestedLODs.SetNumZeroed(Operation->NumObjectComponents);
 
-		for (int32 ComponentIndex = 0; ComponentIndex < Operation->NumComponents; ++ComponentIndex)
+		for (int32 ComponentIndex = 0; ComponentIndex < Operation->NumObjectComponents; ++ComponentIndex)
 		{
 			if (Operation->bStreamMeshLODs)
 			{
@@ -1789,7 +1783,7 @@ namespace impl
 	{
 		MUTABLE_CPUPROFILER_SCOPE(Subtask_Mutable_PrepareSkeletonData);
 
-		OperationData->InstanceUpdateData.Skeletons.SetNum(OperationData->NumComponents);
+		OperationData->InstanceUpdateData.Skeletons.SetNum(OperationData->NumObjectComponents);
 
 		for (int32 ComponentIndex = 0; ComponentIndex<OperationData->InstanceUpdateData.Components.Num(); ++ComponentIndex)
 		{
@@ -1866,7 +1860,7 @@ namespace impl
 			FInstanceUpdateData::FRealTimeMorphsComponentData& ComponentMorphTargetsData = 
 					OperationData->InstanceUpdateData.RealTimeMorphTargets[Component.Id];
 
-			ComponentMorphTargetsData.ComponentIndex = Component.Id;
+			ComponentMorphTargetsData.ObjectComponentIndex = Component.Id;
 		}
 
 		const int32 NumComponents = OperationData->InstanceUpdateData.RealTimeMorphTargets.Num();
@@ -1875,7 +1869,7 @@ namespace impl
 			FInstanceUpdateData::FRealTimeMorphsComponentData& ComponentMorphTargetsData = 
 					OperationData->InstanceUpdateData.RealTimeMorphTargets[ComponentIndex];
 			
-			if (ComponentMorphTargetsData.ComponentIndex == INDEX_NONE)
+			if (ComponentMorphTargetsData.ObjectComponentIndex == INDEX_NONE)
 			{
 				continue;
 			}
@@ -2195,9 +2189,9 @@ namespace impl
 
 		const TSharedRef<TArray<FGetImageData>> GetImagesData = MakeShared<TArray<FGetImageData>>();
 		
-		for (int32 ComponentIndex = 0; ComponentIndex < OperationData->NumComponents; ++ComponentIndex)
+		for (int32 InstanceComponentIndex = 0; InstanceComponentIndex < OperationData->NumInstanceComponents; ++InstanceComponentIndex)
 		{
-			FInstanceUpdateData::FComponent& Component = OperationData->InstanceUpdateData.Components[ComponentIndex];
+			FInstanceUpdateData::FComponent& Component = OperationData->InstanceUpdateData.Components[InstanceComponentIndex];
 			for (int32 MutableLODIndex = 0; MutableLODIndex < Component.LODCount; ++MutableLODIndex)
 			{
 				// Skip LODs outside the range we want to generate
@@ -2218,7 +2212,7 @@ namespace impl
 				
 
 				// This lambda does all the work to fill up the surface data
-				auto AddSurface = [&LOD, &SurfacesSharedId, &ModelResources, &GetImagesData, OperationData, MutableInstance, CustomizableObject, ComponentIndex, MutableLODIndex](uint32 SurfaceId, int32 InstanceSurfaceIndex)
+				auto AddSurface = [&LOD, &SurfacesSharedId, &ModelResources, &GetImagesData, OperationData, MutableInstance, CustomizableObject, InstanceComponentIndex, MutableLODIndex](uint32 SurfaceId, int32 InstanceSurfaceIndex)
 					{
 						int32 BaseSurfaceIndex = InstanceSurfaceIndex;
 						int32 BaseLODIndex = MutableLODIndex;
@@ -2230,7 +2224,7 @@ namespace impl
 						// Now Surface.MaterialIndex is decoded from a parameter at the end of this if()
 						Surface.SurfaceId = SurfaceId;
 
-						const int32 SharedSurfaceId = MutableInstance->GetSharedSurfaceId(ComponentIndex, MutableLODIndex, InstanceSurfaceIndex);
+						const int32 SharedSurfaceId = MutableInstance->GetSharedSurfaceId(InstanceComponentIndex, MutableLODIndex, InstanceSurfaceIndex);
 						const int32 SharedSurfaceIndex = SurfacesSharedId.Find(SharedSurfaceId);
 
 						SurfacesSharedId.Add(SharedSurfaceId);
@@ -2244,32 +2238,32 @@ namespace impl
 							}
 
 							// Find the first LOD where this surface can be found
-							MutableInstance->FindBaseSurfaceBySharedId(ComponentIndex, SharedSurfaceId, BaseSurfaceIndex, BaseLODIndex);
+							MutableInstance->FindBaseSurfaceBySharedId(InstanceComponentIndex, SharedSurfaceId, BaseSurfaceIndex, BaseLODIndex);
 
-							Surface.SurfaceId = MutableInstance->GetSurfaceId(ComponentIndex, BaseLODIndex, BaseSurfaceIndex);
+							Surface.SurfaceId = MutableInstance->GetSurfaceId(InstanceComponentIndex, BaseLODIndex, BaseSurfaceIndex);
 						}
 
 						// Vectors
 						Surface.FirstVector = OperationData->InstanceUpdateData.Vectors.Num();
-						Surface.VectorCount = MutableInstance->GetVectorCount(ComponentIndex, BaseLODIndex, BaseSurfaceIndex);
+						Surface.VectorCount = MutableInstance->GetVectorCount(InstanceComponentIndex, BaseLODIndex, BaseSurfaceIndex);
 						for (int32 VectorIndex = 0; VectorIndex < Surface.VectorCount; ++VectorIndex)
 						{
 							MUTABLE_CPUPROFILER_SCOPE(GetVector);
 							OperationData->InstanceUpdateData.Vectors.Push({});
 							FInstanceUpdateData::FVector& Vector = OperationData->InstanceUpdateData.Vectors.Last();
-							Vector.Name = MutableInstance->GetVectorName(ComponentIndex, BaseLODIndex, BaseSurfaceIndex, VectorIndex);
-							Vector.Vector = MutableInstance->GetVector(ComponentIndex, BaseLODIndex, BaseSurfaceIndex, VectorIndex);
+							Vector.Name = MutableInstance->GetVectorName(InstanceComponentIndex, BaseLODIndex, BaseSurfaceIndex, VectorIndex);
+							Vector.Vector = MutableInstance->GetVector(InstanceComponentIndex, BaseLODIndex, BaseSurfaceIndex, VectorIndex);
 						}
 
 						// Scalars
 						Surface.FirstScalar = OperationData->InstanceUpdateData.Scalars.Num();
-						Surface.ScalarCount = MutableInstance->GetScalarCount(ComponentIndex, BaseLODIndex, BaseSurfaceIndex);
+						Surface.ScalarCount = MutableInstance->GetScalarCount(InstanceComponentIndex, BaseLODIndex, BaseSurfaceIndex);
 						for (int32 ScalarIndex = 0; ScalarIndex < Surface.ScalarCount; ++ScalarIndex)
 						{
 							MUTABLE_CPUPROFILER_SCOPE(GetScalar)
 
-								const FName ScalarName = MutableInstance->GetScalarName(ComponentIndex, BaseLODIndex, BaseSurfaceIndex, ScalarIndex);
-							const float ScalarValue = MutableInstance->GetScalar(ComponentIndex, BaseLODIndex, BaseSurfaceIndex, ScalarIndex);
+								const FName ScalarName = MutableInstance->GetScalarName(InstanceComponentIndex, BaseLODIndex, BaseSurfaceIndex, ScalarIndex);
+							const float ScalarValue = MutableInstance->GetScalar(InstanceComponentIndex, BaseLODIndex, BaseSurfaceIndex, ScalarIndex);
 
 							FString EncodingMaterialIdString = "__MutableMaterialId";
 
@@ -2289,15 +2283,15 @@ namespace impl
 
 						// Images
 						Surface.FirstImage = OperationData->InstanceUpdateData.Images.Num();
-						Surface.ImageCount = MutableInstance->GetImageCount(ComponentIndex, BaseLODIndex, BaseSurfaceIndex);
+						Surface.ImageCount = MutableInstance->GetImageCount(InstanceComponentIndex, BaseLODIndex, BaseSurfaceIndex);
 						for (int32 ImageIndex = 0; ImageIndex < Surface.ImageCount; ++ImageIndex)
 						{
 							MUTABLE_CPUPROFILER_SCOPE(GetImageId);
 
 							const int32 UpdateDataImageIndex = OperationData->InstanceUpdateData.Images.AddDefaulted();
 							FInstanceUpdateData::FImage& Image = OperationData->InstanceUpdateData.Images.Last();
-							Image.Name = MutableInstance->GetImageName(ComponentIndex, BaseLODIndex, BaseSurfaceIndex, ImageIndex);
-							Image.ImageID = MutableInstance->GetImageId(ComponentIndex, BaseLODIndex, BaseSurfaceIndex, ImageIndex);
+							Image.Name = MutableInstance->GetImageName(InstanceComponentIndex, BaseLODIndex, BaseSurfaceIndex, ImageIndex);
+							Image.ImageID = MutableInstance->GetImageId(InstanceComponentIndex, BaseLODIndex, BaseSurfaceIndex, ImageIndex);
 							Image.FullImageSizeX = 0;
 							Image.FullImageSizeY = 0;
 							Image.BaseLOD = BaseLODIndex;
@@ -2335,10 +2329,10 @@ namespace impl
 				// and all defined surfaces will be present. 
 				if (LOD.Mesh->IsReference())
 				{
-					const int32 SurfaceCount = MutableInstance->GetSurfaceCount(ComponentIndex, MutableLODIndex);
+					const int32 SurfaceCount = MutableInstance->GetSurfaceCount(InstanceComponentIndex, MutableLODIndex);
 					for (int32 SurfaceIndex = 0; SurfaceIndex < SurfaceCount; ++SurfaceIndex)
 					{
-						uint32 SurfaceId = MutableInstance->GetSurfaceId(ComponentIndex, MutableLODIndex, SurfaceIndex);
+						uint32 SurfaceId = MutableInstance->GetSurfaceId(InstanceComponentIndex, MutableLODIndex, SurfaceIndex);
 						AddSurface(SurfaceId, SurfaceIndex);
 					}
 				}
@@ -2351,7 +2345,7 @@ namespace impl
 					{
 						const uint32 SurfaceId = LOD.Mesh->GetSurfaceId(MeshSurfaceIndex);
 
-						const int32 InstanceSurfaceIndex = MutableInstance->FindSurfaceById(ComponentIndex, MutableLODIndex, SurfaceId);
+						const int32 InstanceSurfaceIndex = MutableInstance->FindSurfaceById(InstanceComponentIndex, MutableLODIndex, SurfaceId);
 						check(LOD.Mesh->GetVertexCount() > 0 || InstanceSurfaceIndex >= 0);
 
 						if (InstanceSurfaceIndex >= 0)
@@ -2482,15 +2476,24 @@ namespace impl
 
 			const TSharedRef<TArray<FGetMeshData>> GetMeshesData = MakeShared<TArray<FGetMeshData>>();
 		
-			OperationData->InstanceUpdateData.Components.SetNum(OperationData->NumComponents);
-			for (int32 ComponentIndex = 0; ComponentIndex < Instance->GetComponentCount(); ++ComponentIndex)
+			OperationData->InstanceUpdateData.Components.SetNum(OperationData->NumInstanceComponents);
+			for (int32 InstanceComponentIndex = 0; InstanceComponentIndex < OperationData->NumInstanceComponents; ++InstanceComponentIndex)
 			{
-				FInstanceUpdateData::FComponent& Component = OperationData->InstanceUpdateData.Components[ComponentIndex];
+				FInstanceUpdateData::FComponent& Component = OperationData->InstanceUpdateData.Components[InstanceComponentIndex];
 				Component.FirstLOD = OperationData->InstanceUpdateData.LODs.Num();
-				Component.LODCount = OperationData->NumLODsAvailablePerComponent[ComponentIndex];
-				Component.Id = Instance->GetComponentId(ComponentIndex);
+				Component.Id = Instance->GetComponentId(InstanceComponentIndex);
+
+				int32 ObjectComponentIndex = Component.Id;
 
-				for (int32 MutableLODIndex = 0; MutableLODIndex < OperationData->NumLODsAvailablePerComponent[ComponentIndex]; ++MutableLODIndex)
+				if (!OperationData->NumLODsAvailablePerComponent.IsValidIndex(ObjectComponentIndex))
+				{
+					// It happens in degenerated cases with empty components.
+					continue;
+				}
+
+				Component.LODCount = OperationData->NumLODsAvailablePerComponent[ObjectComponentIndex];
+
+				for (int32 MutableLODIndex = 0; MutableLODIndex < Component.LODCount; ++MutableLODIndex)
 				{
 					// If the LOD is not generated we still add an empty one to keep indexes aligned.
 					const int32 UpdateDataLODIndex = OperationData->InstanceUpdateData.LODs.AddDefaulted();
@@ -2503,13 +2506,13 @@ namespace impl
 
 					FInstanceUpdateData::FLOD& LOD = OperationData->InstanceUpdateData.LODs.Last();
 
-					const bool bGenerateLOD = RequestedLODs.IsValidIndex(ComponentIndex) ? RequestedLODs[ComponentIndex] <= MutableLODIndex : true;
+					const bool bGenerateLOD = RequestedLODs.IsValidIndex(ObjectComponentIndex) ? RequestedLODs[ObjectComponentIndex] <= MutableLODIndex : true;
 
 					// Mesh
 					{
 						MUTABLE_CPUPROFILER_SCOPE(GetMesh);
 
-						LOD.MeshID = Instance->GetMeshId(ComponentIndex, MutableLODIndex);
+						LOD.MeshID = Instance->GetMeshId(InstanceComponentIndex, MutableLODIndex);
 
 						if (bGenerateLOD)
 						{
@@ -3229,23 +3232,22 @@ namespace impl
 		CreateMutableInstance(Operation);
 		FixLODs(Operation);
 
-		const int32 NumComponents = Operation->NumComponents;
-
-		Operation->MeshDescriptors.SetNum(NumComponents);
+		Operation->MeshDescriptors.SetNum(Operation->NumObjectComponents);
 
 		const TArray<uint16>& RequestedLODs = Operation->GetRequestedLODs();
-		for (int32 ComponentIndex = 0; ComponentIndex < NumComponents; ++ComponentIndex)
+		for (int32 InstanceComponentIndex = 0; InstanceComponentIndex < Operation->NumInstanceComponents; ++InstanceComponentIndex)
 		{
-			TArray<mu::FResourceID>& MeshId = Operation->MeshDescriptors[ComponentIndex];
+			int32 ObjectComponentIndex = Operation->MutableInstance->GetComponentId(InstanceComponentIndex);
+			TArray<mu::FResourceID>& MeshId = Operation->MeshDescriptors[ObjectComponentIndex];
 			MeshId.Init(MAX_uint64, MAX_MESH_LOD_COUNT);
 
 			
-			for (int32 LODIndex = Operation->GetMinLOD(); LODIndex < Operation->NumLODsAvailablePerComponent[ComponentIndex]; ++LODIndex)
+			for (int32 LODIndex = Operation->GetMinLOD(); LODIndex < Operation->NumLODsAvailablePerComponent[ObjectComponentIndex]; ++LODIndex)
 			{
-				const bool bGenerateLOD = RequestedLODs.IsValidIndex(ComponentIndex) ? RequestedLODs[ComponentIndex] <= LODIndex : true;
+				const bool bGenerateLOD = RequestedLODs.IsValidIndex(ObjectComponentIndex) ? RequestedLODs[ObjectComponentIndex] <= LODIndex : true;
 				if (bGenerateLOD)
 				{
-					MeshId[LODIndex] = Operation->MutableInstance->GetMeshId(ComponentIndex, LODIndex);
+					MeshId[LODIndex] = Operation->MutableInstance->GetMeshId(InstanceComponentIndex, LODIndex);
 				}
 			}
 		}
@@ -3457,7 +3459,7 @@ namespace impl
 		if (bRequestAllLODs)
 		{
 			TArray<uint16> RequestedLODs = Operation->GetRequestedLODs();
-			RequestedLODs.Init(0, Operation->NumComponents);
+			RequestedLODs.Init(0, Operation->NumObjectComponents);
 
 			Operation->SetRequestedLODs(RequestedLODs);
 		}
diff --git a/Engine/Plugins/Experimental/Mutable/Source/CustomizableObject/Public/MuCO/CustomizableObjectInstanceUsage.h b/Engine/Plugins/Experimental/Mutable/Source/CustomizableObject/Public/MuCO/CustomizableObjectInstanceUsage.h
--- a/Engine/Plugins/Experimental/Mutable/Source/CustomizableObject/Public/MuCO/CustomizableObjectInstanceUsage.h
+++ b/Engine/Plugins/Experimental/Mutable/Source/CustomizableObject/Public/MuCO/CustomizableObjectInstanceUsage.h
@@ -37,6 +37,8 @@ public:
 	void SetCustomizableObjectInstance(UCustomizableObjectInstance* CustomizableObjectInstance);
 	UCustomizableObjectInstance* GetCustomizableObjectInstance() const;
 
+	// TODO: What is the proper way to refer to dynamic components?
+	// For now assume this is an *Instance* component index and not an *Object* component index
 	void SetComponentIndex(int32 ComponentIndex);
 	int32 GetComponentIndex() const;
 
diff --git a/Engine/Plugins/Experimental/Mutable/Source/CustomizableObject/Public/MuCO/CustomizableObjectSystem.h b/Engine/Plugins/Experimental/Mutable/Source/CustomizableObject/Public/MuCO/CustomizableObjectSystem.h
--- a/Engine/Plugins/Experimental/Mutable/Source/CustomizableObject/Public/MuCO/CustomizableObjectSystem.h
+++ b/Engine/Plugins/Experimental/Mutable/Source/CustomizableObject/Public/MuCO/CustomizableObjectSystem.h
@@ -93,7 +93,10 @@ struct FMutableImageReference
 	uint32 SurfaceId = 0;
 
 	uint8 LOD = 0;
+	
+	/** Object component index. */
 	uint8 Component = 0;
+
 	uint8 Image = 0;
 
 	uint8 BaseMip = 0;
diff --git a/Engine/Plugins/Experimental/Mutable/Source/CustomizableObjectEditor/Private/MuCOE/GenerateMutableSource/GenerateMutableSourceComponent.cpp b/Engine/Plugins/Experimental/Mutable/Source/CustomizableObjectEditor/Private/MuCOE/GenerateMutableSource/GenerateMutableSourceComponent.cpp
--- a/Engine/Plugins/Experimental/Mutable/Source/CustomizableObjectEditor/Private/MuCOE/GenerateMutableSource/GenerateMutableSourceComponent.cpp
+++ b/Engine/Plugins/Experimental/Mutable/Source/CustomizableObjectEditor/Private/MuCOE/GenerateMutableSource/GenerateMutableSourceComponent.cpp
@@ -105,6 +105,13 @@ mu::Ptr<mu::NodeComponent> GenerateMutableSourceComponent(const UEdGraphPin * Pi
 
 						GenerationContext.ComponentMeshOverride = nullptr;
 					}
+					else
+					{
+						// Add an empty surface node anyway.
+						mu::Ptr<mu::NodeSurfaceNew> SurfaceNode = new mu::NodeSurfaceNew;
+						SurfaceNode->Mesh = MeshNode;
+						LODNode->Surfaces.Add(SurfaceNode);
+					}
 				}
 			}
 		}
diff --git a/Engine/Plugins/Experimental/Mutable/Source/CustomizableObjectEditor/Private/MuCOE/SMutableCodeViewer.cpp b/Engine/Plugins/Experimental/Mutable/Source/CustomizableObjectEditor/Private/MuCOE/SMutableCodeViewer.cpp
--- a/Engine/Plugins/Experimental/Mutable/Source/CustomizableObjectEditor/Private/MuCOE/SMutableCodeViewer.cpp
+++ b/Engine/Plugins/Experimental/Mutable/Source/CustomizableObjectEditor/Private/MuCOE/SMutableCodeViewer.cpp
@@ -18,6 +18,7 @@
 #include "MuCOE/SMutableIntViewer.h"
 #include "MuCOE/SMutableLayoutViewer.h"
 #include "MuCOE/SMutableMeshViewer.h"
+#include "MuCOE/SMutableInstanceViewer.h"
 #include "MuCOE/SMutableParametersWidget.h"
 #include "MuCOE/SMutableProjectorViewer.h"
 #include "MuCOE/SMutableScalarViewer.h"
@@ -1778,6 +1779,13 @@ void SMutableCodeViewer::OnSelectionChanged(TSharedPtr<FMutableCodeTreeElement>
 		break;
 	}
 
+	case mu::DT_INSTANCE:
+	{
+		// Create or reuse the UI
+		PrepareInstanceViewer();
+		break;
+	}
+
 	case mu::DT_SCALAR:
 	{
 		// Create or reuse the UI
@@ -2621,7 +2629,7 @@ void SMutableCodeViewer::Tick(const FGeometry& AllottedGeometry, const double In
 	case mu::DT_IMAGE:
 	{
 		check(PreviewImageViewer);
-		mu::ImagePtrConst MutableImage = System->GetPrivate()->BuildImage(MutableModel, PreviewParameters.get(), SelectedOperationAddress, MipsToSkip, 0);
+		mu::Ptr<const mu::Image> MutableImage = System->GetPrivate()->BuildImage(MutableModel, PreviewParameters.get(), SelectedOperationAddress, MipsToSkip, 0);
 		PreviewImageViewer->SetImage(MutableImage, 0);
 		break;
 	}
@@ -2629,11 +2637,19 @@ void SMutableCodeViewer::Tick(const FGeometry& AllottedGeometry, const double In
 	case mu::DT_MESH:
 	{
 		check(PreviewMeshViewer);
-		mu::MeshPtrConst MutableMesh = System->GetPrivate()->BuildMesh(MutableModel, PreviewParameters.get(), SelectedOperationAddress);
+		mu::Ptr<const mu::Mesh> MutableMesh = System->GetPrivate()->BuildMesh(MutableModel, PreviewParameters.get(), SelectedOperationAddress);
 		PreviewMeshViewer->SetMesh(MutableMesh);
 		break;
 	}
 
+	case mu::DT_INSTANCE:
+	{
+		check(PreviewInstanceViewer);
+		mu::Ptr<const mu::Instance> MutableInstance = System->GetPrivate()->BuildInstance(MutableModel, PreviewParameters.get(), SelectedOperationAddress);
+		PreviewInstanceViewer->SetInstance(MutableInstance, MutableModel, PreviewParameters, System);
+		break;
+	}
+
 	case mu::DT_BOOL:
 	{
 		check(PreviewBoolViewer);
@@ -2645,7 +2661,7 @@ void SMutableCodeViewer::Tick(const FGeometry& AllottedGeometry, const double In
 	case mu::DT_INT:
 	{
 		check(PreviewIntViewer);
-		const int MutableInt = System->GetPrivate()->BuildInt(MutableModel, PreviewParameters.get(), SelectedOperationAddress);
+		const int32 MutableInt = System->GetPrivate()->BuildInt(MutableModel, PreviewParameters.get(), SelectedOperationAddress);
 		PreviewIntViewer->SetInt(MutableInt);
 		break;
 	}
@@ -2733,6 +2749,16 @@ void SMutableCodeViewer::PrepareMeshViewer()
 	PreviewBorder->SetContent(PreviewMeshViewer.ToSharedRef());
 }
 
+void SMutableCodeViewer::PrepareInstanceViewer()
+{
+	if (!PreviewInstanceViewer)
+	{
+		PreviewInstanceViewer = SNew(SMutableInstanceViewer);
+	}
+
+	PreviewBorder->SetContent(PreviewInstanceViewer.ToSharedRef());
+}
+
 void SMutableCodeViewer::PrepareLayoutViewer()
 {
 	if (!PreviewLayoutViewer)
diff --git a/Engine/Plugins/Experimental/Mutable/Source/CustomizableObjectEditor/Private/MuCOE/SMutableCodeViewer.h b/Engine/Plugins/Experimental/Mutable/Source/CustomizableObjectEditor/Private/MuCOE/SMutableCodeViewer.h
--- a/Engine/Plugins/Experimental/Mutable/Source/CustomizableObjectEditor/Private/MuCOE/SMutableCodeViewer.h
+++ b/Engine/Plugins/Experimental/Mutable/Source/CustomizableObjectEditor/Private/MuCOE/SMutableCodeViewer.h
@@ -31,6 +31,7 @@ class SMutableImageViewer;
 class SMutableIntViewer;
 class SMutableLayoutViewer;
 class SMutableMeshViewer;
+class SMutableInstanceViewer;
 class SMutableParametersWidget;
 class SMutableProjectorViewer;
 class SMutableScalarViewer;
@@ -107,7 +108,10 @@ private:
 	/** Widget used to show the preview of image operation results. Once created it is reused to preserve the settings. */
 	TSharedPtr<SMutableImageViewer> PreviewImageViewer;
 
-	/** Widget used to show a preview of the mesh and the metadata it holds */
+	/** Widget used to show a preview of an instance. */
+	TSharedPtr<SMutableInstanceViewer> PreviewInstanceViewer;
+
+	/** Widget used to show a preview of a mesh and the metadata it holds */
 	TSharedPtr<SMutableMeshViewer> PreviewMeshViewer;
 
 	/** Widget used to show a preview of a mutable bool value */
@@ -767,7 +771,7 @@ private:
 	 * Control of "Skip Mips" for image operations.
 	 */
 
-	/** Operation type we are using to search for tree nodes. Driven by the UI */
+	/** */
 	int32 MipsToSkip = 0;
 
 	/** */
@@ -801,6 +805,7 @@ private:
 	void PrepareStringViewer();
 	void PrepareImageViewer();
 	void PrepareMeshViewer();
+	void PrepareInstanceViewer();
 	void PrepareLayoutViewer();
 	void PrepareProjectorViewer();
 };
diff --git a/Engine/Plugins/Experimental/Mutable/Source/CustomizableObjectEditor/Private/MuCOE/SMutableInstanceViewer.cpp b/Engine/Plugins/Experimental/Mutable/Source/CustomizableObjectEditor/Private/MuCOE/SMutableInstanceViewer.cpp
new file mode 100644
--- /dev/null
+++ b/Engine/Plugins/Experimental/Mutable/Source/CustomizableObjectEditor/Private/MuCOE/SMutableInstanceViewer.cpp
@@ -0,0 +1,611 @@
+// Copyright Epic Games, Inc. All Rights Reserved.
+
+#include "MuCOE/SMutableInstanceViewer.h"
+
+#include "Engine/SkeletalMesh.h"
+#include "Framework/Views/TableViewMetadata.h"
+#include "MuCOE/SMutableMeshViewport.h"
+#include "MuCOE/SMutableSkeletonViewer.h"
+#include "MuT/TypeInfo.h"
+#include "PropertyCustomizationHelpers.h"
+#include "Widgets/Layout/SScrollBox.h"
+#include "Widgets/Views/SListView.h"
+#include "MuR/SystemPrivate.h"
+
+class ITableRow;
+class STableViewBase;
+class SWidget;
+class USkeleton;
+
+#define LOCTEXT_NAMESPACE "CustomizableObjectEditor"
+
+
+/** */
+namespace MutableSurfaceListColumns
+{
+	static const FName IdColumnID("Id");
+	static const FName SharedIdID("SharedId");
+	static const FName CustomColumnID("CustomId");
+	static const FName ImageCountID("Images");
+	static const FName VectorCountID("Vectors");
+	static const FName ScalarCountID("Scalars");
+	static const FName StringCountID("Strings");
+};
+
+/* **/
+class SMutableInstanceSurfaceListRow : public SMultiColumnTableRow<TSharedPtr<FMutableInstanceViewerSurfaceElement>>
+{
+public:
+	void Construct(const FArguments& Args, const TSharedRef<STableViewBase>& InOwnerTableView, const TSharedPtr<FMutableInstanceViewerSurfaceElement>& InRowItem)
+	{
+		RowItem = InRowItem;
+		
+		SMultiColumnTableRow<TSharedPtr<FMutableInstanceViewerSurfaceElement>>::Construct(
+			STableRow::FArguments()
+			.ShowSelection(true)
+			, InOwnerTableView
+		);
+	}
+
+
+	virtual TSharedRef<SWidget> GenerateWidgetForColumn(const FName& InColumnName) override
+	{
+		if (InColumnName == MutableSurfaceListColumns::IdColumnID)
+		{
+			return SNew(SHorizontalBox) +
+				SHorizontalBox::Slot()
+				.Padding(4,0)
+			[
+				SNew(STextBlock).Text(RowItem->Id)
+			];
+		}
+
+		if (InColumnName == MutableSurfaceListColumns::SharedIdID)
+		{
+			return SNew(SHorizontalBox) +
+				SHorizontalBox::Slot()
+			[
+				SNew(STextBlock).Text(RowItem->SharedId)
+			];
+		}
+
+		if (InColumnName == MutableSurfaceListColumns::CustomColumnID)
+		{
+			return SNew(SHorizontalBox) + SHorizontalBox::Slot()
+			[
+				SNew(STextBlock).Text(RowItem->CustomId)
+			];
+		}
+
+		if (InColumnName == MutableSurfaceListColumns::ImageCountID)
+		{
+			return SNew(SHorizontalBox) + SHorizontalBox::Slot()
+				[
+					SNew(STextBlock).Text(RowItem->ImageCount)
+				];
+		}
+
+		if (InColumnName == MutableSurfaceListColumns::VectorCountID)
+		{
+			return SNew(SHorizontalBox) + SHorizontalBox::Slot()
+				[
+					SNew(STextBlock).Text(RowItem->VectorCount)
+				];
+		}
+
+		if (InColumnName == MutableSurfaceListColumns::ScalarCountID)
+		{
+			return SNew(SHorizontalBox) + SHorizontalBox::Slot()
+				[
+					SNew(STextBlock).Text(RowItem->ScalarCount)
+				];
+		}
+
+		if (InColumnName == MutableSurfaceListColumns::StringCountID)
+		{
+			return SNew(SHorizontalBox) + SHorizontalBox::Slot()
+				[
+					SNew(STextBlock).Text(RowItem->StringCount)
+				];
+		}
+
+		// Invalid column name so no widget will be produced 
+		return SNullWidget::NullWidget;
+	}
+
+private:
+	TSharedPtr<FMutableInstanceViewerSurfaceElement> RowItem;
+};
+
+
+/* **/
+namespace MutableLODListColumns
+{
+	static const FName LODIndexColumnID("LOD Index");
+	static const FName MeshIdColumnID("Mesh ID");
+	static const FName SurfacesColumnID("Surfaces");
+};
+
+/* **/
+class SMutableInstanceLODListRow : public SMultiColumnTableRow<TSharedPtr<FMutableInstanceViewerLODElement>>
+{
+public:
+	void Construct(const FArguments& Args, const TSharedRef<STableViewBase>& InOwnerTableView, 
+		const TSharedPtr<FMutableInstanceViewerLODElement>& InRowItem, 
+		TSharedPtr<SMutableInstanceViewer> InHost)
+	{
+		HostMutableInstanceViewer = InHost;
+		RowItem = InRowItem;
+
+		SMultiColumnTableRow<TSharedPtr<FMutableInstanceViewerLODElement>>::Construct(
+			STableRow::FArguments()
+			.ShowSelection(true)
+			, InOwnerTableView
+		);
+	}
+
+
+	virtual TSharedRef<SWidget> GenerateWidgetForColumn(const FName& InColumnName) override
+	{
+		if (InColumnName == MutableLODListColumns::LODIndexColumnID)
+		{
+			return SNew(SHorizontalBox) +
+				SHorizontalBox::Slot()
+				.Padding(4, 0)
+				[
+					SNew(STextBlock).Text(RowItem->LODIndex)
+				];
+		}
+
+		if (InColumnName == MutableLODListColumns::MeshIdColumnID)
+		{
+			return SNew(SHorizontalBox) +
+				SHorizontalBox::Slot()
+				[
+					SNew(STextBlock).Text(RowItem->MeshId)
+				];
+		}
+
+		if (InColumnName == MutableLODListColumns::SurfacesColumnID)
+		{
+			const TSharedRef<SWidget> GeneratedSurfaces = HostMutableInstanceViewer->GenerateSurfaceListView(RowItem->Surfaces);
+
+			return SNew(SHorizontalBox)
+				+ SHorizontalBox::Slot()
+				[
+					GeneratedSurfaces
+				];
+		}
+
+		// Invalid column name so no widget will be produced 
+		return SNullWidget::NullWidget;
+	}
+
+private:
+	TSharedPtr<FMutableInstanceViewerLODElement> RowItem;
+	TSharedPtr<SMutableInstanceViewer> HostMutableInstanceViewer;
+};
+
+
+/** Namespace containing the IDs for the header on the buffers list */
+namespace MutableInstanceComponentsListColumns
+{
+	static const FName ComponentIndexColumnID("Component Index");
+	static const FName LODsColumnID("LODs");
+}
+
+
+class SMutableInstanceComponentListRow : public SMultiColumnTableRow<TSharedPtr<FMutableInstanceViewerComponentElement>>
+{
+public:
+	void Construct(const FArguments& Args, const TSharedRef<STableViewBase>& InOwnerTableView,
+	               const TSharedPtr<FMutableInstanceViewerComponentElement>& InRowItem, TSharedPtr<SMutableInstanceViewer> InHost)
+	{
+		HostMutableInstanceViewer = InHost;
+		RowItem = InRowItem;
+		
+		SMultiColumnTableRow<TSharedPtr<FMutableInstanceViewerComponentElement>>::Construct(
+			STableRow::FArguments()
+			.ShowSelection(true)
+			, InOwnerTableView
+		);
+	}
+
+	virtual TSharedRef<SWidget> GenerateWidgetForColumn(const FName& InColumnName) override
+	{
+		// Column with the index for the buffer .
+		// Useful for knowing the channels on what buffer reside
+		if (InColumnName == MutableInstanceComponentsListColumns::ComponentIndexColumnID)
+		{
+			return SNew(SBorder)
+			.Content()
+			[
+				SNew(SHorizontalBox)
+				+ SHorizontalBox::Slot()
+				.HAlign(HAlign_Center)
+				.VAlign(EVerticalAlignment::VAlign_Center)
+				[
+					SNew(STextBlock).Text(RowItem->ComponentIndex)
+				]
+			];
+		}
+
+		// Generate the sub table here
+		if (InColumnName == MutableInstanceComponentsListColumns::LODsColumnID)
+		{
+			const TSharedRef<SWidget> GeneratedLODList = HostMutableInstanceViewer->GenerateLODsListView(RowItem->LODs);
+			
+			return SNew(SHorizontalBox)
+			+ SHorizontalBox::Slot()
+			[
+				GeneratedLODList
+			];
+		}
+
+		// Invalid column name so no widget will be produced 
+		return SNullWidget::NullWidget;
+	}
+
+private:
+	TSharedPtr<FMutableInstanceViewerComponentElement> RowItem;
+	TSharedPtr<SMutableInstanceViewer> HostMutableInstanceViewer;
+};
+
+
+void SMutableInstanceViewer::Construct(const FArguments& InArgs)
+{
+	// Splitter values
+	constexpr float TablesSplitterValue = 0.5f;
+	constexpr float ViewportSplitterValue = 0.5f;
+
+	ChildSlot
+	[
+		SNew(SSplitter)
+
+		+ SSplitter::Slot()
+		.Value(TablesSplitterValue)
+		[
+			GenerateDataTableSlates()
+		]
+
+		+ SSplitter::Slot()
+		.Value(ViewportSplitterValue)
+		[
+			GenerateViewportSlates()
+		]
+	];
+}
+
+
+void SMutableInstanceViewer::SetInstance(const mu::InstancePtrConst& InInstance, const TSharedPtr<const mu::Model>& Model, const mu::Ptr<const mu::Parameters>& Parameters, const mu::Ptr<mu::System>& System)
+{
+	if (InInstance != MutableInstance)
+	{
+		MutableInstance = InInstance;
+
+		if (MutableInstance)
+		{
+			// Make sure no data is left from previous runs
+			Components.Empty();
+
+			const int32 ComponentCount = MutableInstance->GetComponentCount();
+			for (int32 ComponentIndex = 0; ComponentIndex < ComponentCount; ++ComponentIndex)
+			{
+				TSharedPtr<FMutableInstanceViewerComponentElement> NewComponent = MakeShareable(new FMutableInstanceViewerComponentElement());
+				Components.Add(NewComponent);
+
+				NewComponent->ComponentIndex = FText::AsNumber(ComponentIndex);
+
+				TSharedPtr<TArray<TSharedPtr<FMutableInstanceViewerLODElement>>> LODArray = MakeShareable(new TArray<TSharedPtr<FMutableInstanceViewerLODElement>>);
+				NewComponent->LODs = LODArray;
+
+				const int32 LODCount = MutableInstance->GetLODCount(ComponentIndex);
+				for (int32 LODIndex = 0; LODIndex < LODCount; ++LODIndex)
+				{
+					TSharedPtr<FMutableInstanceViewerLODElement> NewLOD = MakeShareable(new FMutableInstanceViewerLODElement());
+					LODArray->Add(NewLOD);
+
+					NewLOD->LODIndex = FText::AsNumber(LODIndex);
+					mu::FResourceID MeshId = MutableInstance->GetMeshId(ComponentIndex, LODIndex);
+					NewLOD->MeshId = FText::AsNumber(MeshId);
+
+					// Generate the mesh to get some data
+					{
+						mu::OP::ADDRESS MeshRootAddress = mu::GetResourceIDRoot(MeshId);
+						mu::Ptr<const mu::Mesh> MutableMesh = System->GetPrivate()->BuildMesh(Model, Parameters.get(), MeshRootAddress);
+						if (MutableMesh && MutableMesh->IsReference())
+						{
+							NewLOD->MeshId = FText::Format(
+								LOCTEXT("MeshReference", "Id [{0}]  Ref [{1}] "), MeshId, MutableMesh->GetReferencedMesh());
+						}
+					}
+
+					TSharedPtr<TArray<TSharedPtr<FMutableInstanceViewerSurfaceElement>>> SurfaceArray = MakeShareable(new TArray<TSharedPtr<FMutableInstanceViewerSurfaceElement>>);
+					NewLOD->Surfaces = SurfaceArray;
+
+					const int32 SurfaceCount = MutableInstance->GetSurfaceCount(ComponentIndex, LODIndex);
+					for (int32 SurfaceIndex = 0; SurfaceIndex < SurfaceCount; ++SurfaceIndex)
+					{
+						TSharedPtr<FMutableInstanceViewerSurfaceElement> NewSurface = MakeShareable(new FMutableInstanceViewerSurfaceElement());
+
+						NewSurface->Id = FText::AsNumber(MutableInstance->GetSurfaceId(ComponentIndex, LODIndex, SurfaceIndex));
+						NewSurface->SharedId = FText::AsNumber(MutableInstance->GetSharedSurfaceId(ComponentIndex, LODIndex, SurfaceIndex));
+						NewSurface->CustomId = FText::AsNumber(MutableInstance->GetSurfaceCustomId(ComponentIndex, LODIndex, SurfaceIndex));
+						NewSurface->ImageCount = FText::AsNumber(MutableInstance->GetImageCount(ComponentIndex, LODIndex, SurfaceIndex));
+						NewSurface->VectorCount = FText::AsNumber(MutableInstance->GetVectorCount(ComponentIndex, LODIndex, SurfaceIndex));
+						NewSurface->ScalarCount = FText::AsNumber(MutableInstance->GetScalarCount(ComponentIndex, LODIndex, SurfaceIndex));
+						NewSurface->StringCount = FText::AsNumber(MutableInstance->GetStringCount(ComponentIndex, LODIndex, SurfaceIndex));
+
+						SurfaceArray->Add(NewSurface);
+					}
+				}
+			}
+
+			// Make sure the list gets refreshed with the new contents
+			ComponentsSlateView->RequestListRefresh();
+
+			// Restore the widths of the columns each time the Instance gets changed.
+			ComponentsSlateView->GetHeaderRow()->ResetColumnWidths();
+		}
+
+		// TODO: No 3D preview for instances yet
+		//InstanceViewport->SetInstance(MutableInstance);
+	}
+}
+
+TSharedRef<SWidget> SMutableInstanceViewer::GenerateViewportSlates()
+{
+	TSharedRef<SWidget> Container = SNew(SVerticalBox)
+
+		// User warning messages 
+		// + SVerticalBox::Slot()
+		// .AutoHeight()
+		// [
+		// 	// TODO: Add a message to tell the user why no Instance is being displayed
+		// 	SNew(SWarningOrErrorBox)
+		// 	.MessageStyle(EMessageStyle::Warning)
+		// 	.Message(FText::FromString(FString("This is just a simulation")))
+		// ]
+
+		// Instance Drawing space
+		+ SVerticalBox::Slot()
+		[
+			SAssignNew(InstanceViewport, SMutableMeshViewport)
+			// TODO: No 3D preview for instances yet
+			//.Instance(MutableInstance)
+		];
+
+
+	return Container;
+}
+
+
+TSharedRef<SWidget> SMutableInstanceViewer::GenerateDataTableSlates()
+{
+	// Formatting
+	constexpr int32 IndentationSpace = 16;
+	constexpr int32 SimpleSpacing = 1;
+
+	constexpr int32 AfterTitleSpacing = 4;
+	constexpr int32 EndOfSectionSpacing = 12;
+
+	// Naming
+	const FText GeneralDataTitle = LOCTEXT("GeneralDataTitle", "General Data");
+	const FText InstanceIdTitle = LOCTEXT("InstanceIdTitle", "Instance Id : ");
+	const FText ComponentsTitle = LOCTEXT("ComponentsTitle", "Components");
+	
+	return SNew(SScrollBox)
+		+ SScrollBox::Slot()
+		[
+			SNew(SVerticalBox)
+
+			// General data ----------------------------------------------------------------
+			+ SVerticalBox::Slot()
+			.AutoHeight()
+			[
+				SNew(STextBlock).
+				Text(GeneralDataTitle)
+			]
+
+			+ SVerticalBox::Slot().
+			  Padding(IndentationSpace, AfterTitleSpacing).
+			  AutoHeight()
+			[
+				SNew(SVerticalBox)
+
+				// Instace ID
+				+ SVerticalBox::Slot().
+				  Padding(0, SimpleSpacing).
+				  AutoHeight()
+				[
+					SNew(SHorizontalBox)
+
+					// Vertices title
+					+ SHorizontalBox::Slot()
+					.AutoWidth()
+					[
+						SNew(STextBlock).
+						Text(InstanceIdTitle)
+					]
+
+					// Vertices value
+					+ SHorizontalBox::Slot()
+					.AutoWidth()
+					[
+						SNew(STextBlock).
+						Text(this, &SMutableInstanceViewer::GetInstanceID)
+					]
+
+				]
+			]
+
+			// Components Data --------------------------------------------------------------
+			+ SVerticalBox::Slot()
+			  .Padding(0, EndOfSectionSpacing)
+			  .AutoHeight()
+			[
+				SNew(SVerticalBox)
+
+				+ SVerticalBox::Slot().
+				AutoHeight()
+				[
+					SNew(STextBlock).
+					Text(ComponentsTitle)
+				]
+
+				+ SVerticalBox::Slot()
+				  .Padding(IndentationSpace, AfterTitleSpacing)
+				  .AutoHeight()
+				[
+					SNew(SVerticalBox)
+
+					// List of components ----------
+					+ SVerticalBox::Slot()
+					.AutoHeight()
+					[
+						GenerateComponentsListView()
+					]
+
+				]
+				
+			]
+		];
+}
+
+
+TSharedRef<SWidget> SMutableInstanceViewer::GenerateComponentsListView()
+{
+	// Headers
+	const FText ComponentIndexTitle = FText(LOCTEXT("ComponentIndexTitle", "Component"));
+	const FText LODTitle = FText(LOCTEXT("LODsTitle", "LODs"));
+
+	// Tooltips
+	const FText ComponentIndexTooltip = FText(LOCTEXT("ComponentIndexTooltip", ""));
+	const FText LODsTooltip = FText(LOCTEXT("LODsTooltip", ""));
+	
+	return SAssignNew(ComponentsSlateView, SListView<TSharedPtr<FMutableInstanceViewerComponentElement>>)
+		.ListItemsSource(&Components)
+		.OnGenerateRow(this, &SMutableInstanceViewer::OnGenerateComponentRow)
+		.SelectionMode(ESelectionMode::None)
+		.HeaderRow
+		(
+			SNew(SHeaderRow)
+		
+			+ SHeaderRow::Column(MutableInstanceComponentsListColumns::ComponentIndexColumnID)
+				.DefaultTooltip(ComponentIndexTitle)
+				.DefaultLabel(ComponentIndexTooltip)
+				.FillWidth(0.1f)
+		
+			+ SHeaderRow::Column(MutableInstanceComponentsListColumns::LODsColumnID)
+				.DefaultTooltip(LODTitle)
+				.DefaultLabel(LODsTooltip)
+			  	.FillWidth(0.9f)
+		);
+}
+
+
+TSharedRef<SWidget> SMutableInstanceViewer::GenerateLODsListView( const TSharedPtr<TArray<TSharedPtr<FMutableInstanceViewerLODElement>>>& InLODs)
+{
+	if (!InLODs.IsValid() || InLODs->IsEmpty())
+	{
+		return SNew(STextBlock).
+			Text(FText(LOCTEXT("NoLODs", "No LODs Found")));
+	}
+
+	// Headers
+	const FText LODIndexTitle = FText(LOCTEXT("ChannelIndexTitle", "Index"));
+	const FText MeshIdTitle = FText(LOCTEXT("MeshIdLabelTitle", "Mesh Id"));
+	const FText SurfacesTitle = FText(LOCTEXT("SurfacesLabelTitle", "Surfaces"));
+	
+	return SNew(SListView<TSharedPtr<FMutableInstanceViewerLODElement>>)
+		.ListItemsSource(InLODs.Get())
+		.OnGenerateRow(this, &SMutableInstanceViewer::OnGenerateLODRow)
+		.SelectionMode(ESelectionMode::None)
+		.HeaderRow
+		(
+			SNew(SHeaderRow)
+
+			+ SHeaderRow::Column(MutableLODListColumns::LODIndexColumnID)
+				.DefaultLabel(LODIndexTitle)
+				.FillWidth(0.14f)
+			
+			+ SHeaderRow::Column(MutableLODListColumns::MeshIdColumnID)
+				.DefaultLabel(MeshIdTitle)
+				.FillWidth(0.14f)
+		
+			+ SHeaderRow::Column(MutableLODListColumns::SurfacesColumnID)
+			  	.DefaultLabel(SurfacesTitle)
+			  	.FillWidth(0.65f)
+		);
+}
+
+
+TSharedRef<SWidget> SMutableInstanceViewer::GenerateSurfaceListView(const TSharedPtr<TArray<TSharedPtr<FMutableInstanceViewerSurfaceElement>>>& InSurfaces)
+{
+	if (!InSurfaces.IsValid() || InSurfaces->IsEmpty())
+	{
+		return SNew(STextBlock).
+			Text(FText(LOCTEXT("NoSurfaces", "No Surfaces Found")));
+	}
+
+	return SNew(SListView<TSharedPtr<FMutableInstanceViewerSurfaceElement>>)
+		.ListItemsSource(InSurfaces.Get())
+		.OnGenerateRow(this, &SMutableInstanceViewer::OnGenerateSurfaceRow)
+		.SelectionMode(ESelectionMode::None)
+		.HeaderRow
+		(
+			SNew(SHeaderRow)
+
+			+ SHeaderRow::Column(MutableSurfaceListColumns::IdColumnID)
+			.DefaultLabel(LOCTEXT("IdTitle", "Id"))
+			.FillWidth(0.14f)
+
+			+ SHeaderRow::Column(MutableSurfaceListColumns::SharedIdID)
+			.DefaultLabel(LOCTEXT("SharedIdTitle", "SharedId"))
+			.FillWidth(0.35f)
+
+			+ SHeaderRow::Column(MutableSurfaceListColumns::CustomColumnID)
+			.DefaultLabel(LOCTEXT("CustomIdTitle", "CustomId"))
+			.FillWidth(0.65f)
+
+			+ SHeaderRow::Column(MutableSurfaceListColumns::ImageCountID)
+			.DefaultLabel(LOCTEXT("ImagesTitle", "Images"))
+			.FillWidth(0.3f)
+
+			+ SHeaderRow::Column(MutableSurfaceListColumns::VectorCountID)
+			.DefaultLabel(LOCTEXT("VectorsTitle", "Vectors"))
+			.FillWidth(0.3f)
+
+			+ SHeaderRow::Column(MutableSurfaceListColumns::ScalarCountID)
+			.DefaultLabel(LOCTEXT("ScalarsTitle", "Scalars"))
+			.FillWidth(0.3f)
+
+			+ SHeaderRow::Column(MutableSurfaceListColumns::StringCountID)
+			.DefaultLabel(LOCTEXT("StringsTitle", "Strings"))
+			.FillWidth(0.3f)
+		);
+}
+
+
+FText SMutableInstanceViewer::GetInstanceID() const
+{
+	return FText::AsNumber(MutableInstance ? MutableInstance->GetId() : 0);
+}
+
+
+TSharedRef<ITableRow> SMutableInstanceViewer::OnGenerateComponentRow(TSharedPtr<FMutableInstanceViewerComponentElement> In, const TSharedRef<STableViewBase>& OwnerTable)
+{
+	return SNew(SMutableInstanceComponentListRow, OwnerTable, In, SharedThis(this) );
+}
+
+TSharedRef<ITableRow> SMutableInstanceViewer::OnGenerateLODRow(TSharedPtr<FMutableInstanceViewerLODElement> In, const TSharedRef<STableViewBase>& OwnerTable)
+{
+	return SNew(SMutableInstanceLODListRow, OwnerTable, In, SharedThis(this));
+}
+
+TSharedRef<ITableRow> SMutableInstanceViewer::OnGenerateSurfaceRow(TSharedPtr<FMutableInstanceViewerSurfaceElement> In, const TSharedRef<STableViewBase>& OwnerTable)
+{
+	return SNew(SMutableInstanceSurfaceListRow, OwnerTable, In);
+}
+
+#undef LOCTEXT_NAMESPACE
diff --git a/Engine/Plugins/Experimental/Mutable/Source/CustomizableObjectEditor/Private/MuCOE/SMutableInstanceViewer.h b/Engine/Plugins/Experimental/Mutable/Source/CustomizableObjectEditor/Private/MuCOE/SMutableInstanceViewer.h
new file mode 100644
--- /dev/null
+++ b/Engine/Plugins/Experimental/Mutable/Source/CustomizableObjectEditor/Private/MuCOE/SMutableInstanceViewer.h
@@ -0,0 +1,110 @@
+// Copyright Epic Games, Inc. All Rights Reserved.
+
+#pragma once
+
+#include "MuR/Instance.h"
+#include "MuR/System.h"
+#include "UObject/GCObject.h"
+#include "Widgets/SCompoundWidget.h"
+
+template <typename ItemType> class SListView;
+
+class ITableRow;
+class SBorder;
+class SMutableMeshViewport;
+class SScrollBox;
+class SSplitter;
+class STableViewBase;
+class SWidget;
+
+
+/** */
+struct FMutableInstanceViewerSurfaceElement
+{
+	FText Id;
+	FText SharedId;
+	FText CustomId;
+	FText ImageCount;
+	FText VectorCount;
+	FText ScalarCount;
+	FText StringCount;
+};
+
+
+/** */
+struct FMutableInstanceViewerLODElement
+{
+	FText LODIndex;
+	FText MeshId;
+	TSharedPtr<TArray<TSharedPtr<FMutableInstanceViewerSurfaceElement>>> Surfaces;
+};
+
+
+/** */
+struct FMutableInstanceViewerComponentElement
+{
+	FText ComponentIndex;
+	TSharedPtr<TArray<TSharedPtr<FMutableInstanceViewerLODElement>>> LODs;
+};
+
+
+/** Widget designed to show the statistical data from a Mutable Instance*/
+class SMutableInstanceViewer : public SCompoundWidget
+{
+public:
+	SLATE_BEGIN_ARGS(SMutableInstanceViewer) {}
+	SLATE_END_ARGS()
+
+	/** Builds the widget */
+	void Construct(const FArguments& InArgs);
+
+	/** Set the Mutable Instance to be used for this widget */
+	void SetInstance(const mu::Ptr<const mu::Instance>&, const TSharedPtr<const mu::Model>&, const mu::Ptr<const mu::Parameters>&, const mu::Ptr<mu::System>&);
+
+	
+	/** */
+	TSharedRef<SWidget> GenerateLODsListView(const TSharedPtr<TArray<TSharedPtr<FMutableInstanceViewerLODElement>>>& InBufferChannelElements);
+	TSharedRef<SWidget> GenerateSurfaceListView(const TSharedPtr<TArray<TSharedPtr<FMutableInstanceViewerSurfaceElement>>>& InBufferChannelElements);
+
+private:
+
+	/** Data backend for the widget. It represents the Instance that is being "displayed" */
+	mu::Ptr<const mu::Instance> MutableInstance = nullptr;
+
+	/** Splitter used to separate the two sides of the slate (tables and viewport) */
+	TSharedPtr<SSplitter> SpaceSplitter;
+
+	/** Slate object containing all the buffer tables alongside with the bone tree */
+	TSharedPtr<SScrollBox> DataSpaceSlate;
+
+	/** Viewport object to preview the current Instance inside an actual Unreal scene */
+	TSharedPtr<SMutableMeshViewport> InstanceViewport;
+	
+	/** */	
+	TSharedPtr<SListView<TSharedPtr<FMutableInstanceViewerComponentElement>>> ComponentsSlateView;
+
+private:
+
+	/** Generates all slate objects related with the Instance Viewport Slate */
+	TSharedRef<SWidget> GenerateViewportSlates();
+
+	/** Generates the tables showing the buffer data on the Instance alongside with the bone tree found on the mutable Instance */
+	TSharedRef<SWidget> GenerateDataTableSlates();
+
+private:
+	
+	/** */
+	FText GetInstanceID() const;
+
+	/** */	
+	TArray<TSharedPtr<FMutableInstanceViewerComponentElement>> Components;
+	
+	/** */
+	TSharedRef<SWidget> GenerateComponentsListView();
+	
+	/**  */
+	TSharedRef<ITableRow> OnGenerateComponentRow(TSharedPtr<FMutableInstanceViewerComponentElement>, const TSharedRef<STableViewBase>& OwnerTable);
+	TSharedRef<ITableRow> OnGenerateLODRow(TSharedPtr<FMutableInstanceViewerLODElement>, const TSharedRef<STableViewBase>& OwnerTable);
+	TSharedRef<ITableRow> OnGenerateSurfaceRow(TSharedPtr<FMutableInstanceViewerSurfaceElement>, const TSharedRef<STableViewBase>& OwnerTable);
+
+};
diff --git a/Engine/Plugins/Experimental/Mutable/Source/CustomizableObjectEditor/Private/MuCOE/SMutableMeshViewer.cpp b/Engine/Plugins/Experimental/Mutable/Source/CustomizableObjectEditor/Private/MuCOE/SMutableMeshViewer.cpp
--- a/Engine/Plugins/Experimental/Mutable/Source/CustomizableObjectEditor/Private/MuCOE/SMutableMeshViewer.cpp
+++ b/Engine/Plugins/Experimental/Mutable/Source/CustomizableObjectEditor/Private/MuCOE/SMutableMeshViewer.cpp
@@ -535,7 +535,7 @@ TSharedRef<SWidget> SMutableMeshViewer::GenerateBuffersListView(
 				.FillWidth(0.1f)
 		
 			+ SHeaderRow::Column(MutableMeshBuffersListColumns::BufferChannelsColumnID)
-				.DefaultLabel(BufferChannelsTooltip)
+				.DefaultTooltip(BufferChannelsTooltip)
 				.DefaultLabel(BufferChannelsTitle)
 			  	.FillWidth(0.9f)
 		);
diff --git a/Engine/Plugins/Experimental/Mutable/Source/CustomizableObjectEditor/Private/MuCOE/SMutableMeshViewer.h b/Engine/Plugins/Experimental/Mutable/Source/CustomizableObjectEditor/Private/MuCOE/SMutableMeshViewer.h
--- a/Engine/Plugins/Experimental/Mutable/Source/CustomizableObjectEditor/Private/MuCOE/SMutableMeshViewer.h
+++ b/Engine/Plugins/Experimental/Mutable/Source/CustomizableObjectEditor/Private/MuCOE/SMutableMeshViewer.h
@@ -8,9 +8,6 @@
 
 template <typename ItemType> class SListView;
 
-class FAssetThumbnail;
-class FAssetThumbnailPool;
-class FReferenceCollector;
 class ITableRow;
 class SBorder;
 class SMutableMeshViewport;
@@ -27,10 +24,10 @@ struct FAssetData;
 /** Container designed to hold the buffer channel data of a mutable mesh buffer to be later used by the UI */
 struct FBufferChannelElement
 {
-	FText SemanticIndex = FText(INVTEXT(""));
-	FText BufferSemantic = FText(INVTEXT(""));
-	FText BufferFormat = FText(INVTEXT(""));
-	FText BufferComponentCount = FText(INVTEXT(""));
+	FText SemanticIndex;
+	FText BufferSemantic;
+	FText BufferFormat;
+	FText BufferComponentCount;
 };
 
 
@@ -38,7 +35,7 @@ struct FBufferChannelElement
 struct FBufferElement
 {
 	/** The index of the buffer on the origin mutable buffer set */
-	FText BufferIndex = FText(INVTEXT(""));
+	FText BufferIndex;
 
 	/** An array of BufferChannels that represent the relative mutable channels on the mutable buffer*/
 	TSharedPtr<TArray<TSharedPtr<FBufferChannelElement>>> BufferChannels;
@@ -75,10 +72,6 @@ private:
 	TArray<TSharedPtr<FString>> MeshTagList;
 	TSharedRef<ITableRow> GenerateTagRow(TSharedPtr<FString> InItem, const TSharedRef<STableViewBase>& OwnerTable);
 
-	/** Thumbnail objects used to render a view of the selected skeletal mesh for the mesh viewport*/
-	TSharedPtr<FAssetThumbnailPool> AssetThumbnailPool;
-	TSharedPtr<FAssetThumbnail> AssetThumbnail;
-
 	/** Data backend for the widget. It represents the mesh that is being "displayed" */
 	mu::MeshPtrConst MutableMesh = nullptr;
 
@@ -107,13 +100,6 @@ private:
 	/** Generates the tables showing the buffer data on the mesh alongside with the bone tree found on the mutable mesh */
 	TSharedRef<SWidget> GenerateDataTableSlates();
 	
-	/** Width and height for the skeletal mesh thumbnail */
-	FVector2f ThumbnailSize = FVector2f();
-	FVector2f WidgetSize = FVector2f();
-
-	// Single property that only draws the combo box widget of the skeletal mesh
-	TSharedPtr<class ISinglePropertyView> SkeletalMeshSelector;
-	
 	/** Method called each time the mesh selected changes so the UI gets updated reliably */
 	void OnMeshChanged();
 
diff --git a/Engine/Plugins/Experimental/Mutable/Source/MutableRuntime/Internal/MuR/SystemPrivate.h b/Engine/Plugins/Experimental/Mutable/Source/MutableRuntime/Internal/MuR/SystemPrivate.h
--- a/Engine/Plugins/Experimental/Mutable/Source/MutableRuntime/Internal/MuR/SystemPrivate.h
+++ b/Engine/Plugins/Experimental/Mutable/Source/MutableRuntime/Internal/MuR/SystemPrivate.h
@@ -29,6 +29,11 @@ namespace mu::MemoryCounters
 
 namespace mu
 {
+	inline uint32 GetResourceIDRoot(FResourceID Id)
+	{
+		return uint32(Id >> 32);
+	}
+
 	class ExtensionDataStreamer;
 
 	// Call the tick of the LLM system (we do this to simulate a frame since the LLM system is not entirelly designed to run over a program)
@@ -1900,6 +1905,7 @@ namespace mu
 		MUTABLERUNTIME_API Ptr<const String> BuildString(const TSharedPtr<const Model>&, const Parameters*, OP::ADDRESS) ;
 		MUTABLERUNTIME_API Ptr<const Image> BuildImage(const TSharedPtr<const Model>&, const Parameters*, OP::ADDRESS, int32 MipsToSkip, int32 LOD) ;
 		MUTABLERUNTIME_API Ptr<const Mesh> BuildMesh(const TSharedPtr<const Model>&, const Parameters*, OP::ADDRESS) ;
+		MUTABLERUNTIME_API Ptr<const Instance> BuildInstance(const TSharedPtr<const Model>&, const Parameters*, OP::ADDRESS);
 		MUTABLERUNTIME_API Ptr<const Layout> BuildLayout(const TSharedPtr<const Model>&, const Parameters*, OP::ADDRESS) ;
     	MUTABLERUNTIME_API FProjector BuildProjector(const TSharedPtr<const Model>&, const Parameters*, OP::ADDRESS) ;
 
diff --git a/Engine/Plugins/Experimental/Mutable/Source/MutableRuntime/Private/MuR/InstancePrivate.h b/Engine/Plugins/Experimental/Mutable/Source/MutableRuntime/Private/MuR/InstancePrivate.h
--- a/Engine/Plugins/Experimental/Mutable/Source/MutableRuntime/Private/MuR/InstancePrivate.h
+++ b/Engine/Plugins/Experimental/Mutable/Source/MutableRuntime/Private/MuR/InstancePrivate.h
@@ -17,11 +17,6 @@ namespace mu
 		return (uint64(RootAddress) << 32) | uint64(ParameterBlobIndex);
 	}
 
-	inline uint32 GetResourceIDRoot(FResourceID Id)
-	{
-		return uint32(Id >> 32);
-	}
-
 	/** */
     struct FInstanceSurface
 	{
diff --git a/Engine/Plugins/Experimental/Mutable/Source/MutableRuntime/Private/MuR/System.cpp b/Engine/Plugins/Experimental/Mutable/Source/MutableRuntime/Private/MuR/System.cpp
--- a/Engine/Plugins/Experimental/Mutable/Source/MutableRuntime/Private/MuR/System.cpp
+++ b/Engine/Plugins/Experimental/Mutable/Source/MutableRuntime/Private/MuR/System.cpp
@@ -962,8 +962,7 @@ namespace mu
 	}
 
 	
-	Ptr<const Image> System::Private::BuildImage(const TSharedPtr<const Model>& pModel,
-		const Parameters* Params, OP::ADDRESS at, int32 MipsToSkip, int32 InImageLOD)
+	Ptr<const Image> System::Private::BuildImage(const TSharedPtr<const Model>& pModel, const Parameters* Params, OP::ADDRESS at, int32 MipsToSkip, int32 InImageLOD)
 	{
 		WorkingMemoryManager.BeginRunnerThread();
 
@@ -985,8 +984,7 @@ namespace mu
 	}
 
 
-	//---------------------------------------------------------------------------------------------
-	MeshPtrConst System::Private::BuildMesh(const TSharedPtr<const Model>& pModel, const Parameters* Params, OP::ADDRESS at)
+	Ptr<const Mesh> System::Private::BuildMesh(const TSharedPtr<const Model>& pModel, const Parameters* Params, OP::ADDRESS at)
 	{
 		WorkingMemoryManager.BeginRunnerThread();
 
@@ -1008,6 +1006,28 @@ namespace mu
 	}
 
 
+	Ptr<const Instance> System::Private::BuildInstance(const TSharedPtr<const Model>& pModel, const Parameters* Params, OP::ADDRESS at)
+	{
+		WorkingMemoryManager.BeginRunnerThread();
+
+		Ptr<const Instance> Result;
+
+		mu::OP_TYPE opType = pModel->GetPrivate()->m_program.GetOpType(at);
+		if (GetOpDataType(opType) == DT_INSTANCE)
+		{
+			RunCode(pModel, Params, at);
+			if (!bUnrecoverableError)
+			{
+				Result = WorkingMemoryManager.CurrentInstanceCache->GetInstance(FCacheAddress(at, 0, 0));
+			}
+		}
+
+		WorkingMemoryManager.EndRunnerThread();
+
+		return Result;
+	}
+
+
 	//---------------------------------------------------------------------------------------------
 	Ptr<const Layout> System::Private::BuildLayout(const TSharedPtr<const Model>& pModel, const Parameters* Params, OP::ADDRESS at)
 	{
diff --git a/Engine/Plugins/Experimental/Mutable/Source/MutableTools/Private/MuT/ASTOpInstanceAdd.h b/Engine/Plugins/Experimental/Mutable/Source/MutableTools/Private/MuT/ASTOpInstanceAdd.h
--- a/Engine/Plugins/Experimental/Mutable/Source/MutableTools/Private/MuT/ASTOpInstanceAdd.h
+++ b/Engine/Plugins/Experimental/Mutable/Source/MutableTools/Private/MuT/ASTOpInstanceAdd.h
@@ -21,7 +21,7 @@ struct FProgram;
 	public:
 
 		//! Type of switch
-		OP_TYPE type;
+		OP_TYPE type = OP_TYPE::NONE;
 
 		ASTChild instance;
 		ASTChild value;
diff --git a/Engine/Plugins/Experimental/Mutable/Source/MutableTools/Private/MuT/CodeGenerator.cpp b/Engine/Plugins/Experimental/Mutable/Source/MutableTools/Private/MuT/CodeGenerator.cpp
--- a/Engine/Plugins/Experimental/Mutable/Source/MutableTools/Private/MuT/CodeGenerator.cpp
+++ b/Engine/Plugins/Experimental/Mutable/Source/MutableTools/Private/MuT/CodeGenerator.cpp
@@ -137,6 +137,8 @@ namespace mu
 				Ptr<ASTOp> stateRoot = Generate_Generic(pNode, Options);
 				States.Emplace(s.Key, stateRoot);
 
+				AdditionalComponents.Empty();
+
 				++CurrentStateIndex;
 			}
 		}
@@ -214,8 +216,9 @@ namespace mu
 		else if (pNode->GetType()->IsA(NodeComponent::GetStaticType()))
 		{
 			const NodeComponent* ComponentNode = static_cast<const NodeComponent*>(pNode.get());
+			FComponentGenerationOptions ComponentOptions( Options, nullptr );
 			FGenericGenerationResult Result;
-			GenerateComponent(Options, Result, ComponentNode);
+			GenerateComponent(ComponentOptions, Result, ComponentNode);
 			return Result.op;
 		}
 
@@ -508,7 +511,7 @@ namespace mu
 
 
 	//---------------------------------------------------------------------------------------------
-	void CodeGenerator::GenerateComponent(const FGenericGenerationOptions& InOptions, FGenericGenerationResult& OutResult, const NodeComponent* InUntypedNode)
+	void CodeGenerator::GenerateComponent(const FComponentGenerationOptions& InOptions, FGenericGenerationResult& OutResult, const NodeComponent* InUntypedNode)
 	{
 		if (!InUntypedNode)
 		{
@@ -537,6 +540,7 @@ namespace mu
 		{
 			// Nothing to do because it is all preprocessed in the first code generator stage
 			//GenerateComponent_Edit(InOptions, OutResult, static_cast<const NodeComponentEdit*>(InUntypedNode));
+			OutResult.op = InOptions.BaseInstance;
 		}
 		else
 		{
@@ -548,7 +552,7 @@ namespace mu
 	}
 
 
-	void CodeGenerator::GenerateComponent_New(const FGenericGenerationOptions& Options, FGenericGenerationResult& Result, const NodeComponentNew* InNode)
+	void CodeGenerator::GenerateComponent_New(const FComponentGenerationOptions& Options, FGenericGenerationResult& Result, const NodeComponentNew* InNode)
 	{
 		const NodeComponentNew& node = *InNode;
 
@@ -571,11 +575,37 @@ namespace mu
 
 		Ptr<ASTOpInstanceAdd> InstanceOp = new ASTOpInstanceAdd();
 		InstanceOp->type = OP_TYPE::IN_ADDCOMPONENT;
-		InstanceOp->instance = nullptr;
+		InstanceOp->instance = Options.BaseInstance;
 		InstanceOp->value = LODsOp;
 		InstanceOp->id = node.Id;
 
 		Result.op = InstanceOp;
+
+		// Add a conditional if this component has conditions
+		for (const FirstPassGenerator::FComponent& Component: FirstPass.Components)
+		{
+			if (Component.Component != InNode)
+			{
+				continue;
+			}
+
+			if (Component.ComponentCondition || Component.ObjectCondition)
+			{
+				// TODO: This could be done earlier?
+				Ptr<ASTOpFixed> ConditionOp = new ASTOpFixed();
+				ConditionOp->op.type = OP_TYPE::BO_AND;
+				ConditionOp->SetChild(ConditionOp->op.args.BoolBinary.a, Component.ObjectCondition);
+				ConditionOp->SetChild(ConditionOp->op.args.BoolBinary.b, Component.ComponentCondition);
+
+				Ptr<ASTOpConditional> IfOp = new ASTOpConditional();
+				IfOp->type = OP_TYPE::IN_CONDITIONAL;
+				IfOp->no = Options.BaseInstance;
+				IfOp->yes = Result.op;
+				IfOp->condition = ConditionOp;
+
+				Result.op = IfOp;
+			}
+		}
 	}
 
 
@@ -653,7 +683,7 @@ namespace mu
         FMeshGenerationResult meshResults;
 
 		// We don't add the mesh here, since it will be added directly at the top of the
-		// component expression in the NodeComponentNew visitor with the right merges
+		// component expression in the NodeComponentNew generator with the right merges
 		// and conditions.
 		// But we store it to be used then.
 
@@ -1254,7 +1284,7 @@ namespace mu
 								imageAd = BlankImageOp;
 							}
 
-							auto UpdateBlockSize = [&BlockPixelsX, &BlockPixelsY, &FinalFormat, &bBlocksHaveMips, &bImageSizeWarning, &formatNode, &BlankImageOp, &node, surfaceNode, &t, this]( FImageDesc BlockDesc, UE::Math::TIntVector2<uint16> LayoutCellSize )
+							auto UpdateBlockSize = [&BlockPixelsX, &BlockPixelsY, &FinalFormat, &bBlocksHaveMips, &bImageSizeWarning, &formatNode, &BlankImageOp, &node, surfaceNode, &t, &Options, this]( FImageDesc BlockDesc, UE::Math::TIntVector2<uint16> LayoutCellSize )
 							{
 								if (BlockPixelsX == 0)
 								{
@@ -1265,7 +1295,8 @@ namespace mu
 										{
 											bImageSizeWarning = true;
 
-											int32 currentLOD = CurrentParents.Last().Lod;
+											check(CurrentParents.Last().Lod== Options.LODIndex);
+											int32 currentLOD = Options.LODIndex;
 											FString Msg = FString::Printf(TEXT("A texture [%s] for material [%s] parameter [%s] in LOD [%d] has been resized because it didn't fit the layout. "),
 												*node.Images[t].Name,
 												*node.Images[t].MaterialName,
@@ -1721,25 +1752,19 @@ namespace mu
                 Ptr<ASTOp> surfaceAt = sop;
 
 				// TODO: This could be done earlier?
-                Ptr<ASTOpFixed> surfCondOp = new ASTOpFixed();
-                surfCondOp->op.type = OP_TYPE::BO_AND;
-                surfCondOp->SetChild(surfCondOp->op.args.BoolBinary.a, its.ObjectCondition );
-                surfCondOp->SetChild(surfCondOp->op.args.BoolBinary.b, its.SurfaceCondition );
-                Ptr<ASTOp> surfaceCondition = surfCondOp;
+                Ptr<ASTOpFixed> SurfaceConditionOp = new ASTOpFixed();
+				SurfaceConditionOp->op.type = OP_TYPE::BO_AND;
+				SurfaceConditionOp->SetChild(SurfaceConditionOp->op.args.BoolBinary.a, its.ObjectCondition );
+				SurfaceConditionOp->SetChild(SurfaceConditionOp->op.args.BoolBinary.b, its.SurfaceCondition );
 
-                if (surfaceCondition)
                 {
                     Ptr<ASTOpConditional> op = new ASTOpConditional();
                     op->type = OP_TYPE::IN_CONDITIONAL;
                     op->no = lastCompOp;
                     op->yes = surfaceAt;
-                    op->condition = surfaceCondition;
+                    op->condition = SurfaceConditionOp;
                     lastCompOp = op;
                 }
-                else
-                {
-                    lastCompOp = surfaceAt;
-                }
 
                 // Add the mesh with its condition
  
@@ -1756,13 +1781,13 @@ namespace mu
                     mergeAd = mop;
                 }
 
-                if (surfaceCondition)
+                if (SurfaceConditionOp)
                 {
                     Ptr<ASTOpConditional> op = new ASTOpConditional();
                     op->type = OP_TYPE::ME_CONDITIONAL;
                     op->no = lastMeshOp;
                     op->yes = mergeAd;
-                    op->condition = surfaceCondition;
+                    op->condition = SurfaceConditionOp;
                     lastMeshOp = op;
                 }
                 else
@@ -1799,6 +1824,10 @@ namespace mu
 	{
 		MUTABLE_CPUPROFILER_SCOPE(NodeObjectNew);
 
+		// There is always at least a null parent
+		bool bIsChildObject = CurrentParents.Num() > 1;
+
+		// Add this object as current parent
         CurrentParents.Add( FParentKey() );
         CurrentParents.Last().ObjectNode = InNode;
 
@@ -1837,6 +1866,12 @@ namespace mu
 
         // Create the expression adding all the components
 		Ptr<ASTOp> LastCompOp;
+		Ptr<ASTOp> PlaceholderOp;
+		if (bIsChildObject)
+		{
+			PlaceholderOp = new ASTOpInstanceAdd;
+			LastCompOp = PlaceholderOp;
+		}
 
 		// Add the components in this node
         for ( int32 t=0; t< InNode->Components.Num(); ++t )
@@ -1844,62 +1879,51 @@ namespace mu
 			const NodeComponent* ComponentNode = InNode->Components[t].get();
             if (ComponentNode)
             {
+				FComponentGenerationOptions ComponentOptions( Options, LastCompOp );
 				FGenericGenerationResult ComponentResult;
-				GenerateComponent(Options, ComponentResult, ComponentNode);
-
-				if (ComponentResult.op)
-				{
-					check(ComponentResult.op->GetOpType() == OP_TYPE::IN_ADDCOMPONENT);
-					ASTOpInstanceAdd* typedOp = static_cast<ASTOpInstanceAdd*>(ComponentResult.op.get());
-
-					check(!typedOp->instance.child());
-					typedOp->instance = LastCompOp;
-					LastCompOp = typedOp;
-				}
+				GenerateComponent(ComponentOptions, ComponentResult, ComponentNode);
+				LastCompOp = ComponentResult.op;
             }
         }
 
+		// If we didn't generate anything, make sure we don't use the placeholder.
+		if (LastCompOp == PlaceholderOp)
+		{
+			LastCompOp = nullptr;
+			PlaceholderOp = nullptr;
+		}
+
 		// Add the components from child objects
-		FAdditionalComponentKey thisKey;
-		thisKey.Lod = CurrentParents.Last().Lod;
-		thisKey.ObjectNode = CurrentParents.Last().ObjectNode;
-		TArray<Ptr<ASTOp>>* addIt = AdditionalComponents.Find(thisKey);
-		if (addIt)
+		FAdditionalComponentKey ThisKey;
+		ThisKey.ObjectNode = CurrentParents.Last().ObjectNode;
+		TArray<FAdditionalComponentData>* ThisAdditionalComponents = AdditionalComponents.Find(ThisKey);
+		if (LastCompOp && ThisAdditionalComponents)
 		{
-			for (const Ptr<ASTOp>& cop : *addIt)
+			for (const FAdditionalComponentData& Additional : *ThisAdditionalComponents)
 			{
-				// Add the additional components after the main ones, this means higher up in the op tree.
-
-				// First find the last op in the chain of IN_ADDCOMPONENT operations
-				check(cop->GetOpType() == OP_TYPE::IN_ADDCOMPONENT);
-				ASTOpInstanceAdd* typedOp = static_cast<ASTOpInstanceAdd*>(cop.get());
-				ASTOpInstanceAdd* bottomOp = typedOp;
-				while (bottomOp->instance)
-				{
-					// Step down
-					check(bottomOp->instance->GetOpType() == OP_TYPE::IN_ADDCOMPONENT);
-					bottomOp = static_cast<ASTOpInstanceAdd*>(bottomOp->instance.child().get());
-				}
-
-				// Chain
-				bottomOp->instance = LastCompOp;
-				LastCompOp = typedOp;
+				check(Additional.PlaceholderOp);
+				ASTOp::Replace(Additional.PlaceholderOp, LastCompOp);
+				LastCompOp = Additional.ComponentOp;
 			}
 		}
 
-		// Store for possible parent objects if necessary
+		// Store this chain of components for use in parent objects if necessary
 		// 2 is because there must be a parent and there is always a null element as well.
-		if (LastCompOp && CurrentParents.Num() > 2)
+		if (LastCompOp && bIsChildObject)
 		{
-			const auto& parentObjectKey = CurrentParents[CurrentParents.Num() - 2];
-			FAdditionalComponentKey parentKey;
-			parentKey.Lod = CurrentParents.Last().Lod;
-			parentKey.ObjectNode = parentObjectKey.ObjectNode;
-			AdditionalComponents.FindOrAdd(parentKey).Add(LastCompOp);
+			// TODO: Directly to the root object?
+			const FParentKey& ParentObjectKey = CurrentParents[CurrentParents.Num() - 2];
+			FAdditionalComponentKey ParentKey;
+			ParentKey.ObjectNode = ParentObjectKey.ObjectNode;
+
+			FAdditionalComponentData Data;
+			Data.ComponentOp = LastCompOp;
+			Data.PlaceholderOp = PlaceholderOp;
+			AdditionalComponents.FindOrAdd(ParentKey).Add(Data);
 		}
 
 
-        Ptr<ASTOp> rootOp = LastCompOp;
+        Ptr<ASTOp> RootOp = LastCompOp;
 
 		// Add an ASTOpAddExtensionData for each connected ExtensionData node
 		for (const NodeObjectNew::FNamedExtensionDataNode& NamedNode : InNode->ExtensionDataNodes)
@@ -1936,7 +1960,7 @@ namespace mu
 			for (const FConditionalExtensionDataOp& SavedOp : ConditionalExtensionDataOps)
 			{
 				Ptr<ASTOpAddExtensionData> ExtensionPinOp = new ASTOpAddExtensionData();
-				ExtensionPinOp->Instance = ASTChild(ExtensionPinOp, rootOp);
+				ExtensionPinOp->Instance = ASTChild(ExtensionPinOp, RootOp);
 				ExtensionPinOp->ExtensionData = ASTChild(ExtensionPinOp, SavedOp.ExtensionDataOp);
 				ExtensionPinOp->ExtensionDataName = SavedOp.ExtensionDataName;
 
@@ -1944,22 +1968,22 @@ namespace mu
 				{
 					Ptr<ASTOpConditional> ConditionOp = new ASTOpConditional();
 					ConditionOp->type = OP_TYPE::IN_CONDITIONAL;
-					ConditionOp->no = rootOp;
+					ConditionOp->no = RootOp;
 					ConditionOp->yes = ExtensionPinOp;
 					ConditionOp->condition = ASTChild(ConditionOp, SavedOp.Condition);
 					
-					rootOp = ConditionOp;
+					RootOp = ConditionOp;
 				}
 				else
 				{
-					rootOp = ExtensionPinOp;
+					RootOp = ExtensionPinOp;
 				}
 			}
 		}
 
         CurrentParents.Pop();
 
-        Result.op = rootOp;
+        Result.op = RootOp;
     }
 
 
diff --git a/Engine/Plugins/Experimental/Mutable/Source/MutableTools/Private/MuT/CodeGenerator.h b/Engine/Plugins/Experimental/Mutable/Source/MutableTools/Private/MuT/CodeGenerator.h
--- a/Engine/Plugins/Experimental/Mutable/Source/MutableTools/Private/MuT/CodeGenerator.h
+++ b/Engine/Plugins/Experimental/Mutable/Source/MutableTools/Private/MuT/CodeGenerator.h
@@ -145,6 +145,19 @@ namespace mu
 			TArray<FString> ActiveTags;
 		};
 
+		struct FComponentGenerationOptions : public FGenericGenerationOptions
+		{
+			FComponentGenerationOptions(const FGenericGenerationOptions& BaseOptions, const Ptr<ASTOp>& InBaseInstance )
+			{
+				BaseInstance = InBaseInstance;
+				State = BaseOptions.State;
+				ActiveTags = BaseOptions.ActiveTags;
+			}
+
+			/** Instance to which the possibly generated components should be added. */
+			Ptr<ASTOp> BaseInstance;
+		};
+
 		struct FLODGenerationOptions : public FGenericGenerationOptions
 		{
 			FLODGenerationOptions(const FGenericGenerationOptions& BaseOptions, int32 InLODIndex, const NodeComponentNew* InComponent)
@@ -258,35 +271,38 @@ namespace mu
 
 		TArray< FParentKey > CurrentParents;
 
-        // List of additional components to add to an object that come from child objects.
-        // The index is the object and lod that should receive the components.
+        /** List of additional components to add to an object that come from child objects.
+        * The index is the object and lod that should receive the components.
+		*/
         struct FAdditionalComponentKey
         {
 			FAdditionalComponentKey()
             {
 				ObjectNode = nullptr;
-				Lod = -1;
             }
 
             const NodeObjectNew* ObjectNode;
-            int32 Lod;
 
 			FORCEINLINE bool operator==(const FAdditionalComponentKey& Other) const
 			{
-				return ObjectNode == Other.ObjectNode
-					&&
-					Lod == Other.Lod;
+				return ObjectNode == Other.ObjectNode;
 			}
 
 			friend FORCEINLINE uint32 GetTypeHash(const FAdditionalComponentKey& InKey)
 			{
 				uint32 KeyHash = 0;
 				KeyHash = HashCombineFast(KeyHash, ::GetTypeHash(InKey.ObjectNode));
-				KeyHash = HashCombineFast(KeyHash, ::GetTypeHash(InKey.Lod));
 				return KeyHash;
 			}
 		};
-        TMap< FAdditionalComponentKey, TArray<Ptr<ASTOp>> > AdditionalComponents;
+
+		struct FAdditionalComponentData
+		{
+			Ptr<ASTOp> ComponentOp;
+			Ptr<ASTOp> PlaceholderOp;
+		};
+
+        TMap< FAdditionalComponentKey, TArray<FAdditionalComponentData> > AdditionalComponents;
 
 
         struct FObjectGenerationData
@@ -349,8 +365,8 @@ namespace mu
 		typedef TMap<FGeneratedComponentCacheKey, FGenericGenerationResult> GeneratedComponentMap;
 		GeneratedComponentMap GeneratedComponents;
 
-		void GenerateComponent(const FGenericGenerationOptions&, FGenericGenerationResult&, const NodeComponent*);
-		void GenerateComponent_New(const FGenericGenerationOptions&, FGenericGenerationResult&, const NodeComponentNew*);
+		void GenerateComponent(const FComponentGenerationOptions&, FGenericGenerationResult&, const NodeComponent*);
+		void GenerateComponent_New(const FComponentGenerationOptions&, FGenericGenerationResult&, const NodeComponentNew*);
 
 		//-----------------------------------------------------------------------------------------
 
diff --git a/Engine/Plugins/Experimental/Mutable/Source/MutableTools/Private/MuT/CodeGenerator_FirstPass.cpp b/Engine/Plugins/Experimental/Mutable/Source/MutableTools/Private/MuT/CodeGenerator_FirstPass.cpp
--- a/Engine/Plugins/Experimental/Mutable/Source/MutableTools/Private/MuT/CodeGenerator_FirstPass.cpp
+++ b/Engine/Plugins/Experimental/Mutable/Source/MutableTools/Private/MuT/CodeGenerator_FirstPass.cpp
@@ -145,6 +145,11 @@ namespace mu
 
 	void FirstPassGenerator::Generate_Generic(const Node* Root)
 	{
+		if (!Root)
+		{
+			return;
+		}
+
 		if (Root->GetType()==NodeSurfaceNew::GetStaticType())
 		{
 			Generate_SurfaceNew(static_cast<const NodeSurfaceNew*>(Root));
@@ -169,6 +174,14 @@ namespace mu
 		{
 			Generate_ComponentEdit(static_cast<const NodeComponentEdit*>(Root));
 		}
+		else if (Root->GetType() == NodeComponentSwitch::GetStaticType())
+		{
+			Generate_ComponentSwitch(static_cast<const NodeComponentSwitch*>(Root));
+		}
+		else if (Root->GetType() == NodeComponentVariation::GetStaticType())
+		{
+			Generate_ComponentVariation(static_cast<const NodeComponentVariation*>(Root));
+		}
 		else if (Root->GetType() == NodeObjectNew::GetStaticType())
 		{
 			Generate_ObjectNew(static_cast<const NodeObjectNew*>(Root));
@@ -456,6 +469,14 @@ namespace mu
 	//---------------------------------------------------------------------------------------------
 	void FirstPassGenerator::Generate_ComponentNew(const NodeComponentNew* InNode)
 	{
+		// Add the data about this surface
+		FComponent ThisData;
+		ThisData.Component = InNode;
+		ThisData.ObjectCondition = CurrentCondition.Last().ObjectCondition;
+		ThisData.PositiveTags = CurrentPositiveTags;
+		ThisData.NegativeTags = CurrentNegativeTags;
+		Components.Add(ThisData);
+
         CurrentComponent = InNode;
 
 		CurrentLOD = 0;
@@ -493,6 +514,90 @@ namespace mu
 	}
 
 
+	//---------------------------------------------------------------------------------------------
+	void FirstPassGenerator::Generate_ComponentVariation(const NodeComponentVariation* InNode)
+	{
+		// Any of the tags in the variations would prevent the default surface
+		TArray<FString> OldNegativeTags = CurrentNegativeTags;
+		for (int32 v = 0; v < InNode->Variations.Num(); ++v)
+		{
+			CurrentNegativeTags.Add(InNode->Variations[v].Tag);
+		}
+
+		Generate_Generic(InNode->DefaultComponent.get());
+
+		CurrentNegativeTags = OldNegativeTags;
+
+		for (int32 v = 0; v < InNode->Variations.Num(); ++v)
+		{
+			CurrentPositiveTags.Add(InNode->Variations[v].Tag);
+			Generate_Generic(InNode->Variations[v].Component.get());
+
+			CurrentPositiveTags.Pop();
+
+			// Tags have an order in a variation node: the current tag should prevent any following variation
+			CurrentNegativeTags.Add(InNode->Variations[v].Tag);
+		}
+
+		CurrentNegativeTags = OldNegativeTags;
+	}
+
+
+	//---------------------------------------------------------------------------------------------
+	void FirstPassGenerator::Generate_ComponentSwitch(const NodeComponentSwitch* InNode)
+	{
+		if (InNode->Options.Num() == 0)
+		{
+			// No options in the switch!
+			return;
+		}
+
+		// Prepare the enumeration parameter
+		CodeGenerator::FGenericGenerationOptions Options;
+		CodeGenerator::FScalarGenerationResult ScalarResult;
+		if (InNode->Parameter)
+		{
+			Generator->GenerateScalar(ScalarResult, Options, InNode->Parameter);
+		}
+		else
+		{
+			// This argument is required
+			ScalarResult.op = Generator->GenerateMissingScalarCode(TEXT("Switch variable"), 0.0f, InNode->GetMessageContext());
+		}
+
+		// Parse the options
+		for (int32 t = 0; t < InNode->Options.Num(); ++t)
+		{
+			// Create a comparison operation as the boolean parameter for the child
+			Ptr<ASTOpFixed> ParamOp = new ASTOpFixed();
+			ParamOp->op.type = OP_TYPE::BO_EQUAL_INT_CONST;
+			ParamOp->SetChild(ParamOp->op.args.BoolEqualScalarConst.value, ScalarResult.op);
+			ParamOp->op.args.BoolEqualScalarConst.constant = (int16)t;
+
+			// Combine the new condition with previous conditions coming from parent objects
+			if (CurrentCondition.Last().ObjectCondition)
+			{
+				Ptr<ASTOpFixed> op = new ASTOpFixed();
+				op->op.type = OP_TYPE::BO_AND;
+				op->SetChild(op->op.args.BoolBinary.a, CurrentCondition.Last().ObjectCondition);
+				op->SetChild(op->op.args.BoolBinary.b, ParamOp);
+				ParamOp = op;
+			}
+
+			FConditionContext data;
+			data.ObjectCondition = ParamOp;
+			CurrentCondition.Push(data);
+
+			if (InNode->Options[t])
+			{
+				Generate_Generic(InNode->Options[t].get());
+			}
+
+			CurrentCondition.Pop();
+		}
+	}
+
+
 	//---------------------------------------------------------------------------------------------
 	void FirstPassGenerator::Generate_LOD(const NodeLOD* InNode)
 	{
diff --git a/Engine/Plugins/Experimental/Mutable/Source/MutableTools/Private/MuT/CodeGenerator_FirstPass.h b/Engine/Plugins/Experimental/Mutable/Source/MutableTools/Private/MuT/CodeGenerator_FirstPass.h
--- a/Engine/Plugins/Experimental/Mutable/Source/MutableTools/Private/MuT/CodeGenerator_FirstPass.h
+++ b/Engine/Plugins/Experimental/Mutable/Source/MutableTools/Private/MuT/CodeGenerator_FirstPass.h
@@ -53,6 +53,8 @@ namespace mu
 		void Generate_SurfaceVariation(const NodeSurfaceVariation*);
 		void Generate_ComponentNew(const NodeComponentNew*);
 		void Generate_ComponentEdit(const NodeComponentEdit*);
+		void Generate_ComponentSwitch(const NodeComponentSwitch*);
+		void Generate_ComponentVariation(const NodeComponentVariation*);
 		void Generate_LOD(const NodeLOD*);
 		void Generate_ObjectNew(const NodeObjectNew*);
 		void Generate_ObjectGroup(const NodeObjectGroup*);
@@ -76,19 +78,43 @@ namespace mu
         //! for the states whose index is true.
         using StateCondition = TArray<uint8>;
 
-		//! Store information about every surface including
-		//! - the component it may be added to
-		//! - the conditions that will enable or disable it
-		//! - all edit operators
-        //! A surface may have different versions depending on the different parents and conditions
-        //! it is reached with.
+		/** Store information about every component found. */
+		struct FComponent
+		{
+			/** Main component node. */
+			const NodeComponentNew* Component = nullptr;
+
+			// List of tags that are required for the presence of this component
+			TArray<FString> PositiveTags;
+
+			// List of tags that block the presence of this component
+			TArray<FString> NegativeTags;
+
+			// This conditions is the condition of the object defining this surface which may not
+			// be the parent object where this surface will be added.
+			Ptr<ASTOp> ObjectCondition;
+
+			// Condition for this component to be added.
+			// This is filled in CodeGenerator_SecondPass.
+			Ptr<ASTOp> ComponentCondition;
+		};
+		TArray<FComponent> Components;
+
+		/** Store information about every surface including
+		* - the component it may be added to
+		* - the conditions that will enable or disable it
+		* - all edit operators
+        * A surface may have different versions depending on the different parents and conditions it is reached with.
+		*/
 		struct FSurface
 		{
             Ptr<const NodeSurfaceNew> Node;
 
-			// Parent Component where this surface will be added. It may be different from the 
-			// Component that defined it (if it was an edit component).
+			/** Parent Component where this surface will be added.It may be different from the
+			* Component that defined it (if it was an edit component).
+			*/
             const NodeComponentNew* Component = nullptr;
+
 			int32 LOD = 0;
 
             // List of tags that are required for the presence of this surface
diff --git a/Engine/Plugins/Experimental/Mutable/Source/MutableTools/Private/MuT/CodeGenerator_SecondPass.cpp b/Engine/Plugins/Experimental/Mutable/Source/MutableTools/Private/MuT/CodeGenerator_SecondPass.cpp
--- a/Engine/Plugins/Experimental/Mutable/Source/MutableTools/Private/MuT/CodeGenerator_SecondPass.cpp
+++ b/Engine/Plugins/Experimental/Mutable/Source/MutableTools/Private/MuT/CodeGenerator_SecondPass.cpp
@@ -33,7 +33,7 @@ namespace mu
 		const set<size_t>& posTag,
 		const set<size_t>& negTag)
 	{
-		auto& t = FirstPass->Tags[tagIndex];
+		FirstPassGenerator::FTag& t = FirstPass->Tags[tagIndex];
 
 		// If this tag is already in the list of positive tags, return true as condition
 		if (posTag.find(tagIndex) != posTag.end())
@@ -66,16 +66,6 @@ namespace mu
 				return it->second;
 			}
 		}
-
-		auto PositiveTagsFunc = [this](size_t SurfaceIndex) -> TArray<FString>&
-		{
-			return FirstPass->Surfaces[SurfaceIndex].PositiveTags;
-		};
-
-		auto NegativeTagsFunc = [this](size_t SurfaceIndex) -> TArray<FString>&
-		{
-			return FirstPass->Surfaces[SurfaceIndex].NegativeTags;
-		};
 		
 		Ptr<ASTOp> c;
 
@@ -96,14 +86,14 @@ namespace mu
 				continue;
 			}
 
-			const auto& surface = FirstPass->Surfaces[surfIndex];
+			const FirstPassGenerator::FSurface& surface = FirstPass->Surfaces[surfIndex];
 
 			auto PositiveTags = posTag;
 			PositiveTags.insert(tagIndex);
 
-			Ptr<ASTOp> surfCondition = GenerateSurfaceOrModifierCodition(surfIndex,
-				PositiveTagsFunc,
-				NegativeTagsFunc,
+			Ptr<ASTOp> surfCondition = GenerateDataCodition(surfIndex,
+				FirstPass->Surfaces[surfIndex].PositiveTags,
+				FirstPass->Surfaces[surfIndex].NegativeTags,
 				posSurf,
 				negSurf,
 				PositiveTags,
@@ -173,9 +163,9 @@ namespace mu
 				auto PositiveTags = posTag;
 				PositiveTags.insert(tagIndex);
 
-				surfCondition = GenerateSurfaceOrModifierCodition(EditKey.Key,
-					PositiveTagsFunc,
-					NegativeTagsFunc,
+				surfCondition = GenerateDataCodition(EditKey.Key,
+					FirstPass->Surfaces[EditKey.Key].PositiveTags,
+					FirstPass->Surfaces[EditKey.Key].NegativeTags,
 					posSurf,
 					negSurf,
 					PositiveTags,
@@ -236,9 +226,9 @@ namespace mu
 	}
 
 	//---------------------------------------------------------------------------------------------
-	mu::Ptr<ASTOp> SecondPassGenerator::GenerateSurfaceOrModifierCodition(size_t Index,
-		TFunction<const TArray<FString>&(size_t)> PositiveTagsFunc,
-		TFunction<const TArray<FString>&(size_t)> NegativeTagsFunc,
+	mu::Ptr<ASTOp> SecondPassGenerator::GenerateDataCodition(size_t Index,
+		const TArray<FString>& PositiveTags,
+		const TArray<FString>& NegativeTags,
 		const set<size_t>& posSurf,
 		const set<size_t>& negSurf,
 		const set<size_t>& posTag,
@@ -258,7 +248,7 @@ namespace mu
 
 		Ptr<ASTOp> c;
 
-		for (const FString& t : PositiveTagsFunc(Index))
+		for (const FString& t : PositiveTags)
 		{
 			const FirstPassGenerator::FTag* it = FirstPass->Tags.FindByPredicate([&](const FirstPassGenerator::FTag& e) { return e.Tag == t; });
 			if (!it)
@@ -328,7 +318,7 @@ namespace mu
 		}
 
 
-		for (const FString& t : NegativeTagsFunc(Index))
+		for (const FString& t : NegativeTags)
 		{
 			const FirstPassGenerator::FTag* it = FirstPass->Tags.FindByPredicate([&](const FirstPassGenerator::FTag& e) { return e.Tag == t; });
 			if (!it)
@@ -337,7 +327,7 @@ namespace mu
 				continue;
 			}
 
-			size_t tagIndex = it - &FirstPass->Tags[0];
+			size_t tagIndex = it - FirstPass->Tags.GetData();
 
 			set<size_t> positiveSurfacesVisited = negSurf;
 			set<size_t> negativeSurfacesVisited = posSurf;
@@ -488,27 +478,22 @@ namespace mu
 			}
 		}
 
-		// Create the conditions for every surface, modifier and individual tag.
+		// Create the conditions for every surface, modifier, component and individual tag.
 		TagConditionGenerationCache.clear();
 
+		set<size_t> Empty;
+
 		for (int32 SurfaceIndex = 0; SurfaceIndex < FirstPass->Surfaces.Num(); ++SurfaceIndex)
 		{
 			FirstPassGenerator::FSurface& Surface = FirstPass->Surfaces[SurfaceIndex];
 
 			{
-				set<size_t> Empty;
+				Ptr<ASTOp> c = GenerateDataCodition(
+					SurfaceIndex, 
+					FirstPass->Surfaces[SurfaceIndex].PositiveTags,
+					FirstPass->Surfaces[SurfaceIndex].NegativeTags,
+					Empty, Empty, Empty, Empty);
 
-				auto PositiveTags = [this](size_t SurfaceIndex) -> const TArray<FString>&
-				{
-					return FirstPass->Surfaces[SurfaceIndex].PositiveTags;
-				};
-
-				auto NegativeTags = [this](size_t SurfaceIndex) -> const TArray<FString>&
-				{
-					return FirstPass->Surfaces[SurfaceIndex].NegativeTags;
-				};
-				
-				Ptr<ASTOp> c = GenerateSurfaceOrModifierCodition(SurfaceIndex, PositiveTags, NegativeTags, Empty, Empty, Empty, Empty);
 				FirstPass->Surfaces[SurfaceIndex].SurfaceCondition = c;
 			}
 			
@@ -516,19 +501,11 @@ namespace mu
 			{
 				FirstPassGenerator::FSurface::FEdit& Edit = Surface.Edits[EditIndex];
 				
-				set<size_t> Empty;
-
-				auto PositiveTags = [&Surface](size_t SurfaceIndex) -> const TArray<FString>&
-				{
-					return Surface.Edits[SurfaceIndex].PositiveTags;
-				};
-
-				auto NegativeTags = [&Surface](size_t SurfaceIndex) -> const TArray<FString>&
-				{
-					return Surface.Edits[SurfaceIndex].NegativeTags;
-				};
-				
-				Ptr<ASTOp> c = GenerateSurfaceOrModifierCodition(EditIndex, PositiveTags, NegativeTags, Empty, Empty, Empty, Empty);
+				Ptr<ASTOp> c = GenerateDataCodition(
+					EditIndex, 
+					Surface.Edits[EditIndex].PositiveTags,
+					Surface.Edits[EditIndex].NegativeTags,
+					Empty, Empty, Empty, Empty);
 
 				Ptr<ASTOpFixed> OpAnd = new ASTOpFixed;
 				OpAnd->op.type = OP_TYPE::BO_AND;
@@ -542,31 +519,33 @@ namespace mu
 
 		for (int32 ModifierIndex = 0; ModifierIndex < FirstPass->Modifiers.Num(); ++ModifierIndex)
 		{
-			set<size_t> Empty;
-
-			auto PositiveTags = [this](size_t SurfaceIndex) -> TArray<FString>&
-			{
-				return FirstPass->Modifiers[SurfaceIndex].PositiveTags;
-			};
-
-			auto NegativeTags = [this](size_t SurfaceIndex) -> TArray<FString>&
-			{
-				return FirstPass->Modifiers[SurfaceIndex].NegativeTags;
-			};
-				
-			Ptr<ASTOp> c = GenerateSurfaceOrModifierCodition(ModifierIndex, PositiveTags, NegativeTags, Empty, Empty, Empty, Empty);
+			Ptr<ASTOp> c = GenerateDataCodition(
+				ModifierIndex, 
+				FirstPass->Modifiers[ModifierIndex].PositiveTags,
+				FirstPass->Modifiers[ModifierIndex].NegativeTags,
+				Empty, Empty, Empty, Empty);
 			
 			FirstPass->Modifiers[ModifierIndex].SurfaceCondition = c;
 		}
 
+		for (int32 ComponentIndex = 0; ComponentIndex < FirstPass->Components.Num(); ++ComponentIndex)
+		{
+			Ptr<ASTOp> c = GenerateDataCodition(
+				ComponentIndex, 
+				FirstPass->Components[ComponentIndex].PositiveTags,
+				FirstPass->Components[ComponentIndex].NegativeTags,
+				Empty, Empty, Empty, Empty);
+
+			FirstPass->Components[ComponentIndex].ComponentCondition = c;
+		}
+
+		// TODO: Do we really need the tag conditions from here on?
 		for (int32 s = 0; s < FirstPass->Tags.Num(); ++s)
 		{
-			set<size_t> empty;
-			Ptr<ASTOp> c = GenerateTagCondition(s, empty, empty, empty, empty);
+			Ptr<ASTOp> c = GenerateTagCondition(s, Empty, Empty, Empty, Empty);
 			FirstPass->Tags[s].GenericCondition = c;
 		}
 
-
 		FirstPass = nullptr;
 
 		return true;
diff --git a/Engine/Plugins/Experimental/Mutable/Source/MutableTools/Private/MuT/CodeGenerator_SecondPass.h b/Engine/Plugins/Experimental/Mutable/Source/MutableTools/Private/MuT/CodeGenerator_SecondPass.h
--- a/Engine/Plugins/Experimental/Mutable/Source/MutableTools/Private/MuT/CodeGenerator_SecondPass.h
+++ b/Engine/Plugins/Experimental/Mutable/Source/MutableTools/Private/MuT/CodeGenerator_SecondPass.h
@@ -80,16 +80,16 @@ namespace mu
                                          const set<size_t>& negTag );
 
         /** Generate Surface, Edit or Modifier condition.
-    	 * @param Index Surface, Edit or Modifier index.
+    	 * @param Index Surface, Edit, Component or Modifier index.
     	 * @param PositiveTags function that given the Surface, Edit or Modifier index, returns its positive tags.
     	 * @param NegativeTags function that given the Surface, Edit or Modifier, returns its negative tags.
     	 * @param posSurf already visited Surfaces, Edits, or Modifiers that participate positively in the condition.
     	 * @param negSurf already visited Surfaces, Edits, or Modifiers that participate negatively in the condition.
       	 * @param posSurf Tags that already belong to the condition (positively).
 		 * @param negSurf Tags that already belong to the condition (negatively). */
-        Ptr<ASTOp> GenerateSurfaceOrModifierCodition(size_t Index,
-									        TFunction<const TArray<FString>&(size_t)> PositiveTags,
-											TFunction<const TArray<FString>&(size_t)> NegativeTags,
+        Ptr<ASTOp> GenerateDataCodition(size_t Index,
+									        const TArray<FString>& PositiveTags,
+											const TArray<FString>& NegativeTags,
                                             const set<size_t>& posSurf,
                                             const set<size_t>& negSurf,
                                             const set<size_t>& posTag,
