Const correctness fixes for UTextures.

diff --git a/Engine/Plugins/Experimental/Mutable/Source/CustomizableObject/Private/MuCO/CustomizableObject.cpp b/Engine/Plugins/Experimental/Mutable/Source/CustomizableObject/Private/MuCO/CustomizableObject.cpp
--- a/Engine/Plugins/Experimental/Mutable/Source/CustomizableObject/Private/MuCO/CustomizableObject.cpp
+++ b/Engine/Plugins/Experimental/Mutable/Source/CustomizableObject/Private/MuCO/CustomizableObject.cpp
@@ -739,7 +739,7 @@ void UCustomizableObjectPrivate::LoadCompiledData(FArchive& MemoryReader, const
 			FString StringRef;
 			MemoryReader << StringRef;
 
-			LocalModelResource.RuntimeReferencedTextures.Add(TSoftObjectPtr<UTexture2D>(FSoftObjectPath(StringRef)));
+			LocalModelResource.RuntimeReferencedTextures.Add(TSoftObjectPtr<const UTexture>(FSoftObjectPath(StringRef)));
 		}
 #endif
 		
diff --git a/Engine/Plugins/Experimental/Mutable/Source/CustomizableObjectEditor/Private/MuCOE/SMutableCodeViewer.cpp b/Engine/Plugins/Experimental/Mutable/Source/CustomizableObjectEditor/Private/MuCOE/SMutableCodeViewer.cpp
--- a/Engine/Plugins/Experimental/Mutable/Source/CustomizableObjectEditor/Private/MuCOE/SMutableCodeViewer.cpp
+++ b/Engine/Plugins/Experimental/Mutable/Source/CustomizableObjectEditor/Private/MuCOE/SMutableCodeViewer.cpp
@@ -2502,8 +2502,8 @@ namespace
 		{
 			check(ReferencedTextures.IsValidIndex(Id));
 
-			TSoftObjectPtr<UTexture> TexturePtr = ReferencedTextures[Id].Get();
-			UTexture2D* Texture = Cast<UTexture2D>( TexturePtr.Get() );
+			TSoftObjectPtr<const UTexture> TexturePtr = ReferencedTextures[Id].Get();
+			const UTexture2D* Texture = Cast<UTexture2D>( TexturePtr.Get() );
 			check(Texture);
 			
 			// In the editor the src data can be directly accessed
