[mutable] Add descriptor support for seldomly used operation to fix crash in old data.
- Some legacy test objects where crashing on compilation because of this missing method.
#rnx

diff --git a/Engine/Plugins/Experimental/Mutable/Source/MutableTools/Private/MuT/ASTOpImageNormalComposite.cpp b/Engine/Plugins/Experimental/Mutable/Source/MutableTools/Private/MuT/ASTOpImageNormalComposite.cpp
--- a/Engine/Plugins/Experimental/Mutable/Source/MutableTools/Private/MuT/ASTOpImageNormalComposite.cpp
+++ b/Engine/Plugins/Experimental/Mutable/Source/MutableTools/Private/MuT/ASTOpImageNormalComposite.cpp
@@ -154,4 +154,41 @@ namespace mu
 		return nullptr;
 	}
 
+
+	FSourceDataDescriptor ASTOpImageNormalComposite::GetSourceDataDescriptor(FGetSourceDataDescriptorContext* Context) const
+	{
+		// Cache management
+		TUniquePtr<FGetSourceDataDescriptorContext> LocalContext;
+		if (!Context)
+		{
+			LocalContext.Reset(new FGetSourceDataDescriptorContext);
+			Context = LocalContext.Get();
+		}
+
+		FSourceDataDescriptor* Found = Context->Cache.Find(this);
+		if (Found)
+		{
+			return *Found;
+		}
+
+		// Not cached: calculate
+		FSourceDataDescriptor Result;
+
+		if (Base)
+		{
+			FSourceDataDescriptor SourceDesc = Base->GetSourceDataDescriptor(Context);
+			Result.CombineWith(SourceDesc);
+		}
+
+		if (Normal)
+		{
+			FSourceDataDescriptor SourceDesc = Normal->GetSourceDataDescriptor(Context);
+			Result.CombineWith(SourceDesc);
+		}
+
+		Context->Cache.Add(this, Result);
+
+		return Result;
+	}
+
 }
diff --git a/Engine/Plugins/Experimental/Mutable/Source/MutableTools/Private/MuT/ASTOpImageNormalComposite.h b/Engine/Plugins/Experimental/Mutable/Source/MutableTools/Private/MuT/ASTOpImageNormalComposite.h
--- a/Engine/Plugins/Experimental/Mutable/Source/MutableTools/Private/MuT/ASTOpImageNormalComposite.h
+++ b/Engine/Plugins/Experimental/Mutable/Source/MutableTools/Private/MuT/ASTOpImageNormalComposite.h
@@ -36,6 +36,7 @@ struct FProgram;
 		FImageDesc GetImageDesc(bool returnBestOption, FGetImageDescContext* context) const override;
 		void GetLayoutBlockSize(int* pBlockX, int* pBlockY) override;
 		Ptr<ImageSizeExpression> GetImageSizeExpression() const override;
+		virtual FSourceDataDescriptor GetSourceDataDescriptor(FGetSourceDataDescriptorContext*) const override;
 	};
 
 }
