[mutable] Fixed Skeleton Retarget settings not propagating to transient merged skeletons.
#tests Editor
#rnx

diff --git a/Engine/Plugins/Experimental/Mutable/Source/CustomizableObject/Internal/MuCO/CustomizableObjectInstancePrivate.h b/Engine/Plugins/Experimental/Mutable/Source/CustomizableObject/Internal/MuCO/CustomizableObjectInstancePrivate.h
--- a/Engine/Plugins/Experimental/Mutable/Source/CustomizableObject/Internal/MuCO/CustomizableObjectInstancePrivate.h
+++ b/Engine/Plugins/Experimental/Mutable/Source/CustomizableObject/Internal/MuCO/CustomizableObjectInstancePrivate.h
@@ -275,7 +275,7 @@ private:
 	void BuildOrCopyClothingData(const TSharedRef<FUpdateContextPrivate>& OperationData, USkeletalMesh* SkeletalMesh, const USkeletalMesh* SrcSkeletalMesh, UCustomizableObjectInstance* CustomizableObjectInstance, int32 InstanceComponentIndex);
 	
 	//
-	USkeleton* MergeSkeletons(UCustomizableObject& CustomizableObject, const FMutableRefSkeletalMeshData& RefSkeletalMeshData, int32 ComponentIndex);
+	USkeleton* MergeSkeletons(UCustomizableObject& CustomizableObject, const FMutableRefSkeletalMeshData& RefSkeletalMeshData, int32 ComponentIndex, bool& bOutCreatedNewSkeleton);
 
 	//
 	UPhysicsAsset* GetOrBuildMainPhysicsAsset(const TSharedRef<FUpdateContextPrivate>& OperationData, TObjectPtr<class UPhysicsAsset> TamplateAsset, const mu::PhysicsBody* PhysicsBody, bool bDisableCollisionBetweenAssets, int32 ComponentIndex);
diff --git a/Engine/Plugins/Experimental/Mutable/Source/CustomizableObject/Private/MuCO/CustomizableObjectInstance.cpp b/Engine/Plugins/Experimental/Mutable/Source/CustomizableObject/Private/MuCO/CustomizableObjectInstance.cpp
--- a/Engine/Plugins/Experimental/Mutable/Source/CustomizableObject/Private/MuCO/CustomizableObjectInstance.cpp
+++ b/Engine/Plugins/Experimental/Mutable/Source/CustomizableObject/Private/MuCO/CustomizableObjectInstance.cpp
@@ -1192,9 +1192,10 @@ bool AreSkeletonsCompatible(const TArray<TObjectPtr<USkeleton>>& InSkeletons)
 #endif
 
 
-USkeleton* UCustomizableInstancePrivate::MergeSkeletons(UCustomizableObject& CustomizableObject, const FMutableRefSkeletalMeshData& RefSkeletalMeshData, int32 ComponentIndex)
+USkeleton* UCustomizableInstancePrivate::MergeSkeletons(UCustomizableObject& CustomizableObject, const FMutableRefSkeletalMeshData& RefSkeletalMeshData, int32 ComponentIndex, bool& bOutCreatedNewSkeleton)
 {
 	MUTABLE_CPUPROFILER_SCOPE(BuildSkeletonData_MergeSkeletons);
+	bOutCreatedNewSkeleton = false;
 
 	FCustomizableInstanceComponentData* ComponentData = GetComponentData(ComponentIndex);
 	check(ComponentData);
@@ -1227,7 +1228,7 @@ USkeleton* UCustomizableInstancePrivate::MergeSkeletons(UCustomizableObject& Cus
 #endif
 
 	FSkeletonMergeParams Params;
-	Exchange(Params.SkeletonsToMerge, ReferencedSkeletons.SkeletonsToMerge);
+	Params.SkeletonsToMerge = ReferencedSkeletons.SkeletonsToMerge;
 
 	USkeleton* FinalSkeleton = USkeletalMergingLibrary::MergeSkeletons(Params);
 	if (!FinalSkeleton)
@@ -1264,6 +1265,8 @@ USkeleton* UCustomizableInstancePrivate::MergeSkeletons(UCustomizableObject& Cus
 		// Add Skeleton to the cache
 		CustomizableObject.GetPrivate()->SkeletonCache.Add(ReferencedSkeletons.SkeletonIds, FinalSkeleton);
 		ReferencedSkeletons.SkeletonIds.Empty();
+
+		bOutCreatedNewSkeleton = true;
 	}
 	
 	return FinalSkeleton;
@@ -3631,7 +3634,9 @@ bool UCustomizableInstancePrivate::BuildSkeletonData(const TSharedRef<FUpdateCon
 {
 	MUTABLE_CPUPROFILER_SCOPE(UCustomizableInstancePrivate::BuildSkeletonData);
 
-	const TObjectPtr<USkeleton> Skeleton = MergeSkeletons(CustomizableObject, RefSkeletalMeshData, ComponentIndex);
+	bool bCreatedNewSkeleton = false;
+
+	const TObjectPtr<USkeleton> Skeleton = MergeSkeletons(CustomizableObject, RefSkeletalMeshData, ComponentIndex, bCreatedNewSkeleton);
 	if (!Skeleton)
 	{
 		return false;
@@ -3722,6 +3727,54 @@ bool UCustomizableInstancePrivate::BuildSkeletonData(const TSharedRef<FUpdateCon
 		SkeletalMesh.CalculateInvRefMatrices();
 	}
 
+	USkeleton* GeneratedSkeleton = SkeletalMesh.GetSkeleton();
+
+	if(GeneratedSkeleton && bCreatedNewSkeleton)
+	{
+		// If the skeleton is new, it means it has just been merged and the retargeting modes need merging too as the
+		// MergeSkeletons function doesn't do it. Only do it for newly generated ones, not for cached or non-transient ones.
+		GeneratedSkeleton->RecreateBoneTree(&SkeletalMesh);
+
+		FCustomizableInstanceComponentData* ComponentData = GetComponentData(ComponentIndex);
+		check(ComponentData);
+
+		TArray<TObjectPtr<USkeleton>>& SkeletonsToMerge = ComponentData->Skeletons.SkeletonsToMerge;
+		check(SkeletonsToMerge.Num() > 1);
+
+		TMap<FName, EBoneTranslationRetargetingMode::Type> BoneNamesToRetargetingMode;
+
+		const int32 NumberOfSkeletons = SkeletonsToMerge.Num();
+
+		for (int32 SkeletonIndex = 0; SkeletonIndex < NumberOfSkeletons; ++SkeletonIndex)
+		{
+			const USkeleton* ToMergeSkeleton = SkeletonsToMerge[SkeletonIndex];
+			const FReferenceSkeleton& ToMergeReferenceSkeleton = ToMergeSkeleton->GetReferenceSkeleton();
+			const TArray<FMeshBoneInfo>& Bones = ToMergeReferenceSkeleton.GetRawRefBoneInfo();
+
+			const int32 NumBones = Bones.Num();
+			for (int32 BoneIndex = 0; BoneIndex < NumBones; ++BoneIndex)
+			{
+				const FMeshBoneInfo& Bone = Bones[BoneIndex];
+
+				EBoneTranslationRetargetingMode::Type RetargetingMode = ToMergeSkeleton->GetBoneTranslationRetargetingMode(BoneIndex, false);
+				BoneNamesToRetargetingMode.Add(Bone.Name, RetargetingMode);
+			}
+		}
+
+		for (const auto& Pair : BoneNamesToRetargetingMode)
+		{
+			const FName& BoneName = Pair.Key;
+			const EBoneTranslationRetargetingMode::Type& RetargetingMode = Pair.Value;
+
+			const int32 BoneIndex = GeneratedSkeleton->GetReferenceSkeleton().FindRawBoneIndex(BoneName);
+
+			if (BoneIndex >= 0)
+			{
+				GeneratedSkeleton->SetBoneTranslationRetargetingMode(BoneIndex, RetargetingMode);
+			}
+		}
+	}
+
 	return true;
 }
 
