[mutable] Add switching and variation nodes for components.
- Implement the unreal side nodes.
- Implement code generation for the mutable side nodes.
- Rename more ComponentIndex variables to make them explicitly refer to object component index or instance component index.
- Fix bugs when optimizing the state root operations: the op reference in the state wasn't updated and could cause infinite loops.
#rnx

diff --git a/Engine/Plugins/Experimental/Mutable/Source/CustomizableObject/Internal/MuCO/CustomizableObjectInstancePrivate.h b/Engine/Plugins/Experimental/Mutable/Source/CustomizableObject/Internal/MuCO/CustomizableObjectInstancePrivate.h
--- a/Engine/Plugins/Experimental/Mutable/Source/CustomizableObject/Internal/MuCO/CustomizableObjectInstancePrivate.h
+++ b/Engine/Plugins/Experimental/Mutable/Source/CustomizableObject/Internal/MuCO/CustomizableObjectInstancePrivate.h
@@ -173,7 +173,7 @@ class CUSTOMIZABLEOBJECT_API UCustomizableInstancePrivate : public UObject
 public:
 	GENERATED_BODY()
 
-	/** The generated skeletal meshes for this Instance, one for each instance component */
+	/** The generated skeletal meshes for this Instance, one for each component in the Object component list which means that some may be null if the component hasn't been generated. */
 	UPROPERTY(Transient, VisibleAnywhere, Category = NoCategory, meta=(DisplayName = "Meshes"))
 	TArray<TObjectPtr<USkeletalMesh>> SkeletalMeshes;
 
@@ -265,9 +265,9 @@ public:
 
 private:
 
-	void InitSkeletalMeshData(const TSharedRef<FUpdateContextPrivate>& OperationData, USkeletalMesh* SkeletalMesh, const FMutableRefSkeletalMeshData& RefSkeletalMeshData, const UCustomizableObject& CustomizableObject, int32 ComponentIndex);
+	void InitSkeletalMeshData(const TSharedRef<FUpdateContextPrivate>& OperationData, USkeletalMesh* SkeletalMesh, const FMutableRefSkeletalMeshData& RefSkeletalMeshData, const UCustomizableObject& CustomizableObject, int32 ObjectComponentIndex);
 
-	bool BuildSkeletonData(const TSharedRef<FUpdateContextPrivate>& OperationData, USkeletalMesh& SkeletalMesh, const FMutableRefSkeletalMeshData& RefSkeletalMeshData, UCustomizableObject& CustomizableObject, int32 ComponentIndex);
+	bool BuildSkeletonData(const TSharedRef<FUpdateContextPrivate>& OperationData, USkeletalMesh& SkeletalMesh, const FMutableRefSkeletalMeshData& RefSkeletalMeshData, UCustomizableObject& CustomizableObject, int32 InstanceComponentIndex);
 	void BuildMeshSockets(const TSharedRef<FUpdateContextPrivate>& OperationData, USkeletalMesh* SkeletalMesh, const FModelResources& ModelResources, const FMutableRefSkeletalMeshData& RefSkeletalMeshData, mu::Ptr<const mu::Mesh> MutableMesh);
 	void BuildOrCopyElementData(const TSharedRef<FUpdateContextPrivate>& OperationData, USkeletalMesh* SkeletalMesh, UCustomizableObjectInstance* CustomizableObjectInstance, int32 InstanceComponentIndex);
 	void BuildOrCopyMorphTargetsData(const TSharedRef<FUpdateContextPrivate>& OperationData, USkeletalMesh* SkeletalMesh, const USkeletalMesh* SrcSkeletalMesh, UCustomizableObjectInstance* CustomizableObjectInstance, int32 InstanceComponentIndex);
@@ -275,10 +275,10 @@ private:
 	void BuildOrCopyClothingData(const TSharedRef<FUpdateContextPrivate>& OperationData, USkeletalMesh* SkeletalMesh, const USkeletalMesh* SrcSkeletalMesh, UCustomizableObjectInstance* CustomizableObjectInstance, int32 InstanceComponentIndex);
 	
 	//
-	USkeleton* MergeSkeletons(UCustomizableObject& CustomizableObject, const FMutableRefSkeletalMeshData& RefSkeletalMeshData, int32 ComponentIndex, bool& bOutCreatedNewSkeleton);
+	USkeleton* MergeSkeletons(UCustomizableObject& CustomizableObject, const FMutableRefSkeletalMeshData& RefSkeletalMeshData, int32 ObjectComponentIndex, bool& bOutCreatedNewSkeleton);
 
 	//
-	UPhysicsAsset* GetOrBuildMainPhysicsAsset(const TSharedRef<FUpdateContextPrivate>& OperationData, TObjectPtr<class UPhysicsAsset> TamplateAsset, const mu::PhysicsBody* PhysicsBody, bool bDisableCollisionBetweenAssets, int32 ComponentIndex);
+	UPhysicsAsset* GetOrBuildMainPhysicsAsset(const TSharedRef<FUpdateContextPrivate>& OperationData, TObjectPtr<class UPhysicsAsset> TamplateAsset, const mu::PhysicsBody* PhysicsBody, bool bDisableCollisionBetweenAssets, int32 InstanceComponentIndex);
 	
 	// Create a transient texture and add it to the TextureTrackerArray
 	UTexture2D* CreateTexture();
diff --git a/Engine/Plugins/Experimental/Mutable/Source/CustomizableObject/Internal/MuCO/CustomizableObjectSystemPrivate.h b/Engine/Plugins/Experimental/Mutable/Source/CustomizableObject/Internal/MuCO/CustomizableObjectSystemPrivate.h
--- a/Engine/Plugins/Experimental/Mutable/Source/CustomizableObject/Internal/MuCO/CustomizableObjectSystemPrivate.h
+++ b/Engine/Plugins/Experimental/Mutable/Source/CustomizableObject/Internal/MuCO/CustomizableObjectSystemPrivate.h
@@ -574,7 +574,7 @@ public:
 	TArray<TArray<mu::FResourceID>> MeshDescriptors;
 
 	/** Used to know if the updated instances' meshes are different from the previous ones. 
-	  * The index of the array is the component's index.
+	  * The index of the array is the instance component's index.
 	  * @return true if the mesh is new or new to this instance (e.g. mesh cached by another instance). */
 	TArray<bool> MeshChanged;
 	
diff --git a/Engine/Plugins/Experimental/Mutable/Source/CustomizableObject/Private/MuCO/CustomizableInstanceLODManagement.cpp b/Engine/Plugins/Experimental/Mutable/Source/CustomizableObject/Private/MuCO/CustomizableInstanceLODManagement.cpp
--- a/Engine/Plugins/Experimental/Mutable/Source/CustomizableObject/Private/MuCO/CustomizableInstanceLODManagement.cpp
+++ b/Engine/Plugins/Experimental/Mutable/Source/CustomizableObject/Private/MuCO/CustomizableInstanceLODManagement.cpp
@@ -359,12 +359,12 @@ void UCustomizableInstanceLODManagement::UpdateInstanceDistsAndLODs(FMutableInst
 				{
 					COI->SetIsBeingUsedByComponentInPlay(true);
 
-					int32 ComponentIndex = CustomizableObjectInstanceUsage->GetComponentIndex();
+					int32 ObjectComponentIndex = CustomizableObjectInstanceUsage->GetComponentIndex();
 
 #if WITH_EDITOR
 					// If the instance is generated but the component doesn't have a mesh, set it.
 					// Can happen when duplicating instances in the editor.
-					const USkeletalMesh* SkeletalMesh = COI->GetSkeletalMesh(ComponentIndex);
+					const USkeletalMesh* SkeletalMesh = COI->GetSkeletalMesh(ObjectComponentIndex);
 					if (SkeletalMesh && !Parent->GetSkeletalMeshAsset())
 					{
 						// As the instance is already generated, this will be very fast and just set the mesh and call the delegates
@@ -386,9 +386,9 @@ void UCustomizableInstanceLODManagement::UpdateInstanceDistsAndLODs(FMutableInst
 					LODTracker.MinLOD = FMath::Min(LODTracker.MinLOD, Parent->bOverrideMinLod ? Parent->MinLodModel : 0);
 
 					// If the parent component have a SkeletalMesh use the RequestedLODLevel of the component as reference to know which LODs mutable should generate.
-					if (UE_MUTABLE_GETSKELETALMESHASSET(Parent) && LODTracker.RequestedLODPerComponent.IsValidIndex(ComponentIndex))
+					if (UE_MUTABLE_GETSKELETALMESHASSET(Parent) && LODTracker.RequestedLODPerComponent.IsValidIndex(ObjectComponentIndex))
 					{
-						uint16& RequestedLOD = LODTracker.RequestedLODPerComponent[ComponentIndex];
+						uint16& RequestedLOD = LODTracker.RequestedLODPerComponent[ObjectComponentIndex];
 						RequestedLOD = FMath::Min((int32)RequestedLOD, Parent->GetPredictedLODLevel());
 					}
 				}
diff --git a/Engine/Plugins/Experimental/Mutable/Source/CustomizableObject/Private/MuCO/CustomizableObject.cpp b/Engine/Plugins/Experimental/Mutable/Source/CustomizableObject/Private/MuCO/CustomizableObject.cpp
--- a/Engine/Plugins/Experimental/Mutable/Source/CustomizableObject/Private/MuCO/CustomizableObject.cpp
+++ b/Engine/Plugins/Experimental/Mutable/Source/CustomizableObject/Private/MuCO/CustomizableObject.cpp
@@ -1130,19 +1130,19 @@ void UCustomizableObjectPrivate::AddUncompiledCOWarning(const FString& Additiona
 }
 
 
-USkeletalMesh* UCustomizableObject::GetRefSkeletalMesh(int32 ComponentIndex) const
+USkeletalMesh* UCustomizableObject::GetRefSkeletalMesh(int32 ObjectComponentIndex) const
 {
 #if WITH_EDITORONLY_DATA
-	if (GetPrivate()->MutableMeshComponents.IsValidIndex(ComponentIndex))
+	if (GetPrivate()->MutableMeshComponents.IsValidIndex(ObjectComponentIndex))
 	{
-		return GetPrivate()->MutableMeshComponents[ComponentIndex].ReferenceSkeletalMesh;
+		return GetPrivate()->MutableMeshComponents[ObjectComponentIndex].ReferenceSkeletalMesh;
 	}
 #else
 	const FModelResources& ModelResources = Private->GetModelResources();
-	if (ModelResources.ReferenceSkeletalMeshesData.IsValidIndex(ComponentIndex))
+	if (ModelResources.ReferenceSkeletalMeshesData.IsValidIndex(ObjectComponentIndex))
 	{
 		// Can be nullptr if RefSkeletalMeshes are not loaded yet.
-		return ModelResources.ReferenceSkeletalMeshesData[ComponentIndex].SkeletalMesh;
+		return ModelResources.ReferenceSkeletalMeshesData[ObjectComponentIndex].SkeletalMesh;
 	}
 #endif
 	return nullptr;
diff --git a/Engine/Plugins/Experimental/Mutable/Source/CustomizableObject/Private/MuCO/CustomizableObjectInstance.cpp b/Engine/Plugins/Experimental/Mutable/Source/CustomizableObject/Private/MuCO/CustomizableObjectInstance.cpp
--- a/Engine/Plugins/Experimental/Mutable/Source/CustomizableObject/Private/MuCO/CustomizableObjectInstance.cpp
+++ b/Engine/Plugins/Experimental/Mutable/Source/CustomizableObject/Private/MuCO/CustomizableObjectInstance.cpp
@@ -1192,12 +1192,12 @@ bool AreSkeletonsCompatible(const TArray<TObjectPtr<USkeleton>>& InSkeletons)
 #endif
 
 
-USkeleton* UCustomizableInstancePrivate::MergeSkeletons(UCustomizableObject& CustomizableObject, const FMutableRefSkeletalMeshData& RefSkeletalMeshData, int32 ComponentIndex, bool& bOutCreatedNewSkeleton)
+USkeleton* UCustomizableInstancePrivate::MergeSkeletons(UCustomizableObject& CustomizableObject, const FMutableRefSkeletalMeshData& RefSkeletalMeshData, int32 ObjectComponentIndex, bool& bOutCreatedNewSkeleton)
 {
 	MUTABLE_CPUPROFILER_SCOPE(BuildSkeletonData_MergeSkeletons);
 	bOutCreatedNewSkeleton = false;
 
-	FCustomizableInstanceComponentData* ComponentData = GetComponentData(ComponentIndex);
+	FCustomizableInstanceComponentData* ComponentData = GetComponentData(ObjectComponentIndex);
 	check(ComponentData);
 
 	FReferencedSkeletons& ReferencedSkeletons = ComponentData->Skeletons;
@@ -1420,7 +1420,7 @@ namespace
 		const TSharedRef<FUpdateContextPrivate>& OperationData,
 		TObjectPtr<UPhysicsAsset> TemplateAsset,
 		const mu::PhysicsBody* MutablePhysics,
-		int32 ComponentIndex)
+		int32 InstanceComponentIndex)
 	{
 		check(TemplateAsset);
 		TObjectPtr<UPhysicsAsset> Result = NewObject<UPhysicsAsset>();
@@ -1435,7 +1435,7 @@ namespace
 
 		Result->bNotForDedicatedServer = TemplateAsset->bNotForDedicatedServer;
 
-		const TMap<mu::FBoneName, TPair<FName, uint16>>& BoneInfoMap = OperationData->InstanceUpdateData.Skeletons[ComponentIndex].BoneInfoMap;
+		const TMap<mu::FBoneName, TPair<FName, uint16>>& BoneInfoMap = OperationData->InstanceUpdateData.Skeletons[InstanceComponentIndex].BoneInfoMap;
 		TMap<FName, int32> BonesInUse;
 
 		const int32 MutablePhysicsBodyCount = MutablePhysics->GetBodyCount();
@@ -1548,7 +1548,7 @@ UPhysicsAsset* UCustomizableInstancePrivate::GetOrBuildMainPhysicsAsset(
 	TObjectPtr<UPhysicsAsset> TemplateAsset,
 	const mu::PhysicsBody* MutablePhysics,
 	bool bDisableCollisionsBetweenDifferentAssets,
-	int32 ComponentIndex)
+	int32 InstanceComponentIndex)
 {
 
 	MUTABLE_CPUPROFILER_SCOPE(MergePhysicsAssets);
@@ -1557,7 +1557,10 @@ UPhysicsAsset* UCustomizableInstancePrivate::GetOrBuildMainPhysicsAsset(
 
 	UPhysicsAsset* Result = nullptr;
 
-	FCustomizableInstanceComponentData* ComponentData = GetComponentData(ComponentIndex);
+	const FInstanceUpdateData::FComponent& Component = OperationData->InstanceUpdateData.Components[InstanceComponentIndex];
+	int32 ObjectComponentIndex = Component.Id;
+
+	FCustomizableInstanceComponentData* ComponentData = GetComponentData(ObjectComponentIndex);
 	check(ComponentData);
 
 	TArray<TObjectPtr<UPhysicsAsset>>& PhysicsAssets = ComponentData->PhysicsAssets.PhysicsAssetsToMerge;
@@ -1601,7 +1604,7 @@ UPhysicsAsset* UCustomizableInstancePrivate::GetOrBuildMainPhysicsAsset(
 
 	Result->bNotForDedicatedServer = TemplateAsset->bNotForDedicatedServer;
 
-	const TMap<mu::FBoneName, TPair<FName, uint16>>& BoneInfoMap = OperationData->InstanceUpdateData.Skeletons[ComponentIndex].BoneInfoMap;
+	const TMap<mu::FBoneName, TPair<FName, uint16>>& BoneInfoMap = OperationData->InstanceUpdateData.Skeletons[InstanceComponentIndex].BoneInfoMap;
 	TMap<FName, int32> BonesInUse;
 
 	const int32 MutablePhysicsBodyCount = MutablePhysics->GetBodyCount();
@@ -1867,31 +1870,24 @@ bool UCustomizableInstancePrivate::DoComponentsNeedUpdate(UCustomizableObjectIns
 		return false;
 	}
 
-	const int32 NumComponents = OperationData->InstanceUpdateData.Components.Num();
+	const int32 NumInstanceComponents = OperationData->InstanceUpdateData.Components.Num();
 
+	// To be indexed with instance component index
 	TArray<bool> ComponentWithMesh;
-	ComponentWithMesh.Init(false, NumComponents);
+	ComponentWithMesh.Init(false, NumInstanceComponents);
 
 	TArray<mu::FResourceID> MeshIDs;
-	MeshIDs.Init(MAX_uint64, NumComponents * MAX_MESH_LOD_COUNT);
+	MeshIDs.Init(MAX_uint64, NumInstanceComponents * MAX_MESH_LOD_COUNT);
 
 	// Gather the Mesh Ids of all components, and validate the integrity of the meshes to generate. 
-	for (int32 ComponentIndex = 0; ComponentIndex < NumComponents; ++ComponentIndex)
+	for (int32 InstanceComponentIndex = 0; InstanceComponentIndex < NumInstanceComponents; ++InstanceComponentIndex)
 	{
-		const FInstanceUpdateData::FComponent& Component = OperationData->InstanceUpdateData.Components[ComponentIndex];
+		const FInstanceUpdateData::FComponent& Component = OperationData->InstanceUpdateData.Components[InstanceComponentIndex];
 		
 		for (int32 LODIndex = OperationData->GetMinLOD(); LODIndex < Component.LODCount; ++LODIndex)
 		{
 			const FInstanceUpdateData::FLOD& LOD = OperationData->InstanceUpdateData.LODs[Component.FirstLOD + LODIndex];
 
-
-			if (ComponentIndex >= NumComponents)
-			{
-				ensureMsgf(false, TEXT("Mutable: Failed to generate SkeletalMesh. Invalid ComponentIndex. Index [%d]. Valid range [0 .. %d)."), ComponentIndex, NumComponents);
-				bHasInvalidMesh = true;
-				return false;
-			}
-
 			if (!LOD.bGenerated || !LOD.Mesh )
 			{
 				continue;
@@ -1907,42 +1903,44 @@ bool UCustomizableInstancePrivate::DoComponentsNeedUpdate(UCustomizableObjectIns
 			{
 				UE_LOG(LogMutable, Error, TEXT("Failed to generate SkeletalMesh for CO Instance [%s]. CO [%s] has invalid geometry for LOD [%d] Component [%d]."),
 					*Public->GetName(), *CustomizableObject->GetName(),
-					LODIndex, ComponentIndex);
+					LODIndex, NumInstanceComponents);
 				bHasInvalidMesh = true;
 				continue;
 			}
 
-			ComponentWithMesh[ComponentIndex] = true;
-			MeshIDs[ComponentIndex * MAX_MESH_LOD_COUNT + LODIndex] = LOD.MeshID;
+			ComponentWithMesh[InstanceComponentIndex] = true;
+			MeshIDs[InstanceComponentIndex * MAX_MESH_LOD_COUNT + LODIndex] = LOD.MeshID;
 		}
 	}
 
 	// Find which components need an update
-	OperationData->MeshChanged.AddDefaulted(NumComponents);
+	OperationData->MeshChanged.Init(false,NumInstanceComponents);
 
-	for (int32 ComponentIndex = 0; ComponentIndex < NumComponents; ++ComponentIndex)
+	for (int32 InstanceComponentIndex = 0; InstanceComponentIndex < NumInstanceComponents; ++InstanceComponentIndex)
 	{
+		const FInstanceUpdateData::FComponent& Component = OperationData->InstanceUpdateData.Components[InstanceComponentIndex];
+		int32 ObjectComponentIndex = Component.Id;
+
 		if (OperationData->bUseMeshCache)
 		{
-			USkeletalMesh* CachedMesh = CustomizableObject->GetPrivate()->MeshCache.Get(OperationData->MeshDescriptors[ComponentIndex]);
+			USkeletalMesh* CachedMesh = CustomizableObject->GetPrivate()->MeshCache.Get(OperationData->MeshDescriptors[ObjectComponentIndex]);
 			if (CachedMesh)
 			{
-				const bool bMeshNeedsUpdate = !SkeletalMeshes.IsValidIndex(ComponentIndex) || (SkeletalMeshes[ComponentIndex] != CachedMesh);
-				OperationData->MeshChanged[ComponentIndex] = bMeshNeedsUpdate;
-				ComponentWithMesh[ComponentIndex] = true;
+				const bool bMeshNeedsUpdate = !SkeletalMeshes.IsValidIndex(ObjectComponentIndex) || (SkeletalMeshes[ObjectComponentIndex] != CachedMesh);
+				OperationData->MeshChanged[InstanceComponentIndex] = bMeshNeedsUpdate;
+				ComponentWithMesh[InstanceComponentIndex] = true;
 				continue;
 			}
 		}
 
 		// Components with mesh must have valid geometry at CurrentMaxLOD
-		const FInstanceUpdateData::FComponent& Component = OperationData->InstanceUpdateData.Components[ComponentIndex];
 
-		const bool bHasSkeletalMesh = SkeletalMeshes.IsValidIndex(ComponentIndex) && SkeletalMeshes[ComponentIndex];
+		const bool bHasSkeletalMesh = SkeletalMeshes.IsValidIndex(ObjectComponentIndex) && SkeletalMeshes[ObjectComponentIndex];
 
 		if (Component.LODCount == 0)
 		{
 			// We don't have a mesh in the component, so it has changed if we had one before.
-			OperationData->MeshChanged[ComponentIndex] = bHasSkeletalMesh;
+			OperationData->MeshChanged[InstanceComponentIndex] = bHasSkeletalMesh;
 			continue;
 		}
 
@@ -1950,34 +1948,39 @@ bool UCustomizableInstancePrivate::DoComponentsNeedUpdate(UCustomizableObjectIns
 		bool bIsReferenced = LOD.Mesh && LOD.Mesh->IsReference();
 		if (!bIsReferenced)
 		{
-			if (ComponentWithMesh[ComponentIndex] && MeshIDs[ComponentIndex * MAX_MESH_LOD_COUNT + OperationData->NumLODsAvailablePerComponent[ComponentIndex] - 1] == MAX_uint64)
+			if (ComponentWithMesh[InstanceComponentIndex] && MeshIDs[InstanceComponentIndex * MAX_MESH_LOD_COUNT + OperationData->NumLODsAvailablePerComponent[ObjectComponentIndex] - 1] == MAX_uint64)
 			{
-				UE_LOG(LogMutable, Error, TEXT("Failed to generate SkeletalMesh for CO Instance [%s]. CO [%s] is missing geometry for LOD [%d] Component [%d]."),
+				UE_LOG(LogMutable, Error, TEXT("Failed to generate SkeletalMesh for CO Instance [%s]. CO [%s] is missing geometry for LOD [%d] Object Component [%d]."),
 					*Public->GetName(), *CustomizableObject->GetName(),
-					OperationData->NumLODsAvailablePerComponent[ComponentIndex] - 1, ComponentIndex);
+					OperationData->NumLODsAvailablePerComponent[ObjectComponentIndex] - 1, ObjectComponentIndex);
 				bHasInvalidMesh = true;
 				continue;
 			}
 		}
 
 		// Update the component if there is a mesh and it shouldn't, or the other way around.
-		OperationData->MeshChanged[ComponentIndex] = (ComponentWithMesh[ComponentIndex] != bHasSkeletalMesh);
+		OperationData->MeshChanged[InstanceComponentIndex] = (ComponentWithMesh[InstanceComponentIndex] != bHasSkeletalMesh);
 
-		const FCustomizableInstanceComponentData* ComponentData = GetComponentData(ComponentIndex);
+		const FCustomizableInstanceComponentData* ComponentData = GetComponentData(ObjectComponentIndex);
 		if (!ComponentData) // Could be nullptr if the component has not been generated.
 		{
 			continue;
 		}
 
 		// Update if MeshIDs are different
-		const int32 ComponentOffset = ComponentIndex * MAX_MESH_LOD_COUNT;
-		for (int32 MeshIndex = 0; !OperationData->MeshChanged[ComponentIndex] && MeshIndex < MAX_MESH_LOD_COUNT; ++MeshIndex)
+		const int32 ComponentOffset = InstanceComponentIndex * MAX_MESH_LOD_COUNT;
+		for (int32 MeshIndex = 0; !OperationData->MeshChanged[InstanceComponentIndex] && MeshIndex < MAX_MESH_LOD_COUNT; ++MeshIndex)
 		{
-			OperationData->MeshChanged[ComponentIndex] = MeshIDs[ComponentOffset + MeshIndex] != ComponentData->LastMeshIdPerLOD[MeshIndex];
+			OperationData->MeshChanged[InstanceComponentIndex] = MeshIDs[ComponentOffset + MeshIndex] != ComponentData->LastMeshIdPerLOD[MeshIndex];
 		}
 	}
 
-	return !bHasInvalidMesh && OperationData->MeshChanged.Find(true) != INDEX_NONE;
+	bool bChanged =
+		OperationData->MeshChanged.Num() != SkeletalMeshes.Num()
+		||
+		OperationData->MeshChanged.Find(true) != INDEX_NONE;
+
+	return !bHasInvalidMesh && bChanged;
 }
 
 
@@ -2029,23 +2032,24 @@ bool UCustomizableInstancePrivate::UpdateSkeletalMesh_PostBeginUpdate0(UCustomiz
 	// Initialize the maximum number of SkeletalMeshes we could possibly have. 
 	SkeletalMeshes.Init(nullptr, CustomizableObject->GetComponentCount());
 
-	const int32 NumMutableComponents = OperationData->InstanceUpdateData.Components.Num();
-	for (int32 ComponentIndex = 0; ComponentIndex < NumMutableComponents; ++ComponentIndex)
+	const int32 NumInstanceComponents = OperationData->InstanceUpdateData.Components.Num();
+	for (int32 InstanceComponentIndex = 0; InstanceComponentIndex < NumInstanceComponents; ++InstanceComponentIndex)
 	{
-		const FInstanceUpdateData::FComponent& Component = OperationData->InstanceUpdateData.Components[ComponentIndex];
+		const FInstanceUpdateData::FComponent& Component = OperationData->InstanceUpdateData.Components[InstanceComponentIndex];
+		int32 ObjectComponentIndex = Component.Id;
 
 		// If the component doesn't need an update copy the previously generated mesh.
-		if (!OperationData->MeshChanged[ComponentIndex])
+		if (!OperationData->MeshChanged[InstanceComponentIndex])
 		{
-			if (OldSkeletalMeshes.IsValidIndex(ComponentIndex))
+			if (OldSkeletalMeshes.IsValidIndex(ObjectComponentIndex))
 			{
-				SkeletalMeshes[ComponentIndex] = OldSkeletalMeshes[ComponentIndex];
+				SkeletalMeshes[ObjectComponentIndex] = OldSkeletalMeshes[ObjectComponentIndex];
 			}
 
 			continue;
 		}
 
-		if (!ComponentsData.IsValidIndex(ComponentIndex))
+		if (!ComponentsData.IsValidIndex(ObjectComponentIndex))
 		{
 			bSuccess = false;
 			ensure(false);
@@ -2054,17 +2058,17 @@ bool UCustomizableInstancePrivate::UpdateSkeletalMesh_PostBeginUpdate0(UCustomiz
 
 		if (OperationData->bUseMeshCache)
 		{
-			if (USkeletalMesh* CachedMesh = CustomizableObject->GetPrivate()->MeshCache.Get(OperationData->MeshDescriptors[ComponentIndex]))
+			if (USkeletalMesh* CachedMesh = CustomizableObject->GetPrivate()->MeshCache.Get(OperationData->MeshDescriptors[ObjectComponentIndex]))
 			{
-				check(OperationData->MeshDescriptors[ComponentIndex].Num() == MAX_MESH_LOD_COUNT);
-				ComponentsData[ComponentIndex].LastMeshIdPerLOD = OperationData->MeshDescriptors[ComponentIndex];
-				SkeletalMeshes[ComponentIndex] = CachedMesh;
+				check(OperationData->MeshDescriptors[ObjectComponentIndex].Num() == MAX_MESH_LOD_COUNT);
+				ComponentsData[ObjectComponentIndex].LastMeshIdPerLOD = OperationData->MeshDescriptors[ObjectComponentIndex];
+				SkeletalMeshes[ObjectComponentIndex] = CachedMesh;
 				continue;
 			}
 		}
 
 		// Reset last mesh IDs.
-		ComponentsData[ComponentIndex].LastMeshIdPerLOD.Init(MAX_uint64, MAX_MESH_LOD_COUNT);
+		ComponentsData[ObjectComponentIndex].LastMeshIdPerLOD.Init(MAX_uint64, MAX_MESH_LOD_COUNT);
 
 		// If it is a referenced resource, only the first LOD is relevant.
 		{
@@ -2085,37 +2089,37 @@ bool UCustomizableInstancePrivate::UpdateSkeletalMesh_PostBeginUpdate0(UCustomiz
 					UE_LOG(LogMutable, Error, TEXT("Referenced skeletal mesh [%s] was not pre-loaded. It will be sync-loaded probably causing a hitch. CO [%s]"), *Ref.ToString(), *GetNameSafe(CustomizableObject));
 				}
 
-				SkeletalMeshes[ComponentIndex] = Ref.LoadSynchronous();
+				SkeletalMeshes[ObjectComponentIndex] = Ref.LoadSynchronous();
 				continue;
 			}
 		}
 
-		if (!ModelResources.ReferenceSkeletalMeshesData.IsValidIndex(ComponentIndex))
+		if (!ModelResources.ReferenceSkeletalMeshesData.IsValidIndex(ObjectComponentIndex))
 		{
 			bSuccess = false;
 			break;
 		}
 
-		const FMutableRefSkeletalMeshData& RefSkeletalMeshData = ModelResources.ReferenceSkeletalMeshesData[ComponentIndex];
+		const FMutableRefSkeletalMeshData& RefSkeletalMeshData = ModelResources.ReferenceSkeletalMeshesData[ObjectComponentIndex];
 
 		// Create and initialize the SkeletalMesh for this component
 		MUTABLE_CPUPROFILER_SCOPE(ConstructMesh);
 
 		if (OperationData->bStreamMeshLODs)
 		{
-			SkeletalMeshes[ComponentIndex] = UCustomizableObjectSkeletalMesh::CreateSkeletalMesh(OperationData, *Public, *CustomizableObject, ComponentIndex);
+			SkeletalMeshes[ObjectComponentIndex] = UCustomizableObjectSkeletalMesh::CreateSkeletalMesh(OperationData, *Public, *CustomizableObject, InstanceComponentIndex);
 		}
 		else
 		{
 			FName SkeletalMeshName = GenerateUniqueNameFromCOInstance(*Public);
-			SkeletalMeshes[ComponentIndex] = NewObject<USkeletalMesh>(GetTransientPackage(), SkeletalMeshName, RF_Transient);
+			SkeletalMeshes[ObjectComponentIndex] = NewObject<USkeletalMesh>(GetTransientPackage(), SkeletalMeshName, RF_Transient);
 		}
 
-		USkeletalMesh* SkeletalMesh = SkeletalMeshes[ComponentIndex];
+		USkeletalMesh* SkeletalMesh = SkeletalMeshes[ObjectComponentIndex];
 		check(SkeletalMesh);
 
 		// Set up the default information any mesh from this component will have (LODArrayInfos, RenderData, Mesh settings, etc). 
-		InitSkeletalMeshData(OperationData, SkeletalMesh, RefSkeletalMeshData, *CustomizableObject, ComponentIndex);
+		InitSkeletalMeshData(OperationData, SkeletalMesh, RefSkeletalMeshData, *CustomizableObject, ObjectComponentIndex);
 
 		// We need the first valid mesh. get it from the component, considering that some LOSs may have been skipped.
 		mu::Ptr<const mu::Mesh> ComponentMesh;
@@ -2127,7 +2131,7 @@ bool UCustomizableInstancePrivate::UpdateSkeletalMesh_PostBeginUpdate0(UCustomiz
 		if (ComponentMesh)
 		{
 			// Construct a new skeleton, fix up ActiveBones and Bonemap arrays and recompute the RefInvMatrices
-			bSuccess = BuildSkeletonData(OperationData, *SkeletalMesh, RefSkeletalMeshData, *CustomizableObject, ComponentIndex);
+			bSuccess = BuildSkeletonData(OperationData, *SkeletalMesh, RefSkeletalMeshData, *CustomizableObject, InstanceComponentIndex);
 
 			if (!bSuccess)
 			{
@@ -2139,7 +2143,7 @@ bool UCustomizableInstancePrivate::UpdateSkeletalMesh_PostBeginUpdate0(UCustomiz
 			{
 				constexpr bool bDisallowCollisionBetweenAssets = true;
 				UPhysicsAsset* PhysicsAssetResult = GetOrBuildMainPhysicsAsset(
-					OperationData, RefSkeletalMeshData.PhysicsAsset, MutablePhysics.get(), bDisallowCollisionBetweenAssets, ComponentIndex);
+					OperationData, RefSkeletalMeshData.PhysicsAsset, MutablePhysics.get(), bDisallowCollisionBetweenAssets, InstanceComponentIndex);
 
 				SkeletalMesh->SetPhysicsAsset(PhysicsAssetResult);
 
@@ -2189,7 +2193,7 @@ bool UCustomizableInstancePrivate::UpdateSkeletalMesh_PostBeginUpdate0(UCustomiz
 
 				Entry.PropertyIndex = Info.PropertyIndex;
 				Entry.PhysicsAsset = MakePhysicsAssetFromTemplateAndMutableBody(
-					OperationData,PhysicsAssetTemplate, AdditionalPhysiscBody.get(), ComponentIndex);
+					OperationData,PhysicsAssetTemplate, AdditionalPhysiscBody.get(), InstanceComponentIndex);
 			}
 
 			// Add sockets from the SkeletalMesh of reference and from the MutableMesh
@@ -2286,7 +2290,7 @@ bool UCustomizableInstancePrivate::UpdateSkeletalMesh_PostBeginUpdate0(UCustomiz
 			// give each one its data.
 			for (const TPair<const UCustomizableObjectExtension*, TArray<FInputPinDataContainer>>& Pair : ExtensionToExtensionData)
 			{
-				Pair.Key->OnSkeletalMeshCreated(Pair.Value, ComponentIndex, SkeletalMesh);
+				Pair.Key->OnSkeletalMeshCreated(Pair.Value, ObjectComponentIndex, SkeletalMesh);
 			}
 		}
 	}
@@ -2976,12 +2980,12 @@ void UCustomizableInstancePrivate::DiscardResources()
 
 	if (SkeletalMeshStatus == ESkeletalMeshStatus::Success)
 	{
-		for (int32 Component = 0; Component < SkeletalMeshes.Num(); ++Component)
+		for (int32 ObjectComponentIndex = 0; ObjectComponentIndex < SkeletalMeshes.Num(); ++ObjectComponentIndex)
 		{
-			if (SkeletalMeshes[Component] && SkeletalMeshes[Component]->IsValidLowLevel())
+			if (SkeletalMeshes[ObjectComponentIndex] && SkeletalMeshes[ObjectComponentIndex]->IsValidLowLevel())
 			{
-				SkeletalMeshes[Component]->ReleaseResources();
-				SkeletalMeshes[Component] = nullptr;
+				SkeletalMeshes[ObjectComponentIndex]->ReleaseResources();
+				SkeletalMeshes[ObjectComponentIndex] = nullptr;
 			}
 		}
 
@@ -3020,8 +3024,8 @@ void UCustomizableInstancePrivate::SetDefaultSkeletalMesh(bool bSetEmptyMesh) co
 			continue;
 		}
 #endif
-		const int32 ComponentIndex = CustomizableObjectInstanceUsage->GetComponentIndex();
-		if (!ModelResources.ReferenceSkeletalMeshesData.IsValidIndex(ComponentIndex))
+		const int32 ObjectComponentIndex = CustomizableObjectInstanceUsage->GetComponentIndex();
+		if (!ModelResources.ReferenceSkeletalMeshesData.IsValidIndex(ObjectComponentIndex))
 		{
 			continue;
 		}
@@ -3031,7 +3035,7 @@ void UCustomizableInstancePrivate::SetDefaultSkeletalMesh(bool bSetEmptyMesh) co
 		if (!bSetEmptyMesh)
 		{
 			// Force load the reference mesh if necessary. 
-			SkeletalMesh = ModelResources.ReferenceSkeletalMeshesData[ComponentIndex].SoftSkeletalMesh.LoadSynchronous();
+			SkeletalMesh = ModelResources.ReferenceSkeletalMeshesData[ObjectComponentIndex].SoftSkeletalMesh.LoadSynchronous();
 		}
 		
 		CustomizableObjectInstanceUsage->SetSkeletalMesh(SkeletalMesh);
@@ -3483,7 +3487,7 @@ void ConvertImage(UTexture2D* Texture, mu::Ptr<const mu::Image> MutableImage, co
 }
 
 
-void UCustomizableInstancePrivate::InitSkeletalMeshData(const TSharedRef<FUpdateContextPrivate>& OperationData, USkeletalMesh* SkeletalMesh, const FMutableRefSkeletalMeshData& RefSkeletalMeshData, const UCustomizableObject& CustomizableObject, int32 ComponentIndex)
+void UCustomizableInstancePrivate::InitSkeletalMeshData(const TSharedRef<FUpdateContextPrivate>& OperationData, USkeletalMesh* SkeletalMesh, const FMutableRefSkeletalMeshData& RefSkeletalMeshData, const UCustomizableObject& CustomizableObject, int32 ObjectComponentIndex)
 {
 	MUTABLE_CPUPROFILER_SCOPE(UCustomizableInstancePrivate::InitSkeletalMesh);
 
@@ -3521,7 +3525,7 @@ void UCustomizableInstancePrivate::InitSkeletalMeshData(const TSharedRef<FUpdate
 			}
 		}
 
-		const FCustomizableInstanceComponentData& ComponentData = ComponentsData[ComponentIndex];
+		const FCustomizableInstanceComponentData& ComponentData = ComponentsData[ObjectComponentIndex];
 		for (TObjectPtr<UAssetUserData> AssetUserData : ComponentData.AssetUserDataArray)
 		{
 			SkeletalMesh->AddAssetUserData(AssetUserData);
@@ -3553,12 +3557,12 @@ void UCustomizableInstancePrivate::InitSkeletalMeshData(const TSharedRef<FUpdate
 		SkeletalMesh->AllocateResourceForRendering();
 
 		FSkeletalMeshRenderData* RenderData = SkeletalMesh->GetResourceForRendering();
-		RenderData->NumInlinedLODs = OperationData->NumLODsAvailablePerComponent[ComponentIndex] - OperationData->FirstResidentLOD;
-		RenderData->NumNonOptionalLODs = OperationData->NumLODsAvailablePerComponent[ComponentIndex] - OperationData->FirstLODAvailable;
+		RenderData->NumInlinedLODs = OperationData->NumLODsAvailablePerComponent[ObjectComponentIndex] - OperationData->FirstResidentLOD;
+		RenderData->NumNonOptionalLODs = OperationData->NumLODsAvailablePerComponent[ObjectComponentIndex] - OperationData->FirstLODAvailable;
 		RenderData->CurrentFirstLODIdx = OperationData->FirstResidentLOD;
 		RenderData->LODBiasModifier = 0;
 
-		for (int32 LODIndex = 0; LODIndex < OperationData->NumLODsAvailablePerComponent[ComponentIndex]; ++LODIndex)
+		for (int32 LODIndex = 0; LODIndex < OperationData->NumLODsAvailablePerComponent[ObjectComponentIndex]; ++LODIndex)
 		{
 			RenderData->LODRenderData.Add(new FSkeletalMeshLODRenderData());
 			
@@ -3630,13 +3634,14 @@ void UCustomizableInstancePrivate::InitSkeletalMeshData(const TSharedRef<FUpdate
 }
 
 
-bool UCustomizableInstancePrivate::BuildSkeletonData(const TSharedRef<FUpdateContextPrivate>& OperationData, USkeletalMesh& SkeletalMesh, const FMutableRefSkeletalMeshData& RefSkeletalMeshData, UCustomizableObject& CustomizableObject, int32 ComponentIndex)
+bool UCustomizableInstancePrivate::BuildSkeletonData(const TSharedRef<FUpdateContextPrivate>& OperationData, USkeletalMesh& SkeletalMesh, const FMutableRefSkeletalMeshData& RefSkeletalMeshData, UCustomizableObject& CustomizableObject, int32 InstanceComponentIndex)
 {
 	MUTABLE_CPUPROFILER_SCOPE(UCustomizableInstancePrivate::BuildSkeletonData);
 
-	bool bCreatedNewSkeleton = false;
+	int32 ObjectComponentIndex = OperationData->InstanceUpdateData.Components[InstanceComponentIndex].Id;
 
-	const TObjectPtr<USkeleton> Skeleton = MergeSkeletons(CustomizableObject, RefSkeletalMeshData, ComponentIndex, bCreatedNewSkeleton);
+	bool bCreatedNewSkeleton = false;
+	const TObjectPtr<USkeleton> Skeleton = MergeSkeletons(CustomizableObject, RefSkeletalMeshData, ObjectComponentIndex, bCreatedNewSkeleton);
 	if (!Skeleton)
 	{
 		return false;
@@ -3650,8 +3655,8 @@ bool UCustomizableInstancePrivate::BuildSkeletonData(const TSharedRef<FUpdateCon
 	const TArray<FMeshBoneInfo>& RawRefBoneInfo = ReferenceSkeleton.GetRawRefBoneInfo();
 	const int32 RawRefBoneCount = ReferenceSkeleton.GetRawBoneNum();
 
-	const TArray<FInstanceUpdateData::FBone>& BonePose = OperationData->InstanceUpdateData.Skeletons[ComponentIndex].BonePose;
-	TMap<mu::FBoneName, TPair<FName,uint16>>& BoneInfoMap = OperationData->InstanceUpdateData.Skeletons[ComponentIndex].BoneInfoMap;
+	const TArray<FInstanceUpdateData::FBone>& BonePose = OperationData->InstanceUpdateData.Skeletons[InstanceComponentIndex].BonePose;
+	TMap<mu::FBoneName, TPair<FName,uint16>>& BoneInfoMap = OperationData->InstanceUpdateData.Skeletons[InstanceComponentIndex].BoneInfoMap;
 
 	{
 		MUTABLE_CPUPROFILER_SCOPE(BuildSkeletonData_BuildBoneInfoMap);
@@ -3735,7 +3740,7 @@ bool UCustomizableInstancePrivate::BuildSkeletonData(const TSharedRef<FUpdateCon
 		// MergeSkeletons function doesn't do it. Only do it for newly generated ones, not for cached or non-transient ones.
 		GeneratedSkeleton->RecreateBoneTree(&SkeletalMesh);
 
-		FCustomizableInstanceComponentData* ComponentData = GetComponentData(ComponentIndex);
+		FCustomizableInstanceComponentData* ComponentData = GetComponentData(ObjectComponentIndex);
 		check(ComponentData);
 
 		TArray<TObjectPtr<USkeleton>>& SkeletonsToMerge = ComponentData->Skeletons.SkeletonsToMerge;
@@ -6684,7 +6689,7 @@ void UCustomizableInstancePrivate::BuildMaterials(const TSharedRef<FUpdateContex
 		}
 
 		// Ensure the number of materials is the same on both sides when using overrides. 
-		check(SkeletalMesh->GetMaterials().Num() == Materials.Num());
+		//check(SkeletalMesh->GetMaterials().Num() == Materials.Num());
 	}
 
 	// Force recreate render state if the mesh is reused and the materials have changed.
@@ -6710,9 +6715,16 @@ void UCustomizableInstancePrivate::BuildMaterials(const TSharedRef<FUpdateContex
 #endif
 
 			// TODO: What is the proper way to refer to dynamic components?
-			const int32 InstanceComponentIndex = CustomizableObjectInstanceUsage->GetComponentIndex();
-			const FInstanceUpdateData::FComponent& Component = OperationData->InstanceUpdateData.Components[InstanceComponentIndex];
-			int32 ObjectComponentIndex = Component.Id;
+			const int32 ObjectComponentIndex = CustomizableObjectInstanceUsage->GetComponentIndex();
+			int32 InstanceComponentIndex = -1;
+			for (int32 CurrentInstanceIndex=0; CurrentInstanceIndex<OperationData->InstanceUpdateData.Components.Num(); ++CurrentInstanceIndex)
+			{
+				if (OperationData->InstanceUpdateData.Components[CurrentInstanceIndex].Id == ObjectComponentIndex)
+				{
+					InstanceComponentIndex = CurrentInstanceIndex;
+					break;
+				}
+			}
 
 			if (!RecreateRenderStateOnInstanceComponent.IsValidIndex(InstanceComponentIndex) || !RecreateRenderStateOnInstanceComponent[InstanceComponentIndex])
 			{
@@ -6980,19 +6992,19 @@ void UCustomizableObjectInstance::Bake(const FBakingConfiguration& InBakingConfi
 #endif
 
 
-USkeletalMesh* UCustomizableObjectInstance::GetSkeletalMesh(int32 ComponentIndex) const
+USkeletalMesh* UCustomizableObjectInstance::GetSkeletalMesh(int32 ObjectComponentIndex) const
 {
 	check(PrivateData);
-	return PrivateData->SkeletalMeshes.IsValidIndex(ComponentIndex) ? PrivateData->SkeletalMeshes[ComponentIndex] : nullptr;
+	return PrivateData->SkeletalMeshes.IsValidIndex(ObjectComponentIndex) ? PrivateData->SkeletalMeshes[ObjectComponentIndex] : nullptr;
 }
 
 
 bool UCustomizableObjectInstance::HasAnySkeletalMesh() const
 {
 	check(PrivateData);
-	for (int32 Index = 0; Index < PrivateData->SkeletalMeshes.Num(); ++Index)
+	for (int32 ObjectComponentIndex = 0; ObjectComponentIndex < PrivateData->SkeletalMeshes.Num(); ++ObjectComponentIndex)
 	{
-		if (PrivateData->SkeletalMeshes[Index])
+		if (PrivateData->SkeletalMeshes[ObjectComponentIndex])
 		{
 			return true;
 		}
diff --git a/Engine/Plugins/Experimental/Mutable/Source/CustomizableObject/Private/MuCO/CustomizableObjectInstanceUsage.cpp b/Engine/Plugins/Experimental/Mutable/Source/CustomizableObject/Private/MuCO/CustomizableObjectInstanceUsage.cpp
--- a/Engine/Plugins/Experimental/Mutable/Source/CustomizableObject/Private/MuCO/CustomizableObjectInstanceUsage.cpp
+++ b/Engine/Plugins/Experimental/Mutable/Source/CustomizableObject/Private/MuCO/CustomizableObjectInstanceUsage.cpp
@@ -385,11 +385,11 @@ void UCustomizableObjectInstanceUsage::UpdateDistFromComponentToLevelEditorCamer
 		}
 
 		USkeletalMesh* AttachedSkeletalMesh = GetAttachedSkeletalMesh();
-		const int32 ComponentIndex = GetComponentIndex();
+		const int32 ObjectComponentIndex = GetComponentIndex();
 
 		const bool bInstanceGenerated = CustomizableObjectInstance->GetPrivate()->SkeletalMeshStatus != ESkeletalMeshStatus::NotGenerated;
-		USkeletalMesh* GeneratedSkeletalMesh = bInstanceGenerated ? CustomizableObjectInstance->GetSkeletalMesh(ComponentIndex) :
-			CustomizableObjectInstance->GetCustomizableObject()->GetRefSkeletalMesh(ComponentIndex);
+		USkeletalMesh* GeneratedSkeletalMesh = bInstanceGenerated ? CustomizableObjectInstance->GetSkeletalMesh(ObjectComponentIndex) :
+			CustomizableObjectInstance->GetCustomizableObject()->GetRefSkeletalMesh(ObjectComponentIndex);
 
 		if (AttachedSkeletalMesh != GeneratedSkeletalMesh)
 		{
@@ -435,9 +435,9 @@ void UCustomizableObjectInstanceUsage::EditorUpdateComponent()
 void UCustomizableObjectInstanceUsage::SetVisibilityOfSkeletalMeshSectionWithMaterialName(bool bInVisible, const FString& MaterialName, int32 LOD)
 {
 	UCustomizableObjectInstance* CustomizableObjectInstance = GetCustomizableObjectInstance();
-	int32 ComponentIndex = GetComponentIndex();
+	int32 ObjectComponentIndex = GetComponentIndex();
 
-	USkeletalMesh* SkeletalMesh = CustomizableObjectInstance->GetSkeletalMesh(ComponentIndex);
+	USkeletalMesh* SkeletalMesh = CustomizableObjectInstance->GetSkeletalMesh(ObjectComponentIndex);
 
 	if (!SkeletalMesh)
 	{
@@ -526,11 +526,11 @@ void UCustomizableObjectInstanceUsage::UpdateDistFromComponentToPlayer(const AAc
 			}
 		}
 
-		int32 ComponentIndex = GetComponentIndex();
+		int32 ObjectComponentIndex = GetComponentIndex();
 
-		if (ParentActor && GetAttachedSkeletalMesh() == nullptr && CustomizableObjectInstance->GetSkeletalMesh(ComponentIndex))
+		if (ParentActor && GetAttachedSkeletalMesh() == nullptr && CustomizableObjectInstance->GetSkeletalMesh(ObjectComponentIndex))
 		{
-			SetSkeletalMesh(CustomizableObjectInstance->GetSkeletalMesh(ComponentIndex));
+			SetSkeletalMesh(CustomizableObjectInstance->GetSkeletalMesh(ObjectComponentIndex));
 		}
 	}
 }
@@ -558,7 +558,7 @@ void UCustomizableObjectInstanceUsage::Tick(float DeltaTime)
 
 	if (USkeletalMeshComponent* Parent = Cast<USkeletalMeshComponent>(GetAttachParent()))
 	{
-		const int32 ComponentIndex = GetComponentIndex();
+		const int32 ObjectComponentIndex = GetComponentIndex();
 
 		USkeletalMesh* SkeletalMesh = nullptr;
 
@@ -566,7 +566,7 @@ void UCustomizableObjectInstanceUsage::Tick(float DeltaTime)
 		if (bInstanceGenerated)
 		{
 			// Generated SkeletalMesh to set, can be null if the component is empty
-			SkeletalMesh = CustomizableObjectInstance->GetSkeletalMesh(ComponentIndex);
+			SkeletalMesh = CustomizableObjectInstance->GetSkeletalMesh(ObjectComponentIndex);
 		}
 		else
 		{
@@ -575,7 +575,7 @@ void UCustomizableObjectInstanceUsage::Tick(float DeltaTime)
 				CustomizableObject->bEnableUseRefSkeletalMeshAsPlaceholder)
 			{
 				// Can be nullptr
-				SkeletalMesh = CustomizableObject->GetRefSkeletalMesh(ComponentIndex);
+				SkeletalMesh = CustomizableObject->GetRefSkeletalMesh(ObjectComponentIndex);
 			}
 		}
 
diff --git a/Engine/Plugins/Experimental/Mutable/Source/CustomizableObject/Private/MuCO/CustomizableObjectSkeletalMesh.cpp b/Engine/Plugins/Experimental/Mutable/Source/CustomizableObject/Private/MuCO/CustomizableObjectSkeletalMesh.cpp
--- a/Engine/Plugins/Experimental/Mutable/Source/CustomizableObject/Private/MuCO/CustomizableObjectSkeletalMesh.cpp
+++ b/Engine/Plugins/Experimental/Mutable/Source/CustomizableObject/Private/MuCO/CustomizableObjectSkeletalMesh.cpp
@@ -20,7 +20,7 @@ FName GenerateUniqueNameFromCOInstance(const UCustomizableObjectInstance& Instan
 
 
 UCustomizableObjectSkeletalMesh* UCustomizableObjectSkeletalMesh::CreateSkeletalMesh(const TSharedRef<FUpdateContextPrivate>& OperationData, 
-	const UCustomizableObjectInstance& Instance, const UCustomizableObject& CustomizableObject, const int32 InComponentIndex)
+	const UCustomizableObjectInstance& Instance, const UCustomizableObject& CustomizableObject, const int32 InstanceComponentIndex)
 {
 	FName SkeletalMeshName = GenerateUniqueNameFromCOInstance(Instance);
 	UCustomizableObjectSkeletalMesh* OutSkeletalMesh = NewObject<UCustomizableObjectSkeletalMesh>(GetTransientPackage(), SkeletalMeshName, RF_Transient);
@@ -38,18 +38,12 @@ UCustomizableObjectSkeletalMesh* UCustomizableObjectSkeletalMesh::CreateSkeletal
 	
 	OutSkeletalMesh->MeshIDs.Init(MAX_uint64, MAX_MESH_LOD_COUNT);
 	
-	for (int32 ComponentIndex = 0; ComponentIndex < OperationData->InstanceUpdateData.Components.Num(); ++ComponentIndex)
+	const FInstanceUpdateData::FComponent& Component = OperationData->InstanceUpdateData.Components[InstanceComponentIndex];
+
+	for (int32 LODIndex = OperationData->FirstLODAvailable; LODIndex < Component.LODCount; ++LODIndex)
 	{
-		const FInstanceUpdateData::FComponent& Component = OperationData->InstanceUpdateData.Components[ComponentIndex];
-
-		if (Component.Id == InComponentIndex)
-		{
-			for (int32 LODIndex = OperationData->FirstLODAvailable; LODIndex < Component.LODCount; ++LODIndex)
-			{
-				const FInstanceUpdateData::FLOD& LOD = OperationData->InstanceUpdateData.LODs[Component.FirstLOD + LODIndex];
-				OutSkeletalMesh->MeshIDs[LODIndex] = LOD.MeshID;
-			}
-		}
+		const FInstanceUpdateData::FLOD& LOD = OperationData->InstanceUpdateData.LODs[Component.FirstLOD + LODIndex];
+		OutSkeletalMesh->MeshIDs[LODIndex] = LOD.MeshID;
 	}
 
 	return OutSkeletalMesh;
diff --git a/Engine/Plugins/Experimental/Mutable/Source/CustomizableObject/Public/MuCO/CustomizableObject.h b/Engine/Plugins/Experimental/Mutable/Source/CustomizableObject/Public/MuCO/CustomizableObject.h
--- a/Engine/Plugins/Experimental/Mutable/Source/CustomizableObject/Public/MuCO/CustomizableObject.h
+++ b/Engine/Plugins/Experimental/Mutable/Source/CustomizableObject/Public/MuCO/CustomizableObject.h
@@ -519,7 +519,7 @@ public:
 	int32 FindIntParameterValue( int32 ParamIndex, const FString& Value ) const;
 	FString FindIntParameterValueName(int32 ParamIndex, int32 ParamValue) const;
 
-	USkeletalMesh* GetRefSkeletalMesh(int32 ComponentIndex = 0) const;
+	USkeletalMesh* GetRefSkeletalMesh(int32 ObjectComponentIndex = 0) const;
 
 	/** Get the default value of a parameter of type Float.
 	  * @param InParameterName The name of the Float parameter to get the default value of.
diff --git a/Engine/Plugins/Experimental/Mutable/Source/CustomizableObject/Public/MuCO/CustomizableObjectExtension.h b/Engine/Plugins/Experimental/Mutable/Source/CustomizableObject/Public/MuCO/CustomizableObjectExtension.h
--- a/Engine/Plugins/Experimental/Mutable/Source/CustomizableObject/Public/MuCO/CustomizableObjectExtension.h
+++ b/Engine/Plugins/Experimental/Mutable/Source/CustomizableObject/Public/MuCO/CustomizableObjectExtension.h
@@ -103,5 +103,5 @@ public:
 	 * 
 	 * @param SkeletalMesh - The Skeletal Mesh that was created.
 	 */
-	virtual void OnSkeletalMeshCreated(const TArray<FInputPinDataContainer>& InputPinData, int32 ComponentIndex, USkeletalMesh* SkeletalMesh) const {}
+	virtual void OnSkeletalMeshCreated(const TArray<FInputPinDataContainer>& InputPinData, int32 ObjectComponentIndex, USkeletalMesh* SkeletalMesh) const {}
 };
diff --git a/Engine/Plugins/Experimental/Mutable/Source/CustomizableObject/Public/MuCO/CustomizableObjectInstanceUsage.h b/Engine/Plugins/Experimental/Mutable/Source/CustomizableObject/Public/MuCO/CustomizableObjectInstanceUsage.h
--- a/Engine/Plugins/Experimental/Mutable/Source/CustomizableObject/Public/MuCO/CustomizableObjectInstanceUsage.h
+++ b/Engine/Plugins/Experimental/Mutable/Source/CustomizableObject/Public/MuCO/CustomizableObjectInstanceUsage.h
@@ -38,8 +38,8 @@ public:
 	UCustomizableObjectInstance* GetCustomizableObjectInstance() const;
 
 	// TODO: What is the proper way to refer to dynamic components?
-	// For now assume this is an *Instance* component index and not an *Object* component index
-	void SetComponentIndex(int32 ComponentIndex);
+	// For now assume this is an *Object* component index (to index the array of all possible components in the CO) and not an *Instance* component index.
+	void SetComponentIndex(int32 ObjectComponentIndex);
 	int32 GetComponentIndex() const;
 
 	void AttachTo(USkeletalMeshComponent* SkeletalMeshComponent);
diff --git a/Engine/Plugins/Experimental/Mutable/Source/CustomizableObject/Public/MuCO/CustomizableObjectSkeletalMesh.h b/Engine/Plugins/Experimental/Mutable/Source/CustomizableObject/Public/MuCO/CustomizableObjectSkeletalMesh.h
--- a/Engine/Plugins/Experimental/Mutable/Source/CustomizableObject/Public/MuCO/CustomizableObjectSkeletalMesh.h
+++ b/Engine/Plugins/Experimental/Mutable/Source/CustomizableObject/Public/MuCO/CustomizableObjectSkeletalMesh.h
@@ -38,7 +38,7 @@ class CUSTOMIZABLEOBJECT_API UCustomizableObjectSkeletalMesh : public USkeletalM
 public: // Public Methods
 
 	static UCustomizableObjectSkeletalMesh* CreateSkeletalMesh(const TSharedRef<FUpdateContextPrivate>& InOperationData, 
-		const UCustomizableObjectInstance& InInstance, const UCustomizableObject& CustomizableObject, const int32 ComponentIndex);
+		const UCustomizableObjectInstance& InInstance, const UCustomizableObject& CustomizableObject, const int32 InstanceComponentIndex);
 
 private: // Private Methods
 
diff --git a/Engine/Plugins/Experimental/Mutable/Source/CustomizableObject/Public/MuCO/CustomizableSkeletalComponent.h b/Engine/Plugins/Experimental/Mutable/Source/CustomizableObject/Public/MuCO/CustomizableSkeletalComponent.h
--- a/Engine/Plugins/Experimental/Mutable/Source/CustomizableObject/Public/MuCO/CustomizableSkeletalComponent.h
+++ b/Engine/Plugins/Experimental/Mutable/Source/CustomizableObject/Public/MuCO/CustomizableSkeletalComponent.h
@@ -20,6 +20,7 @@ public:
 	UPROPERTY(BlueprintReadWrite, EditAnywhere, Category = CustomizableSkeletalMesh)
 	TObjectPtr<UCustomizableObjectInstance> CustomizableObjectInstance;
 
+	/** This component index refers to the object list of components */
 	UPROPERTY(BlueprintReadWrite, EditAnywhere, Category = CustomizableSkeletalMesh)
 	int32 ComponentIndex;
 
diff --git a/Engine/Plugins/Experimental/Mutable/Source/CustomizableObjectEditor/Private/MuCOE/EdGraphSchema_CustomizableObject.cpp b/Engine/Plugins/Experimental/Mutable/Source/CustomizableObjectEditor/Private/MuCOE/EdGraphSchema_CustomizableObject.cpp
--- a/Engine/Plugins/Experimental/Mutable/Source/CustomizableObjectEditor/Private/MuCOE/EdGraphSchema_CustomizableObject.cpp
+++ b/Engine/Plugins/Experimental/Mutable/Source/CustomizableObjectEditor/Private/MuCOE/EdGraphSchema_CustomizableObject.cpp
@@ -52,6 +52,8 @@
 #include "MuCOE/Nodes/CustomizableObjectNodeObjectChild.h"
 #include "MuCOE/Nodes/CustomizableObjectNodeObjectGroup.h"
 #include "MuCOE/Nodes/CustomizableObjectNodeComponentMesh.h"
+#include "MuCOE/Nodes/CustomizableObjectNodeComponentVariation.h"
+#include "MuCOE/Nodes/CustomizableObjectNodeComponentSwitch.h"
 #include "MuCOE/Nodes/CustomizableObjectNodeProjectorConstant.h"
 #include "MuCOE/Nodes/CustomizableObjectNodeRemoveMesh.h"
 #include "MuCOE/Nodes/CustomizableObjectNodeRemoveMeshBlocks.h"
@@ -402,6 +404,8 @@ void UEdGraphSchema_CustomizableObject::GetGraphContextActions(FGraphContextMenu
 		UCustomizableObjectNode* ComponentTemplateNodes[]
 		{
 			ContextMenuBuilder.CreateTemplateNode<UCustomizableObjectNodeComponentMesh>(),
+			ContextMenuBuilder.CreateTemplateNode<UCustomizableObjectNodeComponentVariation>(),
+			ContextMenuBuilder.CreateTemplateNode<UCustomizableObjectNodeComponentSwitch>(),
 		};
 
 		AddNewNodeCategoryActionsFiltered(ComponentTemplateNodes, ContextMenuBuilder, TEXT("Component"), GeneralGrouping, Filter);
@@ -1323,6 +1327,10 @@ FText UEdGraphSchema_CustomizableObject::GetPinCategoryName(const FName& PinCate
 	{
 		return LOCTEXT("Object_Pin_Category", "Object");
 	}
+	else if (PinCategory == UEdGraphSchema_CustomizableObject::PC_Component)
+	{
+		return LOCTEXT("Component_Pin_Category", "Component");
+	}
 	else if (PinCategory == UEdGraphSchema_CustomizableObject::PC_Material)
 	{
 		return LOCTEXT("Material_Pin_Category", "Material");
diff --git a/Engine/Plugins/Experimental/Mutable/Source/CustomizableObjectEditor/Private/MuCOE/GenerateMutableSource/GenerateMutableSourceComponent.cpp b/Engine/Plugins/Experimental/Mutable/Source/CustomizableObjectEditor/Private/MuCOE/GenerateMutableSource/GenerateMutableSourceComponent.cpp
--- a/Engine/Plugins/Experimental/Mutable/Source/CustomizableObjectEditor/Private/MuCOE/GenerateMutableSource/GenerateMutableSourceComponent.cpp
+++ b/Engine/Plugins/Experimental/Mutable/Source/CustomizableObjectEditor/Private/MuCOE/GenerateMutableSource/GenerateMutableSourceComponent.cpp
@@ -5,11 +5,16 @@
 #include "MuCOE/CustomizableObjectCompiler.h"
 #include "MuCOE/GenerateMutableSource/GenerateMutableSourceMesh.h"
 #include "MuCOE/GenerateMutableSource/GenerateMutableSourceSurface.h"
+#include "MuCOE/GenerateMutableSource/GenerateMutableSourceFloat.h"
 #include "MuCOE/GraphTraversal.h"
 #include "MuCOE/MutableUtils.h"
 #include "MuCOE/Nodes/CustomizableObjectNodeComponentMesh.h"
+#include "MuCOE/Nodes/CustomizableObjectNodeComponentSwitch.h"
+#include "MuCOE/Nodes/CustomizableObjectNodeComponentVariation.h"
 #include "MuT/NodeComponent.h"
 #include "MuT/NodeComponentNew.h"
+#include "MuT/NodeComponentSwitch.h"
+#include "MuT/NodeComponentVariation.h"
 #include "MuT/NodeLOD.h"
 #include "MuT/NodeSurfaceNew.h"
 #include "MuT/NodeMeshConstant.h"
@@ -120,9 +125,128 @@ mu::Ptr<mu::NodeComponent> GenerateMutableSourceComponent(const UEdGraphPin * Pi
 		Result = ComponentNode;
 	}
 
+
+	else if (const UCustomizableObjectNodeComponentSwitch* TypedNodeSwitch = Cast<UCustomizableObjectNodeComponentSwitch>(Node))
+	{
+		// Using a lambda so control flow is easier to manage.
+		Result = [&]()
+			{
+				const UEdGraphPin* SwitchParameter = TypedNodeSwitch->SwitchParameter();
+
+				// Check Switch Parameter arity preconditions.
+				if (const UEdGraphPin* EnumPin = FollowInputPin(*SwitchParameter))
+				{
+					mu::Ptr<mu::NodeScalar> SwitchParam = GenerateMutableSourceFloat(EnumPin, GenerationContext);
+
+					// Switch Param not generated
+					if (!SwitchParam)
+					{
+						// Warn about a failure.
+						if (EnumPin)
+						{
+							const FText Message = LOCTEXT("FailedToGenerateSwitchParam", "Could not generate switch enum parameter. Please refesh the switch node and connect an enum.");
+							GenerationContext.Compiler->CompilerLog(Message, Node);
+						}
+
+						return Result;
+					}
+
+					if (SwitchParam->GetType() != mu::NodeScalarEnumParameter::GetStaticType())
+					{
+						const FText Message = LOCTEXT("WrongSwitchParamType", "Switch parameter of incorrect type.");
+						GenerationContext.Compiler->CompilerLog(Message, Node);
+
+						return Result;
+					}
+
+					const int32 NumSwitchOptions = TypedNodeSwitch->GetNumElements();
+
+					mu::NodeScalarEnumParameter* EnumParameter = static_cast<mu::NodeScalarEnumParameter*>(SwitchParam.get());
+					if (NumSwitchOptions != EnumParameter->GetValueCount())
+					{
+						const FText Message = LOCTEXT("MismatchedSwitch", "Switch enum and switch node have different number of options. Please refresh the switch node to make sure the outcomes are labeled properly.");
+						GenerationContext.Compiler->CompilerLog(Message, Node);
+					}
+
+					mu::Ptr<mu::NodeComponentSwitch> SwitchNode = new mu::NodeComponentSwitch;
+					SwitchNode->Parameter = SwitchParam;
+					SwitchNode->Options.SetNum(NumSwitchOptions);
+
+					for (int32 SelectorIndex = 0; SelectorIndex < NumSwitchOptions; ++SelectorIndex)
+					{
+						if (const UEdGraphPin* ConnectedPin = FollowInputPin(*TypedNodeSwitch->GetElementPin(SelectorIndex)))
+						{
+							mu::Ptr<mu::NodeComponent> ChildNode = GenerateMutableSourceComponent(ConnectedPin, GenerationContext);
+							if (ChildNode)
+							{
+								SwitchNode->Options[SelectorIndex] = ChildNode;
+							}
+							else
+							{
+								// Probably ok
+							}
+						}
+					}
+
+					Result = SwitchNode;
+					return Result;
+				}
+				else
+				{
+					GenerationContext.Compiler->CompilerLog(LOCTEXT("NoEnumParamInSwitch", "Switch nodes must have an enum switch parameter. Please connect an enum and refesh the switch node."), Node);
+					return Result;
+				}
+			}(); // invoke lambda.
+	}
+
+	else if (const UCustomizableObjectNodeComponentVariation* TypedNodeVar = Cast<UCustomizableObjectNodeComponentVariation>(Node))
+	{
+		mu::Ptr<mu::NodeComponentVariation> SurfNode = new mu::NodeComponentVariation();
+		Result = SurfNode;
+
+		for (const UEdGraphPin* ConnectedPin : FollowInputPinArray(*TypedNodeVar->DefaultPin()))
+		{
+			mu::Ptr<mu::NodeComponent> ChildNode = GenerateMutableSourceComponent(ConnectedPin, GenerationContext);
+			if (ChildNode)
+			{
+				SurfNode->DefaultComponent = ChildNode;
+			}
+			else
+			{
+				GenerationContext.Compiler->CompilerLog(LOCTEXT("ComponentFailed", "Component generation failed."), Node);
+			}
+		}
+
+		const int32 NumVariations = TypedNodeVar->GetNumVariations();
+		SurfNode->Variations.SetNum(NumVariations);
+		for (int VariationIndex = 0; VariationIndex < NumVariations; ++VariationIndex)
+		{
+			mu::NodeSurfacePtr VariationSurfaceNode;
+
+			if (UEdGraphPin* VariationPin = TypedNodeVar->VariationPin(VariationIndex))
+			{
+				SurfNode->Variations[VariationIndex].Tag = TypedNodeVar->GetVariation(VariationIndex).Tag;
+				for (const UEdGraphPin* ConnectedPin : FollowInputPinArray(*VariationPin))
+				{
+					// Is it a modifier?
+					mu::Ptr<mu::NodeComponent> ChildNode = GenerateMutableSourceComponent(ConnectedPin, GenerationContext);
+					if (ChildNode)
+					{
+						SurfNode->Variations[VariationIndex].Component = ChildNode;
+					}
+					else
+					{
+						GenerationContext.Compiler->CompilerLog(LOCTEXT("ComponentFailed", "Component generation failed."), Node);
+					}
+				}
+			}
+		}
+	}
+
 	else
 	{
 		GenerationContext.Compiler->CompilerLog(LOCTEXT("UnimplementedNode", "Node type not implemented yet."), Node);
+		ensure(false);
 	}
 
 	GenerationContext.Generated.Add(Key, FGeneratedData(Node, Result));
diff --git a/Engine/Plugins/Experimental/Mutable/Source/CustomizableObjectEditor/Private/MuCOE/Nodes/CustomizableObjectNodeComponentSwitch.cpp b/Engine/Plugins/Experimental/Mutable/Source/CustomizableObjectEditor/Private/MuCOE/Nodes/CustomizableObjectNodeComponentSwitch.cpp
new file mode 100644
--- /dev/null
+++ b/Engine/Plugins/Experimental/Mutable/Source/CustomizableObjectEditor/Private/MuCOE/Nodes/CustomizableObjectNodeComponentSwitch.cpp
@@ -0,0 +1,17 @@
+// Copyright Epic Games, Inc. All Rights Reserved.
+
+#include "MuCOE/Nodes/CustomizableObjectNodeComponentSwitch.h"
+
+#include "MuCOE/EdGraphSchema_CustomizableObject.h"
+
+#define LOCTEXT_NAMESPACE "CustomizableObjectEditor"
+
+
+FName UCustomizableObjectNodeComponentSwitch::GetCategory() const
+{
+	return UEdGraphSchema_CustomizableObject::PC_Component;
+}
+
+
+#undef LOCTEXT_NAMESPACE
+
diff --git a/Engine/Plugins/Experimental/Mutable/Source/CustomizableObjectEditor/Private/MuCOE/Nodes/CustomizableObjectNodeComponentVariation.cpp b/Engine/Plugins/Experimental/Mutable/Source/CustomizableObjectEditor/Private/MuCOE/Nodes/CustomizableObjectNodeComponentVariation.cpp
new file mode 100644
--- /dev/null
+++ b/Engine/Plugins/Experimental/Mutable/Source/CustomizableObjectEditor/Private/MuCOE/Nodes/CustomizableObjectNodeComponentVariation.cpp
@@ -0,0 +1,26 @@
+// Copyright Epic Games, Inc. All Rights Reserved.
+
+#include "MuCOE/Nodes/CustomizableObjectNodeComponentVariation.h"
+
+#include "MuCO/CustomizableObjectCustomVersion.h"
+#include "MuCOE/CustomizableObjectEditor_Deprecated.h"
+#include "MuCOE/EdGraphSchema_CustomizableObject.h"
+
+
+
+FName UCustomizableObjectNodeComponentVariation::GetCategory() const
+{
+	return UEdGraphSchema_CustomizableObject::PC_Component;
+}
+
+
+bool UCustomizableObjectNodeComponentVariation::IsInputPinArray() const
+{
+	return true;
+}
+
+
+bool UCustomizableObjectNodeComponentVariation::IsSingleOutputNode() const
+{
+	return true;
+}
diff --git a/Engine/Plugins/Experimental/Mutable/Source/CustomizableObjectEditor/Public/MuCOE/Nodes/CustomizableObjectNodeComponentSwitch.h b/Engine/Plugins/Experimental/Mutable/Source/CustomizableObjectEditor/Public/MuCOE/Nodes/CustomizableObjectNodeComponentSwitch.h
new file mode 100644
--- /dev/null
+++ b/Engine/Plugins/Experimental/Mutable/Source/CustomizableObjectEditor/Public/MuCOE/Nodes/CustomizableObjectNodeComponentSwitch.h
@@ -0,0 +1,19 @@
+// Copyright Epic Games, Inc. All Rights Reserved.
+
+#pragma once
+
+#include "MuCOE/Nodes/CustomizableObjectNodeSwitchBase.h"
+
+#include "CustomizableObjectNodeComponentSwitch.generated.h"
+
+
+UCLASS()
+class CUSTOMIZABLEOBJECTEDITOR_API UCustomizableObjectNodeComponentSwitch : public UCustomizableObjectNodeSwitchBase
+{
+public:
+	GENERATED_BODY()
+	
+	// UCustomizableObjectNodeSwitchBase interface
+	virtual FName GetCategory() const override;
+};
+
diff --git a/Engine/Plugins/Experimental/Mutable/Source/CustomizableObjectEditor/Public/MuCOE/Nodes/CustomizableObjectNodeComponentVariation.h b/Engine/Plugins/Experimental/Mutable/Source/CustomizableObjectEditor/Public/MuCOE/Nodes/CustomizableObjectNodeComponentVariation.h
new file mode 100644
--- /dev/null
+++ b/Engine/Plugins/Experimental/Mutable/Source/CustomizableObjectEditor/Public/MuCOE/Nodes/CustomizableObjectNodeComponentVariation.h
@@ -0,0 +1,37 @@
+// Copyright Epic Games, Inc. All Rights Reserved.
+
+#pragma once
+
+#include "MuCOE/Nodes/CustomizableObjectNodeVariation.h"
+
+#include "CustomizableObjectNodeComponentVariation.generated.h"
+
+struct FCustomizableObjectComponentVariation;
+
+
+UENUM(BlueprintType)
+enum class ECustomizableObjectNodeComponentVariationType : uint8
+{
+	Tag 		UMETA(DisplayName = "Tag"),
+	State 		UMETA(DisplayName = "State"),
+};
+
+
+UCLASS()
+class CUSTOMIZABLEOBJECTEDITOR_API UCustomizableObjectNodeComponentVariation : public UCustomizableObjectNodeVariation
+{
+public:
+	GENERATED_BODY()
+
+	//UPROPERTY(EditAnywhere, Category = CustomizableObject)
+	//ECustomizableObjectNodeComponentVariationType Type = ECustomizableObjectNodeComponentVariationType::Tag;
+
+public:
+	// UCustomizableObjectNode interface
+	virtual bool IsSingleOutputNode() const override;
+	
+	// UCustomizableObjectNodeVariation interface
+	virtual FName GetCategory() const override;
+	virtual bool IsInputPinArray() const override;
+};
+
diff --git a/Engine/Plugins/Experimental/Mutable/Source/MutableTools/Private/MuT/CodeGenerator.cpp b/Engine/Plugins/Experimental/Mutable/Source/MutableTools/Private/MuT/CodeGenerator.cpp
--- a/Engine/Plugins/Experimental/Mutable/Source/MutableTools/Private/MuT/CodeGenerator.cpp
+++ b/Engine/Plugins/Experimental/Mutable/Source/MutableTools/Private/MuT/CodeGenerator.cpp
@@ -48,6 +48,8 @@
 #include "MuT/NodeColour.h"
 #include "MuT/NodeColourConstant.h"
 #include "MuT/NodeComponent.h"
+#include "MuT/NodeComponentSwitch.h"
+#include "MuT/NodeComponentVariation.h"
 #include "MuT/NodeImage.h"
 #include "MuT/NodeImageFormat.h"
 #include "MuT/NodeImageFormatPrivate.h"
@@ -542,6 +544,14 @@ namespace mu
 			//GenerateComponent_Edit(InOptions, OutResult, static_cast<const NodeComponentEdit*>(InUntypedNode));
 			OutResult.op = InOptions.BaseInstance;
 		}
+		else if (Type == NodeComponentSwitch::GetStaticType())
+		{
+			GenerateComponent_Switch(InOptions, OutResult, static_cast<const NodeComponentSwitch*>(InUntypedNode));
+		}
+		else if (Type == NodeComponentVariation::GetStaticType())
+		{
+			GenerateComponent_Variation(InOptions, OutResult, static_cast<const NodeComponentVariation*>(InUntypedNode));
+		}
 		else
 		{
 			check(false);
@@ -609,6 +619,114 @@ namespace mu
 	}
 
 
+	void CodeGenerator::GenerateComponent_Switch(const FComponentGenerationOptions& Options, FGenericGenerationResult& Result, const NodeComponentSwitch* Node)
+	{
+		MUTABLE_CPUPROFILER_SCOPE(NodeComponentSwitch);
+
+		if (Node->Options.Num() == 0)
+		{
+			// No options in the switch!
+			Result.op = Options.BaseInstance;
+			return;
+		}
+
+		Ptr<ASTOpSwitch> Op = new ASTOpSwitch();
+		Op->type = OP_TYPE::IN_SWITCH;
+
+		// Variable value
+		if (Node->Parameter)
+		{
+			Op->variable = Generate_Generic(Node->Parameter.get(), Options);
+		}
+		else
+		{
+			// This argument is required
+			Op->variable = GenerateMissingScalarCode(TEXT("Switch variable"), 0.0f, Node->GetMessageContext());
+		}
+
+		// Options
+		for (int32 OptionIndex = 0; OptionIndex < Node->Options.Num(); ++OptionIndex)
+		{
+			Ptr<ASTOp> Branch;
+
+			if (Node->Options[OptionIndex])
+			{
+				FGenericGenerationResult BaseResult;
+				GenerateComponent(Options, BaseResult, Node->Options[OptionIndex].get());
+				Branch = BaseResult.op;
+			}
+			else
+			{
+				// This argument is not required
+				Branch = Options.BaseInstance;
+			}
+
+			Op->cases.Emplace(OptionIndex, Op, Branch);
+		}
+
+		Result.op = Op;
+	}
+
+
+	void CodeGenerator::GenerateComponent_Variation(const FComponentGenerationOptions& Options, FGenericGenerationResult& Result, const NodeComponentVariation* Node)
+	{
+		Ptr<ASTOp> CurrentMeshOp = Options.BaseInstance;
+
+		// Default case
+		if (Node->DefaultComponent)
+		{
+			FGenericGenerationResult BranchResults;
+
+			GenerateComponent(Options, BranchResults, Node->DefaultComponent.get());
+			CurrentMeshOp = BranchResults.op;
+		}
+
+		// Process variations in reverse order, since conditionals are built bottom-up.
+		for (int32 VariationIndex = Node->Variations.Num() - 1; VariationIndex >= 0; --VariationIndex)
+		{
+			int32 TagIndex = -1;
+			const FString& Tag = Node->Variations[VariationIndex].Tag;
+			for (int32 i = 0; i < FirstPass.Tags.Num(); ++i)
+			{
+				if (FirstPass.Tags[i].Tag == Tag)
+				{
+					TagIndex = i;
+				}
+			}
+
+			if (TagIndex < 0)
+			{
+				ErrorLog->GetPrivate()->Add(
+					FString::Printf(TEXT("Unknown tag found in component variation [%s]."), *Tag),
+					ELMT_WARNING,
+					Node->GetMessageContext(),
+					ELMSB_UNKNOWN_TAG
+				);
+				continue;
+			}
+
+			Ptr<ASTOp> VariationMeshOp = Options.BaseInstance;
+			if (Node->Variations[VariationIndex].Component)
+			{
+				FGenericGenerationResult BranchResults;
+				GenerateComponent(Options, BranchResults, Node->Variations[VariationIndex].Component.get());
+
+				VariationMeshOp = BranchResults.op;
+			}
+
+			Ptr<ASTOpConditional> Conditional = new ASTOpConditional;
+			Conditional->type = OP_TYPE::IN_CONDITIONAL;
+			Conditional->no = CurrentMeshOp;
+			Conditional->yes = VariationMeshOp;
+			Conditional->condition = FirstPass.Tags[TagIndex].GenericCondition;
+
+			CurrentMeshOp = Conditional;
+		}
+
+		Result.op = CurrentMeshOp;
+	}
+
+
 	Ptr<ASTOp> CodeGenerator::ApplyTiling(Ptr<ASTOp> Source, UE::Math::TIntVector2<int32> Size, EImageFormat Format)
 	{
 		// For now always apply tiling
diff --git a/Engine/Plugins/Experimental/Mutable/Source/MutableTools/Private/MuT/CodeGenerator.h b/Engine/Plugins/Experimental/Mutable/Source/MutableTools/Private/MuT/CodeGenerator.h
--- a/Engine/Plugins/Experimental/Mutable/Source/MutableTools/Private/MuT/CodeGenerator.h
+++ b/Engine/Plugins/Experimental/Mutable/Source/MutableTools/Private/MuT/CodeGenerator.h
@@ -367,6 +367,8 @@ namespace mu
 
 		void GenerateComponent(const FComponentGenerationOptions&, FGenericGenerationResult&, const NodeComponent*);
 		void GenerateComponent_New(const FComponentGenerationOptions&, FGenericGenerationResult&, const NodeComponentNew*);
+		void GenerateComponent_Switch(const FComponentGenerationOptions&, FGenericGenerationResult&, const NodeComponentSwitch*);
+		void GenerateComponent_Variation(const FComponentGenerationOptions&, FGenericGenerationResult&, const NodeComponentVariation*);
 
 		//-----------------------------------------------------------------------------------------
 
diff --git a/Engine/Plugins/Experimental/Mutable/Source/MutableTools/Private/MuT/CodeOptimiser.cpp b/Engine/Plugins/Experimental/Mutable/Source/MutableTools/Private/MuT/CodeOptimiser.cpp
--- a/Engine/Plugins/Experimental/Mutable/Source/MutableTools/Private/MuT/CodeOptimiser.cpp
+++ b/Engine/Plugins/Experimental/Mutable/Source/MutableTools/Private/MuT/CodeOptimiser.cpp
@@ -322,11 +322,11 @@ namespace mu
 		MUTABLE_CPUPROFILER_SCOPE(DuplicatedCodeRemoverAST);
 
 		// Visited nodes, per type
-		std::unordered_set<Ptr<ASTOp>,op_pointer_hash,op_pointer_equal> visited[int(OP_TYPE::COUNT)];
+		std::unordered_set<Ptr<ASTOp>,op_pointer_hash,op_pointer_equal> visited[int32(OP_TYPE::COUNT)];
 
 		ASTOp::Traverse_BottomUp_Unique_NonReentrant( roots, [&](Ptr<ASTOp>& n)
 		{
-			auto& container = visited[(int)n->GetOpType()];
+			auto& container = visited[(int32)n->GetOpType()];
 
 			// Insert will tell us if it was already there
 			auto it = container.insert(n);
@@ -356,7 +356,7 @@ namespace mu
 		// Don't early out to be able to complete parent op cached flags
 		while ( !m_pending.IsEmpty() )
 		{
-			TPair<bool,int> item = m_pending.Pop();
+			TPair<bool,int32> item = m_pending.Pop();
 			OP::ADDRESS at = item.Value;
 
 			// Not cached?
@@ -1066,7 +1066,7 @@ namespace mu
 	void CodeOptimiser::FullOptimiseAST( ASTOpList& roots, int32 Pass )
 	{
 		bool modified = true;
-		int numIterations = 0;
+		int32 numIterations = 0;
 		while (modified && (!m_optimizeIterationsMax || (m_optimizeIterationsLeft>0) || !numIterations) )
 		{
 			--m_optimizeIterationsLeft;
@@ -1081,7 +1081,7 @@ namespace mu
 			UE_LOG(LogMutableCore, Verbose, TEXT(" - semantic optimiser"));
 			modified |= SemanticOptimiserAST( roots, m_options->GetPrivate()->OptimisationOptions, Pass );
 			//AXE_INT_VALUE("Mutable", Verbose, "ast size", (int64_t)ASTOp::CountNodes(roots));
-			UE_LOG(LogMutableCore, Verbose, TEXT("(int) %s : %ld"), TEXT("ast size"), int64(ASTOp::CountNodes(roots)));
+			//UE_LOG(LogMutableCore, Verbose, TEXT("(int) %s : %ld"), TEXT("ast size"), int64(ASTOp::CountNodes(roots)));
 			ASTOp::LogHistogram(roots);
 
 			UE_LOG(LogMutableCore, Verbose, TEXT(" - sink optimiser"));
@@ -1109,27 +1109,27 @@ namespace mu
 
 		UE_LOG(LogMutableCore, Verbose, TEXT(" - duplicated code remover"));
 		DuplicatedCodeRemoverAST( roots );
-		UE_LOG(LogMutableCore, Verbose, TEXT("(int) %s : %ld"), TEXT("ast size"), int64(ASTOp::CountNodes(roots)));
+		//UE_LOG(LogMutableCore, Verbose, TEXT("(int) %s : %ld"), TEXT("ast size"), int64(ASTOp::CountNodes(roots)));
 
 		ASTOp::LogHistogram(roots);
 
 		UE_LOG(LogMutableCore, Verbose, TEXT(" - duplicated data remover"));
 		DuplicatedDataRemoverAST( roots );
-		UE_LOG(LogMutableCore, Verbose, TEXT("(int) %s : %ld"), TEXT("ast size"), int64(ASTOp::CountNodes(roots)));
+		//UE_LOG(LogMutableCore, Verbose, TEXT("(int) %s : %ld"), TEXT("ast size"), int64(ASTOp::CountNodes(roots)));
 
 		ASTOp::LogHistogram(roots);
 
 		// Generate constants
 		for ( Ptr<ASTOp>& Root: roots )
 		{
-			UE_LOG(LogMutableCore, Verbose, TEXT("(int) %s : %ld"), TEXT("ast size"), int64(ASTOp::CountNodes(roots)));
+			//UE_LOG(LogMutableCore, Verbose, TEXT("(int) %s : %ld"), TEXT("ast size"), int64(ASTOp::CountNodes(roots)));
 			UE_LOG(LogMutableCore, Verbose, TEXT(" - constant generator"));
 
 			// Constant subtree generation
 			modified = ConstantGeneratorAST( m_options->GetPrivate(), Root, Pass );
 		}
 
-		UE_LOG(LogMutableCore, Verbose, TEXT("(int) %s : %ld"), TEXT("ast size"), int64(ASTOp::CountNodes(roots)));
+		//UE_LOG(LogMutableCore, Verbose, TEXT("(int) %s : %ld"), TEXT("ast size"), int64(ASTOp::CountNodes(roots)));
 
 	//    ASTOp::LogHistogram(roots);
 
@@ -1142,7 +1142,7 @@ namespace mu
 
 		UE_LOG(LogMutableCore, Verbose, TEXT(" - duplicated data remover"));
 		DuplicatedDataRemoverAST( roots );
-		UE_LOG(LogMutableCore, Verbose, TEXT("(int) %s : %ld"), TEXT("ast size"), int64(ASTOp::CountNodes(roots)));
+		//UE_LOG(LogMutableCore, Verbose, TEXT("(int) %s : %ld"), TEXT("ast size"), int64(ASTOp::CountNodes(roots)));
 
 		//if (!modified)
 		{
@@ -1404,14 +1404,14 @@ namespace mu
 	{
 		MUTABLE_CPUPROFILER_SCOPE(OptimiseAST);
 
-		// Gather all the roots
+		// Gather all the roots (one for each state)
 		TArray<Ptr<ASTOp>> roots;
-		for(const auto& s:m_states)
+		for(const FStateCompilationData& s:m_states)
 		{
 			roots.Add(s.root);
 		}
 
-		UE_LOG(LogMutableCore, Verbose, TEXT("(int) %s : %ld"), TEXT("ast size"), int64(ASTOp::CountNodes(roots)));
+		//UE_LOG(LogMutableCore, Verbose, TEXT("(int) %s : %ld"), TEXT("ast size"), int64(ASTOp::CountNodes(roots)));
 
 		if ( m_options->GetPrivate()->OptimisationOptions.bEnabled )
 		{
@@ -1426,13 +1426,13 @@ namespace mu
 			// creation
 			UE_LOG(LogMutableCore, Verbose, TEXT(" - duplicated data remover"));
 			DuplicatedDataRemoverAST( roots );
-			UE_LOG(LogMutableCore, Verbose, TEXT("(int) %s : %ld"), TEXT("ast size"), int64(ASTOp::CountNodes(roots)));
+			//UE_LOG(LogMutableCore, Verbose, TEXT("(int) %s : %ld"), TEXT("ast size"), int64(ASTOp::CountNodes(roots)));
 
 			ASTOp::LogHistogram(roots);
 
 			UE_LOG(LogMutableCore, Verbose, TEXT(" - duplicated code remover"));
 			DuplicatedCodeRemoverAST( roots );
-			UE_LOG(LogMutableCore, Verbose, TEXT("(int) %s : %ld"), TEXT("ast size"), int64(ASTOp::CountNodes(roots)));
+			//UE_LOG(LogMutableCore, Verbose, TEXT("(int) %s : %ld"), TEXT("ast size"), int64(ASTOp::CountNodes(roots)));
 
 			// Special optimization stages
 			if ( m_options->GetPrivate()->OptimisationOptions.bUniformizeSkeleton )
@@ -1441,7 +1441,7 @@ namespace mu
 				ASTOp::LogHistogram(roots);
 
 				SkeletonCleanerAST( roots, m_options->GetPrivate()->OptimisationOptions );
-				UE_LOG(LogMutableCore, Verbose, TEXT("(int) %s : %ld"), TEXT("ast size"), int64(ASTOp::CountNodes(roots)));
+				//UE_LOG(LogMutableCore, Verbose, TEXT("(int) %s : %ld"), TEXT("ast size"), int64(ASTOp::CountNodes(roots)));
 				ASTOp::LogHistogram(roots);
 			}
 
@@ -1449,7 +1449,7 @@ namespace mu
 			// because some operations cannot be applied correctly until the image size is known
 			// like the grow-map generation.
 			bool modified = true;
-			int numIterations = 0;
+			int32 numIterations = 0;
 			while (modified)
 			{
 				MUTABLE_CPUPROFILER_SCOPE(FirstStage);
@@ -1465,18 +1465,18 @@ namespace mu
 				modified |= SizeOptimiserAST( roots );
 			}
 
-			UE_LOG(LogMutableCore, Verbose, TEXT("(int) %s : %ld"), TEXT("ast size"), int64(ASTOp::CountNodes(roots)));
+			//UE_LOG(LogMutableCore, Verbose, TEXT("(int) %s : %ld"), TEXT("ast size"), int64(ASTOp::CountNodes(roots)));
 			ASTOp::LogHistogram(roots);
 
 			// Main optimisation stage
 			{
 				MUTABLE_CPUPROFILER_SCOPE(MainStage);
 				FullOptimiseAST( roots, 0 );
-				UE_LOG(LogMutableCore, Verbose, TEXT("(int) %s : %ld"), TEXT("ast size"), int64(ASTOp::CountNodes(roots)));
+				//UE_LOG(LogMutableCore, Verbose, TEXT("(int) %s : %ld"), TEXT("ast size"), int64(ASTOp::CountNodes(roots)));
 				ASTOp::LogHistogram(roots);
 
 				FullOptimiseAST( roots, 1 );
-				UE_LOG(LogMutableCore, Verbose, TEXT("(int) %s : %ld"), TEXT("ast size"), int64(ASTOp::CountNodes(roots)));
+				//UE_LOG(LogMutableCore, Verbose, TEXT("(int) %s : %ld"), TEXT("ast size"), int64(ASTOp::CountNodes(roots)));
 				ASTOp::LogHistogram(roots);
 			}
 
@@ -1502,11 +1502,11 @@ namespace mu
 			{
 				MUTABLE_CPUPROFILER_SCOPE(FinalStage);
 				FullOptimiseAST(roots, 0);
-				UE_LOG(LogMutableCore, Verbose, TEXT("(int) %s : %ld"), TEXT("ast size"), int64(ASTOp::CountNodes(roots)));
+				//UE_LOG(LogMutableCore, Verbose, TEXT("(int) %s : %ld"), TEXT("ast size"), int64(ASTOp::CountNodes(roots)));
 				ASTOp::LogHistogram(roots);
 
 				FullOptimiseAST(roots, 1);
-				UE_LOG(LogMutableCore, Verbose, TEXT("(int) %s : %ld"), TEXT("ast size"), int64(ASTOp::CountNodes(roots)));
+				//UE_LOG(LogMutableCore, Verbose, TEXT("(int) %s : %ld"), TEXT("ast size"), int64(ASTOp::CountNodes(roots)));
 				ASTOp::LogHistogram(roots);
 			}
 
@@ -1522,14 +1522,14 @@ namespace mu
 				}
 			});
 
-			// Make sure we didn't lose track of pointers
-			for ( int32 s=0;  s<m_states.Num(); ++s )
+			ASTOp::LogHistogram(roots);
+
+			// Reset the state root operations in case they have changed due to optimization
+			for (int32 RootIndex = 0; RootIndex < m_states.Num(); ++RootIndex)
 			{
-				check( roots.Contains( m_states[s].root ) );
+				m_states[RootIndex].root = roots[RootIndex];
 			}
 
-			ASTOp::LogHistogram(roots);
-
 			{
 				MUTABLE_CPUPROFILER_SCOPE(StatesStage);
 
@@ -1556,7 +1556,7 @@ namespace mu
 
 			UE_LOG(LogMutableCore, Verbose, TEXT(" - duplicated code remover"));
 			DuplicatedCodeRemoverAST( roots );
-			UE_LOG(LogMutableCore, Verbose, TEXT("(int) %s : %ld"), TEXT("ast size"), int64(ASTOp::CountNodes(roots)));
+			//UE_LOG(LogMutableCore, Verbose, TEXT("(int) %s : %ld"), TEXT("ast size"), int64(ASTOp::CountNodes(roots)));
 
 			// Constant resolution stage: resolve referenced assets.
 			{
@@ -1564,12 +1564,12 @@ namespace mu
 				FullOptimiseAST(roots, 2);
 			}
 
-			for ( int32 s=0;  s<m_states.Num(); ++s )
+			for ( int32 StateIndex=0; StateIndex <m_states.Num(); ++StateIndex)
 			{
 				constexpr int32 Pass = 1;
 
 				UE_LOG(LogMutableCore, Verbose, TEXT(" - constant generator"));
-				ConstantGeneratorAST( m_options->GetPrivate(), m_states[s].root, Pass );
+				ConstantGeneratorAST( m_options->GetPrivate(), roots[StateIndex], Pass);
 				//AXE_INT_VALUE("Mutable", Verbose, "ast size", (int64_t)ASTOp::CountNodes(roots));
 			}
 
@@ -1581,10 +1581,10 @@ namespace mu
 			DuplicatedCodeRemoverAST( roots );
 			//AXE_INT_VALUE("Mutable", Verbose, "ast size", (int64_t)ASTOp::CountNodes(roots));
 
-			// Make sure we didn't lose track of pointers
-			for ( int32 s=0;  s<m_states.Num(); ++s )
+			// Reset the state root operations in case they have changed due to optimization
+			for (int32 RootIndex = 0; RootIndex < m_states.Num(); ++RootIndex)
 			{
-				check( roots.Contains( m_states[s].root ) );
+				m_states[RootIndex].root = roots[RootIndex];
 			}
 		}
 
diff --git a/Engine/Plugins/Experimental/Mutable/Source/MutableTools/Private/MuT/CodeOptimiser_Parameters.cpp b/Engine/Plugins/Experimental/Mutable/Source/MutableTools/Private/MuT/CodeOptimiser_Parameters.cpp
--- a/Engine/Plugins/Experimental/Mutable/Source/MutableTools/Private/MuT/CodeOptimiser_Parameters.cpp
+++ b/Engine/Plugins/Experimental/Mutable/Source/MutableTools/Private/MuT/CodeOptimiser_Parameters.cpp
@@ -169,7 +169,7 @@ namespace mu
                         childItem.at = typed->children[typed->op.args.MeshInterpolate.base].child();
                         AddIfNeeded(childItem);
 
-                        for (int t=0;t<MUTABLE_OP_MAX_INTERPOLATE_COUNT-1;++t)
+                        for (int32 t=0;t<MUTABLE_OP_MAX_INTERPOLATE_COUNT-1;++t)
                         {
                             childItem.at = typed->children[typed->op.args.MeshInterpolate.targets[t]].child();
                             AddIfNeeded(childItem);
@@ -1963,7 +1963,7 @@ namespace mu
     //---------------------------------------------------------------------------------------------
     //---------------------------------------------------------------------------------------------
     //---------------------------------------------------------------------------------------------
-    LODCountReducerAST::LODCountReducerAST( Ptr<ASTOp>& root, int lodCount )
+    LODCountReducerAST::LODCountReducerAST( Ptr<ASTOp>& root, int32 lodCount )
     {
         m_lodCount = lodCount;
         Traverse( root );
@@ -2041,7 +2041,7 @@ namespace mu
                 // Promote the intructions that depend on runtime parameters, and sink new
                 // format instructions.
                 bool modified = true;
-                int numIterations = 0;
+                int32 numIterations = 0;
                 while (modified && (!m_optimizeIterationsMax || m_optimizeIterationsLeft>0 || !numIterations ))
                 {
                     modified = false;
@@ -2060,17 +2060,17 @@ namespace mu
                     roots.Add(m_states[s].root);
 
                     UE_LOG(LogMutableCore, Verbose, TEXT(" - after parameter optimiser"));
-					UE_LOG(LogMutableCore, Verbose, TEXT("(int) %s : %ld"), TEXT("ast size"), int64(ASTOp::CountNodes(roots)));
+					//UE_LOG(LogMutableCore, Verbose, TEXT("(int) %s : %ld"), TEXT("ast size"), int64(ASTOp::CountNodes(roots)));
 
                     // All kind of optimisations that depend on the meaning of each operation
                     UE_LOG(LogMutableCore, Verbose, TEXT(" - semantic optimiser"));
                     modified |= SemanticOptimiserAST( roots, m_options->GetPrivate()->OptimisationOptions, 1 );
-					UE_LOG(LogMutableCore, Verbose, TEXT("(int) %s : %ld"), TEXT("ast size"), int64(ASTOp::CountNodes(roots)));
+					//UE_LOG(LogMutableCore, Verbose, TEXT("(int) %s : %ld"), TEXT("ast size"), int64(ASTOp::CountNodes(roots)));
 					//ASTOp::LogHistogram(roots);
 
                     UE_LOG(LogMutableCore, Verbose, TEXT(" - sink optimiser"));
                     modified |= SinkOptimiserAST( roots, m_options->GetPrivate()->OptimisationOptions );
-					UE_LOG(LogMutableCore, Verbose, TEXT("(int) %s : %ld"), TEXT("ast size"), int64(ASTOp::CountNodes(roots)));
+					//UE_LOG(LogMutableCore, Verbose, TEXT("(int) %s : %ld"), TEXT("ast size"), int64(ASTOp::CountNodes(roots)));
 					//ASTOp::LogHistogram(roots);
 
                     // Image size operations are treated separately
@@ -2097,18 +2097,25 @@ namespace mu
             }
         }
 
+
+		TArray<Ptr<ASTOp>> roots;
+		for (const FStateCompilationData& s : m_states)
+		{
+			roots.Add(s.root);
+		}
+
         // Mark the instructions that don't depend on runtime parameters to be cached. This is
         // necessary at this stage before GPU optimisation.
         {
-            TArray<Ptr<ASTOp>> roots;
-            for(const FStateCompilationData& s:m_states)
-            {
-                roots.Add(s.root);
-            }
-
             AccumulateAllImageFormatsOpAST opFormats;
             opFormats.Run(roots);
 
+			// Reset the state root operations in case they have changed due to optimization
+			for (int32 RootIndex = 0; RootIndex < m_states.Num(); ++RootIndex)
+			{
+				m_states[RootIndex].root = roots[RootIndex];
+			}
+
             for (FStateCompilationData& s: m_states )
             {
                 {
@@ -2133,12 +2140,6 @@ namespace mu
 			int32 Pass = 1;
             while (modified && (!m_optimizeIterationsMax || m_optimizeIterationsLeft>0 || !numIterations ))
             {
-                TArray<Ptr<ASTOp>> roots;
-                for(const FStateCompilationData& s:m_states)
-                {
-                    roots.Add(s.root);
-                }
-
                 ++numIterations;
                 --m_optimizeIterationsLeft;
                 UE_LOG(LogMutableCore, Verbose, TEXT("State reoptimise iteration %d, max %d, left %d"),
@@ -2147,28 +2148,22 @@ namespace mu
                 modified = false;
 
                 UE_LOG(LogMutableCore, Verbose, TEXT(" - semantic optimiser"));
-                modified |=
-                    SemanticOptimiserAST( roots, m_options->GetPrivate()->OptimisationOptions, Pass );
-				UE_LOG(LogMutableCore, Verbose, TEXT("(int) %s : %ld"), TEXT("ast size"), int64(ASTOp::CountNodes(roots)));
+                modified |= SemanticOptimiserAST( roots, m_options->GetPrivate()->OptimisationOptions, Pass );
+				//UE_LOG(LogMutableCore, Verbose, TEXT("(int) %s : %ld"), TEXT("ast size"), int64(ASTOp::CountNodes(roots)));
 
                 // Image size operations are treated separately
                 UE_LOG(LogMutableCore, Verbose, TEXT(" - size optimiser"));
                 modified |= SizeOptimiserAST( roots );
-				UE_LOG(LogMutableCore, Verbose, TEXT("(int) %s : %ld"), TEXT("ast size"), int64(ASTOp::CountNodes(roots)));
+				//UE_LOG(LogMutableCore, Verbose, TEXT("(int) %s : %ld"), TEXT("ast size"), int64(ASTOp::CountNodes(roots)));
 			}
 
-            for(FStateCompilationData& s:m_states)
+            for(Ptr<ASTOp>& Root : roots)
             {
                 UE_LOG(LogMutableCore, Verbose, TEXT(" - constant optimiser"));
-				modified = ConstantGeneratorAST( m_options->GetPrivate(), s.root, Pass );
+				modified = ConstantGeneratorAST( m_options->GetPrivate(), Root, Pass );
 			}
 
-            TArray<Ptr<ASTOp>> roots;
-            for(const FStateCompilationData& s:m_states)
-            {
-                roots.Add(s.root);
-            }
-			UE_LOG(LogMutableCore, Verbose, TEXT("(int) %s : %ld"), TEXT("ast size"), int64(ASTOp::CountNodes(roots)));
+			//UE_LOG(LogMutableCore, Verbose, TEXT("(int) %s : %ld"), TEXT("ast size"), int64(ASTOp::CountNodes(roots)));
 
             UE_LOG(LogMutableCore, Verbose, TEXT(" - duplicated data remover"));
             DuplicatedDataRemoverAST( roots );
@@ -2179,12 +2174,11 @@ namespace mu
 			//UE_LOG(LogMutableCore, Verbose, TEXT("(int) %s : %ld"), TEXT("ast size"), int64(ASTOp::CountNodes(roots)));
 		}
 
-        // Gather all the current roots
-        TArray<Ptr<ASTOp>> roots;
-        for(const FStateCompilationData& s:m_states)
-        {
-            roots.Add(s.root);
-        }
+		// Reset the state root operations in case they have changed due to optimization
+		for (int32 RootIndex = 0; RootIndex < m_states.Num(); ++RootIndex)
+		{
+			m_states[RootIndex].root = roots[RootIndex];
+		}
 
         // Optimise the data formats
         {
diff --git a/Engine/Plugins/Experimental/Mutable/Source/MutableTools/Private/MuT/CodeOptimiser_Semantic.cpp b/Engine/Plugins/Experimental/Mutable/Source/MutableTools/Private/MuT/CodeOptimiser_Semantic.cpp
--- a/Engine/Plugins/Experimental/Mutable/Source/MutableTools/Private/MuT/CodeOptimiser_Semantic.cpp
+++ b/Engine/Plugins/Experimental/Mutable/Source/MutableTools/Private/MuT/CodeOptimiser_Semantic.cpp
@@ -37,29 +37,38 @@ namespace mu
     //---------------------------------------------------------------------------------------------
     //---------------------------------------------------------------------------------------------
     bool SemanticOptimiserAST(
-		ASTOpList& roots,
-		const FModelOptimizationOptions& optimisationOptions,
+		ASTOpList& Roots,
+		const FModelOptimizationOptions& OptimisationOptions,
 		int32 Pass
 	)
     {
         MUTABLE_CPUPROFILER_SCOPE(SemanticOptimiserAST);
 
-        bool modified = false;
+        bool bModified = false;
 
         // TODO: isn't top down better suited?
-        ASTOp::Traverse_BottomUp_Unique( roots, [&](Ptr<ASTOp>& n)
+        ASTOp::Traverse_BottomUp_Unique( Roots, [&](Ptr<ASTOp>& CurrentOp)
         {
-            auto o = n->OptimiseSemantic(optimisationOptions, Pass);
+            Ptr<ASTOp> OptimizedOp = CurrentOp->OptimiseSemantic(OptimisationOptions, Pass);
 
             // If the returned value is null it means no change.
-            if (o && o!=n)
+            if (OptimizedOp && OptimizedOp !=CurrentOp)
             {
-                modified = true;
-                ASTOp::Replace(n,o);
+				bModified = true;
+                ASTOp::Replace(CurrentOp, OptimizedOp);
+
+				// Check if we are replacing one of the root operations and update the parameter array.
+				for (int32 RootIndex=0; RootIndex<Roots.Num(); ++RootIndex)
+				{
+					if (Roots[RootIndex]== CurrentOp)
+					{
+						Roots[RootIndex] = OptimizedOp;
+					}
+				}
             }
         });
 
-        return modified;
+        return bModified;
     }
 
 
diff --git a/Engine/Plugins/Experimental/Mutable/Source/MutableTools/Public/MuT/NodeComponentSwitch.h b/Engine/Plugins/Experimental/Mutable/Source/MutableTools/Public/MuT/NodeComponentSwitch.h
--- a/Engine/Plugins/Experimental/Mutable/Source/MutableTools/Public/MuT/NodeComponentSwitch.h
+++ b/Engine/Plugins/Experimental/Mutable/Source/MutableTools/Public/MuT/NodeComponentSwitch.h
@@ -27,6 +27,9 @@ namespace mu
 		virtual const FNodeType* GetType() const { return &StaticType; }
 		static const FNodeType* GetStaticType() { return &StaticType; }
 
+		// NodeComponent interface
+		virtual const class NodeComponentNew* GetParentComponentNew() const override { check(false); return nullptr; }
+
 	protected:
 
 		//! Forbidden. Manage with the Ptr<> template.
diff --git a/Engine/Plugins/Experimental/Mutable/Source/MutableTools/Public/MuT/NodeComponentVariation.h b/Engine/Plugins/Experimental/Mutable/Source/MutableTools/Public/MuT/NodeComponentVariation.h
--- a/Engine/Plugins/Experimental/Mutable/Source/MutableTools/Public/MuT/NodeComponentVariation.h
+++ b/Engine/Plugins/Experimental/Mutable/Source/MutableTools/Public/MuT/NodeComponentVariation.h
@@ -32,6 +32,9 @@ namespace mu
 		virtual const FNodeType* GetType() const override { return &StaticType; }
 		static const FNodeType* GetStaticType() { return &StaticType; }
 
+		// NodeComponent interface
+		virtual const class NodeComponentNew* GetParentComponentNew() const override { check(false); return nullptr; }
+
     protected:
 
         //! Forbidden. Manage with the Ptr<> template.
