[mutable] Added a default layer to Texture Layer Nodes.
#tests Editor
#rnx

diff --git a/Engine/Plugins/Experimental/Mutable/Source/CustomizableObjectEditor/Private/MuCOE/Nodes/CustomizableObjectNodeTextureLayer.cpp b/Engine/Plugins/Experimental/Mutable/Source/CustomizableObjectEditor/Private/MuCOE/Nodes/CustomizableObjectNodeTextureLayer.cpp
--- a/Engine/Plugins/Experimental/Mutable/Source/CustomizableObjectEditor/Private/MuCOE/Nodes/CustomizableObjectNodeTextureLayer.cpp
+++ b/Engine/Plugins/Experimental/Mutable/Source/CustomizableObjectEditor/Private/MuCOE/Nodes/CustomizableObjectNodeTextureLayer.cpp
@@ -16,7 +16,7 @@ TArray<UCustomizableObjectNodeTextureLayer::TextureLayerLayersData> UCustomizabl
 UCustomizableObjectNodeTextureLayer::UCustomizableObjectNodeTextureLayer()
 	: Super()
 {
-
+	Layers.Add(FCustomizableObjectTextureLayer());
 }
 
 
