[mutable] Fixed multiple compilations of shaders during CIS testing.
- The redundant compilation was due to us clearing all data (GC) so the materials required recompilation (shadermap) when reloading them.

#rnx

diff --git a/Engine/Plugins/Experimental/Mutable/Source/MutableValidation/Private/MuV/COIBulkUpdateTestCommandlet.cpp b/Engine/Plugins/Experimental/Mutable/Source/MutableValidation/Private/MuV/COIBulkUpdateTestCommandlet.cpp
--- a/Engine/Plugins/Experimental/Mutable/Source/MutableValidation/Private/MuV/COIBulkUpdateTestCommandlet.cpp
+++ b/Engine/Plugins/Experimental/Mutable/Source/MutableValidation/Private/MuV/COIBulkUpdateTestCommandlet.cpp
@@ -11,10 +11,16 @@
 #include "MuCO/CustomizableObjectInstance.h"
 #include "Interfaces/ITargetPlatformManagerModule.h"
 #include "MuCO/CustomizableObjectPrivate.h"
+#include "MuCO/CustomizableObjectSystem.h"
+#include "MuCOE/CustomizableObjectBenchmarkingUtils.h"
 
 
 int32 UCOIBulkUpdateTestCommandlet::Main(const FString& Params)
 {
+	// Ensure we have the cvars used for our testing set
+	UCustomizableObjectSystem::SetBenchmarkState(true);
+	//CustomizableObjectBenchmarkingUtils::SetMutableCVarsForCIS();
+	
 	// Ensure we do not show any OK dialog since we are not an user that can interact with them
 	GIsRunningUnattendedScript = true;
 
@@ -141,12 +147,9 @@ int32 UCOIBulkUpdateTestCommandlet::Main(const FString& Params)
 		const FString CustomizableObjectName = CustomizableObject->GetName();
 		
 		// Set the compilation platform based on what the system is currently running on
-		FCompilationOptions CompilationOptions = CustomizableObject->GetPrivate()->GetCompileOptions();
+		FCompilationOptions CompilationOptions = GetCompilationOptionsForBenchmarking(*CustomizableObject);
 		CompilationOptions.TargetPlatform = TargetCompilationPlatform;
-		CompilationOptions.bUseDiskCompilation = false;
-		CompilationOptions.OptimizationLevel = UE_MUTABLE_MAX_OPTIMIZATION;
-		CompilationOptions.bSilentCompilation = false;
-
+		
 		// Compile the current CO object
 		if (!CompilationUtility->CompileCustomizableObject(CustomizableObject, false, &CompilationOptions))	// Do not log mutable data since mongoDB will not be able to handle it correctly 
 		{
@@ -166,14 +169,17 @@ int32 UCOIBulkUpdateTestCommandlet::Main(const FString& Params)
 			Instances.RemoveAt(0);
 			UE_LOG(LogMutable,Display,TEXT("\t( %u / %u ) Processing instance : \"%s\" ."),CurrentInstanceIndex++, TotalAmountOfInstances ,*Instance->GetName());
 
-			CollectGarbage(RF_NoFlags, true);
+			CollectGarbage(GARBAGE_COLLECTION_KEEPFLAGS, true);
 			
 			// Update each one of the instances and notify if the update failed in any manner
 			InstanceUpdatingUtility->UpdateInstance(Instance.Get());
+			
+			// Remove standalone flag from the instance so we can GC it while keeping other standalone objects
+			Instance->ClearFlags(EObjectFlags::RF_Standalone);
 		}
 
 		ResourcesIterator.RemoveCurrent();
-		CollectGarbage(RF_NoFlags, true);
+		CollectGarbage(GARBAGE_COLLECTION_KEEPFLAGS, true);
 	}
 
 	UE_LOG(LogMutable,Display,TEXT("Mutable commandlet finished."));
diff --git a/Engine/Plugins/Experimental/Mutable/Source/MutableValidation/Private/MuV/CustomizableObjectValidationCommandlet.cpp b/Engine/Plugins/Experimental/Mutable/Source/MutableValidation/Private/MuV/CustomizableObjectValidationCommandlet.cpp
--- a/Engine/Plugins/Experimental/Mutable/Source/MutableValidation/Private/MuV/CustomizableObjectValidationCommandlet.cpp
+++ b/Engine/Plugins/Experimental/Mutable/Source/MutableValidation/Private/MuV/CustomizableObjectValidationCommandlet.cpp
@@ -16,6 +16,7 @@
 #include "MuR/Model.h"
 #include "Interfaces/ITargetPlatformManagerModule.h"
 #include "Interfaces/ITargetPlatform.h"
+#include "MuCOE/CustomizableObjectBenchmarkingUtils.h"
 #include "Serialization/MemoryWriter.h"
 
 
@@ -26,6 +27,10 @@ int32 UCustomizableObjectValidationCommandlet::Main(const FString& Params)
 	// Execution arguments for commandlet from IDE
 	// -run=CustomizableObjectValidation -CustomizableObject=(PathToCO)
 
+	// Ensure we have the cvars used for our testing set
+	UCustomizableObjectSystem::SetBenchmarkState(true);
+	//CustomizableObjectBenchmarkingUtils::SetMutableCVarsForCIS();
+	
 	// Ensure we do not show any OK dialog since we are not an user that can interact with them
 	GIsRunningUnattendedScript = true;
 	
@@ -80,13 +85,10 @@ int32 UCustomizableObjectValidationCommandlet::Main(const FString& Params)
 	bool bWasCoCompilationSuccessful = false;
 	{
 		LLM_SCOPE_BYNAME(TEXT("CustomizableObjectValidationCommandlet/Compile"));
-		
-		// Override some configurations that may have been changed by the user
-		FCompilationOptions CompilationOptions = ToTestCustomizableObject->GetPrivate()->GetCompileOptions();
-		CompilationOptions.bSilentCompilation = false;
-		CompilationOptions.OptimizationLevel = UE_MUTABLE_MAX_OPTIMIZATION;
-		CompilationOptions.TextureCompression = ECustomizableObjectTextureCompression::Fast;
 
+		// Override some configurations that may have been changed by the user
+		FCompilationOptions CompilationOptions = GetCompilationOptionsForBenchmarking(*ToTestCustomizableObject);
+		
 		// Set the target compilation platform based on what the caller wants
 		CompilationOptions.TargetPlatform = TargetCompilationPlatform;
 
@@ -159,48 +161,8 @@ int32 UCustomizableObjectValidationCommandlet::Main(const FString& Params)
 	{
 		LLM_SCOPE_BYNAME(TEXT("CustomizableObjectValidationCommandlet/GenerateInstances"));
 		
-		// Test this parameter configuration in all the states of the CO
-		const uint32 StateCount = ToTestCustomizableObject->GetStateCount();
-		check(StateCount >= 1);
-
-		UE_LOG(LogMutable, Display, TEXT("Requeted Instances Count : %i"), InstancesToGenerate);
-		UE_LOG(LogMutable, Display, TEXT("State Count = %i"), StateCount);
-	
-		// Compute actual total amount of instances to generate
-		const uint32 TotalInstancesToTestCount = InstancesToGenerate * StateCount;
-		UE_LOG(LogMutable,Display,TEXT("Generating %i instances (states * requested instances)..."), TotalInstancesToTestCount);
-		
-		// Create randomization stream for the parameters of the instance
-		FRandomStream RandomizationStream = FRandomStream(0);
-		
-		// Generate a series of instances to later update
-		for (uint32 CustomizableObjectInstanceIndex = 0; CustomizableObjectInstanceIndex < InstancesToGenerate; CustomizableObjectInstanceIndex++)
-		{
-			UCustomizableObjectInstance* GeneratedInstance = ToTestCustomizableObject->CreateInstance();
-			if (GeneratedInstance)
-			{
-				// Force generation of all LODS
-				TArray<uint16> RequestedLodLevels{};
-				RequestedLodLevels.Init(0, GeneratedInstance->GetNumComponents());
-				GeneratedInstance->GetPrivate()->GetDescriptor().SetRequestedLODLevels(RequestedLodLevels);
-				
-				// Randomize instance values
-				GeneratedInstance->SetRandomValuesFromStream(RandomizationStream);
-				
-				for (uint32 State = 0; State < StateCount; State++)
-				{
-					// Set the state for the instance and store it for later update.
-					GeneratedInstance->GetPrivate()->SetState(State);
-                	InstancesToProcess.Enqueue(GeneratedInstance->Clone());
-					GeneratedInstances++;
-				}
-			}
-			else
-			{
-				UE_LOG(LogMutable,Error,TEXT("Failed to generate COI for the %s CO."),*ToTestCustomizableObject->GetName());
-				bWasInstancesCreationSuccessful = false;
-			}
-		}
+		// Create a set of instances so we can later test them out
+		bWasInstancesCreationSuccessful = CustomizableObjectBenchmarkingUtils::GenerateDeterministicSetOfInstances(ToTestCustomizableObject,InstancesToGenerate, InstancesToProcess, GeneratedInstances);
 	}
 	// ---------------------------------------------------------------------------------------------------------- //
 
@@ -218,8 +180,8 @@ int32 UCustomizableObjectValidationCommandlet::Main(const FString& Params)
 		TStrongObjectPtr<UCustomizableObjectInstance> InstanceToUpdate;
 		while (InstancesToProcess.Dequeue(InstanceToUpdate))
 		{
-			CollectGarbage(RF_NoFlags, true);
-			
+			CollectGarbage(GARBAGE_COLLECTION_KEEPFLAGS, true);
+
 			if (!InstanceUpdatingUtility->UpdateInstance(InstanceToUpdate.Get()))
 			{
 				bInstanceFailedUpdate = true;
