[mutable] Fixed multiple compilation of shaders during CIS testing.
- The redundant compilation was due to us clearing all data (GC) so the materials required recompilation (shadermap) when reloading them.

#rnx

diff --git a/Engine/Plugins/Experimental/Mutable/Source/MutableValidation/Private/MuV/COIBulkUpdateTestCommandlet.cpp b/Engine/Plugins/Experimental/Mutable/Source/MutableValidation/Private/MuV/COIBulkUpdateTestCommandlet.cpp
--- a/Engine/Plugins/Experimental/Mutable/Source/MutableValidation/Private/MuV/COIBulkUpdateTestCommandlet.cpp
+++ b/Engine/Plugins/Experimental/Mutable/Source/MutableValidation/Private/MuV/COIBulkUpdateTestCommandlet.cpp
@@ -166,14 +166,17 @@ int32 UCOIBulkUpdateTestCommandlet::Main(const FString& Params)
 			Instances.RemoveAt(0);
 			UE_LOG(LogMutable,Display,TEXT("\t( %u / %u ) Processing instance : \"%s\" ."),CurrentInstanceIndex++, TotalAmountOfInstances ,*Instance->GetName());
 
-			CollectGarbage(RF_NoFlags, true);
+			CollectGarbage(GARBAGE_COLLECTION_KEEPFLAGS, true);
 			
 			// Update each one of the instances and notify if the update failed in any manner
 			InstanceUpdatingUtility->UpdateInstance(Instance.Get());
+			
+			// Remove standalone flag from the instance so we can GC it while keeping other standalone objects
+			Instance->ClearFlags(EObjectFlags::RF_Standalone);
 		}
 
 		ResourcesIterator.RemoveCurrent();
-		CollectGarbage(RF_NoFlags, true);
+		CollectGarbage(GARBAGE_COLLECTION_KEEPFLAGS, true);
 	}
 
 	UE_LOG(LogMutable,Display,TEXT("Mutable commandlet finished."));
diff --git a/Engine/Plugins/Experimental/Mutable/Source/MutableValidation/Private/MuV/CustomizableObjectValidationCommandlet.cpp b/Engine/Plugins/Experimental/Mutable/Source/MutableValidation/Private/MuV/CustomizableObjectValidationCommandlet.cpp
--- a/Engine/Plugins/Experimental/Mutable/Source/MutableValidation/Private/MuV/CustomizableObjectValidationCommandlet.cpp
+++ b/Engine/Plugins/Experimental/Mutable/Source/MutableValidation/Private/MuV/CustomizableObjectValidationCommandlet.cpp
@@ -218,8 +218,8 @@ int32 UCustomizableObjectValidationCommandlet::Main(const FString& Params)
 		TStrongObjectPtr<UCustomizableObjectInstance> InstanceToUpdate;
 		while (InstancesToProcess.Dequeue(InstanceToUpdate))
 		{
-			CollectGarbage(RF_NoFlags, true);
-			
+			CollectGarbage(GARBAGE_COLLECTION_KEEPFLAGS, true);
+
 			if (!InstanceUpdatingUtility->UpdateInstance(InstanceToUpdate.Get()))
 			{
 				bInstanceFailedUpdate = true;
