[Mutable] Fix Viewport LOD menu getting out of sync after a compilation.

#rnx

diff --git a/Engine/Plugins/Experimental/Mutable/Source/CustomizableObjectEditor/Private/MuCOE/SCustomizableObjectEditorViewport.cpp b/Engine/Plugins/Experimental/Mutable/Source/CustomizableObjectEditor/Private/MuCOE/SCustomizableObjectEditorViewport.cpp
--- a/Engine/Plugins/Experimental/Mutable/Source/CustomizableObjectEditor/Private/MuCOE/SCustomizableObjectEditorViewport.cpp
+++ b/Engine/Plugins/Experimental/Mutable/Source/CustomizableObjectEditor/Private/MuCOE/SCustomizableObjectEditorViewport.cpp
@@ -405,6 +405,8 @@ void SCustomizableObjectEditorViewportTabBody::SetPreviewActor(const TWeakObject
 {
 	PreviewSkeletalMeshComponents = InSkeletalMeshComponents;
 	LevelViewportClient->SetPreviewActor(InActor, InInstance, InSkeletalMeshComponents);
+	
+	LODSelection = 0;
 }
 
 
@@ -956,7 +958,7 @@ int32 SCustomizableObjectEditorViewportTabBody::GetLODModelCount() const
 
 bool SCustomizableObjectEditorViewportTabBody::IsLODModelSelected(int32 LODSelectionType) const
 {
-	return (LODSelection == LODSelectionType) ? true : false;
+	return LODSelection == LODSelectionType;
 }
 
 
