[Mutable] Fix crash with Custom Pin Context Menu.

- Fix crash when perfomring an undo/redo transation with the context menu opened.
- Now only ReconstructNode actions will close the context menu.

#rnx

diff --git a/Engine/Plugins/Experimental/Mutable/Source/CustomizableObjectEditor/Private/MuCOE/CustomizableObjectEditor.cpp b/Engine/Plugins/Experimental/Mutable/Source/CustomizableObjectEditor/Private/MuCOE/CustomizableObjectEditor.cpp
--- a/Engine/Plugins/Experimental/Mutable/Source/CustomizableObjectEditor/Private/MuCOE/CustomizableObjectEditor.cpp
+++ b/Engine/Plugins/Experimental/Mutable/Source/CustomizableObjectEditor/Private/MuCOE/CustomizableObjectEditor.cpp
@@ -1166,6 +1166,8 @@ void FCustomizableObjectEditor::PostUndo(bool bSuccess)
 
 		GraphEditor->NotifyGraphChanged();
 		CustomizableObject->MarkPackageDirty();
+
+		FSlateApplication::Get().DismissAllMenus();
 	}
 }
 
diff --git a/Engine/Plugins/Experimental/Mutable/Source/CustomizableObjectEditor/Private/MuCOE/EdGraphSchema_CustomizableObject.cpp b/Engine/Plugins/Experimental/Mutable/Source/CustomizableObjectEditor/Private/MuCOE/EdGraphSchema_CustomizableObject.cpp
--- a/Engine/Plugins/Experimental/Mutable/Source/CustomizableObjectEditor/Private/MuCOE/EdGraphSchema_CustomizableObject.cpp
+++ b/Engine/Plugins/Experimental/Mutable/Source/CustomizableObjectEditor/Private/MuCOE/EdGraphSchema_CustomizableObject.cpp
@@ -913,7 +913,7 @@ void UEdGraphSchema_CustomizableObject::GetContextMenuActions(UToolMenu* Menu, U
 	}
 	else // On Pin right click
 	{
-		if (const UCustomizableObjectNode* Node = Cast<UCustomizableObjectNode>(Context->Node))
+		if (UCustomizableObjectNode* Node = Cast<UCustomizableObjectNode>(const_cast<UEdGraphNode*>(Context->Node.Get())))
 		{
 			if (const UEdGraphPin* Pin = Context->Pin)
 			{
@@ -933,7 +933,7 @@ void UEdGraphSchema_CustomizableObject::GetContextMenuActions(UToolMenu* Menu, U
 				
 				if (TSharedPtr<IDetailsView> Widget = Node->CustomizePinDetails(*Pin))
 				{
-					Widget->OnFinishedChangingProperties().AddLambda([WeakMenu = Widget->AsWeak()](const FPropertyChangedEvent& Event){
+					Node->PostReconstructNodeDelegate.AddLambda([WeakMenu = Widget->AsWeak()](){
 						if (TSharedPtr<SWidget> Menu = WeakMenu.Pin())
 						{
 							FSlateApplication::Get().DismissMenuByWidget(Menu.ToSharedRef());
