[mutable] Added a new performance reporting slate (Performance Analyzer)
- Will eventually replace the Performance Reporter object
- Generates a deterministic set of instances and reports performance metrics (LogBenchmarkUtil)
- Reports if the compilation settings are not ideal for benchmark reporting

diff --git a/Engine/Plugins/Experimental/Mutable/Source/CustomizableObject/Internal/MuCO/LogBenchmarkUtil.h b/Engine/Plugins/Experimental/Mutable/Source/CustomizableObject/Internal/MuCO/LogBenchmarkUtil.h
--- a/Engine/Plugins/Experimental/Mutable/Source/CustomizableObject/Internal/MuCO/LogBenchmarkUtil.h
+++ b/Engine/Plugins/Experimental/Mutable/Source/CustomizableObject/Internal/MuCO/LogBenchmarkUtil.h
@@ -6,6 +6,7 @@
 #include "UObject/NameTypes.h"
 #include "Stats/Stats.h"
 #include "HAL/IConsoleManager.h"
+#include "MuCO/CustomizableObjectInstanceUsage.h"
 
 class UCustomizableObjectInstance;
 class UTexture2D;
@@ -69,13 +70,60 @@ DECLARE_BENCHMARK_INSIGHTS(InstanceBuildTimeAvrg, TEXT("Avrg Instance Build Time
 
 extern TAutoConsoleVariable<bool> CVarEnableBenchmark;
 
+/** Object representing the update data of an instance. In practical terms represents one row of the CSV we generate */
+struct FInstanceUpdateStats
+{
+	FString CustomizableObjectPathName = "";
+	FString CustomizableObjectInstancePathName = "";
+	FString UpdateType = "";
+	FString Descriptor = "";			// Simplified for compatibility (maybe in the future we keep the entire descriptor)
+
+	EUpdateResult UpdateResult = EUpdateResult::Error;
+
+	bool bLevelBegunPlay = false;
+	
+	double QueueTime = 0.0;
+	double UpdateTime = 0.0;
+	double TaskGetMeshTime = 0.0;
+	double TaskLockCacheTime = 0.0;
+
+	double TaskGetImagesTime = 0.0;
+	double TaskConvertResourcesTime = 0.0;
+	double TaskCallbacksTime = 0.0;
+	
+	double UpdatePeakMemory = 0.0;
+	double UpdateRealPeakMemory = 0.0;
+
+	double TaskUpdateImageTime = 0.0;
+	double TaskUpdateImagePeakMemory = 0.0;
+	double TaskUpdateImageRealPeakMemory = 0.0;
+};
+
+
+// Delegate signature
+DECLARE_MULTICAST_DELEGATE_TwoParams(FOnMeshUpdateReportedSignature, const TSharedRef<FUpdateContextPrivate> /* Context */, FInstanceUpdateStats)
+DECLARE_MULTICAST_DELEGATE_OneParam(FOnImageUpdateReportedSignature, FInstanceUpdateStats)
+
 
 /** Benchmarking system. Gathers stats and send it to Insights an Benchmarking Files. */
 class FLogBenchmarkUtil
 {
 public:
 	~FLogBenchmarkUtil();
-
+	
+	/**
+	 * Enables or disables the benchmarking system from code. Usefull for enabling the benchmarking without having to mess
+	 * with CVarEnableBenchmark
+	 * @param bIsEnabled True to enable, false to disable and let CVarEnableBenchmark decide .
+	 */
+	CUSTOMIZABLEOBJECT_API static void SetBenchmarkReportingStateOverride(bool bIsEnabled);
+	
+	/**
+	 * Get to know if the benchmarking is active or not
+	 * @return True if the CVarEnableBenchmark CVar or the override are set to true, false otherwise
+	 */
+	CUSTOMIZABLEOBJECT_API static bool IsBenchmarkingReportingEnabled();
+	
 	/** Get stats. */
 	void GetInstancesStats(int32& OutNumInstances, int32& OutNumBuiltInstances, int32& OutNumInstancesLOD0, int32& OutNumInstancesLOD1, int32& OutNumInstancesLOD2, int32& OutNumAllocatedSkeletalMeshes) const;
 
@@ -107,8 +155,19 @@ public:
 	TArray<TWeakObjectPtr<UTexture2D>> TextureTrackerArray;
 
 private:
+	
+	inline static bool bIsEnabledOverride = false;
+	
 	double TotalUpdateTime = 0;
 	uint32 NumUpdates = 0;
 
 	TSharedPtr<FArchive> Archive;
+
+public:
+
+	/** Delegate invoked each time a new mesh update is reported by this utility */
+	FOnMeshUpdateReportedSignature OnMeshUpdateReported;
+
+	/** Delegate invoked each time a new mip update is reported by this utility */
+	FOnImageUpdateReportedSignature OnImageUpdateReported;
 };
diff --git a/Engine/Plugins/Experimental/Mutable/Source/CustomizableObject/Private/MuCO/BusyWaits_Deprecated.cpp b/Engine/Plugins/Experimental/Mutable/Source/CustomizableObject/Private/MuCO/BusyWaits_Deprecated.cpp
--- a/Engine/Plugins/Experimental/Mutable/Source/CustomizableObject/Private/MuCO/BusyWaits_Deprecated.cpp
+++ b/Engine/Plugins/Experimental/Mutable/Source/CustomizableObject/Private/MuCO/BusyWaits_Deprecated.cpp
@@ -56,7 +56,7 @@ namespace CustomizableObjectSystem::ImplDeprecated
 			}
 
 			const bool bCached = ImagesInThisInstance.Contains(Image.ImageID) || // See if it is cached from this same instance (can happen with LODs)
-				(CVarReuseImagesBetweenInstances.GetValueOnAnyThread() && CustomizableObjectSystemPrivateData->ProtectedObjectCachedImages.Contains(Image.ImageID)); // See if it is cached from another instance
+				(UCustomizableObjectSystem::ShouldReuseTexturesBetweenInstances() && CustomizableObjectSystemPrivateData->ProtectedObjectCachedImages.Contains(Image.ImageID)); // See if it is cached from another instance
 
 			if (bCached)
 			{
@@ -483,7 +483,7 @@ namespace CustomizableObjectMipDataProvider::ImplDeprecated
 		MUTABLE_CPUPROFILER_SCOPE(Task_Mutable_UpdateImage);
 		const double StartTime = FPlatformTime::Seconds();
 		
-		if (CVarEnableBenchmark.GetValueOnAnyThread())
+		if (FLogBenchmarkUtil::IsBenchmarkingReportingEnabled())
 		{
 			// Cache memory used when starting the update of the image
 			OperationData->ImageUpdateStartBytes = mu::FGlobalMemoryCounter::GetAbsoluteCounter();
@@ -589,7 +589,7 @@ namespace CustomizableObjectMipDataProvider::ImplDeprecated
 					System->EndUpdate(InstanceID);
 					System->ReleaseInstance(InstanceID);
 
-					if (CVarClearWorkingMemoryOnUpdateEnd.GetValueOnAnyThread())
+					if (UCustomizableObjectSystem::ShouldClearWorkingMemoryOnUpdateEnd())
 					{
 						System->ClearWorkingMemory();
 					}
@@ -597,7 +597,7 @@ namespace CustomizableObjectMipDataProvider::ImplDeprecated
 			}
 		}
 
-		if (CVarEnableBenchmark.GetValueOnAnyThread())
+		if (FLogBenchmarkUtil::IsBenchmarkingReportingEnabled())
 		{
 			double Time = FPlatformTime::Seconds() - StartTime;
 			// Report the peak memory used by the operation
@@ -690,7 +690,7 @@ namespace CustomizableObjectMeshUpdate::ImplDeprecated
 		System->EndUpdate(InstanceID);
 		System->ReleaseInstance(InstanceID);
 
-		if (CVarClearWorkingMemoryOnUpdateEnd.GetValueOnAnyThread())
+		if (UCustomizableObjectSystem::ShouldClearWorkingMemoryOnUpdateEnd())
 		{
 			System->ClearWorkingMemory();
 		}
diff --git a/Engine/Plugins/Experimental/Mutable/Source/CustomizableObject/Private/MuCO/CustomizableObjectInstance.cpp b/Engine/Plugins/Experimental/Mutable/Source/CustomizableObject/Private/MuCO/CustomizableObjectInstance.cpp
--- a/Engine/Plugins/Experimental/Mutable/Source/CustomizableObject/Private/MuCO/CustomizableObjectInstance.cpp
+++ b/Engine/Plugins/Experimental/Mutable/Source/CustomizableObject/Private/MuCO/CustomizableObjectInstance.cpp
@@ -6064,7 +6064,7 @@ void UCustomizableInstancePrivate::BuildMaterials(const TSharedRef<FUpdateContex
 	    OperationData->GetCapturedDescriptor().GetState());
 
 	// Cache the descriptor as a string if we want to later report it using our benchmark utility. 
-	if (CVarEnableBenchmark.GetValueOnAnyThread())
+	if (FLogBenchmarkUtil::IsBenchmarkingReportingEnabled())
 	{
 		UpdateContext->CapturedDescriptor = OperationData->GetCapturedDescriptor().ToString();
 		if (GWorld)
diff --git a/Engine/Plugins/Experimental/Mutable/Source/CustomizableObject/Private/MuCO/CustomizableObjectInstanceUsage.cpp b/Engine/Plugins/Experimental/Mutable/Source/CustomizableObject/Private/MuCO/CustomizableObjectInstanceUsage.cpp
--- a/Engine/Plugins/Experimental/Mutable/Source/CustomizableObject/Private/MuCO/CustomizableObjectInstanceUsage.cpp
+++ b/Engine/Plugins/Experimental/Mutable/Source/CustomizableObject/Private/MuCO/CustomizableObjectInstanceUsage.cpp
@@ -269,7 +269,7 @@ void UCustomizableObjectInstanceUsage::SetSkeletalMeshAndOverrideMaterials(USkel
 	}
 	
 	const bool bIsTransientMesh = SkeletalMesh ? SkeletalMesh->HasAllFlags(EObjectFlags::RF_Transient) : false;
-	const bool bUseOverrideMaterials = !bIsTransientMesh || (CustomizableObject->bEnableMeshCache && CVarEnableMeshCache.GetValueOnAnyThread());
+	const bool bUseOverrideMaterials = !bIsTransientMesh || (CustomizableObject->bEnableMeshCache && UCustomizableObjectSystem::IsMeshCacheEnabled());
 
 	if (bUseOverrideMaterials)
 	{
diff --git a/Engine/Plugins/Experimental/Mutable/Source/CustomizableObject/Private/MuCO/CustomizableObjectMeshUpdate.cpp b/Engine/Plugins/Experimental/Mutable/Source/CustomizableObject/Private/MuCO/CustomizableObjectMeshUpdate.cpp
--- a/Engine/Plugins/Experimental/Mutable/Source/CustomizableObject/Private/MuCO/CustomizableObjectMeshUpdate.cpp
+++ b/Engine/Plugins/Experimental/Mutable/Source/CustomizableObject/Private/MuCO/CustomizableObjectMeshUpdate.cpp
@@ -116,7 +116,7 @@ namespace impl
 		OperationData->System->EndUpdate(InstanceID);
 		OperationData->System->ReleaseInstance(InstanceID);
 
-		if (CVarClearWorkingMemoryOnUpdateEnd.GetValueOnAnyThread())
+		if (UCustomizableObjectSystem::ShouldClearWorkingMemoryOnUpdateEnd())
 		{
 			OperationData->System->ClearWorkingMemory();
 		}
diff --git a/Engine/Plugins/Experimental/Mutable/Source/CustomizableObject/Private/MuCO/CustomizableObjectMipDataProvider.cpp b/Engine/Plugins/Experimental/Mutable/Source/CustomizableObject/Private/MuCO/CustomizableObjectMipDataProvider.cpp
--- a/Engine/Plugins/Experimental/Mutable/Source/CustomizableObject/Private/MuCO/CustomizableObjectMipDataProvider.cpp
+++ b/Engine/Plugins/Experimental/Mutable/Source/CustomizableObject/Private/MuCO/CustomizableObjectMipDataProvider.cpp
@@ -136,7 +136,7 @@ namespace Impl
 	{
 		const double StartTime = FPlatformTime::Seconds();
 		
-		if (CVarEnableBenchmark.GetValueOnAnyThread())
+		if (FLogBenchmarkUtil::IsBenchmarkingReportingEnabled())
 		{
 			// Cache memory used when starting the update of the image
 			OperationData->ImageUpdateStartBytes = mu::FGlobalMemoryCounter::GetAbsoluteCounter();
@@ -265,13 +265,13 @@ namespace Impl
 						System->EndUpdate(InstanceID);
 						System->ReleaseInstance(InstanceID);
 
-						if (CVarClearWorkingMemoryOnUpdateEnd.GetValueOnAnyThread())
+						if (UCustomizableObjectSystem::ShouldClearWorkingMemoryOnUpdateEnd())
 						{
 							System->ClearWorkingMemory();
 						}
 					}
 
-					if (CVarEnableBenchmark.GetValueOnAnyThread())
+					if (FLogBenchmarkUtil::IsBenchmarkingReportingEnabled())
 					{
 						double Time = FPlatformTime::Seconds() - StartTime;
 						// Report the peak memory used by the operation
diff --git a/Engine/Plugins/Experimental/Mutable/Source/CustomizableObject/Private/MuCO/CustomizableObjectSystem.cpp b/Engine/Plugins/Experimental/Mutable/Source/CustomizableObject/Private/MuCO/CustomizableObjectSystem.cpp
--- a/Engine/Plugins/Experimental/Mutable/Source/CustomizableObject/Private/MuCO/CustomizableObjectSystem.cpp
+++ b/Engine/Plugins/Experimental/Mutable/Source/CustomizableObject/Private/MuCO/CustomizableObjectSystem.cpp
@@ -900,8 +900,7 @@ void FinishUpdateGlobal(const TSharedRef<FUpdateContextPrivate>& Context)
 	const uint32 InstanceId = Instance ? Instance->GetUniqueID() : 0;
 	UE_LOG(LogMutable, Log, TEXT("Finished UpdateSkeletalMesh Async. Instance=%d, Frame=%d, QueueTime=%f, UpdateTime=%f"), InstanceId, GFrameNumber, Context->QueueTime, Context->UpdateTime);
 
-	if (SystemPrivate &&
-		CVarEnableBenchmark.GetValueOnAnyThread())
+	if (SystemPrivate && FLogBenchmarkUtil::IsBenchmarkingReportingEnabled())
 	{
 		FFunctionGraphTask::CreateAndDispatchWhenReady( // Calling Benchmark in a task so we make sure we exited all scopes.
 		[Context]()
@@ -1651,7 +1650,7 @@ namespace impl
 		UCustomizableObjectSystem* System = UCustomizableObjectSystem::GetInstanceChecked(); // Save since UCustomizableObjectSystem::BeginDestroy always waits for all tasks to finish
 		const UCustomizableObjectSystemPrivate* SystemPrivate = System->GetPrivate();
 		
-		if (CVarEnableBenchmark.GetValueOnAnyThread())
+		if (FLogBenchmarkUtil::IsBenchmarkingReportingEnabled())
 		{
 			// Get the amount of mutable memory in use now
 			Operation->UpdateStartBytes = mu::FGlobalMemoryCounter::GetAbsoluteCounter();
@@ -2684,7 +2683,7 @@ namespace impl
 		}
 
 		const bool bCached = ImagesInThisInstance->Contains(Image.ImageID) || // See if it is cached from this same instance (can happen with LODs)
-			(CVarReuseImagesBetweenInstances.GetValueOnAnyThread() && CustomizableObjectSystemPrivateData->ProtectedObjectCachedImages.Contains(Image.ImageID)); // See if it is cached from another instance
+			(UCustomizableObjectSystem::ShouldReuseTexturesBetweenInstances() && CustomizableObjectSystemPrivateData->ProtectedObjectCachedImages.Contains(Image.ImageID)); // See if it is cached from another instance
 
 		if (bCached)
 		{
@@ -2805,7 +2804,7 @@ namespace impl
 
 		MutableSystem->SetImagePixelConversionOverride(nullptr);
 
-		if (CVarClearWorkingMemoryOnUpdateEnd.GetValueOnAnyThread())
+		if (UCustomizableObjectSystem::ShouldClearWorkingMemoryOnUpdateEnd())
 		{
 			MutableSystem->ClearWorkingMemory();
 		}
@@ -2824,7 +2823,7 @@ namespace impl
 			MutableSystem->ReleaseInstance(InstanceID);
 		}
 
-		if (CVarClearWorkingMemoryOnUpdateEnd.GetValueOnAnyThread())
+		if (UCustomizableObjectSystem::ShouldClearWorkingMemoryOnUpdateEnd())
 		{
 			MutableSystem->ClearWorkingMemory();
 		}
@@ -2957,7 +2956,7 @@ namespace impl
 			}
 		} // if (!bInstanceValid)
 
-		if (CVarEnableBenchmark.GetValueOnAnyThread())
+		if (FLogBenchmarkUtil::IsBenchmarkingReportingEnabled())
 		{
 			// Memory used in the context of this the update of mesh
 			OperationData->UpdateEndPeakBytes = mu::FGlobalMemoryCounter::GetPeak();
@@ -3259,7 +3258,7 @@ namespace impl
 		MUTABLE_CPUPROFILER_SCOPE(Task_Game_StartUpdate)
 
 		// Check if a level has been loaded
-		if (CVarEnableBenchmark.GetValueOnAnyThread() && GWorld)
+		if (FLogBenchmarkUtil::IsBenchmarkingReportingEnabled() && GWorld)
 		{
 			Operation->bLevelBegunPlay = GWorld->GetBegunPlay();
 		}
@@ -3405,7 +3404,7 @@ namespace impl
 		// Task: Mutable Update and GetMesh
 		//-------------------------------------------------------------
 		Operation->InstanceID = Operation->bLiveUpdateMode ? CandidateInstancePrivateData->LiveUpdateModeInstanceID : 0;
-		Operation->bUseMeshCache = CustomizableObject->bEnableMeshCache && !Operation->bLiveUpdateMode && CVarEnableMeshCache.GetValueOnGameThread();
+		Operation->bUseMeshCache = CustomizableObject->bEnableMeshCache && !Operation->bLiveUpdateMode && UCustomizableObjectSystem::IsMeshCacheEnabled(true);
 
 		const bool bStreamingEnabled = (CustomizableObject->bEnableMeshStreaming || bForceStreamMeshLODs) && bStreamMeshLODs;
 		Operation->bStreamMeshLODs = bStreamingEnabled && IStreamingManager::Get().IsRenderAssetStreamingEnabled(EStreamableRenderAssetType::SkeletalMesh);
@@ -4105,21 +4104,39 @@ void UCustomizableObjectSystem::AddUncompiledCOWarning(const UCustomizableObject
 }
 
 
-void UCustomizableObjectSystem::EnableBenchmark()
+void UCustomizableObjectSystem::SetReleaseMutableTexturesImmediately(bool bReleaseTextures)
 {
-	CVarEnableBenchmark->Set(true);
+	GetPrivate()->bReleaseTexturesImmediately = bReleaseTextures;
 }
 
 
-void UCustomizableObjectSystem::EndBenchmark()
+void UCustomizableObjectSystem::SetBenchmarkState(bool bIsEnabled)
 {
-	CVarEnableBenchmark->Set(false);
+	bIsBenchmarking = bIsEnabled;
 }
 
 
-void UCustomizableObjectSystem::SetReleaseMutableTexturesImmediately(bool bReleaseTextures)
+bool UCustomizableObjectSystem::IsBenchmarking()
 {
-	GetPrivate()->bReleaseTexturesImmediately = bReleaseTextures;
+	return bIsBenchmarking;
+}
+
+
+bool UCustomizableObjectSystem::IsMeshCacheEnabled(bool bCheckCVarOnGameThread /** = false */)
+{
+	return bIsBenchmarking ? false : CVarEnableMeshCache.GetValueOnAnyThread(bCheckCVarOnGameThread);
+}
+
+
+bool UCustomizableObjectSystem::ShouldClearWorkingMemoryOnUpdateEnd()
+{
+	return bIsBenchmarking ? true : CVarClearWorkingMemoryOnUpdateEnd.GetValueOnAnyThread();
+}
+
+
+bool UCustomizableObjectSystem::ShouldReuseTexturesBetweenInstances()
+{
+	return bIsBenchmarking ? false : CVarReuseImagesBetweenInstances.GetValueOnAnyThread();
 }
 
 
@@ -4134,7 +4151,7 @@ void UCustomizableObjectSystem::SetWorkingMemory(int32 Bytes)
 
 int32 UCustomizableObjectSystem::GetWorkingMemory() const
 {
-	return WorkingMemory;
+	return bIsBenchmarking ? 16384 : WorkingMemory;
 }
 
 
diff --git a/Engine/Plugins/Experimental/Mutable/Source/CustomizableObject/Private/MuCO/LogBenchmarkUtil.cpp b/Engine/Plugins/Experimental/Mutable/Source/CustomizableObject/Private/MuCO/LogBenchmarkUtil.cpp
--- a/Engine/Plugins/Experimental/Mutable/Source/CustomizableObject/Private/MuCO/LogBenchmarkUtil.cpp
+++ b/Engine/Plugins/Experimental/Mutable/Source/CustomizableObject/Private/MuCO/LogBenchmarkUtil.cpp
@@ -120,7 +120,7 @@ void FLogBenchmarkUtil::AddTexture(UTexture2D& Texture)
 {
 	check(IsInGameThread());
 
-	if (!CVarEnableBenchmark.GetValueOnGameThread())
+	if (!IsBenchmarkingReportingEnabled())
 	{
 		return;
 	}
@@ -133,7 +133,7 @@ void FLogBenchmarkUtil::UpdateStats()
 {
 	check(IsInGameThread());
 	
-	if (!CVarEnableBenchmark.GetValueOnGameThread())
+	if (!IsBenchmarkingReportingEnabled())
 	{
 		return;
 	}
@@ -285,7 +285,7 @@ void FLogBenchmarkUtil::FinishUpdateMesh(const TSharedRef<FUpdateContextPrivate>
 {
 	check(IsInGameThread());
 
-	if (!CVarEnableBenchmark.GetValueOnGameThread())
+	if (!IsBenchmarkingReportingEnabled())
 	{
 		return;
 	}
@@ -311,29 +311,46 @@ void FLogBenchmarkUtil::FinishUpdateMesh(const TSharedRef<FUpdateContextPrivate>
 	{
 		Archive = CreateFile();
 	}
-
-	const FString ID_CO = Context->Instance->GetCustomizableObject()->GetPathName();
-	const FString ID_COI = Instance->GetPathName();
-	const FString ID_UpdateType = TEXT("Mesh");
-	const FString ID_Descriptor = Context->GetCapturedDescriptor().ToString();
-	const FString ID_UpdateResult = StaticEnum<EUpdateResult>()->GetValueAsString(Context->UpdateResult);
-	const double Time_Queue = Context->QueueTime * 1000;
-	const double Time_Update = Context->UpdateTime * 1000;
-	const double Time_TaskGetMesh = Context->TaskGetMeshTime * 1000;
-	const double Time_TaskLockCache = Context->TaskLockCacheTime * 1000;
-	const double Time_TaskGetImages = Context->TaskGetImagesTime * 1000;
-	const double Time_TaskConvertResources = Context->TaskConvertResourcesTime * 1000;
-	const double Time_TaskCallbacks =  Context->TaskCallbacksTime * 1000;
-
-	// Context data to know in what situation did the update happen
-	const FString Context_LevelBegunPlay =  Context->bLevelBegunPlay ? TEXT("true") : TEXT("false");
-
-	const double Memory_UpdateEndPeakMB = (Context->UpdateEndPeakBytes / 1024.0) / 1024.0;
-	const double Memory_UpdateEndRealPeakMB = (Context->UpdateEndRealPeakBytes / 1024.0) / 1024.0;
 	
-	const FString UpdateString = FString::Printf(TEXT("%s;%s;%s;%s;%s;%s;%f;%f;%f;%f;%f;%f;%f;%f;%f"), *ID_CO, *ID_COI, *ID_UpdateType, *ID_Descriptor, *ID_UpdateResult, *Context_LevelBegunPlay, Time_Queue, Time_Update, Time_TaskGetMesh, Time_TaskLockCache, Time_TaskGetImages, Time_TaskConvertResources, Time_TaskCallbacks, Memory_UpdateEndPeakMB, Memory_UpdateEndRealPeakMB);
+	// Identifying data
+	FInstanceUpdateStats UpdateData;
+	UpdateData.UpdateType = TEXT("Mesh");
+	UpdateData.CustomizableObjectPathName = Context->Instance->GetCustomizableObject()->GetPathName();
+	UpdateData.CustomizableObjectInstancePathName = Instance->GetPathName();
+	UpdateData.Descriptor = Context->GetCapturedDescriptor().ToString();
+	UpdateData.UpdateResult = Context->UpdateResult;
+	UpdateData.bLevelBegunPlay = Context->bLevelBegunPlay;
+	UpdateData.QueueTime = Context->QueueTime * 1000;
+	UpdateData.UpdateTime = Context->UpdateTime * 1000;
+	UpdateData.TaskGetMeshTime = Context->TaskGetMeshTime * 1000;
+	UpdateData.TaskLockCacheTime = Context->TaskLockCacheTime * 1000;
+	UpdateData.TaskGetImagesTime = Context->TaskGetImagesTime * 1000;
+	UpdateData.TaskConvertResourcesTime = Context->TaskConvertResourcesTime * 1000;
+	UpdateData.TaskCallbacksTime = Context->TaskCallbacksTime * 1000;
+	UpdateData.UpdatePeakMemory = (Context->UpdateEndPeakBytes / 1024.0) / 1024.0;
+	UpdateData.UpdateRealPeakMemory = (Context->UpdateEndRealPeakBytes / 1024.0) / 1024.0;
+
+	// todo: Find a better way of handling the construction of this string so we know the symmetry with the header row will not get lost when adding new elements
+	const FString UpdateString = FString::Printf(TEXT("%s;%s;%s;%s;%s;%s;%f;%f;%f;%f;%f;%f;%f;%f;%f"),
+		*UpdateData.CustomizableObjectPathName,
+		*UpdateData.CustomizableObjectInstancePathName,
+		*UpdateData.UpdateType,
+		*UpdateData.Descriptor,
+		*StaticEnum<EUpdateResult>()->GetValueAsString(UpdateData.UpdateResult),
+		(UpdateData.bLevelBegunPlay ? TEXT("true") : TEXT("false")),
+		UpdateData.QueueTime,
+		UpdateData.UpdateTime,
+		UpdateData.TaskGetMeshTime,
+		UpdateData.TaskLockCacheTime,
+		UpdateData.TaskGetImagesTime,
+		UpdateData.TaskConvertResourcesTime,
+		UpdateData.TaskCallbacksTime,
+		UpdateData.UpdatePeakMemory,
+		UpdateData.UpdateRealPeakMemory);
 	LogBenchmarkUtil::Write(*Archive, UpdateString);
 	Archive->Flush();
+
+	OnMeshUpdateReported.Broadcast(Context, UpdateData);
 }
 
 
@@ -341,22 +358,45 @@ void FLogBenchmarkUtil::FinishUpdateImage(const FString& CustomizableObjectPathN
 {
 	check(IsInGameThread());
 
-	if (!CVarEnableBenchmark.GetValueOnGameThread())
+	if (!IsBenchmarkingReportingEnabled())
 	{
 		return;
 	}
 
-	const FString& ID_CO = CustomizableObjectPathName;
-	const FString& ID_COI = InstancePathName;
-	const FString ID_Descriptor = InstanceDescriptor;
-	const FString Context_LevelBegunPlay =  bDidLevelBeginPlay ? TEXT("true") : TEXT("false");
-	const FString ID_UpdateType = TEXT("Image");
-	const double Time_TaskUpdateImage = TaskUpdateImageTime * 1000;
-	const double Memory_TaskUpdateImagePeakMB = (TaskUpdateImageMemoryPeak / 1024.0) / 1024.0;
-	const double Memory_TaskUpdateImageRealPeakMB = (TaskUpdateImageRealMemoryPeak / 1024.0) / 1024.0;
+	FInstanceUpdateStats MipsUpdateData;
+	MipsUpdateData.UpdateType = TEXT("Image");
+	MipsUpdateData.CustomizableObjectPathName = CustomizableObjectPathName;
+	MipsUpdateData.CustomizableObjectInstancePathName = InstancePathName;
+	MipsUpdateData.Descriptor = InstanceDescriptor;
+	MipsUpdateData.bLevelBegunPlay = bDidLevelBeginPlay;
+	MipsUpdateData.TaskUpdateImageTime = TaskUpdateImageTime * 1000;
+	MipsUpdateData.TaskUpdateImagePeakMemory = (TaskUpdateImageMemoryPeak / 1024.0) / 1024.0;
+	MipsUpdateData.TaskUpdateImageRealPeakMemory = (TaskUpdateImageRealMemoryPeak / 1024.0) / 1024.0;
+	
 
-	const FString UpdateString = FString::Printf(TEXT("%s;%s;%s;%s;;%s;;;;;;;;;;%f;%f;%f"), *ID_CO, *ID_COI, *ID_UpdateType, *ID_Descriptor, *Context_LevelBegunPlay, Time_TaskUpdateImage, Memory_TaskUpdateImagePeakMB,Memory_TaskUpdateImageRealPeakMB);
+	// todo: Find a better way of handling the construction of this string so we know the symmetry with the header row will not get lost when adding new elements
+	const FString UpdateString = FString::Printf(TEXT("%s;%s;%s;%s;;%s;;;;;;;;;;%f;%f;%f"),
+		*MipsUpdateData.CustomizableObjectPathName,
+		*MipsUpdateData.CustomizableObjectInstancePathName,
+		*MipsUpdateData.UpdateType,
+		*MipsUpdateData.Descriptor,
+		(MipsUpdateData.bLevelBegunPlay ? TEXT("true") : TEXT("false")),
+		MipsUpdateData.TaskUpdateImageTime,
+		MipsUpdateData.TaskUpdateImagePeakMemory,
+		MipsUpdateData.TaskUpdateImageRealPeakMemory);
 	LogBenchmarkUtil::Write(*Archive, UpdateString);
-	Archive->Flush();	
+	Archive->Flush();
+
+	OnImageUpdateReported.Broadcast(MipsUpdateData);
+}
+
+void FLogBenchmarkUtil::SetBenchmarkReportingStateOverride(bool bIsEnabled)
+{
+	bIsEnabledOverride = bIsEnabled;
+}
+
+bool FLogBenchmarkUtil::IsBenchmarkingReportingEnabled()
+{
+	return CVarEnableBenchmark.GetValueOnAnyThread() || bIsEnabledOverride;
 }
 
diff --git a/Engine/Plugins/Experimental/Mutable/Source/CustomizableObject/Public/MuCO/CustomizableObjectSystem.h b/Engine/Plugins/Experimental/Mutable/Source/CustomizableObject/Public/MuCO/CustomizableObjectSystem.h
--- a/Engine/Plugins/Experimental/Mutable/Source/CustomizableObject/Public/MuCO/CustomizableObjectSystem.h
+++ b/Engine/Plugins/Experimental/Mutable/Source/CustomizableObject/Public/MuCO/CustomizableObjectSystem.h
@@ -271,11 +271,6 @@ public:
 	// being used and it's not compiled.  Callers can add additional information to the error log.
 	void AddUncompiledCOWarning(const UCustomizableObject& InObject, FString const* OptionalLogInfo = nullptr);
 
-	// Enables the collection of internal Mutable performance data. It has a performance cost.
-	void EnableBenchmark();
-	// Writes the benchmark results
-	void EndBenchmark();
-
 	// Show data about all UCustomizableObjectInstance existing elements
 	void LogShowData(bool bFullInfo, bool ShowMaterialInfo) const;
 
@@ -354,6 +349,26 @@ public:
 	/** Get Mutable's working memory limit (bytes). See SetWorkingMemory(int32). */
 	int32 GetWorkingMemory() const;
 
+	/**
+	 * Get if the mutable mesh cache for the instance meshes is enabled or not.
+	 * @param bCheckCVarOnGameThread Force the checking of the relevant CVar to use the GameThread as target thread
+	 * @return True if it is enabled and false otherwise
+	 */
+	static bool IsMeshCacheEnabled(bool bCheckCVarOnGameThread = false);
+
+	/**
+	 * Get if mutable should clear its working memory between instance updates.
+	 * @return True if the clearing of memory will be performed, false otherwise.
+	 */
+	static bool ShouldClearWorkingMemoryOnUpdateEnd();
+
+	/**
+	 * Get if mutable will try to reuse textures in between instances.
+	 * @return True if mutable will try to reuse them, false otherwise.
+	 */
+	static bool ShouldReuseTexturesBetweenInstances();
+
+
 #if WITH_EDITOR
 	// Copy of the Mutable Editor Settings tied to CO compilation. They are updated whenever changed
 	FEditorCompileSettings EditorSettings;
@@ -362,6 +377,26 @@ public:
 private:
 	UPROPERTY(Transient)
 	TObjectPtr<UCustomizableObjectSystemPrivate> Private = nullptr;
+
+
+public:
+	
+	/**
+	 * Enable or disable the benchmarking mode for the MutableSystem.
+	 * @param bIsEnabled True to enable it and false to disable it
+	 */
+	static void SetBenchmarkState(bool bIsEnabled);
+	
+	/**
+	 * Get if the system is in benchmarking mode or not.
+	 * @return True if it is in benchmarking mode and false otherwise
+	 */
+	static bool IsBenchmarking();
+	
+private:
+
+	/** Flag that controls some of the settings used for the generation of instances. */
+	inline static bool bIsBenchmarking = false;
 };
 
 
diff --git a/Engine/Plugins/Experimental/Mutable/Source/CustomizableObjectEditor/Internal/MuCOE/CustomizableObjectBenchmarkingUtils.h b/Engine/Plugins/Experimental/Mutable/Source/CustomizableObjectEditor/Internal/MuCOE/CustomizableObjectBenchmarkingUtils.h
new file mode 100644
--- /dev/null
+++ b/Engine/Plugins/Experimental/Mutable/Source/CustomizableObjectEditor/Internal/MuCOE/CustomizableObjectBenchmarkingUtils.h
@@ -0,0 +1,31 @@
+﻿// Copyright Epic Games, Inc. All Rights Reserved.
+#pragma once
+
+#include "Containers/SpscQueue.h"
+#include "UObject/StrongObjectPtr.h"
+
+
+class UCustomizableObject;
+class UCustomizableObjectInstance;
+struct FCompilationOptions;
+enum class ECustomizableObjectTextureCompression : uint8;
+
+namespace CustomizableObjectBenchmarkingUtils
+{
+	/**
+	 * Generates a set of instances. If the CO does not change the instances generated from one run to the other will be equal.
+	 * @param TargetCustomizableObject The customizable object used to generate the set of deterministic instances.
+	 * @param InstancesPerState The target amount of instances we want generated for each state of the CO. Ex: If the CO has 2 states and we define 4 here we will get a total of 2 * 4 instances as output.
+	 * @param OutGeneratedInstances Queue with the generated instances. Based on the amount of states and the value set in InstancesPerState
+	 * @param OutSuccesfullyGeneratedInstanceCount The total of instances generated successfully [0 , InstancesPerState * COStateCount]
+	 * @return 
+	 */
+	CUSTOMIZABLEOBJECTEDITOR_API bool GenerateDeterministicSetOfInstances(TObjectPtr<UCustomizableObject> TargetCustomizableObject, const uint16 InstancesPerState, TSpscQueue<TStrongObjectPtr<UCustomizableObjectInstance>>& OutGeneratedInstances, uint32& OutSuccesfullyGeneratedInstanceCount);
+
+	
+	/**
+	 * Returns the optimization level to be used in benchmarking runs
+	 * @return Am integer value representing the optimization level that should be used for benchmarking runs
+	 */
+	CUSTOMIZABLEOBJECTEDITOR_API int32 GetOptimizationLevelForBenchmarking();
+}
\ No newline at end of file
diff --git a/Engine/Plugins/Experimental/Mutable/Source/CustomizableObjectEditor/Private/MuCOE/CustomizableObjectBenchmarkingUtils.cpp b/Engine/Plugins/Experimental/Mutable/Source/CustomizableObjectEditor/Private/MuCOE/CustomizableObjectBenchmarkingUtils.cpp
new file mode 100644
--- /dev/null
+++ b/Engine/Plugins/Experimental/Mutable/Source/CustomizableObjectEditor/Private/MuCOE/CustomizableObjectBenchmarkingUtils.cpp
@@ -0,0 +1,65 @@
+﻿// Copyright Epic Games, Inc. All Rights Reserved.
+
+#include "MuCOE/CustomizableObjectBenchmarkingUtils.h"
+#include "MuCO/CustomizableObjectInstancePrivate.h"
+#include "MuCO/CustomizableObjectSystem.h"
+
+
+bool CustomizableObjectBenchmarkingUtils::GenerateDeterministicSetOfInstances(TObjectPtr<UCustomizableObject> TargetCustomizableObject, const uint16 InstancesPerState,TSpscQueue<TStrongObjectPtr<UCustomizableObjectInstance>>& OutGeneratedInstances,  uint32& OutSuccesfullyGeneratedInstanceCount)
+{
+	if (!TargetCustomizableObject || !TargetCustomizableObject->IsCompiled())
+	{
+		return false;
+	}
+	
+	// Test this parameter configuration in all the states of the CO
+	const uint32 StateCount = TargetCustomizableObject->GetStateCount();
+	check(StateCount >= 1);
+
+	UE_LOG(LogMutable, Display, TEXT("Requested Instances Count : %i"), InstancesPerState);
+	UE_LOG(LogMutable, Display, TEXT("State Count = %i"), StateCount);
+	
+	// Compute actual total amount of instances to generate
+	const uint32 TotalInstancesToTestCount = InstancesPerState * StateCount;
+	UE_LOG(LogMutable,Display,TEXT("Generating %i instances (states * requested instances)..."), TotalInstancesToTestCount);
+	
+	// Create randomization stream for the parameters of the instance
+	FRandomStream RandomizationStream = FRandomStream(0);
+		
+	// Generate a series of instances to later update
+	for (uint32 CustomizableObjectInstanceIndex = 0; CustomizableObjectInstanceIndex < InstancesPerState; CustomizableObjectInstanceIndex++)
+	{
+		if (UCustomizableObjectInstance* GeneratedInstance = TargetCustomizableObject->CreateInstance())
+		{
+			// Force generation of all LODS
+			TArray<uint16> RequestedLodLevels{};
+			RequestedLodLevels.Init(0, GeneratedInstance->GetNumComponents());
+			GeneratedInstance->GetPrivate()->GetDescriptor().SetRequestedLODLevels(RequestedLodLevels);
+				
+			// Randomize instance values
+			GeneratedInstance->SetRandomValuesFromStream(RandomizationStream);
+				
+			for (uint32 State = 0; State < StateCount; State++)
+			{
+				// Set the state for the instance and store it for later update.
+				GeneratedInstance->GetPrivate()->SetState(State);
+				OutGeneratedInstances.Enqueue(GeneratedInstance->Clone());
+				OutSuccesfullyGeneratedInstanceCount++;
+			}
+		}
+		else
+		{
+			UE_LOG(LogMutable,Error,TEXT("Failed to generate COI for the %s CO."),*TargetCustomizableObject->GetName());
+			return false;
+		}
+	}
+
+	return true;
+}
+
+
+int32 CustomizableObjectBenchmarkingUtils::GetOptimizationLevelForBenchmarking()
+{
+	return UE_MUTABLE_MAX_OPTIMIZATION;
+}
+
diff --git a/Engine/Plugins/Experimental/Mutable/Source/CustomizableObjectEditor/Private/MuCOE/CustomizableObjectEditor.cpp b/Engine/Plugins/Experimental/Mutable/Source/CustomizableObjectEditor/Private/MuCOE/CustomizableObjectEditor.cpp
--- a/Engine/Plugins/Experimental/Mutable/Source/CustomizableObjectEditor/Private/MuCOE/CustomizableObjectEditor.cpp
+++ b/Engine/Plugins/Experimental/Mutable/Source/CustomizableObjectEditor/Private/MuCOE/CustomizableObjectEditor.cpp
@@ -2,6 +2,7 @@
 
 #include "MuCOE/CustomizableObjectEditor.h"
 
+#include "CustomizableObjectEditorPerformanceAnalyzer.h"
 #include "Animation/DebugSkelMeshComponent.h"
 #include "AssetRegistry/ARFilter.h"
 #include "AssetRegistry/AssetRegistryModule.h"
@@ -84,6 +85,7 @@ const FName FCustomizableObjectEditor::GraphNodePropertiesTabId( TEXT( "Customiz
 const FName FCustomizableObjectEditor::AdvancedPreviewSettingsTabId(TEXT("CustomizableObjectEditor_AdvancedPreviewSettings"));
 const FName FCustomizableObjectEditor::TextureAnalyzerTabId(TEXT("CustomizableObjectEditor_TextureAnalyzer"));
 const FName FCustomizableObjectEditor::PerformanceReportTabId(TEXT("CustomizableObjectEditor_PerformanceReport"));
+const FName FCustomizableObjectEditor::PerformanceAnalyzerTabId(TEXT("CustomizableObjectEditor_MewPerformanceReport"));
 const FName FCustomizableObjectEditor::TagExplorerTabId(TEXT("CustomizableObjectEditor_TagExplorer"));
 const FName FCustomizableObjectEditor::ObjectDebuggerTabId(TEXT("CustomizableObjectEditor_ObjectDebugger"));
 const FName FCustomizableObjectEditor::PopulationClassTagManagerTabId(TEXT("CustomizableObjectEditor_PopulationClassTabManager"));
@@ -134,6 +136,10 @@ void FCustomizableObjectEditor::RegisterTabSpawners(const TSharedRef<class FTabM
 		.SetDisplayName(LOCTEXT("PerformanceReport", "Performance Report"))
 		.SetGroup(WorkspaceMenuCategoryRef);
 
+	InTabManager->RegisterTabSpawner(PerformanceAnalyzerTabId, FOnSpawnTab::CreateSP(this, &FCustomizableObjectEditor::SpawnTab_PerformanceAnalyzer))
+		.SetDisplayName(LOCTEXT("PerformanceAnalyzer", "Performance Analyzer"))
+		.SetGroup(WorkspaceMenuCategoryRef);
+	
 	InTabManager->RegisterTabSpawner(TagExplorerTabId, FOnSpawnTab::CreateSP(this, &FCustomizableObjectEditor::SpawnTab_TagExplorer))
 		.SetDisplayName(LOCTEXT("TagExplorerTab", "Tag Explorer"))
 		.SetGroup(WorkspaceMenuCategoryRef);
@@ -150,7 +156,9 @@ void FCustomizableObjectEditor::UnregisterTabSpawners(const TSharedRef<class FTa
 	InTabManager->UnregisterTabSpawner(AdvancedPreviewSettingsTabId);
 	InTabManager->UnregisterTabSpawner(TextureAnalyzerTabId);
 	InTabManager->UnregisterTabSpawner(PerformanceReportTabId);
-}
+	InTabManager->UnregisterTabSpawner(PerformanceAnalyzerTabId);
+}	
+
 
 
 FCustomizableObjectEditor::~FCustomizableObjectEditor()
@@ -659,6 +667,13 @@ void FCustomizableObjectEditor::BindCommands()
 		FCanExecuteAction(),
 		FIsActionChecked());
 
+	// Performance Analyzer
+	ToolkitCommands->MapAction(
+		Commands.PerformanceAnalyzer,
+		FExecuteAction::CreateSP(this, &FCustomizableObjectEditor::OpenPerformanceAnalyzerTab),
+		FCanExecuteAction(),
+		FIsActionChecked());
+	
 	// Undo-Redo
 	ToolkitCommands->MapAction(
 		FGenericCommands::Get().Undo,
@@ -1201,6 +1216,7 @@ void FCustomizableObjectEditor::ExtendToolbar()
 			ToolbarBuilder.BeginSection("Information");
 			ToolbarBuilder.AddToolBarButton(FCustomizableObjectEditorCommands::Get().TextureAnalyzer);
 			ToolbarBuilder.AddToolBarButton(FCustomizableObjectEditorCommands::Get().PerformanceReport);
+			ToolbarBuilder.AddToolBarButton(FCustomizableObjectEditorCommands::Get().PerformanceAnalyzer);
 			ToolbarBuilder.EndSection();
 		}
 	};
@@ -2227,6 +2243,10 @@ void FCustomizableObjectEditor::OpenPerformanceReportTab()
 	TabManager->TryInvokeTab(PerformanceReportTabId);
 }
 
+void FCustomizableObjectEditor::OpenPerformanceAnalyzerTab()
+{
+	TabManager->TryInvokeTab(PerformanceAnalyzerTabId);
+}
 
 TSharedRef<SDockTab> FCustomizableObjectEditor::SpawnTab_TextureAnalyzer(const FSpawnTabArgs& Args)
 {
@@ -2258,6 +2278,24 @@ TSharedRef<SDockTab> FCustomizableObjectEditor::SpawnTab_PerformanceReport(const
 }
 
 
+TSharedRef<SDockTab> FCustomizableObjectEditor::SpawnTab_PerformanceAnalyzer(const FSpawnTabArgs& Args)
+{
+	check(Args.GetTabId() == PerformanceAnalyzerTabId);
+	check(CustomizableObject);
+
+	if (!PerformanceAnalyzer.IsValid())
+	{
+		PerformanceAnalyzer = SNew(SCustomizableObjectEditorPerformanceAnalyzer).CustomizableObject(CustomizableObject);
+	}
+
+	return SNew(SDockTab)
+	.Label(LOCTEXT("Performance Analyzer", "Performance Analyzer"))
+	[
+		PerformanceAnalyzer.ToSharedRef()
+	];
+}
+
+
 TSharedRef<SDockTab> FCustomizableObjectEditor::SpawnTab_TagExplorer(const FSpawnTabArgs& Args)
 {
 	check(Args.GetTabId() == TagExplorerTabId);
diff --git a/Engine/Plugins/Experimental/Mutable/Source/CustomizableObjectEditor/Private/MuCOE/CustomizableObjectEditor.h b/Engine/Plugins/Experimental/Mutable/Source/CustomizableObjectEditor/Private/MuCOE/CustomizableObjectEditor.h
--- a/Engine/Plugins/Experimental/Mutable/Source/CustomizableObjectEditor/Private/MuCOE/CustomizableObjectEditor.h
+++ b/Engine/Plugins/Experimental/Mutable/Source/CustomizableObjectEditor/Private/MuCOE/CustomizableObjectEditor.h
@@ -15,6 +15,7 @@
 enum class ECustomizableObjectProjectorType : uint8;
 namespace ESelectInfo { enum Type : int; }
 
+class SCustomizableObjectEditorPerformanceAnalyzer;
 class FCustomizableObjectEditorViewportClient;
 class FProperty;
 class FSpawnTabArgs;
@@ -277,6 +278,7 @@ private:
 	TSharedRef<SDockTab> SpawnTab_AdvancedPreviewSettings(const FSpawnTabArgs& Args);
 	TSharedRef<SDockTab> SpawnTab_TextureAnalyzer(const FSpawnTabArgs& Args);
 	TSharedRef<SDockTab> SpawnTab_PerformanceReport(const FSpawnTabArgs& Args);
+	TSharedRef<SDockTab> SpawnTab_PerformanceAnalyzer(const FSpawnTabArgs& Args);
 	TSharedRef<SDockTab> SpawnTab_TagExplorer(const FSpawnTabArgs& Args);
 	
 	/** Binds commands associated with the Static Mesh Editor. */
@@ -359,6 +361,9 @@ private:
 	/** Open the Performance Report tab */
 	void OpenPerformanceReportTab();
 
+	/** Open the Performance Analyzer tab */
+	void OpenPerformanceAnalyzerTab();
+	
 	/** Creates the actor and necessary components for the preview of the CO instance */
 	void CreatePreviewActor();
 
@@ -389,6 +394,7 @@ public:
 	static const FName AdvancedPreviewSettingsTabId;
 	static const FName TextureAnalyzerTabId;
 	static const FName PerformanceReportTabId;
+	static const FName PerformanceAnalyzerTabId;
 	static const FName TagExplorerTabId;
 	static const FName ObjectDebuggerTabId;
 	static const FName PopulationClassTagManagerTabId;
@@ -440,6 +446,9 @@ private:
 	/** Performance report widget to test and analyze the current customizable object resource demands */
 	TSharedPtr<class SCustomizableObjecEditorPerformanceReport> PerformanceReport;
 
+	/** New performance analyzer widget */
+	TSharedPtr<SCustomizableObjectEditorPerformanceAnalyzer> PerformanceAnalyzer;
+	
 	/** Widget to explore all the tags related with the Customizable Object open in the editor */
 	TSharedPtr<class SCustomizableObjectEditorTagExplorer> TagExplorer;
 
diff --git a/Engine/Plugins/Experimental/Mutable/Source/CustomizableObjectEditor/Private/MuCOE/CustomizableObjectEditorActions.cpp b/Engine/Plugins/Experimental/Mutable/Source/CustomizableObjectEditor/Private/MuCOE/CustomizableObjectEditorActions.cpp
--- a/Engine/Plugins/Experimental/Mutable/Source/CustomizableObjectEditor/Private/MuCOE/CustomizableObjectEditorActions.cpp
+++ b/Engine/Plugins/Experimental/Mutable/Source/CustomizableObjectEditor/Private/MuCOE/CustomizableObjectEditorActions.cpp
@@ -28,6 +28,7 @@ void FCustomizableObjectEditorCommands::RegisterCommands()
 	UI_COMMAND(Debug, "Debug", "Open the CustomizableObject debugger tab for this object.", EUserInterfaceActionType::Button, FInputChord());
 
 	UI_COMMAND(PerformanceReport, "Performance Report", "Open the Performance Report window to set up and perform all tests relevant to Customizable Objects and access worst cases data.", EUserInterfaceActionType::Button, FInputChord());
+	UI_COMMAND(PerformanceAnalyzer, "Performance Analyzer", "Open the Performance Analyzer window to set up and perform all tests relevant to Customizable Objects.", EUserInterfaceActionType::Button, FInputChord());
 	UI_COMMAND(TextureAnalyzer, "Texture Memory Analyzer", "Open the Texture Analyzer window to check all the information of the textures created by Mutable.", EUserInterfaceActionType::Button, FInputChord());
 
 	UI_COMMAND(CompileGatherReferences, "Compile and Gather References", "Compile and gather all asset references used in this Customizable Object. Marks the object as modified.", EUserInterfaceActionType::Button, FInputChord());
diff --git a/Engine/Plugins/Experimental/Mutable/Source/CustomizableObjectEditor/Private/MuCOE/CustomizableObjectEditorPerformanceAnalyzer.cpp b/Engine/Plugins/Experimental/Mutable/Source/CustomizableObjectEditor/Private/MuCOE/CustomizableObjectEditorPerformanceAnalyzer.cpp
new file mode 100644
--- /dev/null
+++ b/Engine/Plugins/Experimental/Mutable/Source/CustomizableObjectEditor/Private/MuCOE/CustomizableObjectEditorPerformanceAnalyzer.cpp
@@ -0,0 +1,878 @@
+﻿// Copyright Epic Games, Inc. All Rights Reserved.
+
+
+#include "CustomizableObjectEditorPerformanceAnalyzer.h"
+
+#include "SWarningOrErrorBox.h"
+#include "MuCOE/CustomizableObjectEditor.h"
+#include "MuCOE/CustomizableObjectGraph.h"
+#include "Editor/PropertyEditor/Public/PropertyCustomizationHelpers.h"
+#include "Editor/UnrealEd/Classes/ThumbnailRendering/ThumbnailManager.h"
+#include "Widgets/SCompoundWidget.h"
+#include "Runtime/UMG/Public/Components/HorizontalBox.h"
+#include "Widgets/Input/SComboBox.h"
+#include "Widgets/Layout/SBorder.h"
+#include "Widgets/Text/STextBlock.h"
+#include "Internationalization/Internationalization.h"
+#include "MuCO/CustomizableObjectSystem.h"
+#include "MuCO/CustomizableObjectSystemPrivate.h"
+#include "MuCO/LogBenchmarkUtil.h"
+#include "MuCOE/CustomizableObjectBenchmarkingUtils.h"
+#include "Widgets/Input/SNumericEntryBox.h"
+#include "Widgets/Layout/SSeparator.h"
+#include "Widgets/Input/SButton.h"
+
+
+#define LOCTEXT_NAMESPACE "CustomizableObjectEditorPerformanceAnalyzer"
+
+// Static vars
+bool SCustomizableObjectEditorPerformanceAnalyzer::bIsPerformanceAnalyzerRunning = false;
+
+namespace InstanceUpdatesMainDataColumns
+{
+	static const FName InstanceNameColumnID(TEXT("InstanceName"));
+	static const FName InstanceDescriptorColumnID(TEXT("Descriptor"));
+	static const FName InstanceUpdateResultColumnID(TEXT("UpdateResult"));
+
+	// Initial Gen
+	static const FName QueueTimeColumnID (TEXT("QueueTime"));
+	static const FName UpdateTimeColumnID (TEXT("UpdateTime"));
+	static const FName GetMeshTimeColumnID (TEXT("GetMeshTime"));
+	static const FName LockCacheTimeColumnID (TEXT("LockCacheTime"));
+	static const FName GetImagesTimeColumnID (TEXT("GetImagesTime"));
+	static const FName ConvertResourcesTimeColumnID (TEXT("ConvertResourcesTime"));
+	static const FName CallbacksTimeColumnID (TEXT("CallbacksTime"));
+	static const FName UpdatePeakMemoryColumnID (TEXT("UpdatePeakMem"));
+	static const FName UpdateRealPeakMemoryColumnID (TEXT("UpdateRealPeakMem"));
+	static const FName InstanceUpdateTypeID(TEXT("UpdateType"));
+};
+
+
+#pragma region SInstanceUpdateDataRow
+
+void SInstanceUpdateDataRow::Construct(const FArguments& InArgs, const TSharedRef<STableViewBase>& InOwnerTableView, TSharedPtr<FInstanceUpdateDataElement> Element)
+{
+	InstanceUpdateElement = Element;
+	
+	SMultiColumnTableRow<TSharedPtr<FInstanceUpdateDataElement>>::Construct(
+		STableRow::FArguments()
+		.ShowSelection(true)
+		, InOwnerTableView
+	);
+}
+
+TSharedRef<SWidget> SInstanceUpdateDataRow::GenerateWidgetForColumn(const FName& InColumnName)
+{
+	if (!InstanceUpdateElement)
+	{
+		return SNullWidget::NullWidget;
+	}
+	
+	if (InColumnName == InstanceUpdatesMainDataColumns::InstanceNameColumnID)
+	{
+		return SNew(STextBlock)
+			.Text(FText::FromString(InstanceUpdateElement->Instance->GetName()));
+	}
+	else if (InColumnName == InstanceUpdatesMainDataColumns::InstanceUpdateTypeID)
+	{
+		return SNew(STextBlock)
+			.Text(FText::FromString(InstanceUpdateElement->UpdateStats.UpdateType));
+	}
+	else if (InColumnName == InstanceUpdatesMainDataColumns::InstanceDescriptorColumnID)
+	{
+		return SNew(STextBlock)
+			.Text(FText::FromString(InstanceUpdateElement->UpdateStats.Descriptor));
+	}
+	else if (InColumnName == InstanceUpdatesMainDataColumns::InstanceUpdateResultColumnID)
+	{
+		return SNew(STextBlock)
+			.Text(UEnum::GetDisplayValueAsText( InstanceUpdateElement->UpdateStats.UpdateResult));
+	}
+	else if (InColumnName == InstanceUpdatesMainDataColumns::QueueTimeColumnID)
+	{
+		return SNew(STextBlock)
+			.Text(FText::FromString(FString::Printf(TEXT("%.2f"),InstanceUpdateElement->UpdateStats.QueueTime)));
+	}
+	else if (InColumnName == InstanceUpdatesMainDataColumns::UpdateTimeColumnID)
+	{
+		return SNew(STextBlock)
+			.Text(FText::FromString(FString::Printf(TEXT("%.2f"),InstanceUpdateElement->UpdateStats.UpdateTime)));
+	}
+	else if (InColumnName == InstanceUpdatesMainDataColumns::GetMeshTimeColumnID)
+	{
+		return SNew(STextBlock)
+			.Text(FText::FromString(FString::Printf(TEXT("%.2f"),InstanceUpdateElement->UpdateStats.TaskGetMeshTime)));
+	}
+	else if (InColumnName == InstanceUpdatesMainDataColumns::LockCacheTimeColumnID)
+	{
+		return SNew(STextBlock)
+			.Text(FText::FromString(FString::Printf(TEXT("%.4f"),InstanceUpdateElement->UpdateStats.TaskLockCacheTime)));
+	}
+	else if (InColumnName == InstanceUpdatesMainDataColumns::GetImagesTimeColumnID)
+	{
+		return SNew(STextBlock)
+			.Text(FText::FromString(FString::Printf(TEXT("%.4f"),InstanceUpdateElement->UpdateStats.TaskGetImagesTime)));
+	}
+	else if (InColumnName == InstanceUpdatesMainDataColumns::ConvertResourcesTimeColumnID)
+	{
+		return SNew(STextBlock)
+			.Text(FText::FromString(FString::Printf(TEXT("%.2f"),InstanceUpdateElement->UpdateStats.TaskConvertResourcesTime)));
+	}
+	else if (InColumnName == InstanceUpdatesMainDataColumns::CallbacksTimeColumnID)
+	{
+		return SNew(STextBlock)
+			.Text(FText::FromString(FString::Printf(TEXT("%.2f"),InstanceUpdateElement->UpdateStats.TaskCallbacksTime)));
+	}
+	else if (InColumnName == InstanceUpdatesMainDataColumns::UpdatePeakMemoryColumnID)
+	{
+		return SNew(STextBlock)
+			.Text(FText::FromString(FString::Printf(TEXT("%.2f"),InstanceUpdateElement->UpdateStats.UpdatePeakMemory)));
+	}
+	else if (InColumnName == InstanceUpdatesMainDataColumns::UpdateRealPeakMemoryColumnID)
+	{
+		return SNew(STextBlock)
+			.Text(FText::FromString(FString::Printf(TEXT("%.2f"),InstanceUpdateElement->UpdateStats.UpdateRealPeakMemory)));
+	}
+
+	// Invalid column name so no widget will be produced 
+	return SNullWidget::NullWidget;
+}
+#pragma endregion 
+
+
+
+// For now limit this slate to only show the generated instance data for the current Co and current compiled data
+void SCustomizableObjectEditorPerformanceAnalyzer::Construct(const FArguments& InArgs)
+{
+	CustomizableObject = TStrongObjectPtr(InArgs._CustomizableObject);
+	System = UCustomizableObjectSystem::GetInstance();
+	
+	// Bind ourselves to the compilation of the CO so we can extract up-to-date data from it
+	FPostCompileDelegate& PostCompilationDelegate = CustomizableObject->GetPostCompileDelegate();
+	PostCompilationDelegate.AddSP(this,&SCustomizableObjectEditorPerformanceAnalyzer::OnCustomizableObjectCompilationFinished);
+	
+	if (CustomizableObject->IsCompiled())
+	{
+		CacheCustomizableObjectModelData();
+	}
+	
+	const FSlateColor StopUpdatesButtonColor = FSlateColor(FLinearColor(1,0,0));
+	
+	// Construct a structure that shows the compilation data for the provided Customizable Object after compiling it
+	ChildSlot
+	[
+		SNew(SHorizontalBox)
+
+		+ SHorizontalBox::Slot()
+		.FillWidth(1)
+		[
+			SNew(SVerticalBox)
+			
+			// Show compiled model information
+			+ SVerticalBox::Slot()
+			.AutoHeight()
+			[
+				SNew(SBorder)
+				.Padding(5,2)
+				.Visibility(this,&SCustomizableObjectEditorPerformanceAnalyzer::GetVisibilityForBenchmarkingSettingsMessage)
+				[
+					SNew(SWarningOrErrorBox)
+					.Message(LOCTEXT("SCustomizableObjectEditorPerformanceAnalyzerBenchmarkSettingsLabel","For a Benchmark run it is recomended to use as \"Optimization Level\" the value of \"MAXIMUM\""))
+					.MessageStyle(EMessageStyle::Warning)
+				]
+			]
+			
+			+ SVerticalBox::Slot()
+			.AutoHeight()
+			[
+				SNew(SSeparator)
+				.Visibility(this,&SCustomizableObjectEditorPerformanceAnalyzer::GetVisibilityForBenchmarkingSettingsMessage)
+			]
+
+			// Instance generation controls 
+			+ SVerticalBox::Slot()
+			.AutoHeight()
+			.VAlign(VAlign_Bottom)
+			.Padding(5,10)
+			[
+				SNew(SVerticalBox)
+				
+				// Set the amount of instances to generate per state
+				+ SVerticalBox::Slot()
+				.AutoHeight()
+				[
+					SNew(SHorizontalBox)
+
+					+ SHorizontalBox::Slot().HAlign(HAlign_Left)
+					.FillWidth(1)
+					[
+						SNew(STextBlock)
+						.Text(LOCTEXT("SCustomizableObjectEditorPerformanceAnalyzerAmountOfInstancesLabel", "Instances per State : "))
+					]
+
+					+ SHorizontalBox::Slot()
+					.AutoWidth()
+					[
+						SNew(SNumericEntryBox<uint32>)
+						.OnValueChanged(this, &SCustomizableObjectEditorPerformanceAnalyzer::OnTargetAmountOfInstancesValueChange)
+						.OnValueCommitted(this, &SCustomizableObjectEditorPerformanceAnalyzer::OnTargetAmountOfInstancesValueCommited)
+						.Value(this, &SCustomizableObjectEditorPerformanceAnalyzer::OnTargetAmountOfInstancesValueRequested)
+						.AllowSpin(false)
+						.Delta(1)
+					]
+				]
+
+				// Button to generate the instances
+				+ SVerticalBox::Slot()
+				.AutoHeight()
+				[
+					SNew(SButton)
+					.OnClicked(this,&SCustomizableObjectEditorPerformanceAnalyzer::OnInstanceUpdateButtonClicked)
+					.IsEnabled(this, &SCustomizableObjectEditorPerformanceAnalyzer::IsInstanceUpdateButtonEnabled)
+					.Text(this, &SCustomizableObjectEditorPerformanceAnalyzer::GetGenerateInstancesButtonText)
+					.HAlign(HAlign_Center)
+				]
+
+				// Button to stop the generation of instances
+				+ SVerticalBox::Slot()
+				.AutoHeight()
+				[
+					SNew(SButton)
+					.OnClicked(this,&SCustomizableObjectEditorPerformanceAnalyzer::OnInstanceUpdateAbortButtonClicked)
+					.IsEnabled(this, &SCustomizableObjectEditorPerformanceAnalyzer::IsInstanceUpdateAbortButtonEnabled)
+					.Visibility(this, &SCustomizableObjectEditorPerformanceAnalyzer::ShouldInstanceUpdateAbortButtonBeVisible)
+					.Text(this, &SCustomizableObjectEditorPerformanceAnalyzer::GetStopUpdatesButtonText)
+					.ButtonColorAndOpacity(StopUpdatesButtonColor)
+					.HAlign(HAlign_Center)
+				]
+			]
+		]
+
+		// Listview space
+		+ SHorizontalBox::Slot()
+		.FillWidth(6)
+		.Padding(5,10)
+		[
+			SAssignNew(InstanceUpdatesListView,SListView<TSharedPtr<FInstanceUpdateDataElement>>)
+			.ListItemsSource(&InstanceUpdateElements)
+			.OnGenerateRow(this,&SCustomizableObjectEditorPerformanceAnalyzer::OnGenerateInstanceUpdateRow)
+			.OnMouseButtonClick(this, &SCustomizableObjectEditorPerformanceAnalyzer::OnInstancesRowClick)
+			.SelectionMode(ESelectionMode::Single)
+			.IsFocusable(true)
+			.Orientation(Orient_Vertical)
+			.HeaderRow
+			(
+				SNew(SHeaderRow)
+
+				+ SHeaderRow::Column(InstanceUpdatesMainDataColumns::InstanceNameColumnID)
+				.DefaultLabel(FText(LOCTEXT("InstanceNameColumnLabel","Name")))
+				.DefaultTooltip(FText(LOCTEXT("InstanceNameColumnToolTip","Customizable Object Instance name")))
+				.OnSort(this, &SCustomizableObjectEditorPerformanceAnalyzer::OnInstanceUpdateListViewSort)
+				.SortMode(this, &SCustomizableObjectEditorPerformanceAnalyzer::GetColumnSortMode, InstanceUpdatesMainDataColumns::InstanceNameColumnID)
+				.HAlignCell(EHorizontalAlignment::HAlign_Left)
+				.OverflowPolicy(ETextOverflowPolicy::Ellipsis)
+				.FillWidth(2)
+				
+				+ SHeaderRow::Column(InstanceUpdatesMainDataColumns::InstanceUpdateResultColumnID)
+				.DefaultLabel(FText(LOCTEXT("InstanceUpdateResultColumnLabel", "Result")))
+				.DefaultTooltip(FText(LOCTEXT("InstanceUpdateResultColumnToolTip", "Instance generation result")))
+				.OnSort(this, &SCustomizableObjectEditorPerformanceAnalyzer::OnInstanceUpdateListViewSort)
+				.SortMode(this, &SCustomizableObjectEditorPerformanceAnalyzer::GetColumnSortMode, InstanceUpdatesMainDataColumns::InstanceUpdateResultColumnID)
+				.HAlignCell(EHorizontalAlignment::HAlign_Left)
+				.OverflowPolicy(ETextOverflowPolicy::Ellipsis)
+				.FillWidth(1)
+
+				+ SHeaderRow::Column(InstanceUpdatesMainDataColumns::QueueTimeColumnID)
+				.DefaultLabel(FText(LOCTEXT("QueueTimeColumnLabel", "Queue (ms)")))
+				.DefaultTooltip(FText(LOCTEXT("QueueTimeColumnToolTip", "Time spent before starting the update")))
+				.OnSort(this, &SCustomizableObjectEditorPerformanceAnalyzer::OnInstanceUpdateListViewSort)
+				.SortMode(this, &SCustomizableObjectEditorPerformanceAnalyzer::GetColumnSortMode, InstanceUpdatesMainDataColumns::QueueTimeColumnID)
+				.HAlignCell(EHorizontalAlignment::HAlign_Right)
+				.FillWidth(1)
+
+				+ SHeaderRow::Column(InstanceUpdatesMainDataColumns::UpdateTimeColumnID)
+				.DefaultLabel(FText(LOCTEXT("UpdateTimeColumnLabel", "Initial Generation (ms)")))
+				.DefaultTooltip(FText(LOCTEXT("UpdateTimeColumnToolTip", "Initial generation - Total time")))
+				.OnSort(this, &SCustomizableObjectEditorPerformanceAnalyzer::OnInstanceUpdateListViewSort)
+				.SortMode(this, &SCustomizableObjectEditorPerformanceAnalyzer::GetColumnSortMode, InstanceUpdatesMainDataColumns::UpdateTimeColumnID)
+				.HAlignCell(EHorizontalAlignment::HAlign_Right)
+				.FillWidth(1)
+
+				+ SHeaderRow::Column(InstanceUpdatesMainDataColumns::GetMeshTimeColumnID)
+				.DefaultLabel(FText(LOCTEXT("GetMeshTimeColumnLabel", "GetMesh (ms)")))
+				.DefaultTooltip( FText(LOCTEXT("GetMeshTimeColumnToolTip", "Initial generation - Mesh generation time")))
+				.OnSort(this, &SCustomizableObjectEditorPerformanceAnalyzer::OnInstanceUpdateListViewSort)
+				.SortMode(this, &SCustomizableObjectEditorPerformanceAnalyzer::GetColumnSortMode, InstanceUpdatesMainDataColumns::GetMeshTimeColumnID)
+				.HAlignCell(EHorizontalAlignment::HAlign_Right)
+				.FillWidth(1)
+
+				+ SHeaderRow::Column(InstanceUpdatesMainDataColumns::LockCacheTimeColumnID)
+				.DefaultLabel(FText(LOCTEXT("LockCacheTimeColumnLabel", "LockCache (ms)")))
+				.DefaultTooltip( FText(LOCTEXT("LockCacheTimeColumnToolTip", "Initial generation - Time spent protecting used UObjects from GC")))
+				.OnSort(this, &SCustomizableObjectEditorPerformanceAnalyzer::OnInstanceUpdateListViewSort)
+				.SortMode(this, &SCustomizableObjectEditorPerformanceAnalyzer::GetColumnSortMode, InstanceUpdatesMainDataColumns::LockCacheTimeColumnID)
+				.HAlignCell(EHorizontalAlignment::HAlign_Right)
+				.FillWidth(1)
+
+				+ SHeaderRow::Column(InstanceUpdatesMainDataColumns::GetImagesTimeColumnID)
+				.DefaultLabel(FText(LOCTEXT("GetImagesTimeColumnLabel", "GetImages (ms)")))
+				.DefaultTooltip(FText(LOCTEXT("GetImagesTimeColumnToolTip", "Initial generation - Image generation time")))
+				.OnSort(this, &SCustomizableObjectEditorPerformanceAnalyzer::OnInstanceUpdateListViewSort)
+				.SortMode(this, &SCustomizableObjectEditorPerformanceAnalyzer::GetColumnSortMode, InstanceUpdatesMainDataColumns::GetImagesTimeColumnID)
+				.HAlignCell(EHorizontalAlignment::HAlign_Right)
+				.FillWidth(1)
+
+				+ SHeaderRow::Column(InstanceUpdatesMainDataColumns::ConvertResourcesTimeColumnID)
+				.DefaultLabel( FText(LOCTEXT("ConvertResourcesTimeColumnLabel", "ConvertResources (ms)")))
+				.DefaultTooltip(FText(LOCTEXT("ConvertResourcesTimeColumnToolTip", "Initial generation - Time converting resources from Mutable format to target format")))
+				.OnSort(this, &SCustomizableObjectEditorPerformanceAnalyzer::OnInstanceUpdateListViewSort)
+				.SortMode(this, &SCustomizableObjectEditorPerformanceAnalyzer::GetColumnSortMode, InstanceUpdatesMainDataColumns::ConvertResourcesTimeColumnID)
+				.HAlignCell(EHorizontalAlignment::HAlign_Right)
+				.FillWidth(1)
+
+				+ SHeaderRow::Column(InstanceUpdatesMainDataColumns::CallbacksTimeColumnID)
+				.DefaultLabel(FText(LOCTEXT("CallbacksTimeColumnLabel", "Callbacks (ms)")))
+				.DefaultTooltip(FText(LOCTEXT("CallbacksTimeColumnToolTip", "Initial generation - User attached callbacks")))
+				.OnSort(this, &SCustomizableObjectEditorPerformanceAnalyzer::OnInstanceUpdateListViewSort)
+				.SortMode(this, &SCustomizableObjectEditorPerformanceAnalyzer::GetColumnSortMode, InstanceUpdatesMainDataColumns::CallbacksTimeColumnID)
+				.HAlignCell(EHorizontalAlignment::HAlign_Right)
+				.FillWidth(1)
+
+				+ SHeaderRow::Column(InstanceUpdatesMainDataColumns::UpdatePeakMemoryColumnID)
+				.DefaultLabel(FText(LOCTEXT("UpdatePeakMemoryColumnLabel", "Update Peak (MB)")))
+				.DefaultTooltip(FText(LOCTEXT("UpdatePeakMemoryColumnToolTip", "Maximum memory used by Mutable by the Update")))
+				.OnSort(this, &SCustomizableObjectEditorPerformanceAnalyzer::OnInstanceUpdateListViewSort)
+				.SortMode(this, &SCustomizableObjectEditorPerformanceAnalyzer::GetColumnSortMode, InstanceUpdatesMainDataColumns::UpdatePeakMemoryColumnID)
+				.HAlignCell(EHorizontalAlignment::HAlign_Right)
+				.FillWidth(1)
+
+				+ SHeaderRow::Column(InstanceUpdatesMainDataColumns::UpdateRealPeakMemoryColumnID)
+				.DefaultLabel(FText(LOCTEXT("UpdateRealPeakMemoryColumnLabel", "Update Real Peak (MB)")))
+				.DefaultTooltip(FText(LOCTEXT("UpdateRealPeakMemoryColumnToolTip", "Maximum memory used by Mutable during the Update (Update Peak + Previously used Mutable memory)")))
+				.OnSort(this, &SCustomizableObjectEditorPerformanceAnalyzer::OnInstanceUpdateListViewSort)
+				.SortMode(this, &SCustomizableObjectEditorPerformanceAnalyzer::GetColumnSortMode, InstanceUpdatesMainDataColumns::UpdateRealPeakMemoryColumnID)
+				.HAlignCell(EHorizontalAlignment::HAlign_Right)
+				.FillWidth(1)
+			)
+		]
+	];
+}
+
+
+SCustomizableObjectEditorPerformanceAnalyzer::~SCustomizableObjectEditorPerformanceAnalyzer()
+{
+	// If we are the instance that made the bIsPerformanceAnalyzerRunning true then reset it.
+	if (AreUpdatesPending())
+	{
+		ClearPendingUpdatesQueue();
+		bIsPerformanceAnalyzerRunning = false;
+	}
+
+	SetMutableBenchmarkingSystemState(false);
+}
+
+
+EVisibility SCustomizableObjectEditorPerformanceAnalyzer::GetVisibilityForBenchmarkingSettingsMessage() const
+{
+	if (LastCompilationOptions)
+	{
+		if (LastCompilationOptions->OptimizationLevel == CustomizableObjectBenchmarkingUtils::GetOptimizationLevelForBenchmarking())
+		{
+			return EVisibility::Collapsed;
+		}
+		else
+		{
+			return EVisibility::Visible;
+		}
+	}
+
+	return EVisibility::Collapsed;
+}
+
+
+void SCustomizableObjectEditorPerformanceAnalyzer::OnTargetAmountOfInstancesValueChange(uint32 NewValue)
+{
+	if (NewValue > 0)
+	{
+		InstancesToGeneratePerState = NewValue;
+	}
+}
+
+
+void SCustomizableObjectEditorPerformanceAnalyzer::OnTargetAmountOfInstancesValueCommited(uint32 NewValue, ETextCommit::Type CommitType)
+{
+	if (NewValue > 0)
+	{
+		InstancesToGeneratePerState = NewValue;
+	}
+	else
+	{
+		InstancesToGeneratePerState = 24;
+	}
+}
+
+
+TOptional<uint32> SCustomizableObjectEditorPerformanceAnalyzer::OnTargetAmountOfInstancesValueRequested() const
+{
+	return InstancesToGeneratePerState;
+}
+
+
+bool SCustomizableObjectEditorPerformanceAnalyzer::IsInstanceUpdateButtonEnabled() const
+{
+	return !bIsPerformanceAnalyzerRunning && CustomizableObject && CustomizableObject->IsCompiled() && !AreUpdatesPending();
+}
+
+
+FReply SCustomizableObjectEditorPerformanceAnalyzer::OnInstanceUpdateButtonClicked()
+{
+	check(CustomizableObject);
+
+	// Ensure there is no work to do
+	check(!AreUpdatesPending());
+		
+	InstanceUpdateElements.Reset();
+	InstanceUpdatesListView->RequestListRefresh();
+	
+	UpdatedInstancesCount = 0;
+	TotalScheduledUpdates = 0;
+
+	// Set the CVArs as in CIS
+	//SetForBenchmarkMutableCVarValues();
+		
+	if (!CustomizableObjectBenchmarkingUtils::GenerateDeterministicSetOfInstances(CustomizableObject.Get(),InstancesToGeneratePerState,InstancesToUpdate,TotalScheduledUpdates))
+	{
+		UE_LOG(LogMutable,Error, TEXT("Mutable Customizable Object Instance generation failed. Aborting perf test."))
+		return FReply::Handled();
+	}
+	check (!InstancesToUpdate.IsEmpty())
+
+	// tell the system a test is already running so no new test can be invoked while that is being done
+	bIsPerformanceAnalyzerRunning = true;
+	ScheduleNextInstanceUpdate();
+	
+	return FReply::Handled();
+}
+
+
+FText SCustomizableObjectEditorPerformanceAnalyzer::GetGenerateInstancesButtonText() const
+{
+	if (CustomizableObject && !CustomizableObject->IsCompiled())
+	{
+		return FText::FromString(TEXT("CO is not compiled"));
+	}
+	
+	if (bIsPerformanceAnalyzerRunning && !AreUpdatesPending())
+	{
+		return FText::FromString(TEXT("Another instance is running"));
+	}
+	else
+	{
+		if (!AreUpdatesPending())
+		{
+			return FText::FromString(TEXT("Generate Random Instances"));
+		}
+		else
+		{
+			return FText::FromString(FString::Printf(TEXT("Updated %u / %u instances"),UpdatedInstancesCount, TotalScheduledUpdates));;
+		}
+	}
+
+	return FText::FromString(TEXT(""));
+}
+
+
+EVisibility SCustomizableObjectEditorPerformanceAnalyzer::ShouldInstanceUpdateAbortButtonBeVisible() const
+{
+	return  AreUpdatesPending() ? EVisibility::Visible :  EVisibility::Hidden;
+}
+
+
+FText SCustomizableObjectEditorPerformanceAnalyzer::GetStopUpdatesButtonText() const
+{
+	// Already cleared the pending updates and we are waiting for the last one
+	if (InstancesToUpdate.IsEmpty() && CurrentInstance)
+	{
+		// Stopping
+		return FText::FromString(TEXT("STOPPING..."));
+	}
+	else
+	{
+		return FText::FromString(TEXT("STOP"));
+	}
+}
+
+
+FReply SCustomizableObjectEditorPerformanceAnalyzer::OnInstanceUpdateAbortButtonClicked()
+{
+	ClearPendingUpdatesQueue();
+	bIsPerformanceAnalyzerRunning = false;
+	
+	// The update finish of the last instance to be updated will take care of clearing the update data (like cvars)
+	return FReply::Handled();
+}
+
+
+bool SCustomizableObjectEditorPerformanceAnalyzer::IsInstanceUpdateAbortButtonEnabled() const
+{
+	// Allow the abortion of the test if one slate is running the test and that slate is us (since we have Instances to yet update)
+	return !InstancesToUpdate.IsEmpty();
+}
+
+
+void SCustomizableObjectEditorPerformanceAnalyzer::ClearPendingUpdatesQueue()
+{
+	while(InstancesToUpdate.Dequeue()) {}
+}
+
+
+void SCustomizableObjectEditorPerformanceAnalyzer::CacheCustomizableObjectModelData()
+{
+	if (!CustomizableObject)
+	{
+		return;
+	}
+	
+	UCustomizableObjectPrivate* CustomizableObjectPrivate = CustomizableObject->GetPrivate();
+	check(CustomizableObjectPrivate);
+	
+	// Store the compilation options once the compilation has been performed. This way we should be able to compare the options
+	// with what we know are the CIS ones
+	LastCompilationOptions = MakeShared<FCompilationOptions>(CustomizableObjectPrivate->GetCompileOptions());
+}
+
+
+void SCustomizableObjectEditorPerformanceAnalyzer::OnCustomizableObjectCompilationFinished()
+{
+	check(CustomizableObject);
+
+	// If updates are pending, since only one instance of this slate can run, mark the execution status of this instance
+	// as halted. Only do this for instances of the slate where the execution has been possible.
+	if (AreUpdatesPending())
+	{
+		bIsPerformanceAnalyzerRunning = false;
+	}
+	
+	// Clear any old data from previous compilation operation
+	//		This will not always happen since we prevent the compilation of the CO if it has instances being updated
+	//		if no CO is actually compiled then this method will not be called
+	{
+		ClearPendingUpdatesQueue();
+		CurrentInstance = nullptr;
+		if (!InstanceUpdateElements.IsEmpty())
+		{
+			// Clear the list of update data
+			InstanceUpdateElements.Empty();
+			InstanceUpdatesListView->RequestListRefresh();
+		}
+	}
+	
+	CacheCustomizableObjectModelData();
+}
+
+
+void SCustomizableObjectEditorPerformanceAnalyzer::SetMutableBenchmarkingSystemState(const bool bNewState)
+{
+	// Tell the benchmarking system that there is no need to continue benchmarking.
+	FLogBenchmarkUtil::SetBenchmarkReportingStateOverride(bNewState);
+
+	// Setting this to false will allow for the system to retrieve the values set by CVars and other non code ways
+	UCustomizableObjectSystem::SetBenchmarkState(bNewState);
+}
+
+
+void SCustomizableObjectEditorPerformanceAnalyzer::AddReferencedObjects(FReferenceCollector& Collector)
+{
+	//Collector.AddReferencedObject(CustomizableObject);
+	Collector.AddReferencedObject(System);
+}
+
+
+FString SCustomizableObjectEditorPerformanceAnalyzer::GetReferencerName() const
+{
+	return FString(TEXT("SCustomizableObjectEditorPerformanceAnalyzer"));
+}
+
+
+EColumnSortMode::Type SCustomizableObjectEditorPerformanceAnalyzer::GetColumnSortMode(FName ColumnName) const
+{
+	if (CurrentSortColumn != ColumnName)
+	{
+		return EColumnSortMode::None;
+	}
+
+	return SortMode;
+}
+
+
+TSharedRef<ITableRow> SCustomizableObjectEditorPerformanceAnalyzer::OnGenerateInstanceUpdateRow(TSharedPtr<FInstanceUpdateDataElement> InstanceUpdateDataElement, const TSharedRef<STableViewBase>& TableViewBase)
+{
+	TSharedRef<SInstanceUpdateDataRow> Row = SNew(SInstanceUpdateDataRow, TableViewBase, InstanceUpdateDataElement);
+	return Row;
+}
+
+
+void SCustomizableObjectEditorPerformanceAnalyzer::OnInstancesRowClick(TSharedPtr<FInstanceUpdateDataElement> InstanceUpdateDataElement)
+{
+	if (InstanceUpdateDataElement && InstanceUpdateDataElement->Instance)
+	{
+		GEditor->GetEditorSubsystem<UAssetEditorSubsystem>()->OpenEditorForAsset(InstanceUpdateDataElement->Instance.Get());
+	}
+}
+
+
+void SCustomizableObjectEditorPerformanceAnalyzer::OnInstanceUpdateListViewSort(EColumnSortPriority::Type ColumnPriority, const FName& ColumnId, EColumnSortMode::Type NewSortMode)
+{
+	CurrentSortColumn = ColumnId;
+	SortMode = NewSortMode;
+
+	if (ColumnId == InstanceUpdatesMainDataColumns::InstanceNameColumnID)
+	{
+		InstanceUpdateElements.StableSort([&](const TSharedPtr<FInstanceUpdateDataElement>& ElementA, const TSharedPtr<FInstanceUpdateDataElement>& ElementB)
+		{
+			if (ElementA.IsValid() && ElementB.IsValid())
+			{
+				return NewSortMode == EColumnSortMode::Ascending ? ElementA->UpdateIndex < ElementB->UpdateIndex : ElementA->UpdateIndex > ElementB->UpdateIndex;
+			}
+
+			return ElementA.IsValid();
+		});
+	}
+	else if (ColumnId == InstanceUpdatesMainDataColumns::InstanceUpdateResultColumnID)
+	{
+		InstanceUpdateElements.StableSort([&](const TSharedPtr<FInstanceUpdateDataElement>& ElementA, const TSharedPtr<FInstanceUpdateDataElement>& ElementB)
+		{
+			if (ElementA.IsValid() && ElementB.IsValid())
+			{
+				const FString NameA = UEnum::GetValueAsString(ElementA->UpdateStats.UpdateResult);
+				const FString NameB =  UEnum::GetValueAsString(ElementB->UpdateStats.UpdateResult);
+
+				return NewSortMode == EColumnSortMode::Ascending ? NameA < NameB : NameA > NameB;
+			}
+
+			return ElementA.IsValid();
+		});
+	}
+	else if (ColumnId == InstanceUpdatesMainDataColumns::QueueTimeColumnID)
+	{
+		InstanceUpdateElements.StableSort([&](const TSharedPtr<FInstanceUpdateDataElement>& ElementA, const TSharedPtr<FInstanceUpdateDataElement>& ElementB)
+		{
+			if (ElementA.IsValid() && ElementB.IsValid())
+			{
+				const double ValA = ElementA->UpdateStats.QueueTime;
+				const double ValB =  ElementB->UpdateStats.QueueTime;
+
+				return NewSortMode == EColumnSortMode::Ascending ? ValA < ValB : ValA > ValB;
+			}
+
+			return ElementA.IsValid();
+		});
+	}
+	else if (ColumnId == InstanceUpdatesMainDataColumns::UpdateTimeColumnID)
+	{
+		InstanceUpdateElements.StableSort([&](const TSharedPtr<FInstanceUpdateDataElement>& ElementA, const TSharedPtr<FInstanceUpdateDataElement>& ElementB)
+		{
+			if (ElementA.IsValid() && ElementB.IsValid())
+			{
+				const double ValA = ElementA->UpdateStats.UpdateTime;
+				const double ValB =  ElementB->UpdateStats.UpdateTime;
+
+				return NewSortMode == EColumnSortMode::Ascending ? ValA < ValB : ValA > ValB;
+			}
+
+			return ElementA.IsValid();
+		});
+	}
+	else if (ColumnId == InstanceUpdatesMainDataColumns::GetMeshTimeColumnID)
+	{
+		InstanceUpdateElements.StableSort([&](const TSharedPtr<FInstanceUpdateDataElement>& ElementA, const TSharedPtr<FInstanceUpdateDataElement>& ElementB)
+		{
+			if (ElementA.IsValid() && ElementB.IsValid())
+			{
+				const double ValA = ElementA->UpdateStats.TaskGetMeshTime;
+				const double ValB =  ElementB->UpdateStats.TaskGetMeshTime;
+
+				return NewSortMode == EColumnSortMode::Ascending ? ValA < ValB : ValA > ValB;
+			}
+
+			return ElementA.IsValid();
+		});
+	}
+	else if (ColumnId == InstanceUpdatesMainDataColumns::LockCacheTimeColumnID)
+	{
+		InstanceUpdateElements.StableSort([&](const TSharedPtr<FInstanceUpdateDataElement>& ElementA, const TSharedPtr<FInstanceUpdateDataElement>& ElementB)
+		{
+			if (ElementA.IsValid() && ElementB.IsValid())
+			{
+				const double ValA = ElementA->UpdateStats.TaskLockCacheTime;
+				const double ValB =  ElementB->UpdateStats.TaskLockCacheTime;
+
+				return NewSortMode == EColumnSortMode::Ascending ? ValA < ValB : ValA > ValB;
+			}
+
+			return ElementA.IsValid();
+		});
+	}
+	else if (ColumnId == InstanceUpdatesMainDataColumns::GetImagesTimeColumnID)
+	{
+		InstanceUpdateElements.StableSort([&](const TSharedPtr<FInstanceUpdateDataElement>& ElementA, const TSharedPtr<FInstanceUpdateDataElement>& ElementB)
+		{
+			if (ElementA.IsValid() && ElementB.IsValid())
+			{
+				const double ValA = ElementA->UpdateStats.TaskGetImagesTime;
+				const double ValB =  ElementB->UpdateStats.TaskGetImagesTime;
+
+				return NewSortMode == EColumnSortMode::Ascending ? ValA < ValB : ValA > ValB;
+			}
+
+			return ElementA.IsValid();
+		});
+	}
+	else if (ColumnId == InstanceUpdatesMainDataColumns::ConvertResourcesTimeColumnID)
+	{
+		InstanceUpdateElements.StableSort([&](const TSharedPtr<FInstanceUpdateDataElement>& ElementA, const TSharedPtr<FInstanceUpdateDataElement>& ElementB)
+		{
+			if (ElementA.IsValid() && ElementB.IsValid())
+			{
+				const double ValA = ElementA->UpdateStats.TaskConvertResourcesTime;
+				const double ValB =  ElementB->UpdateStats.TaskConvertResourcesTime;
+
+				return NewSortMode == EColumnSortMode::Ascending ? ValA < ValB : ValA > ValB;
+			}
+
+			return ElementA.IsValid();
+		});
+	}
+	else if (ColumnId == InstanceUpdatesMainDataColumns::CallbacksTimeColumnID)
+	{
+		InstanceUpdateElements.StableSort([&](const TSharedPtr<FInstanceUpdateDataElement>& ElementA, const TSharedPtr<FInstanceUpdateDataElement>& ElementB)
+		{
+			if (ElementA.IsValid() && ElementB.IsValid())
+			{
+				const double ValA = ElementA->UpdateStats.TaskCallbacksTime;
+				const double ValB =  ElementB->UpdateStats.TaskCallbacksTime;
+
+				return NewSortMode == EColumnSortMode::Ascending ? ValA < ValB : ValA > ValB;
+			}
+
+			return ElementA.IsValid();
+		});
+	}
+	else if (ColumnId == InstanceUpdatesMainDataColumns::UpdatePeakMemoryColumnID)
+	{
+		InstanceUpdateElements.StableSort([&](const TSharedPtr<FInstanceUpdateDataElement>& ElementA, const TSharedPtr<FInstanceUpdateDataElement>& ElementB)
+		{
+			if (ElementA.IsValid() && ElementB.IsValid())
+			{
+				const double ValA = ElementA->UpdateStats.UpdatePeakMemory;
+				const double ValB =  ElementB->UpdateStats.UpdatePeakMemory;
+
+				return NewSortMode == EColumnSortMode::Ascending ? ValA < ValB : ValA > ValB;
+			}
+
+			return ElementA.IsValid();
+		});
+	}
+	else if (ColumnId == InstanceUpdatesMainDataColumns::UpdateRealPeakMemoryColumnID)
+	{
+		InstanceUpdateElements.StableSort([&](const TSharedPtr<FInstanceUpdateDataElement>& ElementA, const TSharedPtr<FInstanceUpdateDataElement>& ElementB)
+		{
+			if (ElementA.IsValid() && ElementB.IsValid())
+			{
+				const double ValA = ElementA->UpdateStats.UpdateRealPeakMemory;
+				const double ValB =  ElementB->UpdateStats.UpdateRealPeakMemory;
+
+				return NewSortMode == EColumnSortMode::Ascending ? ValA < ValB : ValA > ValB;
+			}
+
+			return ElementA.IsValid();
+		});
+	}
+	else if (ColumnId == InstanceUpdatesMainDataColumns::InstanceUpdateTypeID)
+	{
+		InstanceUpdateElements.StableSort([&](const TSharedPtr<FInstanceUpdateDataElement>& ElementA, const TSharedPtr<FInstanceUpdateDataElement>& ElementB)
+		{
+			if (ElementA.IsValid() && ElementB.IsValid())
+			{
+				const FString NameA = ElementA->UpdateStats.UpdateType;
+				const FString NameB =  ElementB->UpdateStats.UpdateType;
+
+				return NewSortMode == EColumnSortMode::Ascending ? NameA < NameB : NameA > NameB;
+			}
+
+			return ElementA.IsValid();
+		});
+	}
+	else
+	{
+		check(false); // Unknown method.
+	}
+
+	InstanceUpdatesListView->RequestListRefresh();
+}
+
+
+void SCustomizableObjectEditorPerformanceAnalyzer::OnBenchmarkMeshUpdated(TSharedRef<FUpdateContextPrivate> UpdateContextPrivate, FInstanceUpdateStats UpdateStats)
+{
+	if (CurrentInstance.Get() != UpdateContextPrivate->Instance)
+	{
+		return;
+	}
+	
+	// Create a new element for the mesh update data for this instance
+	TSharedPtr<FInstanceUpdateDataElement> InstanceElement = MakeShared<FInstanceUpdateDataElement>();
+	InstanceElement->Instance = CurrentInstance;
+	InstanceElement->UpdateStats = UpdateStats;
+	InstanceElement->UpdateIndex = UpdatedInstancesCount++;
+
+	// 3) Store the element to keep the instance alive
+	check(InstanceElement);
+	InstanceUpdateElements.Push(InstanceElement);
+	
+	// Reset control flags for the next instance
+	CurrentInstance = nullptr;
+	
+	// And once the updates have completed update the list view
+	check(InstanceUpdatesListView);
+	InstanceUpdatesListView->RequestListRefresh();
+	
+	// Stop listening for instance update reported data
+	System->GetPrivate()->LogBenchmarkUtil.OnMeshUpdateReported.Remove(OnMeshUpdatePerfReportDelegateHandle);
+
+	// Handle here the waiting for the update of the next instances
+	if (!InstancesToUpdate.IsEmpty())
+	{
+		ScheduleNextInstanceUpdate();
+	}
+	else
+	{
+		SetMutableBenchmarkingSystemState(false);
+		bIsPerformanceAnalyzerRunning = false;
+	}
+}
+
+
+void SCustomizableObjectEditorPerformanceAnalyzer::ScheduleNextInstanceUpdate()
+{
+	// We always expect to find an instance here after the Dequeue operation
+	InstancesToUpdate.Dequeue(CurrentInstance);
+	check(CurrentInstance)
+	
+	// Enable the benchmarking support mutable systems
+	SetMutableBenchmarkingSystemState(true);
+
+	// Request the actual update
+	FInstanceUpdateNativeDelegate InstanceUpdateNativeDelegate;
+	CurrentInstance->UpdateSkeletalMeshAsyncResult(InstanceUpdateNativeDelegate, true, true);
+
+	// Bind our listView update to the update of the benchmarking tool
+	OnMeshUpdatePerfReportDelegateHandle = System->GetPrivate()->LogBenchmarkUtil.OnMeshUpdateReported.AddSP(this, &SCustomizableObjectEditorPerformanceAnalyzer::OnBenchmarkMeshUpdated);
+}
+
+
+bool SCustomizableObjectEditorPerformanceAnalyzer::AreUpdatesPending() const
+{
+	return CurrentInstance || !InstancesToUpdate.IsEmpty();
+}
+
+#undef LOCTEXT_NAMESPACE
diff --git a/Engine/Plugins/Experimental/Mutable/Source/CustomizableObjectEditor/Private/MuCOE/CustomizableObjectEditorPerformanceAnalyzer.h b/Engine/Plugins/Experimental/Mutable/Source/CustomizableObjectEditor/Private/MuCOE/CustomizableObjectEditorPerformanceAnalyzer.h
new file mode 100644
--- /dev/null
+++ b/Engine/Plugins/Experimental/Mutable/Source/CustomizableObjectEditor/Private/MuCOE/CustomizableObjectEditorPerformanceAnalyzer.h
@@ -0,0 +1,170 @@
+﻿// Copyright Epic Games, Inc. All Rights Reserved.
+#pragma once
+
+#include "Containers/SpscQueue.h"
+#include "MuCO/CustomizableObjectCompilerTypes.h"
+#include "MuCO/LogBenchmarkUtil.h"
+#include "Widgets/SCompoundWidget.h"
+#include "UObject/StrongObjectPtr.h"
+#include "Widgets/Views/SListView.h"
+#include "Widgets/Views/STableRow.h"
+
+
+
+struct FCompilationRequest;
+struct FCompilationOptions;
+struct InstanceUpdateData;
+class SObjectPropertyEntryBox;
+class UCustomizableObject;
+class UCustomizableObjectInstance;
+class FCustomizableObjectEditor;
+class FCustomizableObjectInstanceEditor;
+class UCustomizableObjectSystem;
+enum class ECustomizableObjectTextureCompression : uint8;
+
+/**
+ * Backend object that represents the update of one instance. Currently only representes the data for the initial geneation of the instance update
+ */
+class FInstanceUpdateDataElement 
+{
+public:
+	/** Instance whose data this object exposes */
+	TStrongObjectPtr<UCustomizableObjectInstance> Instance = nullptr;
+
+	/** Index to aid in the sorting of the FInstanceUpdateDataElement in the SListView that contains them */
+	uint16 UpdateIndex = 0;
+	
+	/** Container with all the perf data in relation to the instance this object represents */
+	FInstanceUpdateStats UpdateStats;
+};
+
+
+/**
+ * Slate that exposes the perf data of one single instance update.
+ */
+class SInstanceUpdateDataRow : public SMultiColumnTableRow<TSharedPtr<FInstanceUpdateDataElement>>
+{
+public:
+	
+	void Construct(const FArguments& InArgs, const TSharedRef<STableViewBase>& InOwnerTableView, TSharedPtr<FInstanceUpdateDataElement> Element);
+	
+private:
+	
+	/** Backend object whose data we are drawing */
+	TSharedPtr<FInstanceUpdateDataElement> InstanceUpdateElement;
+	
+	virtual TSharedRef<SWidget> GenerateWidgetForColumn(const FName& InColumnName) override;
+};
+
+
+/**
+ * Slate that exposes perf data of a deterministic set of instances generated and then updated in the context of the editor.
+ * Will update the values of some CVARS that will, at the end, be restored.
+ */
+class SCustomizableObjectEditorPerformanceAnalyzer : public SCompoundWidget, public FGCObject
+{
+public:
+	SLATE_BEGIN_ARGS(SCustomizableObjectEditorPerformanceAnalyzer){}
+		SLATE_ARGUMENT(UCustomizableObject*,CustomizableObject)
+	SLATE_END_ARGS()
+
+	void Construct(const FArguments& InArgs);
+	virtual ~SCustomizableObjectEditorPerformanceAnalyzer() override;
+
+private:
+
+	// FGCObject
+	virtual void AddReferencedObjects(FReferenceCollector& Collector) override;
+	virtual FString GetReferencerName() const override;
+	// FGCObject
+	
+	EVisibility GetVisibilityForBenchmarkingSettingsMessage() const;
+
+	/** Instance amount to generate and update controls */
+	void OnTargetAmountOfInstancesValueChange(uint32 NewValue);
+	void OnTargetAmountOfInstancesValueCommited(uint32 NewValue, ETextCommit::Type CommitType);
+	TOptional<uint32> OnTargetAmountOfInstancesValueRequested() const;
+
+	/** Random Instances generation and Update button */
+	FReply OnInstanceUpdateButtonClicked();
+	bool IsInstanceUpdateButtonEnabled() const;
+	FText GetGenerateInstancesButtonText() const;
+
+	/** Methods for the button that allows the user to cancel the execution of the test */
+	FReply OnInstanceUpdateAbortButtonClicked();
+	EVisibility ShouldInstanceUpdateAbortButtonBeVisible() const;
+	bool IsInstanceUpdateAbortButtonEnabled() const;
+	FText GetStopUpdatesButtonText() const;
+
+	/** Utility methods */
+	void ClearPendingUpdatesQueue();
+
+	/** Compilation handling */
+	void OnCustomizableObjectCompilationFinished();
+	void CacheCustomizableObjectModelData();
+
+	/**
+	 * Enable or disable the mutable benchmarking settings.
+	 * @param bNewState Set it to true to enable the benchmarking configurations and false to disable them.
+	 */
+	void SetMutableBenchmarkingSystemState(const bool bNewState);
+	
+	/** ListView methods */
+	TSharedRef<ITableRow> OnGenerateInstanceUpdateRow(TSharedPtr<FInstanceUpdateDataElement> InstanceUpdateDataElement, const TSharedRef<STableViewBase>& TableViewBase);
+	void OnInstancesRowClick(TSharedPtr<FInstanceUpdateDataElement> InstanceUpdateDataElement);
+	void OnInstanceUpdateListViewSort(EColumnSortPriority::Type ColumnPriority, const FName& ColumnId, EColumnSortMode::Type NewSortMode);
+	EColumnSortMode::Type GetColumnSortMode(FName ColumnName) const;
+
+	/**
+	 * Callback invoked each time the LogBenchmarkUtil notifies us that an initial generation update (mesh Update) has been produced
+	 * @param UpdateContextPrivate The context of the instance that got updated
+	 * @param UpdateStats The Perf data of the instance that was updated.
+	 */
+	void OnBenchmarkMeshUpdated(TSharedRef<FUpdateContextPrivate> UpdateContextPrivate, FInstanceUpdateStats UpdateStats);
+
+	/** Grab the next available instance to be updated*/
+	void ScheduleNextInstanceUpdate();
+
+	/** Determines if there are updates to wait for. Checks not only the queue but also the current instnace being udpated*/
+	bool AreUpdatesPending() const;
+	
+private:
+
+	/** Delegate pointing to the LogBenchmarkUtil on Mesh Updated event */
+	FDelegateHandle OnMeshUpdatePerfReportDelegateHandle;
+
+	/** Is this slate running? Used to determine if another instance of the slate is performing the test */
+	static bool bIsPerformanceAnalyzerRunning;
+
+	/** The amount of instances per state we want to generate and then update */
+	uint32 InstancesToGeneratePerState = 24;
+
+	/** Cached Customizable Object System object */
+	TObjectPtr<UCustomizableObjectSystem> System;
+	
+	/** CO whose data we are we reporting. */
+	TStrongObjectPtr<UCustomizableObject> CustomizableObject;
+	
+	/** List view showing all the instance update data elements */
+	TSharedPtr<SListView<TSharedPtr<FInstanceUpdateDataElement>>> InstanceUpdatesListView;
+	
+	/** Instance update elements (size == RandomInstancesToGenerate) at the time of the update request */
+	TArray<TSharedPtr<FInstanceUpdateDataElement>> InstanceUpdateElements;
+
+	/** The name of the last column the user wanted to sort */
+	FName CurrentSortColumn;
+	/** The sort type applied to the last column that got sorted */
+	EColumnSortMode::Type SortMode = EColumnSortMode::Type::None;
+
+	/** Updating process variables */
+	TSpscQueue<TStrongObjectPtr<UCustomizableObjectInstance>> InstancesToUpdate;
+	/** Current instance being updated */
+	TStrongObjectPtr<UCustomizableObjectInstance> CurrentInstance = nullptr;
+	
+	/** Control counters (UI) */
+	uint32 UpdatedInstancesCount = 0;
+	uint32 TotalScheduledUpdates = 0;
+
+	/** The compilation options the last time we compiled the CO */
+	TSharedPtr<FCompilationOptions> LastCompilationOptions = nullptr;
+};
diff --git a/Engine/Plugins/Experimental/Mutable/Source/CustomizableObjectEditor/Private/MuCOE/CustomizableObjectEditorStyle.cpp b/Engine/Plugins/Experimental/Mutable/Source/CustomizableObjectEditor/Private/MuCOE/CustomizableObjectEditorStyle.cpp
--- a/Engine/Plugins/Experimental/Mutable/Source/CustomizableObjectEditor/Private/MuCOE/CustomizableObjectEditorStyle.cpp
+++ b/Engine/Plugins/Experimental/Mutable/Source/CustomizableObjectEditor/Private/MuCOE/CustomizableObjectEditorStyle.cpp
@@ -88,6 +88,8 @@ TSharedRef<FSlateStyleSet> FCustomizableObjectEditorStyle::Create()
 	Style->Set("CustomizableObjectEditor.TextureAnalyzer.Small", new IMAGE_BRUSH_SVG(Style, "Editor/Slate/Starship/Common/IssueTracker", Icon20x20));
 	Style->Set("CustomizableObjectEditor.PerformanceReport", new IMAGE_BRUSH_SVG(Style, "Editor/Slate/Starship/MainToolbar/simulate", Icon40x40));
 	Style->Set("CustomizableObjectEditor.PerformanceReport.Small", new IMAGE_BRUSH_SVG(Style, "Editor/Slate/Starship/MainToolbar/simulate", Icon20x20));
+	Style->Set("CustomizableObjectEditor.PerformanceAnalyzer", new IMAGE_BRUSH_SVG(Style, "Editor/Slate/Starship/MainToolbar/simulate", Icon40x40));
+	Style->Set("CustomizableObjectEditor.PerformanceAnalyzer.Small", new IMAGE_BRUSH_SVG(Style, "Editor/Slate/Starship/MainToolbar/simulate", Icon20x20));
 
 	Style->Set("CustomizableObjectEditor.CompileGatherReferences", new IMAGE_BRUSH_SVG(Style, "Editor/Slate/Starship/MainToolbar/compile", Icon40x40));
 	Style->Set("CustomizableObjectEditor.CompileGatherReferences.Small", new IMAGE_BRUSH_SVG(Style, "Editor/Slate/Starship/MainToolbar/compile", Icon20x20));
diff --git a/Engine/Plugins/Experimental/Mutable/Source/CustomizableObjectEditor/Public/MuCOE/CustomizableObjectEditorActions.h b/Engine/Plugins/Experimental/Mutable/Source/CustomizableObjectEditor/Public/MuCOE/CustomizableObjectEditorActions.h
--- a/Engine/Plugins/Experimental/Mutable/Source/CustomizableObjectEditor/Public/MuCOE/CustomizableObjectEditorActions.h
+++ b/Engine/Plugins/Experimental/Mutable/Source/CustomizableObjectEditor/Public/MuCOE/CustomizableObjectEditorActions.h
@@ -21,6 +21,7 @@ public:
 	TSharedPtr< FUICommandInfo > Debug;
 
 	TSharedPtr< FUICommandInfo > PerformanceReport;
+	TSharedPtr< FUICommandInfo > PerformanceAnalyzer;
 	TSharedPtr< FUICommandInfo > ResetPerformanceReportOptions;
 	TSharedPtr< FUICommandInfo > TextureAnalyzer;
 
diff --git a/Engine/Plugins/Experimental/Mutable/Source/MutableValidation/Private/MuV/COIBakingTestCommandlet.cpp b/Engine/Plugins/Experimental/Mutable/Source/MutableValidation/Private/MuV/COIBakingTestCommandlet.cpp
--- a/Engine/Plugins/Experimental/Mutable/Source/MutableValidation/Private/MuV/COIBakingTestCommandlet.cpp
+++ b/Engine/Plugins/Experimental/Mutable/Source/MutableValidation/Private/MuV/COIBakingTestCommandlet.cpp
@@ -13,6 +13,7 @@
 #include "HAL/FileManager.h"
 #include "Interfaces/ITargetPlatformManagerModule.h"
 #include "MuCO/CustomizableObjectPrivate.h"
+#include "MuCOE/CustomizableObjectBenchmarkingUtils.h"
 
 
 /** Flag useful to know if we are currently updating an instance or not */
@@ -36,6 +37,8 @@ void OnInstanceUpdate(const FUpdateContext& Result)
 
 int32 UCOIBakingTestCommandlet::Main(const FString& Params)
 {
+	UCustomizableObjectSystem::SetBenchmarkState(true);
+	
 	// Ensure we do not show any OK dialog since we are not an user that can interact with them
 	GIsRunningUnattendedScript = true;
 	
diff --git a/Engine/Plugins/Experimental/Mutable/Source/MutableValidation/Private/MuV/COIBulkUpdateTestCommandlet.cpp b/Engine/Plugins/Experimental/Mutable/Source/MutableValidation/Private/MuV/COIBulkUpdateTestCommandlet.cpp
--- a/Engine/Plugins/Experimental/Mutable/Source/MutableValidation/Private/MuV/COIBulkUpdateTestCommandlet.cpp
+++ b/Engine/Plugins/Experimental/Mutable/Source/MutableValidation/Private/MuV/COIBulkUpdateTestCommandlet.cpp
@@ -11,10 +11,15 @@
 #include "MuCO/CustomizableObjectInstance.h"
 #include "Interfaces/ITargetPlatformManagerModule.h"
 #include "MuCO/CustomizableObjectPrivate.h"
+#include "MuCO/CustomizableObjectSystem.h"
+#include "MuCOE/CustomizableObjectBenchmarkingUtils.h"
 
 
 int32 UCOIBulkUpdateTestCommandlet::Main(const FString& Params)
 {
+	// Ensure we have the cvars used for our testing set
+	UCustomizableObjectSystem::SetBenchmarkState(true);
+	
 	// Ensure we do not show any OK dialog since we are not an user that can interact with them
 	GIsRunningUnattendedScript = true;
 
@@ -141,12 +146,9 @@ int32 UCOIBulkUpdateTestCommandlet::Main(const FString& Params)
 		const FString CustomizableObjectName = CustomizableObject->GetName();
 		
 		// Set the compilation platform based on what the system is currently running on
-		FCompilationOptions CompilationOptions = CustomizableObject->GetPrivate()->GetCompileOptions();
+		FCompilationOptions CompilationOptions = GetCompilationOptionsForBenchmarking(*CustomizableObject);
 		CompilationOptions.TargetPlatform = TargetCompilationPlatform;
-		CompilationOptions.bUseDiskCompilation = false;
-		CompilationOptions.OptimizationLevel = UE_MUTABLE_MAX_OPTIMIZATION;
-		CompilationOptions.bSilentCompilation = false;
-
+		
 		// Compile the current CO object
 		if (!CompilationUtility->CompileCustomizableObject(CustomizableObject, false, &CompilationOptions))	// Do not log mutable data since mongoDB will not be able to handle it correctly 
 		{
diff --git a/Engine/Plugins/Experimental/Mutable/Source/MutableValidation/Private/MuV/CustomizableObjectValidationCommandlet.cpp b/Engine/Plugins/Experimental/Mutable/Source/MutableValidation/Private/MuV/CustomizableObjectValidationCommandlet.cpp
--- a/Engine/Plugins/Experimental/Mutable/Source/MutableValidation/Private/MuV/CustomizableObjectValidationCommandlet.cpp
+++ b/Engine/Plugins/Experimental/Mutable/Source/MutableValidation/Private/MuV/CustomizableObjectValidationCommandlet.cpp
@@ -16,6 +16,7 @@
 #include "MuR/Model.h"
 #include "Interfaces/ITargetPlatformManagerModule.h"
 #include "Interfaces/ITargetPlatform.h"
+#include "MuCOE/CustomizableObjectBenchmarkingUtils.h"
 #include "Serialization/MemoryWriter.h"
 
 
@@ -26,6 +27,9 @@ int32 UCustomizableObjectValidationCommandlet::Main(const FString& Params)
 	// Execution arguments for commandlet from IDE
 	// -run=CustomizableObjectValidation -CustomizableObject=(PathToCO)
 
+	// Ensure we have the cvars used for our testing set
+	UCustomizableObjectSystem::SetBenchmarkState(true);
+	
 	// Ensure we do not show any OK dialog since we are not an user that can interact with them
 	GIsRunningUnattendedScript = true;
 	
@@ -80,13 +84,10 @@ int32 UCustomizableObjectValidationCommandlet::Main(const FString& Params)
 	bool bWasCoCompilationSuccessful = false;
 	{
 		LLM_SCOPE_BYNAME(TEXT("CustomizableObjectValidationCommandlet/Compile"));
-		
-		// Override some configurations that may have been changed by the user
-		FCompilationOptions CompilationOptions = ToTestCustomizableObject->GetPrivate()->GetCompileOptions();
-		CompilationOptions.bSilentCompilation = false;
-		CompilationOptions.OptimizationLevel = UE_MUTABLE_MAX_OPTIMIZATION;
-		CompilationOptions.TextureCompression = ECustomizableObjectTextureCompression::Fast;
 
+		// Override some configurations that may have been changed by the user
+		FCompilationOptions CompilationOptions = GetCompilationOptionsForBenchmarking(*ToTestCustomizableObject);
+		
 		// Set the target compilation platform based on what the caller wants
 		CompilationOptions.TargetPlatform = TargetCompilationPlatform;
 
@@ -159,48 +160,8 @@ int32 UCustomizableObjectValidationCommandlet::Main(const FString& Params)
 	{
 		LLM_SCOPE_BYNAME(TEXT("CustomizableObjectValidationCommandlet/GenerateInstances"));
 		
-		// Test this parameter configuration in all the states of the CO
-		const uint32 StateCount = ToTestCustomizableObject->GetStateCount();
-		check(StateCount >= 1);
-
-		UE_LOG(LogMutable, Display, TEXT("Requeted Instances Count : %i"), InstancesToGenerate);
-		UE_LOG(LogMutable, Display, TEXT("State Count = %i"), StateCount);
-	
-		// Compute actual total amount of instances to generate
-		const uint32 TotalInstancesToTestCount = InstancesToGenerate * StateCount;
-		UE_LOG(LogMutable,Display,TEXT("Generating %i instances (states * requested instances)..."), TotalInstancesToTestCount);
-		
-		// Create randomization stream for the parameters of the instance
-		FRandomStream RandomizationStream = FRandomStream(0);
-		
-		// Generate a series of instances to later update
-		for (uint32 CustomizableObjectInstanceIndex = 0; CustomizableObjectInstanceIndex < InstancesToGenerate; CustomizableObjectInstanceIndex++)
-		{
-			UCustomizableObjectInstance* GeneratedInstance = ToTestCustomizableObject->CreateInstance();
-			if (GeneratedInstance)
-			{
-				// Force generation of all LODS
-				TArray<uint16> RequestedLodLevels{};
-				RequestedLodLevels.Init(0, GeneratedInstance->GetNumComponents());
-				GeneratedInstance->GetPrivate()->GetDescriptor().SetRequestedLODLevels(RequestedLodLevels);
-				
-				// Randomize instance values
-				GeneratedInstance->SetRandomValuesFromStream(RandomizationStream);
-				
-				for (uint32 State = 0; State < StateCount; State++)
-				{
-					// Set the state for the instance and store it for later update.
-					GeneratedInstance->GetPrivate()->SetState(State);
-                	InstancesToProcess.Enqueue(GeneratedInstance->Clone());
-					GeneratedInstances++;
-				}
-			}
-			else
-			{
-				UE_LOG(LogMutable,Error,TEXT("Failed to generate COI for the %s CO."),*ToTestCustomizableObject->GetName());
-				bWasInstancesCreationSuccessful = false;
-			}
-		}
+		// Create a set of instances so we can later test them out
+		bWasInstancesCreationSuccessful = CustomizableObjectBenchmarkingUtils::GenerateDeterministicSetOfInstances(ToTestCustomizableObject,InstancesToGenerate, InstancesToProcess, GeneratedInstances);
 	}
 	// ---------------------------------------------------------------------------------------------------------- //
 
diff --git a/Engine/Plugins/Experimental/Mutable/Source/MutableValidation/Private/MuV/ValidationUtils.cpp b/Engine/Plugins/Experimental/Mutable/Source/MutableValidation/Private/MuV/ValidationUtils.cpp
--- a/Engine/Plugins/Experimental/Mutable/Source/MutableValidation/Private/MuV/ValidationUtils.cpp
+++ b/Engine/Plugins/Experimental/Mutable/Source/MutableValidation/Private/MuV/ValidationUtils.cpp
@@ -7,6 +7,9 @@
 #include "MuCO/CustomizableObject.h"
 #include "MuCO/CustomizableObjectSystem.h"
 #include "RHIGlobals.h"
+#include "MuCO/CustomizableObjectPrivate.h"
+#include "MuCOE/CustomizableObjectBenchmarkingUtils.h"
+#include "UObject/ObjectPtr.h"
 
 void PrepareAssetRegistry()
 {
@@ -58,4 +61,16 @@ void Wait(const double ToWaitSeconds)
 }
 
 
+FCompilationOptions GetCompilationOptionsForBenchmarking (UCustomizableObject& ReferenceCustomizableObject)
+{
+	// Override some configurations that may have been changed by the user
+	FCompilationOptions CISCompilationOptions = ReferenceCustomizableObject.GetPrivate()->GetCompileOptions();
+	CISCompilationOptions.bSilentCompilation = false;										// totally optional
+	CISCompilationOptions.OptimizationLevel = CustomizableObjectBenchmarkingUtils::GetOptimizationLevelForBenchmarking();
+	CISCompilationOptions.TextureCompression = ECustomizableObjectTextureCompression::Fast;	// Does not affect instance update speed but does compilation
+
+	return CISCompilationOptions;
+}
+
+
 
diff --git a/Engine/Plugins/Experimental/Mutable/Source/MutableValidation/Private/MuV/ValidationUtils.h b/Engine/Plugins/Experimental/Mutable/Source/MutableValidation/Private/MuV/ValidationUtils.h
--- a/Engine/Plugins/Experimental/Mutable/Source/MutableValidation/Private/MuV/ValidationUtils.h
+++ b/Engine/Plugins/Experimental/Mutable/Source/MutableValidation/Private/MuV/ValidationUtils.h
@@ -1,6 +1,9 @@
 ﻿// Copyright Epic Games, Inc. All Rights Reserved.
 #pragma once
 
+struct FCompilationOptions;
+class UCustomizableObject;
+
 /**
  * Prepare the asset registry so we can later use it to search assets. It is required by Mutable to compile.
  */
@@ -21,3 +24,11 @@ void Wait(const double ToWaitSeconds);
  */
 void LogGlobalSettings();
 
+
+/**
+ * Returns the settings used by CIS based on the compilation options of the provided CO. 
+ * @param ReferenceCustomizableObject CO used to get the base FCompilationOptions we want. 
+ * @return The FCompilationOptions for the provided CO but with some settings changed to be adecuate for a benchmark
+ * oriented compilation.
+ */
+FCompilationOptions GetCompilationOptionsForBenchmarking (UCustomizableObject& ReferenceCustomizableObject);
