[mutable] Fixed unreachable code

#rnx

diff --git a/Engine/Plugins/Experimental/Mutable/Source/CustomizableObjectEditor/Private/MuCOE/CustomizableObjectEditorPerformanceAnalyzer.cpp b/Engine/Plugins/Experimental/Mutable/Source/CustomizableObjectEditor/Private/MuCOE/CustomizableObjectEditorPerformanceAnalyzer.cpp
--- a/Engine/Plugins/Experimental/Mutable/Source/CustomizableObjectEditor/Private/MuCOE/CustomizableObjectEditorPerformanceAnalyzer.cpp
+++ b/Engine/Plugins/Experimental/Mutable/Source/CustomizableObjectEditor/Private/MuCOE/CustomizableObjectEditorPerformanceAnalyzer.cpp
@@ -466,19 +466,15 @@ FText SCustomizableObjectEditorPerformanceAnalyzer::GetGenerateInstancesButtonTe
 	{
 		return FText::FromString(TEXT("Another instance is running"));
 	}
+
+	if (!AreUpdatesPending())
+	{
+		return FText::FromString(TEXT("Generate Random Instances"));
+	}
 	else
 	{
-		if (!AreUpdatesPending())
-		{
-			return FText::FromString(TEXT("Generate Random Instances"));
-		}
-		else
-		{
-			return FText::FromString(FString::Printf(TEXT("Updated %u / %u instances"),UpdatedInstancesCount, TotalScheduledUpdates));;
-		}
+		return FText::FromString(FString::Printf(TEXT("Updated %u / %u instances"),UpdatedInstancesCount, TotalScheduledUpdates));;
 	}
-
-	return FText::FromString(TEXT(""));
 }
 
 
@@ -583,7 +579,6 @@ void SCustomizableObjectEditorPerformanceAnalyzer::SetMutableBenchmarkingSystemS
 
 void SCustomizableObjectEditorPerformanceAnalyzer::AddReferencedObjects(FReferenceCollector& Collector)
 {
-	//Collector.AddReferencedObject(CustomizableObject);
 	Collector.AddReferencedObject(System);
 }
 
