[Mutable] Make MeshMorph operations inplace.

#rnx

diff --git a/Engine/Plugins/Experimental/Mutable/Source/MutableRuntime/Private/MuR/CodeRunner.cpp b/Engine/Plugins/Experimental/Mutable/Source/MutableRuntime/Private/MuR/CodeRunner.cpp
--- a/Engine/Plugins/Experimental/Mutable/Source/MutableRuntime/Private/MuR/CodeRunner.cpp
+++ b/Engine/Plugins/Experimental/Mutable/Source/MutableRuntime/Private/MuR/CodeRunner.cpp
@@ -1186,19 +1186,19 @@ namespace mu
 
         case OP_TYPE::ME_MORPH:
         {
-			const uint8* data = Program.GetOpArgsPointer(item.At);
+			const uint8* Data = Program.GetOpArgsPointer(item.At);
 
 			OP::ADDRESS FactorAt = 0;
-			FMemory::Memcpy(&FactorAt, data, sizeof(OP::ADDRESS)); 
-			data += sizeof(OP::ADDRESS);
+			FMemory::Memcpy(&FactorAt, Data, sizeof(OP::ADDRESS)); 
+			Data += sizeof(OP::ADDRESS);
 			
 			OP::ADDRESS BaseAt = 0;
-			FMemory::Memcpy(&BaseAt, data, sizeof(OP::ADDRESS)); 
-			data += sizeof(OP::ADDRESS);
+			FMemory::Memcpy(&BaseAt, Data, sizeof(OP::ADDRESS)); 
+			Data += sizeof(OP::ADDRESS);
 
 			OP::ADDRESS TargetAt = 0;
-			FMemory::Memcpy(&TargetAt, data, sizeof(OP::ADDRESS)); 
-			data += sizeof(OP::ADDRESS);
+			FMemory::Memcpy(&TargetAt, Data, sizeof(OP::ADDRESS)); 
+			Data += sizeof(OP::ADDRESS);
 
 			switch (item.Stage)
             {
@@ -1225,19 +1225,19 @@ namespace mu
 
                 FScheduledOpData HeapData;
 				HeapData.Interpolate.Bifactor = Factor;
-				uint32 dataAddress = uint32(m_heapData.Add(HeapData));
+				uint32 DataAddress = uint32(m_heapData.Add(HeapData));
 
                 // No morph
 				if (FMath::IsNearlyZero(Factor))
                 {                        
-                    AddOp(FScheduledOp(item.At, item, 2, dataAddress),
+                    AddOp(FScheduledOp(item.At, item, 2, DataAddress),
 						FScheduledOp(BaseAt, item));
                 }
                 // The Morph, partial or full
                 else
                 {
                     // We will need the base again
-                    AddOp(FScheduledOp(item.At, item, 2, dataAddress),
+                    AddOp(FScheduledOp(item.At, item, 2, DataAddress),
 						FScheduledOp(BaseAt, item),
 						FScheduledOp(TargetAt, item));
                 }
@@ -1249,46 +1249,35 @@ namespace mu
             {
        		    MUTABLE_CPUPROFILER_SCOPE(ME_MORPH_2)
 
-                Ptr<const Mesh> pBase = LoadMesh(FCacheAddress(BaseAt, item));
+                Ptr<const Mesh> BaseMesh = LoadMesh(FCacheAddress(BaseAt, item));
 
                 // Factor from 0 to 1 between the two targets
-                const FScheduledOpData& HeapData = m_heapData[(size_t)item.CustomState];
+                const FScheduledOpData& HeapData = m_heapData[item.CustomState];
                 float Factor = HeapData.Interpolate.Bifactor;
 
-                if (pBase)
+                if (BaseMesh)
                 {
 					// No morph
 					if (FMath::IsNearlyZero(Factor))
                     {
-						StoreMesh(item, pBase);
+						StoreMesh(item, BaseMesh);
                     }
 					// The Morph, partial or full
                     else 
                     {
-                        Ptr<const Mesh> pMorph = LoadMesh(FCacheAddress(TargetAt,item));
+                        Ptr<const Mesh> MorphMesh = LoadMesh(FCacheAddress(TargetAt, item));
 						
-						if (pMorph)
+						if (MorphMesh)
 						{
-							Ptr<Mesh> Result = CreateMesh(pBase->GetDataSize());
-							bool bOutSuccess = false;
-							MeshMorph(Result.get(), pBase.get(), pMorph.get(), Factor, bOutSuccess);
-
-							Release(pMorph);
+							Ptr<Mesh> Result = CloneOrTakeOver(BaseMesh);
+							MeshMorph(Result.get(), MorphMesh.get(), Factor);
 
-							if (!bOutSuccess)
-							{
-								Release(Result);
-								StoreMesh(item, pBase);
-							}
-							else
-							{
-								Release(pBase);
-								StoreMesh(item, Result);
-							}
+							Release(MorphMesh);
+							StoreMesh(item, Result);
 						}
 						else
 						{
-							StoreMesh(item, pBase);
+							StoreMesh(item, BaseMesh);
 						}
                     }
                 }
@@ -1389,16 +1378,16 @@ namespace mu
 
         case OP_TYPE::ME_INTERPOLATE:
         {
-			OP::MeshInterpolateArgs args = Program.GetOpArgs<OP::MeshInterpolateArgs>(item.At);
+			OP::MeshInterpolateArgs Args = Program.GetOpArgs<OP::MeshInterpolateArgs>(item.At);
             switch (item.Stage)
             {
             case 0:
 			{
-				if (args.base)
+				if (Args.base)
 				{
 					AddOp(FScheduledOp(item.At, item, 1),
-						FScheduledOp(args.base, item),
-						FScheduledOp(args.factor, item));
+						FScheduledOp(Args.base, item),
+						FScheduledOp(Args.factor, item));
 				}
 				else
 				{
@@ -1410,72 +1399,70 @@ namespace mu
             {
             	MUTABLE_CPUPROFILER_SCOPE(ME_INTERPOLATE_1)
 
-                int count = 1;
-                for ( int i=0
-                    ; i<MUTABLE_OP_MAX_INTERPOLATE_COUNT-1 && args.targets[i]
-                    ; ++i )
+                int32 Count = 1;
+                for (int32 TargetIndex = 0; TargetIndex < MUTABLE_OP_MAX_INTERPOLATE_COUNT-1 && Args.targets[TargetIndex]; ++TargetIndex)
                 {
-                    count++;
+                    ++Count;
                 }
 
                 // Factor from 0 to 1 across all targets
-                float factor = LoadScalar(FCacheAddress(args.factor, item));
+                float Factor = LoadScalar(FCacheAddress(Args.factor, item));
 
-                float delta = 1.0f/(count-1);
-                int min = (int)floorf( factor/delta );
-                int max = (int)ceilf( factor/delta );
+                float Delta = 1.0f/(Count-1);
+                int32 Min = FMath::FloorToInt(Factor/Delta);
+                int32 Max = FMath::CeilToInt(Factor/Delta);
 
                 // Factor from 0 to 1 between the two targets
-                float bifactor = factor/delta - min;
+                float Bifactor = Factor/Delta - Min;
 
-                FScheduledOpData data;
-                data.Interpolate.Bifactor = bifactor;
-				data.Interpolate.Min = FMath::Clamp(min, 0, count - 1);
-				data.Interpolate.Max = FMath::Clamp(max, 0, count - 1);
-				uint32 dataAddress = uint32(m_heapData.Num());
+                FScheduledOpData Data;
+                Data.Interpolate.Bifactor = Bifactor;
+				Data.Interpolate.Min = FMath::Clamp(Min, 0, Count - 1);
+				Data.Interpolate.Max = FMath::Clamp(Max, 0, Count - 1);
+				uint32 DataAddress = uint32(m_heapData.Num());
 
                 // Just the first of the targets
-                if (bifactor < UE_SMALL_NUMBER)
+                if (Bifactor < UE_SMALL_NUMBER)
                 {
-                    if (min == 0)
+                    if (Min == 0)
                     {
                         // Just the base
-						Ptr<const Mesh> pBase = LoadMesh(FCacheAddress(args.base, item));
-						StoreMesh(item, pBase);
+						Ptr<const Mesh> BaseMesh = LoadMesh(FCacheAddress(Args.base, item));
+						StoreMesh(item, BaseMesh);
 					}
                     else
                     {
                         // Base with one full morph
-                        m_heapData.Add(data);
-						AddOp(FScheduledOp(item.At, item, 2, dataAddress),
-							FScheduledOp(args.base, item),
-							FScheduledOp(args.targets[min-1], item));
+                        m_heapData.Add(Data);
+						AddOp(FScheduledOp(item.At, item, 2, DataAddress),
+							FScheduledOp(Args.base, item),
+							FScheduledOp(Args.targets[Min-1], item));
 					}
 				}
                 // Just the second of the targets
-                else if (bifactor > 1.0f-UE_SMALL_NUMBER)
+                else if (Bifactor > 1.0f-UE_SMALL_NUMBER)
                 {
-                    m_heapData.Add(data);
-					AddOp(FScheduledOp(item.At, item, 2, dataAddress),
-						FScheduledOp(args.base, item),
-						FScheduledOp(args.targets[max-1], item));
+                    m_heapData.Add(Data);
+					AddOp(FScheduledOp(item.At, item, 2, DataAddress),
+						FScheduledOp(Args.base, item),
+						FScheduledOp(Args.targets[Max-1], item));
 				}
                 // Mix the first target on the base
-                else if (min == 0)
+                else if (Min == 0)
                 {
-                    m_heapData.Add(data);
-					AddOp(FScheduledOp(item.At, item, 2, dataAddress),
-						FScheduledOp(args.base, item),
-						FScheduledOp(args.targets[0], item));
+                    m_heapData.Add(Data);
+					AddOp(FScheduledOp(item.At, item, 2, DataAddress),
+						FScheduledOp(Args.base, item),
+						FScheduledOp(Args.targets[0], item));
 				}
                 // Mix two targets on the base
                 else
                 {
-                    m_heapData.Add(data);
-					AddOp(FScheduledOp(item.At, item, 2, dataAddress),
-						FScheduledOp(args.base, item),
-						FScheduledOp(args.targets[min-1], item),
-						FScheduledOp(args.targets[max-1], item));
+                    m_heapData.Add(Data);
+					AddOp(FScheduledOp(item.At, item, 2, DataAddress),
+						FScheduledOp(Args.base, item),
+						FScheduledOp(Args.targets[Min-1], item),
+						FScheduledOp(Args.targets[Max-1], item));
 				}
 
                 break;
@@ -1485,29 +1472,27 @@ namespace mu
             {
 				MUTABLE_CPUPROFILER_SCOPE(ME_INTERPOLATE_2)
 
-                int count = 1;
-                for ( int i=0
-                    ; i<MUTABLE_OP_MAX_INTERPOLATE_COUNT-1 && args.targets[i]
-                    ; ++i )
+                int32 Count = 1;
+                for (int32 TargetIndex = 0; TargetIndex < MUTABLE_OP_MAX_INTERPOLATE_COUNT-1 && Args.targets[TargetIndex]; ++TargetIndex)
                 {
-                    count++;
+                    ++Count;
                 }
 
-                const FScheduledOpData& data = m_heapData[ (size_t)item.CustomState ];
+                const FScheduledOpData& Data = m_heapData[ (size_t)item.CustomState ];
 
                 // Factor from 0 to 1 between the two targets
-                float bifactor = data.Interpolate.Bifactor;
-                int min = data.Interpolate.Min;
-                int max = data.Interpolate.Max;
+                float Bifactor = Data.Interpolate.Bifactor;
+                int32 Min = Data.Interpolate.Min;
+                int32 Max = Data.Interpolate.Max;
 
-                Ptr<const Mesh> pBase = LoadMesh(FCacheAddress(args.base, item));
+                Ptr<const Mesh> BaseMesh = LoadMesh(FCacheAddress(Args.base, item));
 
-                if (pBase)
+                if (BaseMesh)
                 {
                     // Just the first of the targets
-                    if (bifactor < UE_SMALL_NUMBER)
+                    if (Bifactor < UE_SMALL_NUMBER)
                     {
-                        if (min == 0)
+                        if (Min == 0)
                         {
                             // Just the base. It should have been dealt with in the previous stage.
                             check(false);
@@ -1515,123 +1500,90 @@ namespace mu
                         else
                         {
                             // Base with one full morph
-                            Ptr<const Mesh> pMorph = LoadMesh(FCacheAddress(args.targets[min-1], item));
-							
-							Ptr<Mesh> Result = CreateMesh(pBase->GetDataSize());
-
-							bool bOutSuccess = false;
-                            MeshMorph(Result.get(), pBase.get(), pMorph.get(), bOutSuccess);
+                            Ptr<const Mesh> MorphMesh = LoadMesh(FCacheAddress(Args.targets[Min-1], item));
 						
-							Release(pMorph);
-
-							if (!bOutSuccess)
+							if (MorphMesh)
 							{
-								Release(Result);
-								StoreMesh(item, pBase);
+								Ptr<Mesh> Result = CloneOrTakeOver(BaseMesh);
+
+								MeshMorph(Result.get(), MorphMesh.get());
+							
+								StoreMesh(item, Result);
 							}
 							else
 							{
-								Release(pBase);
-								StoreMesh(item, Result);
+								StoreMesh(item, BaseMesh);
 							}
+							
+							Release(MorphMesh);
                         }
                     }
                     // Just the second of the targets
-                    else if (bifactor > 1.0f-UE_SMALL_NUMBER)
+                    else if (Bifactor > 1.0f-UE_SMALL_NUMBER)
                     {
-                        check(max > 0);
-                        Ptr<const Mesh> pMorph = LoadMesh(FCacheAddress(args.targets[max-1], item));
+                        check(Max > 0);
+                        Ptr<const Mesh> MorphMesh = LoadMesh(FCacheAddress(Args.targets[Max-1], item));
 
-                        if (pMorph)
+                        if (MorphMesh)
                         {
-							Ptr<Mesh> Result = CreateMesh(pBase->GetDataSize());
+							Ptr<Mesh> Result = CloneOrTakeOver(BaseMesh);
 							
-							bool bOutSuccess = false;
-                            MeshMorph(Result.get(), pBase.get(), pMorph.get(), bOutSuccess);
-
-							Release(pMorph);
-							if (!bOutSuccess)
-							{
-								Release(Result);
-								StoreMesh(item, pBase);
-							}
-							else
-							{
-								Release(pBase);
-								StoreMesh(item, Result);
-							}
+                            MeshMorph(Result.get(), MorphMesh.get());
 
+							StoreMesh(item, Result);
                         }
                         else
                         {
-							StoreMesh(item, pBase);
+							StoreMesh(item, BaseMesh);
                         }
+						
+						Release(MorphMesh);
                     }
                     // Mix the first target on the base
-                    else if (min == 0)
+                    else if (Min == 0)
                     {
-                        Ptr<const Mesh> pMorph = LoadMesh(FCacheAddress(args.targets[0], item));
-                        if (pMorph)
+                        Ptr<const Mesh> MorphMesh = LoadMesh(FCacheAddress(Args.targets[0], item));
+                        if (MorphMesh)
                         {
-							Ptr<Mesh> Result = CreateMesh(pBase->GetDataSize());
+							Ptr<Mesh> Result = CloneOrTakeOver(BaseMesh);
 
-							bool bOutSuccess = false;
-                            MeshMorph(Result.get(), pBase.get(), pMorph.get(), bifactor, bOutSuccess);
+                            MeshMorph(Result.get(), MorphMesh.get(), Bifactor);
 
-							Release(pMorph);
-
-							if (!bOutSuccess)
-							{
-								Release(Result);
-								StoreMesh(item, pBase);
-							}
-							else
-							{
-								Release(pBase);
-								StoreMesh(item, Result);
-							}
+							StoreMesh(item, Result);
                         }
                         else
                         {
-							StoreMesh(item, pBase);
+							StoreMesh(item, BaseMesh);
                         }
+						
+						Release(MorphMesh);
                     }
                     // Mix two targets on the base
                     else
                     {
-                        Ptr<const Mesh> pMin = LoadMesh(FCacheAddress(args.targets[min-1], item));
-                        Ptr<const Mesh> pMax = LoadMesh(FCacheAddress(args.targets[max-1], item));
+                        Ptr<const Mesh> MinMesh = LoadMesh(FCacheAddress(Args.targets[Min-1], item));
+                        Ptr<const Mesh> MaxMesh = LoadMesh(FCacheAddress(Args.targets[Max-1], item));
 
-                        if (pMin && pMax)
+                        if (MinMesh && MaxMesh)
                         {
-							Ptr<Mesh> Result = CreateMesh(pBase->GetDataSize());
+							Ptr<Mesh> Result = CloneOrTakeOver(BaseMesh);
 
-							bool bOutSuccess = false;
-                            MeshMorph2(Result.get(), pBase.get(), pMin.get(), pMax.get(), bifactor, bOutSuccess);
-
-							Release(pMin);
-							Release(pMax);
+                            MeshMorph2(Result.get(), MinMesh.get(), MaxMesh.get(), Bifactor);
 
-							if (!bOutSuccess)
-							{
-								Release(Result);
-								StoreMesh(item, pBase);
-							}
-							else
-							{
-								Release(pBase);
-								StoreMesh(item, Result);
-							}
+							StoreMesh(item, Result);
                         }
                         else
                         {
-							StoreMesh(item, pBase);
+							StoreMesh(item, BaseMesh);
                         }
+
+						Release(MinMesh);
+						Release(MaxMesh);
                     }
                 }
 				else
 				{
-					StoreMesh(item, pBase);
+					StoreMesh(item, BaseMesh);
 				}
 
                 break;
diff --git a/Engine/Plugins/Experimental/Mutable/Source/MutableRuntime/Private/MuR/OpMeshMorph.h b/Engine/Plugins/Experimental/Mutable/Source/MutableRuntime/Private/MuR/OpMeshMorph.h
--- a/Engine/Plugins/Experimental/Mutable/Source/MutableRuntime/Private/MuR/OpMeshMorph.h
+++ b/Engine/Plugins/Experimental/Mutable/Source/MutableRuntime/Private/MuR/OpMeshMorph.h
@@ -10,18 +10,15 @@
 
 namespace mu
 {
-
-	//---------------------------------------------------------------------------------------------
-	//! Optimized linear factor version for morphing 2 targets
-	//---------------------------------------------------------------------------------------------
-    inline void MeshMorph2(Mesh* Result, const Mesh* pBase, const Mesh* pMin, const Mesh* pMax, const float Factor, bool& bOutSuccess)
+	/**
+	 * Optimized linear factor version for morphing 2 targets
+	 */
+    inline void MeshMorph2(Mesh* BaseMesh, const Mesh* MinMesh, const Mesh* MaxMesh, const float Factor)
     {
         MUTABLE_CPUPROFILER_SCOPE(MeshMorph2);
-		bOutSuccess = true;
 
-		if (!pBase)
+		if (!BaseMesh)
 		{
-			bOutSuccess = true; // return an empty mesh.
 			return;
 		}
 
@@ -228,19 +225,16 @@ namespace mu
 
 
 		// Number of vertices to modify
-		const int32 MinNum = pMin ? pMin->GetVertexBuffers().GetElementCount() : 0;
-		const int32 MaxNum = pMax ? pMax->GetVertexBuffers().GetElementCount() : 0;
-		const int32 BaseNum = pBase ? pBase->GetVertexBuffers().GetElementCount() : 0;
-		const Mesh* RefTarget = MinNum > 0 ? pMin : pMax;
+		const int32 MinNum = MinMesh ? MinMesh->GetVertexBuffers().GetElementCount() : 0;
+		const int32 MaxNum = MaxMesh ? MaxMesh->GetVertexBuffers().GetElementCount() : 0;
+		const int32 BaseNum = BaseMesh ? BaseMesh->GetVertexBuffers().GetElementCount() : 0;
+		const Mesh* RefTarget = MinNum > 0 ? MinMesh : MaxMesh;
 
 		if (BaseNum == 0 || (MinNum + MaxNum) == 0)
 		{
-			bOutSuccess = false; // Use the passed pBase as result.
 			return;
 		}
 
-		Result->CopyFrom(*pBase);
-
 		if (RefTarget)
 		{
 			constexpr int32 MorphBufferDataChannel = 0;
@@ -258,7 +252,7 @@ namespace mu
 			UntypedMeshBufferIteratorConst MinNormalChannelIter;
 			UntypedMeshBufferIteratorConst MaxNormalChannelIter;
 
-			const bool bBaseHasNormals = UntypedMeshBufferIteratorConst(pBase->GetVertexBuffers(), MBS_NORMAL, 0).ptr() != nullptr;
+			const bool bBaseHasNormals = UntypedMeshBufferIteratorConst(BaseMesh->GetVertexBuffers(), MBS_NORMAL, 0).ptr() != nullptr;
 			for (int32 ChannelIndex = 0; ChannelIndex < ChannelsNum; ++ChannelIndex)
 			{
 				const FMeshBufferSet& MBSPriv = RefTarget->GetVertexBuffers();
@@ -268,45 +262,45 @@ namespace mu
 			
 				if (Sem == MBS_NORMAL && bBaseHasNormals)
 				{
-					BaseTangentFrameChannelsIters[2] = UntypedMeshBufferIterator(Result->GetVertexBuffers(), Sem, SemIndex);
+					BaseTangentFrameChannelsIters[2] = UntypedMeshBufferIterator(BaseMesh->GetVertexBuffers(), Sem, SemIndex);
 					if (MinNum > 0)
 					{
-						MinNormalChannelIter = UntypedMeshBufferIteratorConst(pMin->GetVertexBuffers(), Sem, SemIndex);
+						MinNormalChannelIter = UntypedMeshBufferIteratorConst(MinMesh->GetVertexBuffers(), Sem, SemIndex);
 					}
 
 					if (MaxNum > 0)
 					{
-						MaxNormalChannelIter = UntypedMeshBufferIteratorConst(pMax->GetVertexBuffers(), Sem, SemIndex);
+						MaxNormalChannelIter = UntypedMeshBufferIteratorConst(MaxMesh->GetVertexBuffers(), Sem, SemIndex);
 					}
 				}
 				else if (Sem == MBS_TANGENT && bBaseHasNormals)
 				{
-					BaseTangentFrameChannelsIters[1] = UntypedMeshBufferIterator(Result->GetVertexBuffers(), Sem, SemIndex);
+					BaseTangentFrameChannelsIters[1] = UntypedMeshBufferIterator(BaseMesh->GetVertexBuffers(), Sem, SemIndex);
 				}
 				else if (Sem == MBS_BINORMAL && bBaseHasNormals)
 				{
-					BaseTangentFrameChannelsIters[0] = UntypedMeshBufferIterator(Result->GetVertexBuffers(), Sem, SemIndex);
+					BaseTangentFrameChannelsIters[0] = UntypedMeshBufferIterator(BaseMesh->GetVertexBuffers(), Sem, SemIndex);
 				}
 				else
 				{
-					BaseChannelsIters[ChannelIndex] = UntypedMeshBufferIterator(Result->GetVertexBuffers(), Sem, SemIndex);
+					BaseChannelsIters[ChannelIndex] = UntypedMeshBufferIterator(BaseMesh->GetVertexBuffers(), Sem, SemIndex);
 					if (MinNum > 0)
 					{
-						MinChannelsIters[ChannelIndex] = UntypedMeshBufferIteratorConst(pMin->GetVertexBuffers(), Sem, SemIndex);
+						MinChannelsIters[ChannelIndex] = UntypedMeshBufferIteratorConst(MinMesh->GetVertexBuffers(), Sem, SemIndex);
 					}
 
 					if (MaxNum > 0)
 					{
-						MaxChannelsIters[ChannelIndex] = UntypedMeshBufferIteratorConst(pMax->GetVertexBuffers(), Sem, SemIndex);
+						MaxChannelsIters[ChannelIndex] = UntypedMeshBufferIteratorConst(MaxMesh->GetVertexBuffers(), Sem, SemIndex);
 					}
 				}
 			}
 			
-			MeshVertexIdIteratorConst BaseIdIter(pBase);
+			MeshVertexIdIteratorConst BaseIdIter(BaseMesh);
 
 			if (MinNum > 0)
 			{
-				MeshVertexIdIteratorConst MinIdIter(pMin);
+				MeshVertexIdIteratorConst MinIdIter(MinMesh);
 				SparseIndexMapSet IndexMap = MakeIndexMap(BaseIdIter, BaseNum, MinIdIter, MinNum);
 
 				ApplyGenericMorph(BaseIdIter, BaseChannelsIters, BaseNum, MinIdIter, MinChannelsIters, MinNum, IndexMap, 1.0f - Factor);
@@ -318,7 +312,7 @@ namespace mu
 
 			if (MaxNum > 0)
 			{
-				MeshVertexIdIteratorConst MaxIdIter(pMax);
+				MeshVertexIdIteratorConst MaxIdIter(MaxMesh);
 				SparseIndexMapSet IndexMap = MakeIndexMap(BaseIdIter, BaseNum, MaxIdIter, MaxNum);
 
 				ApplyGenericMorph(BaseIdIter, BaseChannelsIters, BaseNum, MaxIdIter, MaxChannelsIters, MaxNum, IndexMap, Factor);
@@ -331,20 +325,16 @@ namespace mu
 		}
     }
 
-	//---------------------------------------------------------------------------------------------
-	//! \TODO Optimized linear factor version
-	//---------------------------------------------------------------------------------------------
-	inline void MeshMorph(Mesh* Result, const Mesh* pBase, const Mesh* pMorph, float factor, bool& bOutSuccess)
+	//TODO Optimized linear factor version
+	inline void MeshMorph(Mesh* BaseMesh, const Mesh* MorphMesh, float Factor)
 	{
-		MeshMorph2(Result, pBase, nullptr, pMorph, factor, bOutSuccess);
+		MeshMorph2(BaseMesh, nullptr, MorphMesh, Factor);
 	}
 
-    //---------------------------------------------------------------------------------------------
-    //! \TODO Optimized Factor-less version
-    //---------------------------------------------------------------------------------------------
-	inline void MeshMorph(Mesh* Result, const Mesh* pBase, const Mesh* pMorph, bool& bOutSuccess)
+    //TODO Optimized Factor-less version
+	inline void MeshMorph(Mesh* BaseMesh, const Mesh* MorphMesh)
     {
         // Trust the compiler to remove the factor
-		MeshMorph(Result, pBase, pMorph, 1.0f, bOutSuccess);
+		MeshMorph(BaseMesh, MorphMesh, 1.0f);
     }
 }
