[Mutable] Added more Viewport options.

- Added Show Mesh Info.
- Added Show Normals/Tangents/Binormals.
- Added Enable Cloth Physics.
- Added Show Cloth Physics Mesh.
- Show Bones options moved to Character submenu.

#rnx

diff --git a/Engine/Plugins/Experimental/Mutable/Source/CustomizableObjectEditor/Private/MuCOE/CustomizableObjectEditorActions.cpp b/Engine/Plugins/Experimental/Mutable/Source/CustomizableObjectEditor/Private/MuCOE/CustomizableObjectEditorActions.cpp
--- a/Engine/Plugins/Experimental/Mutable/Source/CustomizableObjectEditor/Private/MuCOE/CustomizableObjectEditorActions.cpp
+++ b/Engine/Plugins/Experimental/Mutable/Source/CustomizableObjectEditor/Private/MuCOE/CustomizableObjectEditorActions.cpp
@@ -61,6 +61,13 @@ void FCustomizableObjectEditorViewportCommands::RegisterCommands()
 	UI_COMMAND(BakeInstance, "Bake Instance", "Create baked unreal resources for the current preview instance.", EUserInterfaceActionType::Button, FInputChord());
 	UI_COMMAND(StateChangeShowData, "Show or hide test results", "Show or hide test results", EUserInterfaceActionType::ToggleButton, FInputChord());
 	UI_COMMAND(StateChangeShowGeometryData, "Show instance geometry data", "Show instance geometry data", EUserInterfaceActionType::ToggleButton, FInputChord());
+
+	UI_COMMAND(ShowDisplayInfo, "Mesh Info", "Display mesh info in the viewport", EUserInterfaceActionType::ToggleButton, FInputChord());
+	UI_COMMAND(EnableClothSimulation, "Enable Cloth Simulation", "Show simulated cloth mesh", EUserInterfaceActionType::ToggleButton, FInputChord());
+	UI_COMMAND(DebugDrawPhysMeshWired, "Physical Mesh (Wireframe)", "Draws the current physical mesh result in wireframe", EUserInterfaceActionType::ToggleButton, FInputChord());
+	UI_COMMAND(SetShowNormals, "Normals", "Toggles display of vertex normals in the Preview Pane.", EUserInterfaceActionType::ToggleButton, FInputChord());
+	UI_COMMAND(SetShowTangents, "Tangents", "Toggles display of vertex tangents in the Preview Pane.", EUserInterfaceActionType::ToggleButton, FInputChord());
+	UI_COMMAND(SetShowBinormals, "Binormals", "Toggles display of vertex binormals (orthogonal vector to normal and tangent) in the Preview Pane.", EUserInterfaceActionType::ToggleButton, FInputChord());
 }
 
 
diff --git a/Engine/Plugins/Experimental/Mutable/Source/CustomizableObjectEditor/Private/MuCOE/CustomizableObjectEditorViewportClient.cpp b/Engine/Plugins/Experimental/Mutable/Source/CustomizableObjectEditor/Private/MuCOE/CustomizableObjectEditorViewportClient.cpp
--- a/Engine/Plugins/Experimental/Mutable/Source/CustomizableObjectEditor/Private/MuCOE/CustomizableObjectEditorViewportClient.cpp
+++ b/Engine/Plugins/Experimental/Mutable/Source/CustomizableObjectEditor/Private/MuCOE/CustomizableObjectEditorViewportClient.cpp
@@ -384,6 +384,24 @@ void FCustomizableObjectEditorViewportClient::Draw(const FSceneView* View, FPrim
 			}
 		}
 	}
+
+	if (bShowDebugClothing)
+	{
+		IConsoleVariable* CVar = IConsoleManager::Get().FindConsoleVariable(TEXT("p.ChaosClothEditor.DebugDrawPhysMeshWired"));
+		check(CVar);
+		const bool bPreviousValue = CVar->GetBool();
+		CVar->Set(true);
+		
+		for (const TWeakObjectPtr<UDebugSkelMeshComponent>& SkeletalMeshComponent : SkeletalMeshComponents)
+		{
+			if (SkeletalMeshComponent.IsValid())
+			{
+				SkeletalMeshComponent->DebugDrawClothing(PDI);
+			}
+		}
+		
+		CVar->Set(bPreviousValue);
+	}
 }
 
 
@@ -676,7 +694,20 @@ void FCustomizableObjectEditorViewportClient::SetPreviewActor(const TWeakObjectP
 	SkeletalMeshComponents = InSkeletalMeshComponents;
 	Actor = InActor;
 
+	for (const TWeakObjectPtr<UDebugSkelMeshComponent>& SkeletalMeshComponent : SkeletalMeshComponents)
+	{
+		if (SkeletalMeshComponent.IsValid())
+		{
+			SkeletalMeshComponent->bDisableClothSimulation = bDisableClothSimulation;
+			SkeletalMeshComponent->bDrawNormals = bDrawNormals;
+			SkeletalMeshComponent->bDrawTangents = bDrawTangents;
+			SkeletalMeshComponent->bDrawBinormals = bDrawBinormals;
+		}
+	}
+	
 	InInstance->UpdatedNativeDelegate.AddSP(SharedThis(this), &FCustomizableObjectEditorViewportClient::OnInstanceUpdate);
+	
+	Invalidate();
 }
 
 
@@ -1803,6 +1834,259 @@ const TArray<ULightComponent*>& FCustomizableObjectEditorViewportClient::GetLigh
 }
 
 
+void FCustomizableObjectEditorViewportClient::OnShowDisplayInfo()
+{
+	bShowDisplayInfo = !bShowDisplayInfo;
+	
+	Invalidate();
+}
+
+
+bool FCustomizableObjectEditorViewportClient::IsShowingMeshInfo() const
+{
+	return bShowDisplayInfo;
+}
+
+
+void FCustomizableObjectEditorViewportClient::OnEnableClothSimulation()
+{
+	bDisableClothSimulation = !bDisableClothSimulation;
+	
+	for (const TWeakObjectPtr<UDebugSkelMeshComponent>& SkeletalMeshComponent : SkeletalMeshComponents)
+	{
+		if (SkeletalMeshComponent.IsValid())
+		{
+			SkeletalMeshComponent->bDisableClothSimulation = bDisableClothSimulation;
+		}
+	}
+
+	Invalidate();
+}
+
+
+bool FCustomizableObjectEditorViewportClient::IsClothSimulationEnabled() const
+{
+	return bDisableClothSimulation;	
+}
+
+
+void FCustomizableObjectEditorViewportClient::OnDebugDrawPhysMeshWired()
+{
+	bShowDebugClothing = !bShowDebugClothing;
+	
+	Invalidate();
+}
+
+
+bool FCustomizableObjectEditorViewportClient::IsDebugDrawPhysMeshWired() const
+{
+	return bShowDebugClothing;
+}
+
+
+FText MergeLine(const FText& InText, const FText& InNewLine)
+{
+	if (InText.IsEmpty())
+	{
+		return InNewLine;
+	}
+
+	return FText::Format(LOCTEXT("ViewportTextNewlineFormatter", "{0}\n{1}"), InText, InNewLine);
+}
+
+
+// Based on FAnimationViewportClient::GetDisplayInfo(bool)
+FText FCustomizableObjectEditorViewportClient::GetMeshInfoText() const
+{
+	FText TextValue;
+
+	for (int32 ComponentIndex = 0; ComponentIndex < SkeletalMeshComponents.Num(); ++ComponentIndex)
+	{
+		TWeakObjectPtr<UDebugSkelMeshComponent> PreviewMeshComponent = SkeletalMeshComponents[ComponentIndex];
+		
+		if (!PreviewMeshComponent.IsValid())
+		{
+			continue;
+		}
+
+		FSkeletalMeshRenderData* SkelMeshResource = PreviewMeshComponent->GetSkeletalMeshRenderData();
+		if (!SkelMeshResource)
+		{
+			continue;
+		}
+
+		// Draw stats about the mesh
+		int32 NumBonesInUse;
+		int32 NumBonesMappedToVerts;
+		int32 NumSectionsInUse;
+
+		const int32 LODIndex = FMath::Clamp(PreviewMeshComponent->GetPredictedLODLevel(), 0, SkelMeshResource->LODRenderData.Num() - 1);
+		FSkeletalMeshLODRenderData& LODData = SkelMeshResource->LODRenderData[LODIndex];
+
+		NumBonesInUse = LODData.RequiredBones.Num();
+		NumBonesMappedToVerts = LODData.ActiveBoneIndices.Num();
+		NumSectionsInUse = LODData.RenderSections.Num();
+
+		// Calculate polys based on non clothing sections so we don't duplicate the counts.
+		uint32 NumTotalTriangles = 0;
+		int32 NumSections = LODData.RenderSections.Num();
+		for(int32 SectionIndex = 0; SectionIndex < NumSections; SectionIndex++)
+		{
+			NumTotalTriangles += LODData.RenderSections[SectionIndex].NumTriangles;
+		}
+
+		if (ComponentIndex > 0)
+		{
+			TextValue = FText::Format(LOCTEXT("MeshInfoComponentSeparation", "{0}\n"), TextValue);
+		}
+		
+		TextValue = MergeLine(TextValue, FText::Format(LOCTEXT("MeshInfoFormat", "Component: {0}, LOD: {1}, Bones: {2} (Mapped to Vertices: {3}), Polys: {4}"),
+			FText::AsNumber(ComponentIndex),
+			FText::AsNumber(LODIndex),
+			FText::AsNumber(NumBonesInUse),
+			FText::AsNumber(NumBonesMappedToVerts),
+			FText::AsNumber(NumTotalTriangles)));
+
+		for (int32 SectionIndex = 0; SectionIndex < LODData.RenderSections.Num(); SectionIndex++)
+		{
+			int32 SectionVerts = LODData.RenderSections[SectionIndex].GetNumVertices();
+
+			FText SectionDisabledText = LODData.RenderSections[SectionIndex].bDisabled ? LOCTEXT("SectionIsDisbable", " Disabled") : FText::GetEmpty();
+			TextValue = MergeLine(TextValue, FText::Format(LOCTEXT("SectionFormat", " [Section {0}]{1} Verts: {2}, Bones: {3}, Max Influences: {4}"),
+				FText::AsNumber(SectionIndex),
+				SectionDisabledText,
+				FText::AsNumber(SectionVerts),
+				FText::AsNumber(LODData.RenderSections[SectionIndex].BoneMap.Num()),
+				FText::AsNumber(LODData.RenderSections[SectionIndex].MaxBoneInfluences)
+				));
+		}
+
+		TextValue = MergeLine(TextValue, FText::Format(LOCTEXT("TotalVerts", "TOTAL Verts: {0}"),
+			FText::AsNumber(LODData.GetNumVertices())));
+
+		TextValue = MergeLine(TextValue, FText::Format(LOCTEXT("Sections", "Sections: {0}"),
+			NumSectionsInUse
+			));
+
+		TArray<FTransform> LocalBoneTransforms = PreviewMeshComponent->GetBoneSpaceTransforms();
+		if (PreviewMeshComponent->BonesOfInterest.Num() > 0)
+		{
+			int32 BoneIndex = PreviewMeshComponent->BonesOfInterest[0];
+			FTransform ReferenceTransform = PreviewMeshComponent->GetReferenceSkeleton().GetRefBonePose()[BoneIndex];
+			FTransform LocalTransform = LocalBoneTransforms[BoneIndex];
+			FTransform ComponentTransform = PreviewMeshComponent->GetDrawTransform(BoneIndex);
+
+			auto GetDisplayTransform = [](const FTransform& InTransform) -> FText
+			{
+				FRotator R(InTransform.GetRotation());
+				FVector T(InTransform.GetTranslation());
+				FVector S(InTransform.GetScale3D());
+
+				FString Output = FString::Printf(TEXT("Rotation: X(Roll) %f Y(Pitch)  %f Z(Yaw) %f\r\n"), R.Roll, R.Pitch, R.Yaw);
+				Output += FString::Printf(TEXT("Translation: %f %f %f\r\n"), T.X, T.Y, T.Z);
+				Output += FString::Printf(TEXT("Scale3D: %f %f %f\r\n"), S.X, S.Y, S.Z);
+
+				return FText::FromString(Output);
+			};
+
+			TextValue = MergeLine(TextValue, FText::Format(LOCTEXT("LocalTransform", "Local: {0}"), GetDisplayTransform(LocalTransform)));
+
+			TextValue = MergeLine(TextValue, FText::Format(LOCTEXT("ComponentTransform", "Component: {0}"), GetDisplayTransform(ComponentTransform)));
+
+			TextValue = MergeLine(TextValue, FText::Format(LOCTEXT("ReferenceTransform", "Reference: {0}"), GetDisplayTransform(ReferenceTransform)));
+		}
+
+		TextValue = MergeLine(TextValue, FText::Format(LOCTEXT("ApproximateSize", "Approximate Size: {0}x{1}x{2}"),
+			FText::AsNumber(FMath::RoundToInt(PreviewMeshComponent->Bounds.BoxExtent.X * 2.0f)),
+			FText::AsNumber(FMath::RoundToInt(PreviewMeshComponent->Bounds.BoxExtent.Y * 2.0f)),
+			FText::AsNumber(FMath::RoundToInt(PreviewMeshComponent->Bounds.BoxExtent.Z * 2.0f))));
+
+		uint32 NumNotiesWithErrors = PreviewMeshComponent->AnimNotifyErrors.Num();
+		for (uint32 i = 0; i < NumNotiesWithErrors; ++i)
+		{
+			uint32 NumErrors = PreviewMeshComponent->AnimNotifyErrors[i].Errors.Num();
+			for (uint32 ErrorIdx = 0; ErrorIdx < NumErrors; ++ErrorIdx)
+			{
+				TextValue = MergeLine(TextValue, FText::FromString(PreviewMeshComponent->AnimNotifyErrors[i].Errors[ErrorIdx]));
+			}
+		}	
+	}
+	
+	return TextValue;
+}
+
+
+void FCustomizableObjectEditorViewportClient::ToggleShowNormals()
+{
+	bDrawNormals = !bDrawNormals;
+	
+	for (const TWeakObjectPtr<UDebugSkelMeshComponent>& SkeletalMeshComponent : SkeletalMeshComponents)
+	{
+		if (SkeletalMeshComponent.IsValid())
+		{
+			SkeletalMeshComponent->bDrawNormals;
+			SkeletalMeshComponent->MarkRenderStateDirty();
+		}
+	}
+
+	Invalidate();
+}
+
+
+bool FCustomizableObjectEditorViewportClient::IsSetShowNormalsChecked() const
+{
+	return bDrawNormals;
+}
+
+
+void FCustomizableObjectEditorViewportClient::ToggleShowTangents()
+{
+	bDrawTangents = !bDrawTangents;
+
+	for (const TWeakObjectPtr<UDebugSkelMeshComponent>& SkeletalMeshComponent : SkeletalMeshComponents)
+	{
+		if (SkeletalMeshComponent.IsValid())
+		{
+			SkeletalMeshComponent->bDrawTangents = bDrawTangents;
+			SkeletalMeshComponent->MarkRenderStateDirty();
+		}
+	}
+
+	Invalidate();
+}
+
+
+bool FCustomizableObjectEditorViewportClient::IsSetShowTangentsChecked() const
+{
+	return bDrawTangents;
+}
+
+
+void FCustomizableObjectEditorViewportClient::ToggleShowBinormals()
+{
+	bDrawBinormals = !bDrawBinormals;
+	
+	for (const TWeakObjectPtr<UDebugSkelMeshComponent>& SkeletalMeshComponent : SkeletalMeshComponents)
+	{
+		if (SkeletalMeshComponent.IsValid())
+		{
+			SkeletalMeshComponent->bDrawBinormals = bDrawBinormals;
+			SkeletalMeshComponent->MarkRenderStateDirty();
+		}
+	}
+	
+	GetWorld()->SendAllEndOfFrameUpdates();
+
+	Invalidate();
+}
+
+
+bool FCustomizableObjectEditorViewportClient::IsSetShowBinormalsChecked() const
+{
+	return bDrawBinormals;
+}
+
+
 void FCustomizableObjectEditorViewportClient::DrawMeshBones(UDebugSkelMeshComponent* MeshComponent, FPrimitiveDrawInterface* PDI)
 {
 	if (!MeshComponent ||
diff --git a/Engine/Plugins/Experimental/Mutable/Source/CustomizableObjectEditor/Private/MuCOE/SCustomizableObjectEditorViewport.cpp b/Engine/Plugins/Experimental/Mutable/Source/CustomizableObjectEditor/Private/MuCOE/SCustomizableObjectEditorViewport.cpp
--- a/Engine/Plugins/Experimental/Mutable/Source/CustomizableObjectEditor/Private/MuCOE/SCustomizableObjectEditorViewport.cpp
+++ b/Engine/Plugins/Experimental/Mutable/Source/CustomizableObjectEditor/Private/MuCOE/SCustomizableObjectEditorViewport.cpp
@@ -85,16 +85,31 @@ void SCustomizableObjectEditorViewport::PopulateViewportOverlays(TSharedRef<SOve
 	[
 		SNew(SVerticalBox)
 		+SVerticalBox::Slot()
+		.VAlign(VAlign_Top)
+		.AutoHeight()
 		[
 			SNew(SCustomizableObjectEditorViewportToolBar, TabBodyPtr.Pin(), SharedThis(this)).Cursor(EMouseCursor::Default)
 		]
 		+SVerticalBox::Slot()
+		.VAlign(VAlign_Top)
+		.AutoHeight()
 		.Padding(4.0, 16.0, 0.0, 0.0)
 		[
 			SNew(STextBlock)
 			.Text(this, &SCustomizableObjectEditorViewport::GetWarningText)
+			.Visibility(this, &SCustomizableObjectEditorViewport::GetWarningTextVisibility)
 			.ColorAndOpacity(FLinearColor::Yellow)
 		]
+		+SVerticalBox::Slot()
+		.VAlign(VAlign_Top)
+		.AutoHeight()
+		.Padding(4.0, 16.0, 0.0, 0.0)
+		[
+			SNew(STextBlock)
+			.Text(this, &SCustomizableObjectEditorViewport::GetMeshInfoText)
+			.Visibility(this, &SCustomizableObjectEditorViewport::GetMeshInfoTextVisibility)
+			.ColorAndOpacity(FLinearColor::White)
+		]
 	];
 	Overlay->AddSlot()
 	.VAlign(VAlign_Center)
@@ -214,6 +229,24 @@ FText SCustomizableObjectEditorViewport::GetWarningText() const
 }
 
 
+EVisibility SCustomizableObjectEditorViewport::GetWarningTextVisibility() const
+{
+	return !GetWarningText().IsEmpty() ? EVisibility::Visible : EVisibility::Collapsed;
+}
+
+
+FText SCustomizableObjectEditorViewport::GetMeshInfoText() const
+{
+	return LevelViewportClient->GetMeshInfoText();
+}
+
+
+EVisibility SCustomizableObjectEditorViewport::GetMeshInfoTextVisibility() const
+{
+	return LevelViewportClient->IsShowingMeshInfo() ? EVisibility::Visible : EVisibility::Collapsed;
+}
+
+
 void SCustomizableObjectEditorViewport::OnUndoRedo()
 {
 	LevelViewportClient->Invalidate();
@@ -585,8 +618,43 @@ void SCustomizableObjectEditorViewportTabBody::BindCommands()
 		FCanExecuteAction(),
 		FIsActionChecked::CreateSP(EditorViewportClientRef, &FCustomizableObjectEditorViewportClient::IsActiveViewportType, LVT_OrthoNegativeXY));
 
-	// all other LODs will be added dynamically 
+	CommandList.MapAction( 
+		Commands.ShowDisplayInfo,
+		FExecuteAction::CreateSP(EditorViewportClientRef, &FCustomizableObjectEditorViewportClient::OnShowDisplayInfo),
+		FCanExecuteAction(),
+		FIsActionChecked::CreateSP(EditorViewportClientRef, &FCustomizableObjectEditorViewportClient::IsShowingMeshInfo));
+
+	CommandList.MapAction( 
+		Commands.EnableClothSimulation,
+		FExecuteAction::CreateSP(EditorViewportClientRef, &FCustomizableObjectEditorViewportClient::OnEnableClothSimulation),
+		FCanExecuteAction(),
+		FIsActionChecked::CreateSP(EditorViewportClientRef, &FCustomizableObjectEditorViewportClient::IsClothSimulationEnabled));
+
+	CommandList.MapAction( 
+		Commands.DebugDrawPhysMeshWired,
+		FExecuteAction::CreateSP(EditorViewportClientRef, &FCustomizableObjectEditorViewportClient::OnDebugDrawPhysMeshWired),
+		FCanExecuteAction(),
+		FIsActionChecked::CreateSP(EditorViewportClientRef, &FCustomizableObjectEditorViewportClient::IsDebugDrawPhysMeshWired));
+	
+	CommandList.MapAction(
+		Commands.SetShowNormals,
+		FExecuteAction::CreateSP( EditorViewportClientRef, &FCustomizableObjectEditorViewportClient::ToggleShowNormals),
+		FCanExecuteAction(),
+		FIsActionChecked::CreateSP( EditorViewportClientRef, &FCustomizableObjectEditorViewportClient::IsSetShowNormalsChecked));
+
+	CommandList.MapAction(
+		Commands.SetShowTangents,
+		FExecuteAction::CreateSP( EditorViewportClientRef, &FCustomizableObjectEditorViewportClient::ToggleShowTangents),
+		FCanExecuteAction(),
+		FIsActionChecked::CreateSP( EditorViewportClientRef, &FCustomizableObjectEditorViewportClient::IsSetShowTangentsChecked)); 
 
+	CommandList.MapAction(
+		Commands.SetShowBinormals,
+		FExecuteAction::CreateSP( EditorViewportClientRef, &FCustomizableObjectEditorViewportClient::ToggleShowBinormals),
+		FCanExecuteAction(),
+		FIsActionChecked::CreateSP( EditorViewportClientRef, &FCustomizableObjectEditorViewportClient::IsSetShowBinormalsChecked));
+	
+	// all other LODs will be added dynamically 
 }
 
 
diff --git a/Engine/Plugins/Experimental/Mutable/Source/CustomizableObjectEditor/Private/MuCOE/SCustomizableObjectEditorViewport.h b/Engine/Plugins/Experimental/Mutable/Source/CustomizableObjectEditor/Private/MuCOE/SCustomizableObjectEditorViewport.h
--- a/Engine/Plugins/Experimental/Mutable/Source/CustomizableObjectEditor/Private/MuCOE/SCustomizableObjectEditorViewport.h
+++ b/Engine/Plugins/Experimental/Mutable/Source/CustomizableObjectEditor/Private/MuCOE/SCustomizableObjectEditorViewport.h
@@ -63,6 +63,12 @@ private:
 	FText GetCompileErrorOverlayText() const;
 
 	FText GetWarningText() const;
+
+	EVisibility GetWarningTextVisibility() const;
+
+	FText GetMeshInfoText() const;
+
+	EVisibility GetMeshInfoTextVisibility() const;
 	
 	// Viewport client
 	TSharedPtr<FCustomizableObjectEditorViewportClient> LevelViewportClient;
diff --git a/Engine/Plugins/Experimental/Mutable/Source/CustomizableObjectEditor/Private/MuCOE/SCustomizableObjectEditorViewportToolBar.cpp b/Engine/Plugins/Experimental/Mutable/Source/CustomizableObjectEditor/Private/MuCOE/SCustomizableObjectEditorViewportToolBar.cpp
--- a/Engine/Plugins/Experimental/Mutable/Source/CustomizableObjectEditor/Private/MuCOE/SCustomizableObjectEditorViewportToolBar.cpp
+++ b/Engine/Plugins/Experimental/Mutable/Source/CustomizableObjectEditor/Private/MuCOE/SCustomizableObjectEditorViewportToolBar.cpp
@@ -26,6 +26,7 @@
 #include "MuCO/CustomizableObject.h"
 #include "MuCO/CustomizableObjectPrivate.h"
 #include "MuCO/CustomizableObjectSystem.h"
+#include "MuCOE/CustomizableObjectEditorActions.h"
 #include "ViewportToolbar/UnrealEdViewportToolbar.h"
 #include "Widgets/Input/SButton.h"
 #include "Widgets/Input/SMenuAnchor.h"
@@ -104,6 +105,18 @@ void SCustomizableObjectEditorViewportToolBar::Construct(const FArguments& InArg
 			.ParentToolBar(SharedThis(this))
 			.Label(LOCTEXT("ViewOptionsMenuLabel","View Options"))
 			.OnGetMenuContent(this, &SCustomizableObjectEditorViewportToolBar::GenerateViewportOptionsMenu)
+		]
+	
+	// Character Menu
+	+SHorizontalBox::Slot()
+		.AutoWidth()
+		.Padding(2.0f, 2.0f)
+		[
+			//Show Bones
+			SNew(SEditorViewportToolbarMenu)
+			.ParentToolBar(SharedThis(this))
+			.Label(LOCTEXT("CharacterMenuLabel","Character"))
+			.OnGetMenuContent(this, &SCustomizableObjectEditorViewportToolBar::GenerateCharacterMenu)
 		];
 	
 	TSharedRef<SWidget> RTSButtons = GenerateRTSButtons();
@@ -270,6 +283,34 @@ TSharedRef<SWidget> SCustomizableObjectEditorViewportToolBar::GenerateLODMenu()
 	return ShowMenuBuilder.MakeWidget();
 }
 
+
+TSharedRef<SWidget> SCustomizableObjectEditorViewportToolBar::GenerateCharacterMenu() const
+{
+	FMenuBuilder MenuBuilder(true, Viewport.Pin()->GetCommandList());
+	
+	MenuBuilder.BeginSection("Mesh", LOCTEXT("Mesh", "Mesh"));
+	MenuBuilder.AddMenuEntry(FCustomizableObjectEditorViewportCommands::Get().ShowDisplayInfo);
+
+	// Uncomment once UE-217529 fixed.
+	// MenuBuilder.AddMenuEntry(FCustomizableObjectEditorViewportCommands::Get().SetShowNormals);
+	// MenuBuilder.AddMenuEntry(FCustomizableObjectEditorViewportCommands::Get().SetShowTangents);
+	// MenuBuilder.AddMenuEntry(FCustomizableObjectEditorViewportCommands::Get().SetShowBinormals);
+
+	MenuBuilder.EndSection();
+
+	MenuBuilder.BeginSection("Bones", LOCTEXT("Bones", "Bones"));
+	MenuBuilder.AddMenuEntry(FCustomizableObjectEditorViewportLODCommands::Get().ShowBones);
+	MenuBuilder.EndSection();
+
+	MenuBuilder.BeginSection("Clothing", LOCTEXT("Clothing", "Clothing"));
+	MenuBuilder.AddMenuEntry(FCustomizableObjectEditorViewportCommands::Get().EnableClothSimulation);
+	MenuBuilder.AddMenuEntry(FCustomizableObjectEditorViewportCommands::Get().DebugDrawPhysMeshWired);
+	MenuBuilder.EndSection();
+	
+	return MenuBuilder.MakeWidget();
+}
+
+
 TSharedRef<SWidget> SCustomizableObjectEditorViewportToolBar::GenerateViewportTypeMenu() const
 {
 	const bool bInShouldCloseWindowAfterMenuSelection = true;
@@ -308,61 +349,45 @@ TSharedRef<SWidget> SCustomizableObjectEditorViewportToolBar::GenerateViewportOp
 	const bool bInShouldCloseWindowAfterMenuSelection = true;
 
 	FMenuBuilder ShowMenuBuilder(bInShouldCloseWindowAfterMenuSelection, Viewport.Pin()->GetCommandList());
-	{
-		ShowMenuBuilder.AddSubMenu(LOCTEXT("OptionsMenu_CameraOptions", "Camera Mode"),
-			LOCTEXT("OptionsMenu_CameraOptionsTooltip", "Select the camera mode"),
-		FNewMenuDelegate::CreateLambda([this, &Actions](FMenuBuilder& SubMenuBuilder)
-		{
-			SubMenuBuilder.BeginSection("Camera");
-			{
-				TSharedPtr<SWidget> BoneSizeWidget = SNew(SVerticalBox)
-				+ SVerticalBox::Slot().AutoHeight()
-				.HAlign(HAlign_Left)
-				.Padding(FMargin(20.0f, 5.0f, 0.0f, 0.0f))
-				[
-					SNew(STextBlock)
-					.Text(LOCTEXT("OptionsMenu_CameraOptions_CameraSpeed_Text", "Camera Speed"))
-					.Font(UE_MUTABLE_GET_FONTSTYLE(TEXT("MenuItem.Font")))
-				]
+	ShowMenuBuilder.BeginSection("Camera", LOCTEXT("Camera", "Camera"));
 
-				+SVerticalBox::Slot().AutoHeight()
-				.HAlign(HAlign_Left)
-				.Padding(FMargin(20.0f, 0.0f, 0.0f, 0.0f))
-				[
-					SNew(SHorizontalBox)
-					+ SHorizontalBox::Slot().AutoWidth()
-					[
-						SNew(SBox).WidthOverride(100.0f)
-						[
-							SNew(SSlider)
-							.Value_Lambda([this]() {return Viewport.Pin().Get()->GetViewportCameraSpeed(); })
-							.MinValue(1).MaxValue(4)
-							.OnValueChanged_Lambda([this](int32 Value) { Viewport.Pin().Get()->SetViewportCameraSpeed(Value); })
-						]
-					]
-					+ SHorizontalBox::Slot().AutoWidth()
-					[
-						SNew(STextBlock).Text_Lambda([this]() { return FText::AsNumber(Viewport.Pin().Get()->GetViewportCameraSpeed()); })
-					]
-				];
-				
-				SubMenuBuilder.AddMenuEntry(Actions.OrbitalCamera);
-				SubMenuBuilder.AddMenuEntry(Actions.FreeCamera);
-				SubMenuBuilder.AddWidget(BoneSizeWidget.ToSharedRef(), LOCTEXT("OptionMenu_CameraOptions_CameraSpeed", ""));
-			}
-			SubMenuBuilder.EndSection();
-		}));
+	ShowMenuBuilder.AddMenuEntry(Actions.OrbitalCamera);
+	ShowMenuBuilder.AddMenuEntry(Actions.FreeCamera);
+	
+	TSharedPtr<SWidget> BoneSizeWidget = SNew(SVerticalBox)
+	+ SVerticalBox::Slot().AutoHeight()
+	.HAlign(HAlign_Left)
+	.Padding(FMargin(20.0f, 5.0f, 0.0f, 0.0f))
+	[
+		SNew(STextBlock)
+		.Text(LOCTEXT("OptionsMenu_CameraOptions_CameraSpeed_Text", "Camera Speed"))
+		.Font(UE_MUTABLE_GET_FONTSTYLE(TEXT("MenuItem.Font")))
+	]
 
-		ShowMenuBuilder.AddSubMenu(LOCTEXT("OptionsMenu_BoneOptions", "Bones"), LOCTEXT("OptionsMenu_BoneOptionsTooltip", "Show/hide bone hierarchy"),
-		FNewMenuDelegate::CreateLambda([&Actions](FMenuBuilder& SubMenuBuilder)
-		{
-			SubMenuBuilder.BeginSection("Bones");
-			{
-				SubMenuBuilder.AddMenuEntry(Actions.ShowBones);
-			}
-			SubMenuBuilder.EndSection();
-		}));
-	}
+	+SVerticalBox::Slot().AutoHeight()
+	.HAlign(HAlign_Left)
+	.Padding(FMargin(20.0f, 0.0f, 0.0f, 0.0f))
+	[
+		SNew(SHorizontalBox)
+		+ SHorizontalBox::Slot().AutoWidth()
+		[
+			SNew(SBox).WidthOverride(100.0f)
+			[
+				SNew(SSlider)
+				.Value_Lambda([this]() {return Viewport.Pin().Get()->GetViewportCameraSpeed(); })
+				.MinValue(1).MaxValue(4)
+				.OnValueChanged_Lambda([this](int32 Value) { Viewport.Pin().Get()->SetViewportCameraSpeed(Value); })
+			]
+		]
+		+ SHorizontalBox::Slot().AutoWidth()
+		[
+			SNew(STextBlock).Text_Lambda([this]() { return FText::AsNumber(Viewport.Pin().Get()->GetViewportCameraSpeed()); })
+		]
+	];
+	
+	ShowMenuBuilder.AddWidget(BoneSizeWidget.ToSharedRef(), LOCTEXT("OptionMenu_CameraOptions_CameraSpeed", ""));
+	
+	ShowMenuBuilder.EndSection();
 
 	return ShowMenuBuilder.MakeWidget();
 }
diff --git a/Engine/Plugins/Experimental/Mutable/Source/CustomizableObjectEditor/Public/MuCOE/CustomizableObjectEditorActions.h b/Engine/Plugins/Experimental/Mutable/Source/CustomizableObjectEditor/Public/MuCOE/CustomizableObjectEditorActions.h
--- a/Engine/Plugins/Experimental/Mutable/Source/CustomizableObjectEditor/Public/MuCOE/CustomizableObjectEditorActions.h
+++ b/Engine/Plugins/Experimental/Mutable/Source/CustomizableObjectEditor/Public/MuCOE/CustomizableObjectEditorActions.h
@@ -44,17 +44,22 @@ class FCustomizableObjectEditorViewportCommands : public TCommands<FCustomizable
 public:
 	FCustomizableObjectEditorViewportCommands();
 	
-	TSharedPtr< FUICommandInfo > SetDrawUVs;
-	TSharedPtr< FUICommandInfo > SetShowGrid;
-	TSharedPtr< FUICommandInfo > SetShowSky;
-	TSharedPtr< FUICommandInfo > SetShowBounds;
-	TSharedPtr< FUICommandInfo > SetShowCollision;
-	TSharedPtr< FUICommandInfo > SetCameraLock;
-	TSharedPtr< FUICommandInfo > SaveThumbnail;
-
-	TSharedPtr< FUICommandInfo > BakeInstance;
-	TSharedPtr< FUICommandInfo > StateChangeShowData;
-	TSharedPtr< FUICommandInfo > StateChangeShowGeometryData;
+	TSharedPtr<FUICommandInfo> SetDrawUVs;
+	TSharedPtr<FUICommandInfo> SetShowGrid;
+	TSharedPtr<FUICommandInfo> SetShowSky;
+	TSharedPtr<FUICommandInfo> SetShowBounds;
+	TSharedPtr<FUICommandInfo> SetShowCollision;
+	TSharedPtr<FUICommandInfo> SetCameraLock;
+	TSharedPtr<FUICommandInfo> SaveThumbnail;
+	TSharedPtr<FUICommandInfo> BakeInstance;
+	TSharedPtr<FUICommandInfo> StateChangeShowData;
+	TSharedPtr<FUICommandInfo> StateChangeShowGeometryData;
+	TSharedPtr<FUICommandInfo> ShowDisplayInfo;
+	TSharedPtr<FUICommandInfo> EnableClothSimulation;
+	TSharedPtr<FUICommandInfo> DebugDrawPhysMeshWired;
+	TSharedPtr<FUICommandInfo> SetShowNormals;
+	TSharedPtr<FUICommandInfo> SetShowTangents;
+	TSharedPtr<FUICommandInfo> SetShowBinormals;
 
 	/**
 	 * Initialize commands
diff --git a/Engine/Plugins/Experimental/Mutable/Source/CustomizableObjectEditor/Public/MuCOE/CustomizableObjectEditorViewportClient.h b/Engine/Plugins/Experimental/Mutable/Source/CustomizableObjectEditor/Public/MuCOE/CustomizableObjectEditorViewportClient.h
--- a/Engine/Plugins/Experimental/Mutable/Source/CustomizableObjectEditor/Public/MuCOE/CustomizableObjectEditorViewportClient.h
+++ b/Engine/Plugins/Experimental/Mutable/Source/CustomizableObjectEditor/Public/MuCOE/CustomizableObjectEditorViewportClient.h
@@ -70,7 +70,7 @@ class FCustomizableObjectEditorViewportClient : public FEditorViewportClient, pu
 public:
 	FCustomizableObjectEditorViewportClient(TWeakPtr<ICustomizableObjectInstanceEditor> InCustomizableObjectEditor, FPreviewScene* InPreviewScene, const TSharedPtr<SEditorViewport>& EditorViewportWidget);
 	~FCustomizableObjectEditorViewportClient();
-
+	
 	/** persona config options **/
 	class UPersonaOptions* ConfigOption;
 
@@ -257,6 +257,32 @@ public:
 	bool IsShowingBones() const;
 
 	const TArray<ULightComponent*>& GetLightComponents() const;
+
+	void OnShowDisplayInfo();
+
+	bool IsShowingMeshInfo() const;
+	
+	void OnEnableClothSimulation();
+	
+	bool IsClothSimulationEnabled() const;
+
+	void OnDebugDrawPhysMeshWired();
+
+	bool IsDebugDrawPhysMeshWired() const;
+	
+	FText GetMeshInfoText() const;
+	
+	void ToggleShowNormals();
+
+	bool IsSetShowNormalsChecked() const;
+	
+	void ToggleShowTangents();
+
+	bool IsSetShowTangentsChecked() const;
+	
+	void ToggleShowBinormals();
+
+	bool IsSetShowBinormalsChecked() const;
 	
 private:
 	/** Draws Mesh Bones in foreground (From: FAnimationViewportClient) */
@@ -337,6 +363,9 @@ private:
 	/** Flag to control the bones visibility in the viewport */
 	bool bShowBones;
 
+	/** Draw wireframe physics mesh. */
+	bool bShowDebugClothing = false;
+	
 	// Temp Instance used in the bake process if a new instance is needed because mutable texture streaming is enabled so the viewport 
 	// instance does not have the high quality mips in the texture's platform data
 	TObjectPtr<UCustomizableObjectInstance> BakeTempInstance = nullptr;
@@ -344,6 +373,21 @@ private:
 	// Cache System configuration cached before performing the mandatory instance update for baking so we can restore it after the bake operation
 	bool bIsProgressiveMipStreamingEnabled = false;
 	bool bIsOnlyGenerateRequestedLODsEnabled = false;
+
+	/** Show detailed mesh info text. */
+	bool bShowDisplayInfo = false;
+
+	/** See USkeletalMeshComponent::bDisableClothSimulation. */
+	bool bDisableClothSimulation = false;
+
+	/** See USkeletalMeshComponent::bDrawNormals. */
+	bool bDrawNormals = false;
+
+	/** See USkeletalMeshComponent::bDrawTangents. */
+	bool bDrawTangents = false;
+
+	/** See USkeletalMeshComponent::bDrawBinormals. */
+	bool bDrawBinormals = false;
 	
 	// The following delegates currently are only used by the EWidgetType::Projector
 	FWidgetLocationDelegate WidgetLocationDelegate;
diff --git a/Engine/Plugins/Experimental/Mutable/Source/CustomizableObjectEditor/Public/MuCOE/SCustomizableObjectEditorViewportToolBar.h b/Engine/Plugins/Experimental/Mutable/Source/CustomizableObjectEditor/Public/MuCOE/SCustomizableObjectEditorViewportToolBar.h
--- a/Engine/Plugins/Experimental/Mutable/Source/CustomizableObjectEditor/Public/MuCOE/SCustomizableObjectEditorViewportToolBar.h
+++ b/Engine/Plugins/Experimental/Mutable/Source/CustomizableObjectEditor/Public/MuCOE/SCustomizableObjectEditorViewportToolBar.h
@@ -53,6 +53,13 @@ private:
 	*/
 	FText GetLODMenuLabel() const;
 
+	/**
+	* Generates the toolbar Character menu content
+	*
+	* @return The widget containing the view menu content
+	*/
+	TSharedRef<SWidget> GenerateCharacterMenu() const;
+	
 	/**
 	* Generates the toolbar viewport type menu content
 	*
