[Mutable] Fix ToggleShowNormals unused expression result.

#rnx

diff --git a/Engine/Plugins/Experimental/Mutable/Source/CustomizableObjectEditor/Private/MuCOE/CustomizableObjectEditorViewportClient.cpp b/Engine/Plugins/Experimental/Mutable/Source/CustomizableObjectEditor/Private/MuCOE/CustomizableObjectEditorViewportClient.cpp
--- a/Engine/Plugins/Experimental/Mutable/Source/CustomizableObjectEditor/Private/MuCOE/CustomizableObjectEditorViewportClient.cpp
+++ b/Engine/Plugins/Experimental/Mutable/Source/CustomizableObjectEditor/Private/MuCOE/CustomizableObjectEditorViewportClient.cpp
@@ -2024,7 +2024,7 @@ void FCustomizableObjectEditorViewportClient::ToggleShowNormals()
 	{
 		if (SkeletalMeshComponent.IsValid())
 		{
-			SkeletalMeshComponent->bDrawNormals;
+			SkeletalMeshComponent->bDrawNormals = bDrawNormals;
 			SkeletalMeshComponent->MarkRenderStateDirty();
 		}
 	}
