[Mutable] Fix hang when BlockTillAllRequestsFinished and waiting for AdditionalAssetsAsyncLoaded.

#rnx

diff --git a/Engine/Plugins/Experimental/Mutable/Source/CustomizableObject/Internal/MuCO/CustomizableObjectInstancePrivate.h b/Engine/Plugins/Experimental/Mutable/Source/CustomizableObject/Internal/MuCO/CustomizableObjectInstancePrivate.h
--- a/Engine/Plugins/Experimental/Mutable/Source/CustomizableObject/Internal/MuCO/CustomizableObjectInstancePrivate.h
+++ b/Engine/Plugins/Experimental/Mutable/Source/CustomizableObject/Internal/MuCO/CustomizableObjectInstancePrivate.h
@@ -12,6 +12,7 @@
 #include "GameplayTagContainer.h"
 #include "MuCO/DescriptorHash.h"
 #include "UObject/Package.h"
+#include "Tasks/Task.h"
 
 #include "CustomizableObjectInstancePrivate.generated.h"
 
@@ -194,12 +195,15 @@ public:
 
 	// Handle used to store a streaming request operation if one is ongoing.
 	TSharedPtr<FStreamableHandle> StreamingHandle;
+	UE::Tasks::FTaskEvent AssetAsyncLoadCompletionEvent = UE::Tasks::FTaskEvent(TEXT(""));
+	TArray<FSoftObjectPath> AssetsToStream;
 
 	// Only used in LiveUpdateMode to reuse core instances between updates and their temp data to speed up updates, but spend way more memory
 	mu::Instance::ID LiveUpdateModeInstanceID = 0;
 
-#if WITH_EDITOR
+	UCustomizableInstancePrivate();
 
+#if WITH_EDITOR
 	virtual void PostDuplicate(bool bDuplicateForPIE) override;
 
 	void BindObjectDelegates(UCustomizableObject* CurrentCustomizableObject, UCustomizableObject* NewCustomizableObject);
@@ -207,7 +211,7 @@ public:
 	void OnPostCompile();
 	void OnObjectStatusChanged(FCustomizableObjectStatus::EState Previous, FCustomizableObjectStatus::EState Next);
 #endif
-
+	
 	/** Invalidates the previously generated data and retrieves information from the CObject after specific actions.
 	 *  It'll be called in the PostLoad, after Compiling the CO, and after changing the CO of the Instance. */
 	void InitCustomizableObjectData(const UCustomizableObject* InCustomizableObject);
@@ -226,7 +230,7 @@ public:
 
 	/** Returns the task that will be called when all assets and data are loaded, may be already completed if no assets or data needs loading.
 	 * If no StreamableManager is provided, it will load assets synchronously. */
-	UE::Tasks::FTask LoadAdditionalAssetsAndData(const TSharedRef<FUpdateContextPrivate>& OperationData, FStreamableManager& StreamableManager, bool bAsync);
+	UE::Tasks::FTask LoadAdditionalAssetsAndData(const TSharedRef<FUpdateContextPrivate>& OperationData, FStreamableManager& StreamableManager);
 
 	void AdditionalAssetsAsyncLoaded(UCustomizableObjectInstance* Public);
 
@@ -304,7 +308,7 @@ public:
 
 	void SetState(int32 InState);
 
-	void AdditionalAssetsAsyncLoaded(UE::Tasks::FTaskEvent CompletionEvent);
+	void AdditionalAssetsAsyncLoaded();
 
 	FCustomizableObjectInstanceDescriptor& GetDescriptor() const;
 
diff --git a/Engine/Plugins/Experimental/Mutable/Source/CustomizableObject/Internal/MuCO/CustomizableObjectSystemPrivate.h b/Engine/Plugins/Experimental/Mutable/Source/CustomizableObject/Internal/MuCO/CustomizableObjectSystemPrivate.h
--- a/Engine/Plugins/Experimental/Mutable/Source/CustomizableObject/Internal/MuCO/CustomizableObjectSystemPrivate.h
+++ b/Engine/Plugins/Experimental/Mutable/Source/CustomizableObject/Internal/MuCO/CustomizableObjectSystemPrivate.h
@@ -791,9 +791,6 @@ public:
 
 	bool bAutoCompileCommandletEnabled = false;
 
-	/** We are in BlockTillAllRequestsFinished scope. */
-	bool bBlocking = false;
-
 	UPROPERTY()
 	TArray<FPendingReleaseSkeletalMeshInfo> PendingReleaseSkeletalMesh;
 	
diff --git a/Engine/Plugins/Experimental/Mutable/Source/CustomizableObject/Private/MuCO/CustomizableObjectInstance.cpp b/Engine/Plugins/Experimental/Mutable/Source/CustomizableObject/Private/MuCO/CustomizableObjectInstance.cpp
--- a/Engine/Plugins/Experimental/Mutable/Source/CustomizableObject/Private/MuCO/CustomizableObjectInstance.cpp
+++ b/Engine/Plugins/Experimental/Mutable/Source/CustomizableObject/Private/MuCO/CustomizableObjectInstance.cpp
@@ -298,9 +298,13 @@ void UCustomizableInstancePrivate::PrepareForUpdate(const TSharedRef<FUpdateCont
 }
 
 
-#if WITH_EDITOR
+UCustomizableInstancePrivate::UCustomizableInstancePrivate()
+{
+	AssetAsyncLoadCompletionEvent.Trigger(); // MakeCompletedEvent does not exist. Trigger the placeholder event.
+}
 
 
+#if WITH_EDITOR
 void UCustomizableInstancePrivate::PostDuplicate(bool bDuplicateForPIE)
 {
 	Super::PostDuplicate(bDuplicateForPIE);
@@ -5164,7 +5168,7 @@ FAutoConsoleVariableRef CVarMutableHighPriorityLoading(
 
 
 UE::Tasks::FTask UCustomizableInstancePrivate::LoadAdditionalAssetsAndData(
-		const TSharedRef<FUpdateContextPrivate>& OperationData, FStreamableManager& StreamableManager, bool bAsync)
+		const TSharedRef<FUpdateContextPrivate>& OperationData, FStreamableManager& StreamableManager)
 {
 	MUTABLE_CPUPROFILER_SCOPE(UCustomizableInstancePrivate::LoadAdditionalAssetsAndDataAsync);
 
@@ -5172,7 +5176,7 @@ UE::Tasks::FTask UCustomizableInstancePrivate::LoadAdditionalAssetsAndData(
 
 	const FModelResources& ModelResources = CustomizableObject->GetPrivate()->GetModelResources();
 
-	TArray<FSoftObjectPath> AssetsToStream;
+	AssetsToStream.Empty();
 	TArray<uint32> RealTimeMorphStreamableBlocksToStream;
 	TArray<uint32> ClothingStreamableBlocksToStream;
 
@@ -5472,8 +5476,9 @@ UE::Tasks::FTask UCustomizableInstancePrivate::LoadAdditionalAssetsAndData(
 
 	TArray<UE::Tasks::FTaskEvent> StreamingCompletionEvents;
 	if (AssetsToStream.Num() > 0)
-	{	
-		UE::Tasks::FTaskEvent AssetAsyncLoadCompletionEvent = StreamingCompletionEvents.Emplace_GetRef(TEXT("AssetAsyncLoadCompletionEvent"));
+	{
+		check(AssetAsyncLoadCompletionEvent.IsCompleted());
+		AssetAsyncLoadCompletionEvent = StreamingCompletionEvents.Emplace_GetRef(TEXT("AssetAsyncLoadCompletionEvent"));
 
 #if WITH_EDITOR
 		// TODO: Remove with UE-217665 when the underlying bug in the ColorPicker is solved
@@ -5482,19 +5487,10 @@ UE::Tasks::FTask UCustomizableInstancePrivate::LoadAdditionalAssetsAndData(
 		FSlateThrottleManager::Get().DisableThrottle(true);
 #endif
 
-		if (bAsync)
-		{
-			StreamingHandle = StreamableManager.RequestAsyncLoad(
-					AssetsToStream, 
-					FStreamableDelegate::CreateUObject(this, &UCustomizableInstancePrivate::AdditionalAssetsAsyncLoaded, AssetAsyncLoadCompletionEvent),
-					bEnableHighPriorityLoading ? FStreamableManager::AsyncLoadHighPriority : FStreamableManager::DefaultAsyncLoadPriority);			
-		}
-		else
-		{
-			StreamableManager.RequestSyncLoad(AssetsToStream);
-			
-			AdditionalAssetsAsyncLoaded(AssetAsyncLoadCompletionEvent);
-		}
+		StreamingHandle = StreamableManager.RequestAsyncLoad(
+				AssetsToStream, 
+				FStreamableDelegate::CreateUObject(this, &UCustomizableInstancePrivate::AdditionalAssetsAsyncLoaded),
+				bEnableHighPriorityLoading ? FStreamableManager::AsyncLoadHighPriority : FStreamableManager::DefaultAsyncLoadPriority);			
 	}
 
 	
@@ -5737,12 +5733,12 @@ UE::Tasks::FTask UCustomizableInstancePrivate::LoadAdditionalAssetsAndData(
 	}
 }
 
-void UCustomizableInstancePrivate::AdditionalAssetsAsyncLoaded(UE::Tasks::FTaskEvent CompletionEvent)
+void UCustomizableInstancePrivate::AdditionalAssetsAsyncLoaded()
 {
 	// TODO: Do we need this separated?
 	//check(IsInGameThread())
 	AdditionalAssetsAsyncLoaded(GetPublic());
-	CompletionEvent.Trigger(); // TODO: we know it is game thread?
+	AssetAsyncLoadCompletionEvent.Trigger(); // TODO: we know it is game thread?
 
 	StreamingHandle = nullptr;
 
diff --git a/Engine/Plugins/Experimental/Mutable/Source/CustomizableObject/Private/MuCO/CustomizableObjectSystem.cpp b/Engine/Plugins/Experimental/Mutable/Source/CustomizableObject/Private/MuCO/CustomizableObjectSystem.cpp
--- a/Engine/Plugins/Experimental/Mutable/Source/CustomizableObject/Private/MuCO/CustomizableObjectSystem.cpp
+++ b/Engine/Plugins/Experimental/Mutable/Source/CustomizableObject/Private/MuCO/CustomizableObjectSystem.cpp
@@ -3141,7 +3141,7 @@ namespace impl
 
 		// Next Task: Load Unreal Assets
 		//-------------------------------------------------------------
-		UE::Tasks::FTask Game_LoadUnrealAssets = ObjectInstancePrivateData->LoadAdditionalAssetsAndData(OperationData, System->GetPrivate()->StreamableManager, !System->GetPrivate()->bBlocking);
+		UE::Tasks::FTask Game_LoadUnrealAssets = ObjectInstancePrivateData->LoadAdditionalAssetsAndData(OperationData, System->GetPrivate()->StreamableManager);
 
 		// Next-next Task: Convert Resources
 		//-------------------------------------------------------------
@@ -3559,12 +3559,12 @@ void UCustomizableObjectSystem::AdvanceCurrentOperation()
 
 bool UCustomizableObjectSystem::Tick(float DeltaTime)
 {
-	TickInternal();
+	TickInternal(false);
 	return true;
 }
 
 
-int32 UCustomizableObjectSystem::TickInternal()
+int32 UCustomizableObjectSystem::TickInternal(const bool bBlocking)
 {
 	MUTABLE_CPUPROFILER_SCOPE(UCustomizableObjectSystem::TickInternal)
 
@@ -3817,13 +3817,25 @@ int32 UCustomizableObjectSystem::TickInternal()
 		static_cast<int32>(LODUpdateCandidateFound != nullptr) + // Still a pending LOD update. We can not use the size of RequestedLODUpdates since not all requests valid in future ticks.
 		RemainingTasks;
 
-#if WITH_EDITOR
-	if (Private->bBlocking)
+	if (bBlocking)
 	{
+#if WITH_EDITOR
 		ICustomizableObjectEditorModule* EditorModule = ICustomizableObjectEditorModule::Get();
 		RemainingWork += EditorModule ? EditorModule->Tick(true) : 0;
-	}
 #endif
+		
+		if (GetPrivate()->CurrentMutableOperation)
+		{
+			UCustomizableInstancePrivate* InstancePrivate = GetPrivate()->CurrentMutableOperation->Instance->GetPrivate();
+			
+			if (InstancePrivate->StreamingHandle)
+			{
+				InstancePrivate->StreamingHandle->CancelHandle();
+				GetPrivate()->StreamableManager.RequestSyncLoad(InstancePrivate->AssetsToStream);
+				InstancePrivate->AdditionalAssetsAsyncLoaded();
+			}
+		}
+	}
 	
 	return RemainingWork;
 }
@@ -4316,15 +4328,12 @@ bool UCustomizableObjectSystem::IsMutableAnimInfoDebuggingEnabled() const
 
 void UCustomizableObjectSystemPrivate::UpdateResourceStreaming(float DeltaTime, bool bProcessEverything)
 {
-	GetPublic()->TickInternal();
+	GetPublic()->TickInternal(false);
 }
 
 
 int32 UCustomizableObjectSystemPrivate::BlockTillAllRequestsFinished(float TimeLimit, bool bLogResults)
 {
-	bBlocking = true;
-	ON_SCOPE_EXIT {bBlocking = false; };
-	
 	const double BlockEndTime = FPlatformTime::Seconds() + TimeLimit;
 
 	int32 RemainingWork = TNumericLimits<int32>::Max();
@@ -4333,7 +4342,7 @@ int32 UCustomizableObjectSystemPrivate::BlockTillAllRequestsFinished(float TimeL
 	{
 		while (RemainingWork > 0)
 		{
-			RemainingWork = GetPublic()->TickInternal();
+			RemainingWork = GetPublic()->TickInternal(true);
 		}
 	}
 	else
@@ -4345,7 +4354,7 @@ int32 UCustomizableObjectSystemPrivate::BlockTillAllRequestsFinished(float TimeL
 				return RemainingWork;
 			}
 			
-			RemainingWork = GetPublic()->TickInternal();
+			RemainingWork = GetPublic()->TickInternal(true);
 		}
 	}
 
diff --git a/Engine/Plugins/Experimental/Mutable/Source/CustomizableObject/Public/MuCO/CustomizableObjectSystem.h b/Engine/Plugins/Experimental/Mutable/Source/CustomizableObject/Public/MuCO/CustomizableObjectSystem.h
--- a/Engine/Plugins/Experimental/Mutable/Source/CustomizableObject/Public/MuCO/CustomizableObjectSystem.h
+++ b/Engine/Plugins/Experimental/Mutable/Source/CustomizableObject/Public/MuCO/CustomizableObjectSystem.h
@@ -327,7 +327,7 @@ private:
 	bool Tick(float DeltaTime);
 
 	/** Returns the number of remaining operations. */
-	int32 TickInternal();
+	int32 TickInternal(bool bBlocking);
 
 	// If there is an on-going operation, advance it.
 	void AdvanceCurrentOperation();
