[Mutable] Remove BusyWaits_Deprecate.

#rnx

diff --git a/Engine/Plugins/Experimental/Mutable/Source/CustomizableObject/Private/MuCO/BusyWaits_Deprecated.cpp b/Engine/Plugins/Experimental/Mutable/Source/CustomizableObject/Private/MuCO/BusyWaits_Deprecated.cpp
deleted file mode 100644
diff --git a/Engine/Plugins/Experimental/Mutable/Source/CustomizableObject/Private/MuCO/BusyWaits_Deprecated.h b/Engine/Plugins/Experimental/Mutable/Source/CustomizableObject/Private/MuCO/BusyWaits_Deprecated.h
deleted file mode 100644
diff --git a/Engine/Plugins/Experimental/Mutable/Source/CustomizableObject/Private/MuCO/CustomizableObjectMeshUpdate.cpp b/Engine/Plugins/Experimental/Mutable/Source/CustomizableObject/Private/MuCO/CustomizableObjectMeshUpdate.cpp
--- a/Engine/Plugins/Experimental/Mutable/Source/CustomizableObject/Private/MuCO/CustomizableObjectMeshUpdate.cpp
+++ b/Engine/Plugins/Experimental/Mutable/Source/CustomizableObject/Private/MuCO/CustomizableObjectMeshUpdate.cpp
@@ -20,8 +20,6 @@ CustomizableObjectMeshUpdate.cpp: Helpers to stream in CustomizableObject skelet
 #include "Rendering/SkeletalMeshRenderData.h"
 
 
-#include "BusyWaits_Deprecated.h"
-
 template class TRenderAssetUpdate<FSkelMeshUpdateContext>;
 
 #define UE_MUTABLE_UPDATE_MESH_REGION		TEXT("Task_Mutable_UpdateMesh")
@@ -217,17 +215,7 @@ void FCustomizableObjectMeshStreamIn::RequestMeshUpdate(const FContext& Context)
 		TEXT("Mutable_MeshUpdate"),
 		[SharedOperationData, RefThis]() mutable
 		{
-			if (CVarEnableNewSplitMutableTask.GetValueOnAnyThread())
-			{
-				impl::Task_Mutable_UpdateMesh(SharedOperationData, RefThis);
-			}
-			else
-			{
-				using namespace CustomizableObjectMeshUpdate;
-				ImplDeprecated::Task_Mutable_UpdateMesh(SharedOperationData);
-				
-				RefThis->OnUpdateMeshFinished();
-			}
+			impl::Task_Mutable_UpdateMesh(SharedOperationData, RefThis);
 		});
 }
 
diff --git a/Engine/Plugins/Experimental/Mutable/Source/CustomizableObject/Private/MuCO/CustomizableObjectMipDataProvider.cpp b/Engine/Plugins/Experimental/Mutable/Source/CustomizableObject/Private/MuCO/CustomizableObjectMipDataProvider.cpp
--- a/Engine/Plugins/Experimental/Mutable/Source/CustomizableObject/Private/MuCO/CustomizableObjectMipDataProvider.cpp
+++ b/Engine/Plugins/Experimental/Mutable/Source/CustomizableObject/Private/MuCO/CustomizableObjectMipDataProvider.cpp
@@ -7,7 +7,6 @@
 #include "MuR/Model.h"
 #include "TextureResource.h"
 #include "UnrealMutableImageProvider.h"
-#include "MuCO/BusyWaits_Deprecated.h"
 
 #include "Containers/Ticker.h"
 
@@ -399,15 +398,7 @@ int32 FMutableTextureMipDataProvider::GetMips(const FTextureUpdateContext& Conte
 				TEXT("Mutable_MipUpdate"),
 				[OperationData = this->OperationData]()
 				{
-					if (CVarEnableNewSplitMutableTask.GetValueOnAnyThread())
-					{
-						Impl::Task_Mutable_UpdateImage(OperationData);
-					}
-					else
-					{
-						using namespace CustomizableObjectMipDataProvider;
-						ImplDeprecated::Task_Mutable_UpdateImage(OperationData);
-					}
+					Impl::Task_Mutable_UpdateImage(OperationData);
 				});
 		}
 
diff --git a/Engine/Plugins/Experimental/Mutable/Source/CustomizableObject/Private/MuCO/CustomizableObjectSystem.cpp b/Engine/Plugins/Experimental/Mutable/Source/CustomizableObject/Private/MuCO/CustomizableObjectSystem.cpp
--- a/Engine/Plugins/Experimental/Mutable/Source/CustomizableObject/Private/MuCO/CustomizableObjectSystem.cpp
+++ b/Engine/Plugins/Experimental/Mutable/Source/CustomizableObject/Private/MuCO/CustomizableObjectSystem.cpp
@@ -43,9 +43,6 @@
 #include "Engine/Engine.h"
 #endif
 
-
-#include "BusyWaits_Deprecated.h"
-
 #include UE_INLINE_GENERATED_CPP_BY_NAME(CustomizableObjectSystem)
 
 class AActor;
@@ -113,12 +110,6 @@ TAutoConsoleVariable<bool> CVarEnableMeshCache(
 	TEXT("Enables or disables the reuse of meshes."),
 	ECVF_Scalability);
 
-TAutoConsoleVariable<bool> CVarEnableNewSplitMutableTask(
-	TEXT("mutable.EnableNewSplitMutableTask"),
-	true,
-	TEXT("Enables or disables the then new split GetImages and GetMesh tasks that remove BusyWaits."),
-	ECVF_Scalability);
-
 TAutoConsoleVariable<bool> CVarEnableUpdateOptimization(
 	TEXT("mutable.EnableUpdateOptimization"),
 	false,
@@ -2428,7 +2419,7 @@ namespace impl
 	}
 
 
-	namespace impl_new
+	namespace Impl
 	{
 		/** Start of the GetMeshes tasks.
 		  * Gathers all GetMeshes that has to be called. */
@@ -2530,14 +2521,7 @@ namespace impl
 
 	void Task_Mutable_GetMeshes(const TSharedRef<FUpdateContextPrivate>& OperationData)
 	{
-		if (CVarEnableNewSplitMutableTask->GetBool())
-		{
-			impl_new::Task_Mutable_GetMeshes(OperationData);
-		}
-		else
-		{
-			CustomizableObjectSystem::ImplDeprecated::Task_Mutable_GetMeshes(OperationData);
-		}
+		Impl::Task_Mutable_GetMeshes(OperationData);
 	}
 
 	
@@ -2753,7 +2737,7 @@ namespace impl
 	}
 
 
-	namespace impl_new
+	namespace Impl
 	{
 		// This runs in a worker thread.
 		void Task_Mutable_GetImages(const TSharedRef<FUpdateContextPrivate>& OperationData)
@@ -2774,14 +2758,7 @@ namespace impl
 	/** Start of the GetImages tasks. */
 	void Task_Mutable_GetImages(const TSharedRef<FUpdateContextPrivate>& OperationData)
 	{
-		if (CVarEnableNewSplitMutableTask->GetBool())
-		{
-			impl_new::Task_Mutable_GetImages(OperationData);
-		}
-		else
-		{
-			CustomizableObjectSystem::ImplDeprecated::Task_Mutable_GetImages(OperationData);						
-		}
+		Impl::Task_Mutable_GetImages(OperationData);
 	}
 
 
diff --git a/Engine/Plugins/Experimental/Mutable/Source/CustomizableObject/Public/MuCO/CustomizableObjectSystem.h b/Engine/Plugins/Experimental/Mutable/Source/CustomizableObject/Public/MuCO/CustomizableObjectSystem.h
--- a/Engine/Plugins/Experimental/Mutable/Source/CustomizableObject/Public/MuCO/CustomizableObjectSystem.h
+++ b/Engine/Plugins/Experimental/Mutable/Source/CustomizableObject/Public/MuCO/CustomizableObjectSystem.h
@@ -39,8 +39,6 @@ extern TAutoConsoleVariable<bool> CVarEnableMeshCache;
 
 extern TAutoConsoleVariable<bool> CVarRollbackFixModelDiskStreamerDataRace;
 
-extern TAutoConsoleVariable<bool> CVarEnableNewSplitMutableTask;
-
 #if WITH_EDITOR
 
 // Struct used to keep a copy of the EditorSettings needed to compile Customizable Objects.
