[Mutable] Add zero initialization to MeshBufferSet. Better follow UE coding standars around MeshBufferSet struct.

- MeshBufferSet SetBuffer and SetElementCount now accept a new enum parameter EMemoryInitPolicy that governs how the newly allocated data will be initialized. Valid values are Uninitialized and Zeroed.

#rnx

diff --git a/Engine/Plugins/Experimental/Mutable/Source/CustomizableObjectEditor/Private/MuCOE/GenerateMutableSource/GenerateMutableSourceMesh.cpp b/Engine/Plugins/Experimental/Mutable/Source/CustomizableObjectEditor/Private/MuCOE/GenerateMutableSource/GenerateMutableSourceMesh.cpp
--- a/Engine/Plugins/Experimental/Mutable/Source/CustomizableObjectEditor/Private/MuCOE/GenerateMutableSource/GenerateMutableSourceMesh.cpp
+++ b/Engine/Plugins/Experimental/Mutable/Source/CustomizableObjectEditor/Private/MuCOE/GenerateMutableSource/GenerateMutableSourceMesh.cpp
@@ -45,6 +45,7 @@
 #include "MuT/NodeMeshSwitch.h"
 #include "MuT/NodeMeshTable.h"
 #include "MuT/NodeMeshVariation.h"
+#include "MuR/Types.h"
 #include "PhysicsEngine/PhysicsAsset.h"
 #include "PhysicsEngine/SkeletalBodySetup.h"
 #include "Engine/SkinnedAssetCommon.h"
@@ -614,33 +615,33 @@ namespace MutablePrivate
 	{
 		MUTABLE_CPUPROFILER_SCOPE(CopyBufferClearingPadding);
 
-		int32 ElementCount = BufferSet.m_elementCount;
+		int32 ElementCount = BufferSet.ElementCount;
 
 		for (int32 BufferIndex = 0; BufferIndex < BufferSet.GetBufferCount(); ++BufferIndex)
 		{
-			int32 ChannelCount = BufferSet.m_buffers[BufferIndex].m_channels.Num();
-			int32 ElementSize = BufferSet.m_buffers[BufferIndex].m_elementSize;
+			int32 ChannelCount = BufferSet.Buffers[BufferIndex].Channels.Num();
+			int32 ElementSize = BufferSet.Buffers[BufferIndex].ElementSize;
 
-			check(SourceBufferFormat.m_buffers.Num()==1);
-			const mu::FMeshBuffer& SourceBuffer = SourceBufferFormat.m_buffers[0];
+			check(SourceBufferFormat.Buffers.Num()==1);
+			const mu::FMeshBuffer& SourceBuffer = SourceBufferFormat.Buffers[0];
 
 			uint8* TargetData = BufferSet.GetBufferData(BufferIndex);
 			TArray<const uint8*, TInlineAllocator<8>> SourceDataPerChannel;
 			SourceDataPerChannel.SetNumUninitialized(ChannelCount);
 			for (int32 ChannelIndex = 0; ChannelIndex < ChannelCount; ++ChannelIndex)
 			{
-				const mu::FMeshBufferChannel& DestinationChannel = BufferSet.m_buffers[BufferIndex].m_channels[ChannelIndex];
+				const mu::FMeshBufferChannel& DestinationChannel = BufferSet.Buffers[BufferIndex].Channels[ChannelIndex];
 
 				int32 SourceBufferIndex = -1;
 				int32 SourceChannelIndex = -1;
-				SourceBufferFormat.FindChannel(DestinationChannel.m_semantic, DestinationChannel.m_semanticIndex, &SourceBufferIndex, &SourceChannelIndex);
+				SourceBufferFormat.FindChannel(DestinationChannel.Semantic, DestinationChannel.SemanticIndex, &SourceBufferIndex, &SourceChannelIndex);
 				check(SourceBufferIndex==0 && SourceChannelIndex >=0);
 
-				const mu::FMeshBufferChannel& SourceChannel = SourceBuffer.m_channels[SourceChannelIndex];
-				check(SourceChannel.m_format == DestinationChannel.m_format);
-				check(SourceChannel.m_componentCount == DestinationChannel.m_componentCount);
+				const mu::FMeshBufferChannel& SourceChannel = SourceBuffer.Channels[SourceChannelIndex];
+				check(SourceChannel.Format == DestinationChannel.Format);
+				check(SourceChannel.ComponentCount == DestinationChannel.ComponentCount);
 
-				SourceDataPerChannel[ChannelIndex] = reinterpret_cast<const uint8*>(InSourceData) + SourceChannel.m_offset;
+				SourceDataPerChannel[ChannelIndex] = reinterpret_cast<const uint8*>(InSourceData) + SourceChannel.Offset;
 			}
 
 			for (int32 Element = 0; Element < ElementCount; ++Element)
@@ -648,9 +649,9 @@ namespace MutablePrivate
 				int32 CurrentOffset = 0;
 				for (int32 ChannelIndex = 0; ChannelIndex < ChannelCount; ++ChannelIndex)
 				{
-					const mu::FMeshBufferChannel& Channel = BufferSet.m_buffers[BufferIndex].m_channels[ChannelIndex];
+					const mu::FMeshBufferChannel& Channel = BufferSet.Buffers[BufferIndex].Channels[ChannelIndex];
 
-					int32 ChannelOffset = Channel.m_offset;
+					int32 ChannelOffset = Channel.Offset;
 
 					int32 PreviousPadding = ChannelOffset - CurrentOffset;
 					if (PreviousPadding > 0)
@@ -660,10 +661,10 @@ namespace MutablePrivate
 						CurrentOffset += PreviousPadding;
 					}
 
-					int32 ChannelSize = Channel.m_componentCount * GetMeshFormatData(Channel.m_format).SizeInBytes;
+					int32 ChannelSize = Channel.ComponentCount * GetMeshFormatData(Channel.Format).SizeInBytes;
 					FMemory::Memcpy(TargetData, SourceDataPerChannel[ChannelIndex], ChannelSize);
 					TargetData += ChannelSize;
-					SourceDataPerChannel[ChannelIndex] += SourceBuffer.m_elementSize;
+					SourceDataPerChannel[ChannelIndex] += SourceBuffer.ElementSize;
 					CurrentOffset += ChannelSize;
 				}
 
@@ -1009,7 +1010,7 @@ mu::MeshPtr ConvertSkeletalMeshToMutable(const USkeletalMesh* InSkeletalMesh, co
 			const int32 ElementSize = sizeof(float) * 13 + 4;
 
 			MutableMesh->GetVertexBuffers().SetBuffer(0, ElementSize, ChannelCount, Semantics, SemanticIndices, Formats, Components, Offsets);
-			check(!MutableMesh->VertexBuffers.m_buffers[0].HasPadding());
+			check(!MutableMesh->VertexBuffers.Buffers[0].HasPadding());
 		}
 
 		// Texture coordinates
@@ -1029,7 +1030,7 @@ mu::MeshPtr ConvertSkeletalMeshToMutable(const USkeletalMesh* InSkeletalMesh, co
 			};
 
 			MutableMesh->GetVertexBuffers().SetBuffer(1, ElementSize, ChannelCount, Semantics, SemanticIndices, Formats, Components, Offsets);
-			check(!MutableMesh->VertexBuffers.m_buffers[1].HasPadding());
+			check(!MutableMesh->VertexBuffers.Buffers[1].HasPadding());
 		}
 	}
 
@@ -1147,7 +1148,7 @@ mu::MeshPtr ConvertSkeletalMeshToMutable(const USkeletalMesh* InSkeletalMesh, co
 
 
 	// TODO: Add Mesh generation flags to not include RT Morph and clothing if not needed.
-	int32 NextBufferIndex = MutableMesh->VertexBuffers.m_buffers.Num();
+	int32 NextBufferIndex = MutableMesh->VertexBuffers.Buffers.Num();
 	if (GenerationContext.Options.bRealTimeMorphTargetsEnabled)
 	{
 
@@ -1274,7 +1275,7 @@ mu::MeshPtr ConvertSkeletalMeshToMutable(const USkeletalMesh* InSkeletalMesh, co
 				int32 Components[ChannelCount] = { 1 };
 				const int32 Offsets[ChannelCount] = { 0 };
 
-				MutableMesh->GetVertexBuffers().SetBuffer(NextBufferIndex, ElementSize, ChannelCount, Semantics, SemanticIndices, Formats, Components, Offsets);
+				MutableMesh->GetVertexBuffers().SetBuffer(NextBufferIndex, ElementSize, ChannelCount, Semantics, SemanticIndices, Formats, Components, Offsets, mu::EMemoryInitPolicy::Zeroed);
 			}
 
 			// MorphTarget vertex morph count.
@@ -1290,7 +1291,7 @@ mu::MeshPtr ConvertSkeletalMeshToMutable(const USkeletalMesh* InSkeletalMesh, co
 				int32 Components[ChannelCount] = { 1 };
 				const int32 Offsets[ChannelCount] = { 0 };
 
-				MutableMesh->GetVertexBuffers().SetBuffer(NextBufferIndex + 1, ElementSize, ChannelCount, Semantics, SemanticIndices, Formats, Components, Offsets);
+				MutableMesh->GetVertexBuffers().SetBuffer(NextBufferIndex + 1, ElementSize, ChannelCount, Semantics, SemanticIndices, Formats, Components, Offsets, mu::EMemoryInitPolicy::Zeroed);
 			}
 
 			// Setup MorphTarget reconstruction data.
@@ -1622,7 +1623,7 @@ mu::MeshPtr ConvertSkeletalMeshToMutable(const USkeletalMesh* InSkeletalMesh, co
 				int32 Components[ChannelCount] = { 1 };
 				const int32 Offsets[ChannelCount] = { 0 };
 
-				MutableMesh->GetVertexBuffers().SetBuffer(NextBufferIndex, ElementSize, ChannelCount, Semantics, SemanticIndices, Formats, Components, Offsets);
+				MutableMesh->GetVertexBuffers().SetBuffer(NextBufferIndex, ElementSize, ChannelCount, Semantics, SemanticIndices, Formats, Components, Offsets, mu::EMemoryInitPolicy::Zeroed);
 			}
 
 			{
@@ -1635,7 +1636,7 @@ mu::MeshPtr ConvertSkeletalMeshToMutable(const USkeletalMesh* InSkeletalMesh, co
 				int32 Components[ChannelCount] = { 1 };
 				const int32 Offsets[ChannelCount] = { 0 };
 
-				MutableMesh->GetVertexBuffers().SetBuffer(NextBufferIndex + 1, ElementSize, ChannelCount, Semantics, SemanticIndices, Formats, Components, Offsets);
+				MutableMesh->GetVertexBuffers().SetBuffer(NextBufferIndex + 1, ElementSize, ChannelCount, Semantics, SemanticIndices, Formats, Components, Offsets, mu::EMemoryInitPolicy::Zeroed);
 			}
 
 			TArrayView<int32> ClothSectionIndexView(reinterpret_cast<int32*>(
diff --git a/Engine/Plugins/Experimental/Mutable/Source/MutableRuntime/Internal/MuR/OpMeshMerge.h b/Engine/Plugins/Experimental/Mutable/Source/MutableRuntime/Internal/MuR/OpMeshMerge.h
--- a/Engine/Plugins/Experimental/Mutable/Source/MutableRuntime/Internal/MuR/OpMeshMerge.h
+++ b/Engine/Plugins/Experimental/Mutable/Source/MutableRuntime/Internal/MuR/OpMeshMerge.h
@@ -45,10 +45,10 @@ namespace mu
 			check(pSecond->GetIndexBuffers().GetBufferCount() <= 1);
 			Result->GetIndexBuffers().SetBufferCount(1);
 
-			FMeshBuffer& ResultIndexBuffer = Result->GetIndexBuffers().m_buffers[0];
+			FMeshBuffer& ResultIndexBuffer = Result->GetIndexBuffers().Buffers[0];
 
-			const FMeshBuffer& FirstIndexBuffer = pFirst->GetIndexBuffers().m_buffers[0];
-			const FMeshBuffer& SecondIndexBuffer = pSecond->GetIndexBuffers().m_buffers[0];
+			const FMeshBuffer& FirstIndexBuffer = pFirst->GetIndexBuffers().Buffers[0];
+			const FMeshBuffer& SecondIndexBuffer = pSecond->GetIndexBuffers().Buffers[0];
 
 			// Avoid unused variable warnings
 			(void)FirstIndexBuffer;
@@ -59,13 +59,13 @@ namespace mu
 
 			if (FirstCount && SecondCount)
 			{
-				check(!FirstIndexBuffer.m_channels.IsEmpty());
-				check(FirstIndexBuffer.m_channels == SecondIndexBuffer.m_channels);
-				check(FirstIndexBuffer.m_elementSize == SecondIndexBuffer.m_elementSize);
+				check(!FirstIndexBuffer.Channels.IsEmpty());
+				check(FirstIndexBuffer.Channels == SecondIndexBuffer.Channels);
+				check(FirstIndexBuffer.ElementSize == SecondIndexBuffer.ElementSize);
 
 				// We need to know the total number of vertices in case we need to adjust the index buffer format.
 				const uint64 totalVertexCount = pFirst->GetVertexBuffers().GetElementCount() + pSecond->GetVertexBuffers().GetElementCount();
-				const uint64 maxValueBits = GetMeshFormatData(pFirst->GetIndexBuffers().m_buffers[0].m_channels[0].m_format).MaxValueBits;
+				const uint64 maxValueBits = GetMeshFormatData(pFirst->GetIndexBuffers().Buffers[0].Channels[0].Format).MaxValueBits;
 				const uint64 maxSupportedVertices = uint64(1) << maxValueBits;
 				
 				if (totalVertexCount > maxSupportedVertices)
@@ -78,61 +78,61 @@ namespace mu
 			if (IndexBufferFormat != MBF_NONE)
 			{
 				// We only support vertex indices in case of having to change the format.
-				check(FirstIndexBuffer.m_channels.Num() == 1);
-
-				ResultIndexBuffer.m_channels.SetNum(1);
-				ResultIndexBuffer.m_channels[0].m_semantic = MBS_VERTEXINDEX;
-				ResultIndexBuffer.m_channels[0].m_format = IndexBufferFormat;
-				ResultIndexBuffer.m_channels[0].m_componentCount = 1;
-				ResultIndexBuffer.m_channels[0].m_semanticIndex = 0;
-				ResultIndexBuffer.m_channels[0].m_offset = 0;
-				ResultIndexBuffer.m_elementSize = GetMeshFormatData(IndexBufferFormat).SizeInBytes;
+				check(FirstIndexBuffer.Channels.Num() == 1);
+
+				ResultIndexBuffer.Channels.SetNum(1);
+				ResultIndexBuffer.Channels[0].Semantic = MBS_VERTEXINDEX;
+				ResultIndexBuffer.Channels[0].Format = IndexBufferFormat;
+				ResultIndexBuffer.Channels[0].ComponentCount = 1;
+				ResultIndexBuffer.Channels[0].SemanticIndex = 0;
+				ResultIndexBuffer.Channels[0].Offset = 0;
+				ResultIndexBuffer.ElementSize = GetMeshFormatData(IndexBufferFormat).SizeInBytes;
 			}
 			else if (FirstCount)
 			{
-				ResultIndexBuffer.m_channels = FirstIndexBuffer.m_channels;
-				ResultIndexBuffer.m_elementSize = FirstIndexBuffer.m_elementSize;
+				ResultIndexBuffer.Channels = FirstIndexBuffer.Channels;
+				ResultIndexBuffer.ElementSize = FirstIndexBuffer.ElementSize;
 			}
 			else if (SecondCount)
 			{
-				ResultIndexBuffer.m_channels = SecondIndexBuffer.m_channels;
-				ResultIndexBuffer.m_elementSize = SecondIndexBuffer.m_elementSize;
+				ResultIndexBuffer.Channels = SecondIndexBuffer.Channels;
+				ResultIndexBuffer.ElementSize = SecondIndexBuffer.ElementSize;
 			}
 
-			ResultIndexBuffer.m_data.SetNum(ResultIndexBuffer.m_elementSize * (FirstCount + SecondCount));
+			ResultIndexBuffer.Data.SetNum(ResultIndexBuffer.ElementSize * (FirstCount + SecondCount));
 
-			check(ResultIndexBuffer.m_channels.Num() == 1);
-			check(ResultIndexBuffer.m_channels[0].m_semantic == MBS_VERTEXINDEX);
+			check(ResultIndexBuffer.Channels.Num() == 1);
+			check(ResultIndexBuffer.Channels[0].Semantic == MBS_VERTEXINDEX);
 
-			if (!ResultIndexBuffer.m_data.IsEmpty())
+			if (!ResultIndexBuffer.Data.IsEmpty())
 			{
 				if (FirstCount)
 				{
 					if (IndexBufferFormat == MBF_NONE
-						|| IndexBufferFormat == FirstIndexBuffer.m_channels[0].m_format)
+						|| IndexBufferFormat == FirstIndexBuffer.Channels[0].Format)
 					{
-						FMemory::Memcpy(&ResultIndexBuffer.m_data[0],
-							&FirstIndexBuffer.m_data[0],
-							FirstIndexBuffer.m_elementSize * FirstCount);
+						FMemory::Memcpy(&ResultIndexBuffer.Data[0],
+							&FirstIndexBuffer.Data[0],
+							FirstIndexBuffer.ElementSize * FirstCount);
 					}
 					else
 					{
 						// Conversion required
-						const uint8_t* pSource = &FirstIndexBuffer.m_data[0];
-						uint8_t* pDest = &ResultIndexBuffer.m_data[0];
+						const uint8_t* pSource = &FirstIndexBuffer.Data[0];
+						uint8_t* pDest = &ResultIndexBuffer.Data[0];
 						switch (IndexBufferFormat)
 						{
 						case MBF_UINT32:
 						{
-							switch (FirstIndexBuffer.m_channels[0].m_format)
+							switch (FirstIndexBuffer.Channels[0].Format)
 							{
 							case MBF_UINT16:
 							{
 								for (int v = 0; v < FirstCount; ++v)
 								{
 									*(uint32_t*)pDest = *(const uint16*)pSource;
-									pSource += FirstIndexBuffer.m_elementSize;
-									pDest += ResultIndexBuffer.m_elementSize;
+									pSource += FirstIndexBuffer.ElementSize;
+									pDest += ResultIndexBuffer.ElementSize;
 								}
 								break;
 							}
@@ -142,8 +142,8 @@ namespace mu
 								for (int v = 0; v < FirstCount; ++v)
 								{
 									*(uint32_t*)pDest = *(const uint8_t*)pSource;
-									pSource += FirstIndexBuffer.m_elementSize;
-									pDest += ResultIndexBuffer.m_elementSize;
+									pSource += FirstIndexBuffer.ElementSize;
+									pDest += ResultIndexBuffer.ElementSize;
 								}
 								break;
 							}
@@ -157,7 +157,7 @@ namespace mu
 
 						case MBF_UINT16:
 						{
-							switch (FirstIndexBuffer.m_channels[0].m_format)
+							switch (FirstIndexBuffer.Channels[0].Format)
 							{
 
 							case MBF_UINT8:
@@ -165,8 +165,8 @@ namespace mu
 								for (int v = 0; v < FirstCount; ++v)
 								{
 									*(uint16*)pDest = *(const uint8_t*)pSource;
-									pSource += FirstIndexBuffer.m_elementSize;
-									pDest += ResultIndexBuffer.m_elementSize;
+									pSource += FirstIndexBuffer.ElementSize;
+									pDest += ResultIndexBuffer.ElementSize;
 								}
 								break;
 							}
@@ -187,15 +187,15 @@ namespace mu
 
 				if (SecondCount)
 				{
-					const uint8_t* pSource = &SecondIndexBuffer.m_data[0];
-					uint8_t* pDest = &ResultIndexBuffer.m_data[ResultIndexBuffer.m_elementSize * FirstCount];
+					const uint8_t* pSource = &SecondIndexBuffer.Data[0];
+					uint8_t* pDest = &ResultIndexBuffer.Data[ResultIndexBuffer.ElementSize * FirstCount];
 
 					uint32_t firstVertexCount = pFirst->GetVertexBuffers().GetElementCount();
 
 					if (IndexBufferFormat == MBF_NONE
-						|| IndexBufferFormat == SecondIndexBuffer.m_channels[0].m_format)
+						|| IndexBufferFormat == SecondIndexBuffer.Channels[0].Format)
 					{
-						switch (SecondIndexBuffer.m_channels[0].m_format)
+						switch (SecondIndexBuffer.Channels[0].Format)
 						{
 						case MBF_INT32:
 						case MBF_UINT32:
@@ -205,8 +205,8 @@ namespace mu
 							for (int v = 0; v < SecondCount; ++v)
 							{
 								*(uint32_t*)pDest = firstVertexCount + *(const uint32_t*)pSource;
-								pSource += SecondIndexBuffer.m_elementSize;
-								pDest += ResultIndexBuffer.m_elementSize;
+								pSource += SecondIndexBuffer.ElementSize;
+								pDest += ResultIndexBuffer.ElementSize;
 							}
 							break;
 						}
@@ -219,8 +219,8 @@ namespace mu
 							for (int v = 0; v < SecondCount; ++v)
 							{
 								*(uint16*)pDest = uint16(firstVertexCount) + *(const uint16*)pSource;
-								pSource += SecondIndexBuffer.m_elementSize;
-								pDest += ResultIndexBuffer.m_elementSize;
+								pSource += SecondIndexBuffer.ElementSize;
+								pDest += ResultIndexBuffer.ElementSize;
 							}
 							break;
 						}
@@ -233,8 +233,8 @@ namespace mu
 							for (int v = 0; v < SecondCount; ++v)
 							{
 								*(uint8_t*)pDest = uint8_t(firstVertexCount) + *(const uint8_t*)pSource;
-								pSource += SecondIndexBuffer.m_elementSize;
-								pDest += ResultIndexBuffer.m_elementSize;
+								pSource += SecondIndexBuffer.ElementSize;
+								pDest += ResultIndexBuffer.ElementSize;
 							}
 							break;
 						}
@@ -252,7 +252,7 @@ namespace mu
 
 						case MBF_UINT32:
 						{
-							switch (SecondIndexBuffer.m_channels[0].m_format)
+							switch (SecondIndexBuffer.Channels[0].Format)
 							{
 							case MBF_INT16:
 							case MBF_UINT16:
@@ -262,8 +262,8 @@ namespace mu
 								for (int v = 0; v < SecondCount; ++v)
 								{
 									*(uint32_t*)pDest = uint32_t(firstVertexCount) + *(const uint16*)pSource;
-									pSource += SecondIndexBuffer.m_elementSize;
-									pDest += ResultIndexBuffer.m_elementSize;
+									pSource += SecondIndexBuffer.ElementSize;
+									pDest += ResultIndexBuffer.ElementSize;
 								}
 								break;
 							}
@@ -276,8 +276,8 @@ namespace mu
 								for (int v = 0; v < SecondCount; ++v)
 								{
 									*(uint32_t*)pDest = uint32_t(firstVertexCount) + *(const uint8_t*)pSource;
-									pSource += SecondIndexBuffer.m_elementSize;
-									pDest += ResultIndexBuffer.m_elementSize;
+									pSource += SecondIndexBuffer.ElementSize;
+									pDest += ResultIndexBuffer.ElementSize;
 								}
 								break;
 							}
@@ -292,7 +292,7 @@ namespace mu
 
 						case MBF_UINT16:
 						{
-							switch (SecondIndexBuffer.m_channels[0].m_format)
+							switch (SecondIndexBuffer.Channels[0].Format)
 							{
 							case MBF_INT8:
 							case MBF_UINT8:
@@ -302,8 +302,8 @@ namespace mu
 								for (int v = 0; v < SecondCount; ++v)
 								{
 									*(uint16*)pDest = uint16(firstVertexCount) + *(const uint8_t*)pSource;
-									pSource += SecondIndexBuffer.m_elementSize;
-									pDest += ResultIndexBuffer.m_elementSize;
+									pSource += SecondIndexBuffer.ElementSize;
+									pDest += ResultIndexBuffer.ElementSize;
 								}
 								break;
 							}
@@ -750,8 +750,8 @@ namespace mu
 		{
             MUTABLE_CPUPROFILER_SCOPE(Vertices);
 
-			const int32 FirstBufferCount = pFirst->VertexBuffers.m_buffers.Num();
-			const int32 SecondBufferCount = pSecond->VertexBuffers.m_buffers.Num();
+			const int32 FirstBufferCount = pFirst->VertexBuffers.Buffers.Num();
+			const int32 SecondBufferCount = pSecond->VertexBuffers.Buffers.Num();
 			const int32 FirstVertexCount = pFirst->GetVertexBuffers().GetElementCount();
 			const int32 SecondVertexCount = pSecond->GetVertexBuffers().GetElementCount();
 
@@ -773,22 +773,22 @@ namespace mu
 				Result->GetVertexBuffers().SetBufferCount(FirstBufferCount);
 				for (int32 BufferIndex = 0; BufferIndex < FirstBufferCount; ++BufferIndex)
 				{
-					Result->VertexBuffers.m_buffers[BufferIndex].m_channels = pFirst->VertexBuffers.m_buffers[BufferIndex].m_channels;
-					Result->VertexBuffers.m_buffers[BufferIndex].m_elementSize = pFirst->VertexBuffers.m_buffers[BufferIndex].m_elementSize;
+					Result->VertexBuffers.Buffers[BufferIndex].Channels = pFirst->VertexBuffers.Buffers[BufferIndex].Channels;
+					Result->VertexBuffers.Buffers[BufferIndex].ElementSize = pFirst->VertexBuffers.Buffers[BufferIndex].ElementSize;
 				}
 
 				// See if we need to add additional buffers from the second mesh (like vertex colours or additional UV Channels)
 				// This is a bit ad-hoc: we only add buffers containing all new channels
-				for (const FMeshBuffer& SecondBuf : pSecond->GetVertexBuffers().m_buffers)
+				for (const FMeshBuffer& SecondBuf : pSecond->GetVertexBuffers().Buffers)
 				{
 					bool bSomeChannel = false;
 					bool bAllNewChannels = true;
-					for (const FMeshBufferChannel& SecondChan : SecondBuf.m_channels)
+					for (const FMeshBufferChannel& SecondChan : SecondBuf.Channels)
 					{
 						// Skip system buffers
-						if (SecondChan.m_semantic == MBS_VERTEXINDEX
+						if (SecondChan.Semantic == MBS_VERTEXINDEX
 							||
-							SecondChan.m_semantic == MBS_LAYOUTBLOCK)
+							SecondChan.Semantic == MBS_LAYOUTBLOCK)
 						{
 							continue;
 						}
@@ -797,7 +797,7 @@ namespace mu
 
 						int32 FoundBuffer = -1;
 						int32 FoundChannel = -1;
-						pFirst->GetVertexBuffers().FindChannel(SecondChan.m_semantic, SecondChan.m_semanticIndex, &FoundBuffer, &FoundChannel);
+						pFirst->GetVertexBuffers().FindChannel(SecondChan.Semantic, SecondChan.SemanticIndex, &FoundBuffer, &FoundChannel);
 						if (FoundBuffer >= 0)
 						{
 							// There's at least one channel that already exists in the first mesh. Don't add the buffer.
@@ -806,18 +806,18 @@ namespace mu
 						}
 
 						// If there are additional UV channels try to add them
-						if (!bAllNewChannels && SecondChan.m_semantic == MBS_TEXCOORDS
+						if (!bAllNewChannels && SecondChan.Semantic == MBS_TEXCOORDS
 							&&
-							SecondChan.m_semanticIndex > 0)
+							SecondChan.SemanticIndex > 0)
 						{
 							// Add additional UV channels if the previous one is found.
 							FMeshBufferSet& ResultVertexBuffers = Result->GetVertexBuffers();
-							ResultVertexBuffers.FindChannel(MBS_TEXCOORDS, SecondChan.m_semanticIndex - 1, &FoundBuffer, &FoundChannel);
+							ResultVertexBuffers.FindChannel(MBS_TEXCOORDS, SecondChan.SemanticIndex - 1, &FoundBuffer, &FoundChannel);
 
 							if (FoundBuffer >= 0)
 							{
-								FMeshBuffer& Buffer = ResultVertexBuffers.m_buffers[FoundBuffer];
-								Buffer.m_channels.Insert(SecondChan, FoundChannel + 1);
+								FMeshBuffer& Buffer = ResultVertexBuffers.Buffers[FoundBuffer];
+								Buffer.Channels.Insert(SecondChan, FoundChannel + 1);
 
 								ResultVertexBuffers.UpdateOffsets(FoundBuffer);
 							}
@@ -826,9 +826,9 @@ namespace mu
 
 					if (bSomeChannel && bAllNewChannels)
 					{
-						int32 NewBufferIndex = Result->GetVertexBuffers().m_buffers.Emplace();
-						Result->VertexBuffers.m_buffers[NewBufferIndex].m_channels = SecondBuf.m_channels;
-						Result->VertexBuffers.m_buffers[NewBufferIndex].m_elementSize = SecondBuf.m_elementSize;
+						int32 NewBufferIndex = Result->GetVertexBuffers().Buffers.Emplace();
+						Result->VertexBuffers.Buffers[NewBufferIndex].Channels = SecondBuf.Channels;
+						Result->VertexBuffers.Buffers[NewBufferIndex].ElementSize = SecondBuf.ElementSize;
 					}
 				}
 
@@ -837,27 +837,27 @@ namespace mu
 				for (int32 BufferIndex = 0; BufferIndex < FMath::Min(ResultBufferCount, FirstBufferCount); ++BufferIndex)
 				{
 					// Expand component counts in vertex channels of the format mesh
-					FMeshBuffer& result = Result->GetVertexBuffers().m_buffers[BufferIndex];
+					FMeshBuffer& result = Result->GetVertexBuffers().Buffers[BufferIndex];
 
 					bool bResetOffsets = false;
-					for (int32 c = 0; c < result.m_channels.Num(); ++c)
+					for (int32 c = 0; c < result.Channels.Num(); ++c)
 					{
 						int32 sb = -1;
 						int32 sc = -1;
 						pSecond->GetVertexBuffers().FindChannel
 						(
-							result.m_channels[c].m_semantic,
-							result.m_channels[c].m_semanticIndex,
+							result.Channels[c].Semantic,
+							result.Channels[c].SemanticIndex,
 							&sb, &sc
 						);
 
 						if (sb >= 0)
 						{
-							const FMeshBuffer& second = pSecond->GetVertexBuffers().m_buffers[sb];
+							const FMeshBuffer& second = pSecond->GetVertexBuffers().Buffers[sb];
 
-							if (second.m_channels[sc].m_componentCount>result.m_channels[c].m_componentCount)
+							if (second.Channels[sc].ComponentCount>result.Channels[c].ComponentCount)
 							{
-								result.m_channels[c].m_componentCount = second.m_channels[sc].m_componentCount;
+								result.Channels[c].ComponentCount = second.Channels[sc].ComponentCount;
 								bResetOffsets = true;
 							}
 						}
@@ -876,27 +876,27 @@ namespace mu
 				{
 					// Iterate all vertex buffers and update the format
 					FMeshBufferSet& VertexBuffers = Result->GetVertexBuffers();
-					for (int32 VertexBufferIndex = 0; VertexBufferIndex < VertexBuffers.m_buffers.Num(); ++VertexBufferIndex)
+					for (int32 VertexBufferIndex = 0; VertexBufferIndex < VertexBuffers.Buffers.Num(); ++VertexBufferIndex)
 					{
-						FMeshBuffer& result = VertexBuffers.m_buffers[VertexBufferIndex];
+						FMeshBuffer& result = VertexBuffers.Buffers[VertexBufferIndex];
 
 						const int32 ChannelsCount = VertexBuffers.GetBufferChannelCount(VertexBufferIndex);
 						for (int32 ChannelIndex = 0; ChannelIndex < ChannelsCount; ++ChannelIndex)
 						{
-							if (result.m_channels[ChannelIndex].m_semantic == MBS_BONEINDICES)
+							if (result.Channels[ChannelIndex].Semantic == MBS_BONEINDICES)
 							{
-								result.m_channels[ChannelIndex].m_format = BoneIndexFormat;
+								result.Channels[ChannelIndex].Format = BoneIndexFormat;
 
 								// Reset offsets
 								int32 offset = 0;
 								for (int32 AuxChannelIndex = 0; AuxChannelIndex < ChannelsCount; ++AuxChannelIndex)
 								{
-									result.m_channels[AuxChannelIndex].m_offset = (uint8_t)offset;
-									offset += result.m_channels[AuxChannelIndex].m_componentCount
+									result.Channels[AuxChannelIndex].Offset = (uint8_t)offset;
+									offset += result.Channels[AuxChannelIndex].ComponentCount
 										*
-										GetMeshFormatData(result.m_channels[AuxChannelIndex].m_format).SizeInBytes;
+										GetMeshFormatData(result.Channels[AuxChannelIndex].Format).SizeInBytes;
 								}
-								result.m_elementSize = offset;
+								result.ElementSize = offset;
 							}
 						}
 					}
@@ -922,22 +922,22 @@ namespace mu
 				int32 ResultBufferCount = Result->GetVertexBuffers().GetBufferCount();
 				for (int32 ResultBufferIndex = 0; ResultBufferIndex < ResultBufferCount; ++ResultBufferIndex)
 				{
-					FMeshBuffer& ResultBuffer = Result->VertexBuffers.m_buffers[ResultBufferIndex];
+					FMeshBuffer& ResultBuffer = Result->VertexBuffers.Buffers[ResultBufferIndex];
 
 					// TODO: Don't assume the buffer order in First and Second matches Result, for more opportunities of fast-path
 					bool bFirstFastPath = (ResultBufferIndex < FirstBufferCount)
 						&&
 						pFirst->VertexBuffers.HasSameFormat(ResultBufferIndex, Result->VertexBuffers, ResultBufferIndex);
 
-					int32 FirstResultBufferSize = ResultBuffer.m_elementSize * FirstVertexCount;
-					uint8* Destination = ResultBuffer.m_data.GetData();
+					int32 FirstResultBufferSize = ResultBuffer.ElementSize * FirstVertexCount;
+					uint8* Destination = ResultBuffer.Data.GetData();
 					if (bFirstFastPath)
 					{
 						MUTABLE_CPUPROFILER_SCOPE(FirstFastPath);
 
-						const FMeshBuffer& FirstBuffer = pFirst->VertexBuffers.m_buffers[ResultBufferIndex];
-						check(FirstResultBufferSize == FirstBuffer.m_data.Num());
-						FMemory::Memcpy(Destination, FirstBuffer.m_data.GetData(), FirstResultBufferSize);
+						const FMeshBuffer& FirstBuffer = pFirst->VertexBuffers.Buffers[ResultBufferIndex];
+						check(FirstResultBufferSize == FirstBuffer.Data.Num());
+						FMemory::Memcpy(Destination, FirstBuffer.Data.GetData(), FirstResultBufferSize);
 					}
 					else
 					{
@@ -952,15 +952,15 @@ namespace mu
 						&&
 						pSecond->VertexBuffers.HasSameFormat(ResultBufferIndex, Result->VertexBuffers, ResultBufferIndex);
 
-					int32 SecondResultBufferSize = ResultBuffer.m_elementSize * SecondVertexCount;
-					Destination = ResultBuffer.m_data.GetData() + FirstResultBufferSize;
+					int32 SecondResultBufferSize = ResultBuffer.ElementSize * SecondVertexCount;
+					Destination = ResultBuffer.Data.GetData() + FirstResultBufferSize;
 					if (bSecondFastPath)
 					{
 						MUTABLE_CPUPROFILER_SCOPE(SecondFastPath);
 
-						const FMeshBuffer& SecondBuffer = pSecond->VertexBuffers.m_buffers[ResultBufferIndex];
-						check(SecondResultBufferSize == SecondBuffer.m_data.Num());
-						FMemory::Memcpy(Destination, SecondBuffer.m_data.GetData(), SecondResultBufferSize);
+						const FMeshBuffer& SecondBuffer = pSecond->VertexBuffers.Buffers[ResultBufferIndex];
+						check(SecondResultBufferSize == SecondBuffer.Data.Num());
+						FMemory::Memcpy(Destination, SecondBuffer.Data.GetData(), SecondResultBufferSize);
 					}
 					else
 					{
@@ -981,9 +981,9 @@ namespace mu
 
 					// Iterate all vertex buffers and update the format
 					FMeshBufferSet& VertexBuffers = Result->GetVertexBuffers();
-					for (int32 VertexBufferIndex = 0; VertexBufferIndex < VertexBuffers.m_buffers.Num(); ++VertexBufferIndex)
+					for (int32 VertexBufferIndex = 0; VertexBufferIndex < VertexBuffers.Buffers.Num(); ++VertexBufferIndex)
 					{
-						FMeshBuffer& ResultBuffer = VertexBuffers.m_buffers[VertexBufferIndex];
+						FMeshBuffer& ResultBuffer = VertexBuffers.Buffers[VertexBufferIndex];
 
 						const int32 ElemSize = VertexBuffers.GetElementSize(VertexBufferIndex);
 						const int32 FirstSize = FirstVertexCount * ElemSize;
@@ -991,14 +991,14 @@ namespace mu
 						const int32 ChannelsCount = VertexBuffers.GetBufferChannelCount(VertexBufferIndex);
 						for (int32 ChannelIndex = 0; ChannelIndex < ChannelsCount; ++ChannelIndex)
 						{
-							if (ResultBuffer.m_channels[ChannelIndex].m_semantic != MBS_BONEINDICES)
+							if (ResultBuffer.Channels[ChannelIndex].Semantic != MBS_BONEINDICES)
 							{
 								continue;
 							}
 
-							int32 ResultOffset = FirstSize + ResultBuffer.m_channels[ChannelIndex].m_offset;
+							int32 ResultOffset = FirstSize + ResultBuffer.Channels[ChannelIndex].Offset;
 
-							const int32 NumComponents = ResultBuffer.m_channels[ChannelIndex].m_componentCount;
+							const int32 NumComponents = ResultBuffer.Channels[ChannelIndex].ComponentCount;
 
 							// Bone indices may need remapping
 							for (int32 VertexIndex = 0; VertexIndex < pSecond->GetVertexCount(); ++VertexIndex)
@@ -1008,7 +1008,7 @@ namespace mu
 								case MBF_INT8:
 								case MBF_UINT8:
 								{
-									uint8* pD = reinterpret_cast<uint8*>(&ResultBuffer.m_data[ResultOffset]);
+									uint8* pD = reinterpret_cast<uint8*>(&ResultBuffer.Data[ResultOffset]);
 
 									for (int32 ComponentIndex = 0; ComponentIndex < NumComponents; ++ComponentIndex)
 									{
@@ -1032,7 +1032,7 @@ namespace mu
 								case MBF_INT16:
 								case MBF_UINT16:
 								{
-									uint16* pD = reinterpret_cast<uint16*>(&ResultBuffer.m_data[ResultOffset]);
+									uint16* pD = reinterpret_cast<uint16*>(&ResultBuffer.Data[ResultOffset]);
 
 									for (int32 ComponentIndex = 0; ComponentIndex < NumComponents; ++ComponentIndex)
 									{
diff --git a/Engine/Plugins/Experimental/Mutable/Source/MutableRuntime/Private/MuR/Mesh.cpp b/Engine/Plugins/Experimental/Mutable/Source/MutableRuntime/Private/MuR/Mesh.cpp
--- a/Engine/Plugins/Experimental/Mutable/Source/MutableRuntime/Private/MuR/Mesh.cpp
+++ b/Engine/Plugins/Experimental/Mutable/Source/MutableRuntime/Private/MuR/Mesh.cpp
@@ -19,24 +19,24 @@ MUTABLE_IMPLEMENT_ENUM_SERIALISABLE(EShapeBindingMethod);
 MUTABLE_IMPLEMENT_ENUM_SERIALISABLE(EVertexColorUsage);
 
 
-void Mesh::Serialise( const Mesh* p, OutputArchive& arch )
+void Mesh::Serialise(const Mesh* MeshPtr, OutputArchive& Arch)
 {
-    //p->m_pD->CheckIntegrity();
-    arch << *p;
+    //MeshPtr->m_pD->CheckIntegrity();
+    Arch << *MeshPtr;
 }
 
 
-MeshPtr Mesh::StaticUnserialise( InputArchive& arch )
+Ptr<Mesh> Mesh::StaticUnserialise(InputArchive& Arch)
 {
     MUTABLE_CPUPROFILER_SCOPE(MeshUnserialise)
 	LLM_SCOPE_BYNAME(TEXT("MutableRuntime"));
 
-    MeshPtr pResult = new Mesh();
-    arch >> *pResult;
+    Ptr<Mesh> Result = new Mesh();
+    Arch >> *Result;
 
-    //pResult->m_pD->CheckIntegrity();
+    //Result->m_pD->CheckIntegrity();
 
-    return pResult;
+    return Result;
 }
 
 
@@ -74,32 +74,32 @@ uint32 Mesh::GetReferencedMesh() const
 }
 
 
-MeshPtr Mesh::Clone() const
+Ptr<Mesh> Mesh::Clone() const
 {
     //MUTABLE_CPUPROFILER_SCOPE(MeshClone);
 	LLM_SCOPE_BYNAME(TEXT("MutableRuntime"));
 
-    MeshPtr pResult = new Mesh();
+    Ptr<Mesh> Result = new Mesh();
 
-    pResult->InternalId = InternalId;
-	pResult->Flags = Flags;
-	pResult->ReferenceID = ReferenceID;
-	pResult->Surfaces = Surfaces;
-	pResult->Skeleton = Skeleton;
-	pResult->PhysicsBody = PhysicsBody;
-	pResult->Tags = Tags;
-	pResult->StreamedResources = StreamedResources;
-	pResult->MeshIDPrefix = MeshIDPrefix;
+    Result->InternalId = InternalId;
+	Result->Flags = Flags;
+	Result->ReferenceID = ReferenceID;
+	Result->Surfaces = Surfaces;
+	Result->Skeleton = Skeleton;
+	Result->PhysicsBody = PhysicsBody;
+	Result->Tags = Tags;
+	Result->StreamedResources = StreamedResources;
+	Result->MeshIDPrefix = MeshIDPrefix;
 
     // Clone the main buffers
-    pResult->VertexBuffers = VertexBuffers;
-    pResult->IndexBuffers = IndexBuffers;
+    Result->VertexBuffers = VertexBuffers;
+    Result->IndexBuffers = IndexBuffers;
 
 	// Clone additional buffers
-	pResult->AdditionalBuffers = AdditionalBuffers;
+	Result->AdditionalBuffers = AdditionalBuffers;
 
     // Clone the layouts
-	pResult->Layouts = Layouts;
+	Result->Layouts = Layouts;
 
     // The skeleton is not cloned because it is not owned by this mesh and it is always assumed
     // to be shared.
@@ -107,14 +107,14 @@ MeshPtr Mesh::Clone() const
 	// physics body doen't need to be deep cloned either as they are also assumed to be shared.
 	
 	// Clone bone poses
-	pResult->BonePoses = BonePoses;
-	pResult->BoneMap = BoneMap;
+	Result->BonePoses = BonePoses;
+	Result->BoneMap = BoneMap;
 
 	// Clone SkeletonIDs
-	pResult->SkeletonIDs = SkeletonIDs;
-	pResult->AdditionalPhysicsBodies = AdditionalPhysicsBodies;
+	Result->SkeletonIDs = SkeletonIDs;
+	Result->AdditionalPhysicsBodies = AdditionalPhysicsBodies;
 
-    return pResult;
+    return Result;
 }
 
 
@@ -123,59 +123,59 @@ Ptr<Mesh> Mesh::Clone(EMeshCopyFlags InFlags) const
     //MUTABLE_CPUPROFILER_SCOPE(MeshClone);
 	LLM_SCOPE_BYNAME(TEXT("MutableRuntime"));
 
-	Ptr<Mesh> pResult = new Mesh();
+	Ptr<Mesh> Result = new Mesh();
 
-    pResult->InternalId = InternalId;
-	pResult->MeshIDPrefix = MeshIDPrefix;
-	pResult->Flags = Flags;
-	pResult->ReferenceID = ReferenceID;
+    Result->InternalId = InternalId;
+	Result->MeshIDPrefix = MeshIDPrefix;
+	Result->Flags = Flags;
+	Result->ReferenceID = ReferenceID;
 
 	if (EnumHasAnyFlags(InFlags, EMeshCopyFlags::WithSurfaces))
 	{
-		pResult->Surfaces = Surfaces;
+		Result->Surfaces = Surfaces;
 	}
 
 	if (EnumHasAnyFlags(InFlags, EMeshCopyFlags::WithSkeleton))
 	{
-		pResult->Skeleton = Skeleton;
+		Result->Skeleton = Skeleton;
 	}
 
 	if (EnumHasAnyFlags(InFlags, EMeshCopyFlags::WithPhysicsBody))
 	{
-		pResult->PhysicsBody = PhysicsBody;
+		Result->PhysicsBody = PhysicsBody;
 	}
 
 	if (EnumHasAnyFlags(InFlags, EMeshCopyFlags::WithTags))
 	{
-		pResult->Tags = Tags;
+		Result->Tags = Tags;
 	}
 
 	if (EnumHasAnyFlags(InFlags, EMeshCopyFlags::WithStreamedResources))
 	{
-		pResult->StreamedResources = StreamedResources;
+		Result->StreamedResources = StreamedResources;
 	}
 
     // Clone the main buffers
 	if (EnumHasAnyFlags(InFlags, EMeshCopyFlags::WithVertexBuffers))
     {
-		pResult->VertexBuffers = VertexBuffers;
+		Result->VertexBuffers = VertexBuffers;
 	}
 
 	if (EnumHasAnyFlags(InFlags, EMeshCopyFlags::WithIndexBuffers))
     {
-		pResult->IndexBuffers = IndexBuffers;
+		Result->IndexBuffers = IndexBuffers;
 	}
 
 	// Clone additional buffers
 	if (EnumHasAnyFlags(InFlags, EMeshCopyFlags::WithAdditionalBuffers))
 	{
-		pResult->AdditionalBuffers = AdditionalBuffers;
+		Result->AdditionalBuffers = AdditionalBuffers;
 	}
 
     // Clone the layout	
 	if (EnumHasAnyFlags(InFlags, EMeshCopyFlags::WithLayouts))
 	{
-		pResult->Layouts = Layouts;
+		Result->Layouts = Layouts;
 	}
     // The skeleton is not cloned because it is not owned by this mesh and it is always assumed
     // to be shared.
@@ -185,27 +185,27 @@ Ptr<Mesh> Mesh::Clone(EMeshCopyFlags InFlags) const
 	// Clone bone poses
 	if (EnumHasAnyFlags(InFlags, EMeshCopyFlags::WithPoses))
 	{
-		pResult->BonePoses = BonePoses;
+		Result->BonePoses = BonePoses;
 	}
 
 	// Clone BoneMap
 	if (EnumHasAnyFlags(InFlags, EMeshCopyFlags::WithBoneMap))
 	{
-		pResult->BoneMap = BoneMap;
+		Result->BoneMap = BoneMap;
 	}
 
 	// Clone SkeletonIDs
 	if (EnumHasAnyFlags(InFlags, EMeshCopyFlags::WithSkeletonIDs))
 	{
-		pResult->SkeletonIDs = SkeletonIDs;
+		Result->SkeletonIDs = SkeletonIDs;
 	}
 
 	if (EnumHasAnyFlags(InFlags, EMeshCopyFlags::WithAdditionalPhysics))
 	{
-		pResult->AdditionalPhysicsBodies = AdditionalPhysicsBodies;
+		Result->AdditionalPhysicsBodies = AdditionalPhysicsBodies;
 	}
 
-    return pResult;
+    return Result;
 }
 
 
@@ -339,11 +339,15 @@ bool Mesh::AreVertexIdsExplicit() const
 	int32 ChannelIndex = -1;
 	VertexBuffers.FindChannel(MBS_VERTEXINDEX, 0, &BufferIndex, &ChannelIndex);
 
-	bool bExplicit = (BufferIndex >= 0) && (ChannelIndex >= 0) && (VertexBuffers.m_buffers[BufferIndex].m_channels[ChannelIndex].m_format == MBF_UINT64);
+	bool bExplicit = 
+			(BufferIndex >= 0) && (ChannelIndex >= 0) && 
+			(VertexBuffers.Buffers[BufferIndex].Channels[ChannelIndex].Format == MBF_UINT64);
+
 	if (bExplicit)
 	{
 		check(MeshIDPrefix == 0);
 	}
+
 	return bExplicit;
 }
 
@@ -359,13 +363,13 @@ void Mesh::MakeVertexIdsRelative()
 	EMeshBufferFormat Format = MBF_UINT32;
 	int32 Components = 1;
 	int32 Offset = 0;
-	VertexBuffers.SetBuffer(NewBuffer, sizeof(uint32), 1, &Semantic, &SemanticIndex, &Format, &Components, &Offset );
-	uint32* pIdData = reinterpret_cast<uint32*>( VertexBuffers.GetBufferData(NewBuffer) );
+	VertexBuffers.SetBuffer(NewBuffer, sizeof(uint32), 1, &Semantic, &SemanticIndex, &Format, &Components, &Offset);
+	uint32* IdDataPtr = reinterpret_cast<uint32*>(VertexBuffers.GetBufferData(NewBuffer));
 
 	int32 VertexCount = GetVertexCount();
-	for (int32 i = 0; i < VertexCount; ++i)
+	for (int32 Index = 0; Index < VertexCount; ++Index)
 	{
-		(*pIdData++) = i;
+		(*IdDataPtr++) = Index;
 	}
 }
 
@@ -375,56 +379,56 @@ void Mesh::MakeIdsExplicit()
 	MUTABLE_CPUPROFILER_SCOPE(Mesh_MakeIdsExplicit);
 
 	int32 VertexCount = GetVertexCount();
-	check(VertexCount==0);
+	check(VertexCount == 0);
 
 	// Vertex IDs
 	{
 		bool bHasRelativeVertexIndices = false;
 
-		int32 OldBuf = -1;
-		int32 OldChan = -1;
-		VertexBuffers.FindChannel(MBS_VERTEXINDEX, 0, &OldBuf, &OldChan);
-		bool bHasVertexIndices = (OldBuf >= 0 && OldChan >= 0);
+		int32 OldBufferIndex = -1;
+		int32 OldChannelIndex = -1;
+		VertexBuffers.FindChannel(MBS_VERTEXINDEX, 0, &OldBufferIndex, &OldChannelIndex);
+		bool bHasVertexIndices = (OldBufferIndex >= 0 && OldChannelIndex >= 0);
 		if (bHasVertexIndices)
 		{
-			check(OldChan == 0 && VertexBuffers.m_buffers[OldBuf].m_channels.Num() == 1);
+			check(OldChannelIndex == 0 && VertexBuffers.Buffers[OldBufferIndex].Channels.Num() == 1);
 
-			FMeshBuffer& Buffer = VertexBuffers.m_buffers[OldBuf];
-			Buffer.m_channels[0].m_format = MBF_UINT64;
-			Buffer.m_elementSize = sizeof(uint64);
+			FMeshBuffer& Buffer = VertexBuffers.Buffers[OldBufferIndex];
+			Buffer.Channels[0].Format = MBF_UINT64;
+			Buffer.ElementSize = sizeof(uint64);
 		}
 
 		else
 		{
 			// The mesh has implicit Ids
 			// Create a new buffer with explicit ids
-			FMeshBuffer& Buffer = VertexBuffers.m_buffers.Emplace_GetRef();
-			Buffer.m_elementSize = sizeof(uint64);
-			FMeshBufferChannel& Channel = Buffer.m_channels.Emplace_GetRef();
-			Channel.m_semantic = MBS_VERTEXINDEX;
-			Channel.m_semanticIndex= 0;
-			Channel.m_format = MBF_UINT64;
-			Channel.m_componentCount = 1;
-			Channel.m_offset = 0;
+			FMeshBuffer& Buffer = VertexBuffers.Buffers.Emplace_GetRef();
+			Buffer.ElementSize = sizeof(uint64);
+			FMeshBufferChannel& Channel = Buffer.Channels.Emplace_GetRef();
+			Channel.Semantic = MBS_VERTEXINDEX;
+			Channel.SemanticIndex= 0;
+			Channel.Format = MBF_UINT64;
+			Channel.ComponentCount = 1;
+			Channel.Offset = 0;
 		}
 	}
 
 	// Layout block IDs
 	{
-		for (FMeshBuffer& Buffer: VertexBuffers.m_buffers)
+		for (FMeshBuffer& Buffer : VertexBuffers.Buffers)
 		{
-			for (FMeshBufferChannel& Channel : Buffer.m_channels)
+			for (FMeshBufferChannel& Channel : Buffer.Channels)
 			{
-				if (Channel.m_semantic != MBS_LAYOUTBLOCK)
+				if (Channel.Semantic != MBS_LAYOUTBLOCK)
 				{
 					continue;
 				}
 
-				check(Buffer.m_channels.Num() == 1);
-				check(Buffer.m_channels[0].m_offset == 0);
+				check(Buffer.Channels.Num() == 1);
+				check(Buffer.Channels[0].Offset == 0);
 
-				Buffer.m_channels[0].m_format = MBF_UINT64;
-				Buffer.m_elementSize = sizeof(uint64);
+				Buffer.Channels[0].Format = MBF_UINT64;
+				Buffer.ElementSize = sizeof(uint64);
 			}
 		}
 	}
@@ -440,9 +444,9 @@ mu::Ptr<const Skeleton> Mesh::GetSkeleton() const
 }
 
 
-void Mesh::SetSkeleton( Ptr<const mu::Skeleton> s )
+void Mesh::SetSkeleton(Ptr<const mu::Skeleton> InSkeleton)
 {
-    Skeleton = s;
+    Skeleton = InSkeleton;
 }
 
 
@@ -452,7 +456,7 @@ mu::Ptr<const PhysicsBody> Mesh::GetPhysicsBody() const
 }
 
 
-void Mesh::SetPhysicsBody( Ptr<const mu::PhysicsBody> InPhysicsBody )
+void Mesh::SetPhysicsBody(Ptr<const mu::PhysicsBody> InPhysicsBody)
 {
     PhysicsBody = InPhysicsBody;
 }
@@ -464,11 +468,11 @@ int32 Mesh::AddAdditionalPhysicsBody(Ptr<const mu::PhysicsBody> Body)
 }
 
 
-Ptr<const PhysicsBody> Mesh::GetAdditionalPhysicsBody(int32 I) const
+Ptr<const PhysicsBody> Mesh::GetAdditionalPhysicsBody(int32 Index) const
 {
-	check(I >= 0 && I < AdditionalPhysicsBodies.Num());
+	check(AdditionalPhysicsBodies.IsValidIndex(Index));
 
-	return AdditionalPhysicsBodies[I];
+	return AdditionalPhysicsBodies[Index];
 }
 
 
@@ -478,7 +482,7 @@ int32 Mesh::GetFaceCount() const
 }
 
 
-int Mesh::GetIndexCount() const
+int32 Mesh::GetIndexCount() const
 {
     return GetIndexBuffers().GetElementCount();
 }
@@ -496,16 +500,17 @@ const FMeshBufferSet& Mesh::GetIndexBuffers() const
 }
 
 
-int Mesh::GetSurfaceCount() const
+int32 Mesh::GetSurfaceCount() const
 {
     return Surfaces.Num();
 }
 
 
-void Mesh::GetSurface(int32 SurfaceIndex,
-                      int32& OutFirstVertex, int32& OutVertexCount,
-                      int32& OutFirstIndex, int32& OutIndexCount,
-					  int32& OutBoneIndex, int32& OutBoneCount) const
+void Mesh::GetSurface(
+		int32 SurfaceIndex,
+        int32& OutFirstVertex, int32& OutVertexCount,
+    	int32& OutFirstIndex, int32& OutIndexCount,
+		int32& OutBoneIndex, int32& OutBoneCount) const
 {
     int32 Count = GetSurfaceCount();
 
@@ -565,55 +570,55 @@ uint32 Mesh::GetSurfaceId(int32 SurfaceIndex) const
 }
 
 
-void Mesh::AddLayout(Ptr<const Layout> pLayout )
+void Mesh::AddLayout(Ptr<const Layout> InLayout)
 {
 	LLM_SCOPE_BYNAME(TEXT("MutableRuntime"));
-    Layouts.Add( pLayout );
+    Layouts.Add(InLayout);
 }
 
 
-int Mesh::GetLayoutCount() const
+int32 Mesh::GetLayoutCount() const
 {
     return Layouts.Num();
 }
 
 
-const Layout* Mesh::GetLayout( int i ) const
+const Layout* Mesh::GetLayout(int32 LayoutIndex) const
 {
-    check( i>=0 && i<Layouts.Num() );
+    check(Layouts.IsValidIndex(LayoutIndex));
 
-    return Layouts[i].get();
+    return Layouts[LayoutIndex].get();
 }
 
 
-void Mesh::SetLayout( int i, Ptr<const Layout> pLayout )
+void Mesh::SetLayout(int32 LayoutIndex, Ptr<const Layout> InLayout)
 {
-    check( i>=0 && i<Layouts.Num() );
+    check(Layouts.IsValidIndex(LayoutIndex));
 
-    Layouts[i] = pLayout;
+    Layouts[LayoutIndex] = InLayout;
 }
 
 
-int Mesh::GetTagCount() const
+int32 Mesh::GetTagCount() const
 {
     return Tags.Num();
 }
 
 
-void Mesh::SetTagCount( int count )
+void Mesh::SetTagCount(int32 Count)
 {
 	LLM_SCOPE_BYNAME(TEXT("MutableRuntime"));
-    Tags.SetNum( count );
+    Tags.SetNum(Count);
 }
 
 
-const FString& Mesh::GetTag( int tagIndex ) const
+const FString& Mesh::GetTag(int32 TagIndex) const
 {
-    check( tagIndex>=0 && tagIndex<GetTagCount() );
+    check(Tags.IsValidIndex(TagIndex));
 
-    if (tagIndex >= 0 && tagIndex < GetTagCount())
+    if (Tags.IsValidIndex(TagIndex))
     {
-        return Tags[tagIndex];
+        return Tags[TagIndex];
     }
     else
     {
@@ -623,14 +628,14 @@ const FString& Mesh::GetTag( int tagIndex ) const
 }
 
 
-void Mesh::SetTag( int tagIndex, const FString& Name )
+void Mesh::SetTag(int32 TagIndex, const FString& Name)
 {
-    check( tagIndex>=0 && tagIndex<GetTagCount() );
+    check(Tags.IsValidIndex(TagIndex));
 	LLM_SCOPE_BYNAME(TEXT("MutableRuntime"));
 
-    if (tagIndex >= 0 && tagIndex < GetTagCount())
+    if (Tags.IsValidIndex(TagIndex))
     {
-        Tags[tagIndex] = Name;
+        Tags[TagIndex] = Name;
     }
 }
 
@@ -748,28 +753,26 @@ int32 Mesh::GetDataSize() const
 }
 
 
-bool Mesh::HasCompatibleFormat( const Mesh* pOther ) const
+bool Mesh::HasCompatibleFormat(const Mesh* Other) const
 {
-    bool compatible = true;
-
-    compatible &= Layouts.Num()==pOther->Layouts.Num();
-    compatible &= VertexBuffers.GetBufferCount()
-            == pOther->VertexBuffers.GetBufferCount();
+    bool bCompatible = true;
 
+    bCompatible = Layouts.Num() == Other->Layouts.Num();
+    bCompatible = bCompatible && VertexBuffers.GetBufferCount() == Other->VertexBuffers.GetBufferCount();
 
     // Indices
     //-----------------
-    if ( IndexBuffers.GetElementCount()>0 && pOther->GetIndexCount()>0 )
+    if (IndexBuffers.GetElementCount() > 0 && Other->GetIndexCount() > 0)
     {
-        check( IndexBuffers.m_buffers.Num() == 1 );
-        check( pOther->GetIndexBuffers().m_buffers.Num() == 1 );
-        check( IndexBuffers.GetBufferChannelCount(0) == 1 );
-        check( pOther->GetIndexBuffers().GetBufferChannelCount(0) == 1 );
+        check(IndexBuffers.Buffers.Num() == 1);
+        check(Other->GetIndexBuffers().Buffers.Num() == 1);
+        check(IndexBuffers.GetBufferChannelCount(0) == 1);
+        check(Other->GetIndexBuffers().GetBufferChannelCount(0) == 1);
 
-        const FMeshBuffer& dest = IndexBuffers.m_buffers[0];
-        const FMeshBuffer& source = pOther->GetIndexBuffers().m_buffers[0];
+        const FMeshBuffer& Dest = IndexBuffers.Buffers[0];
+        const FMeshBuffer& Source = Other->GetIndexBuffers().Buffers[0];
 
-        compatible &= dest.m_channels[0].m_format == source.m_channels[0].m_format;
+        bCompatible = bCompatible && Dest.Channels[0].Format == Source.Channels[0].Format;
     }
 
 
@@ -780,159 +783,162 @@ bool Mesh::HasCompatibleFormat( const Mesh* pOther ) const
 
     // Vertices
     //-----------------
-    for ( int vb = 0; vb<VertexBuffers.GetBufferCount(); ++vb )
+	const int32 NumVertexBuffers =  VertexBuffers.GetBufferCount();
+    for (int32 VertexBufferIndex = 0; VertexBufferIndex < NumVertexBuffers; ++VertexBufferIndex)
     {
-        const FMeshBuffer& dest = VertexBuffers.m_buffers[vb];
-        const FMeshBuffer& source = pOther->GetVertexBuffers().m_buffers[vb];
+        const FMeshBuffer& Dest = VertexBuffers.Buffers[VertexBufferIndex];
+        const FMeshBuffer& Source = Other->GetVertexBuffers().Buffers[VertexBufferIndex];
 
         // TODO: More checks about channels formats and semantics
-        //compatible &= GetVertexBufferElementSize(vb) == pOther->GetVertexBufferElementSize(vb);
-        compatible &= dest.m_channels.Num()==source.m_channels.Num();
+        //bCompatible = bCompatible && GetVertexBufferElementSize(VertexBufferIndex) == Other->GetVertexBufferElementSize(VertexBufferIndex);
+        bCompatible = bCompatible && Dest.Channels.Num() == Source.Channels.Num();
     }
 
-    return compatible;
+    return bCompatible;
 }
 
 
-UE::Math::TIntVector3<uint32_t>  Mesh::GetFaceVertexIndices( int f ) const
+UE::Math::TIntVector3<uint32> Mesh::GetFaceVertexIndices(int32 FaceIndex) const
 {
-	UE::Math::TIntVector3<uint32> res;
+	UE::Math::TIntVector3<uint32> Result;
 
-    MeshBufferIteratorConst<MBF_UINT32,uint32_t,1> it( IndexBuffers, MBS_VERTEXINDEX );
-    it += f*3;
+    MeshBufferIteratorConst<MBF_UINT32, uint32, 1> Iter(IndexBuffers, MBS_VERTEXINDEX);
+    Iter += FaceIndex*3;
 
-    res[0] = (*it)[0];
-    ++it;
+    Result[0] = (*Iter)[0];
+    ++Iter;
 
-    res[1] = (*it)[0];
-    ++it;
+    Result[1] = (*Iter)[0];
+    ++Iter;
 
-    res[2] = (*it)[0];
-    ++it;
+    Result[2] = (*Iter)[0];
+    ++Iter;
 
-    return res;
+    return Result;
 }
 
 
-bool Mesh::FVertexMatchMap::DoMatch(int32 v, int32 ov) const
+bool Mesh::FVertexMatchMap::DoMatch(int32 Vertex, int32 OtherVertex) const
 {
-	if (v >= 0 && v < FirstMatch.Num())
+	if (Vertex >= 0 && Vertex < FirstMatch.Num())
 	{
-		int32 start = FirstMatch[v];
-		int32 end = v + 1 < FirstMatch.Num() ? FirstMatch[v + 1] : Matches.Num();
-		bool res = false;
+		int32 Start = FirstMatch[Vertex];
+		int32 End = Vertex + 1 < FirstMatch.Num() ? FirstMatch[Vertex + 1] : Matches.Num();
 
-		while (!res && start < end)
+		bool bResult = false;
+		while (!bResult && Start < End)
 		{
-			if (Matches[start] == ov)
+			if (Matches[Start] == OtherVertex)
 			{
-				res = true;
+				bResult = true;
 			}
-			++start;
+
+			++Start;
 		}
 
-		return res;
+		return bResult;
 	}
 
 	return false;
 }
 
 
-void Mesh::GetVertexMap( const Mesh& other, FVertexMatchMap& vertexMap, float tolerance ) const
+void Mesh::GetVertexMap(const Mesh& Other, FVertexMatchMap& VertexMap, float Tolerance) const
 {
-    int32 vertexCount = VertexBuffers.GetElementCount();
-    vertexMap.FirstMatch.SetNum( vertexCount );
-    vertexMap.Matches.SetNum( vertexCount+(vertexCount>>2) );
+    int32 VertexCount = VertexBuffers.GetElementCount();
+    VertexMap.FirstMatch.SetNum(VertexCount);
+    VertexMap.Matches.SetNum(VertexCount + (VertexCount >> 2));
 
-    int otherVertexCount = other.VertexBuffers.GetElementCount();
+    int32 OtherVertexCount = Other.VertexBuffers.GetElementCount();
 
-    if ( !vertexCount || !otherVertexCount )
+    if (!VertexCount || !OtherVertexCount)
     {
         return;
     }
 
 
-    MeshBufferIteratorConst< MBF_FLOAT32, float, 3 > itp( VertexBuffers, MBS_POSITION);
-    MeshBufferIteratorConst< MBF_FLOAT32, float, 3 > itopBegin( other.VertexBuffers, MBS_POSITION);
+    MeshBufferIteratorConst<MBF_FLOAT32, float, 3> ItPosition(VertexBuffers, MBS_POSITION);
+    MeshBufferIteratorConst<MBF_FLOAT32, float, 3> ItOtherPositionBegin(Other.VertexBuffers, MBS_POSITION);
 
 
     // Bucket the other mesh
 #define MUTABLE_NUM_BUCKETS 256
 #define MUTABLE_BUCKET_CHANNEL 0
 
-    float rangeMin = TNumericLimits<float>::Max();
-    float rangeMax = -TNumericLimits<float>::Max();
-    MeshBufferIteratorConst< MBF_FLOAT32, float, 3 >  itop = itopBegin;
-    for ( int32 ov=0; ov<otherVertexCount; ++ov )
+    float RangeMin = TNumericLimits<float>::Max();
+    float RangeMax = -TNumericLimits<float>::Max();
+    MeshBufferIteratorConst< MBF_FLOAT32, float, 3 >  ItOtherPosition = ItOtherPositionBegin;
+    
+	for (int32 OtherVertex = 0; OtherVertex < OtherVertexCount; ++OtherVertex)
     {
-        float v = (*itop)[MUTABLE_BUCKET_CHANNEL];
-        rangeMin = FMath::Min( rangeMin, v );
-        rangeMax = FMath::Max( rangeMax, v );
-        ++itop;
+        float V = (*ItOtherPosition)[MUTABLE_BUCKET_CHANNEL];
+        RangeMin = FMath::Min(RangeMin, V);
+        RangeMax = FMath::Max(RangeMax, V);
+        ++ItOtherPosition;
     }
-    rangeMin -= tolerance;
-    rangeMax += tolerance;
+    RangeMin -= Tolerance;
+    RangeMax += Tolerance;
 
-    TArray<int32> buckets[MUTABLE_NUM_BUCKETS];
-    for ( int32 b=0; b<MUTABLE_NUM_BUCKETS; ++b )
+    TArray<int32> Buckets[MUTABLE_NUM_BUCKETS];
+    for (int32 BucketIndex = 0; BucketIndex < MUTABLE_NUM_BUCKETS; ++BucketIndex)
     {
-        buckets[b].Reserve( otherVertexCount/MUTABLE_NUM_BUCKETS*2 );
+        Buckets[BucketIndex].Reserve(OtherVertexCount/MUTABLE_NUM_BUCKETS*2);
     }
 
-    float bucketSize = (rangeMax-rangeMin)/float(MUTABLE_NUM_BUCKETS);
-    itop = itopBegin;
-    for ( int32 ov=0; ov<otherVertexCount; ++ov )
+    float BucketSize = (RangeMax-RangeMin)/float(MUTABLE_NUM_BUCKETS);
+    ItOtherPosition = ItOtherPositionBegin;
+    for (int32 OtherVertex = 0; OtherVertex < OtherVertexCount; ++OtherVertex)
     {
-        float v = (*itop)[MUTABLE_BUCKET_CHANNEL];
+        float V = (*ItOtherPosition)[MUTABLE_BUCKET_CHANNEL];
 
-        int32 bucket0 = int32( floor( (v-tolerance-rangeMin)/bucketSize ) );
-        bucket0 = FMath::Min( MUTABLE_NUM_BUCKETS-1, FMath::Max( 0, bucket0 ) );
-        buckets[bucket0].Add(ov);
+        int32 Bucket0 = FMath::FloorToInt((V-Tolerance-RangeMin)/BucketSize);
+        Bucket0 = FMath::Min(MUTABLE_NUM_BUCKETS-1, FMath::Max(0, Bucket0));
+        Buckets[Bucket0].Add(OtherVertex);
 
-        int32 bucket1 = int32( floor( (v+tolerance-rangeMin)/bucketSize ) );
-        bucket1 = FMath::Min( MUTABLE_NUM_BUCKETS-1, FMath::Max( 0, bucket1 ) );
+        int32 Bucket1 = FMath::FloorToInt((V + Tolerance - RangeMin)/BucketSize);
+        Bucket1 = FMath::Min(MUTABLE_NUM_BUCKETS-1, FMath::Max(0, Bucket1));
 
-        if (bucket1!=bucket0)
+        if (Bucket1 != Bucket0)
         {
-            buckets[bucket1].Add(ov);
+            Buckets[Bucket1].Add(OtherVertex);
         }
 
-        ++itop;
+        ++ItOtherPosition;
     }
 
     // TODO Compare only positions?
 
     // Use buckets
-    for ( int32 v=0; v<vertexCount; ++v )
+    for (int32 VertexIndex = 0; VertexIndex < VertexCount; ++VertexIndex)
     {
-        vertexMap.FirstMatch[v] = vertexMap.Matches.Num();
+        VertexMap.FirstMatch[VertexIndex] = VertexMap.Matches.Num();
 
-        float vbucket = (*itp)[MUTABLE_BUCKET_CHANNEL];
-        int32 bucket = int32( floor( (vbucket-rangeMin)/bucketSize ) );
+        float VBucket = (*ItPosition)[MUTABLE_BUCKET_CHANNEL];
+        int32 Bucket = FMath::FloorToInt((VBucket - RangeMin)/BucketSize);
 
-        if (bucket>=0 && bucket<MUTABLE_NUM_BUCKETS)
+        if (Bucket >= 0 && Bucket < MUTABLE_NUM_BUCKETS)
         {
-            int32 bucketVertexCount = buckets[bucket].Num();
-            for ( int32 ov=0; ov<bucketVertexCount; ++ov )
+            int32 BucketVertexCount = Buckets[Bucket].Num();
+            for (int32 OtherVertex = 0; OtherVertex < BucketVertexCount; ++OtherVertex)
             {
-                int32 otherVertexIndex = buckets[bucket][ov];
-                FVector3f p = (itopBegin+otherVertexIndex).GetAsVec3f();
+                int32 OtherVertexIndex = Buckets[Bucket][OtherVertex];
+                FVector3f Position = (ItOtherPositionBegin + OtherVertexIndex).GetAsVec3f();
 
-                bool same = true;
-                for ( int32 d=0; same && d<3; ++d )
+                bool bSame = true;
+                for (int32 Dim = 0; bSame && Dim < 3; ++Dim)
                 {
-                    float diff = fabs( (*itp)[d] - p[d] );
-                    same = diff <= tolerance;
+                    float Diff = FMath::Abs((*ItPosition)[Dim] - Position[Dim]);
+                    bSame = Diff <= Tolerance;
                 }
 
-                if ( same )
+                if (bSame)
                 {
-                    vertexMap.Matches.Add( otherVertexIndex );
+                    VertexMap.Matches.Add(OtherVertexIndex);
                 }
             }
         }
 
-        ++itp;
+        ++ItPosition;
     }
 
 }
@@ -984,19 +990,19 @@ void FMeshSurface::Unserialise(InputArchive& Arch)
 }
 
 
-void Mesh::FBonePose::Serialise(OutputArchive& arch) const
+void Mesh::FBonePose::Serialise(OutputArchive& Arch) const
 {
-	arch << BoneId;
-	arch << BoneUsageFlags;
-	arch << BoneTransform;
+	Arch << BoneId;
+	Arch << BoneUsageFlags;
+	Arch << BoneTransform;
 }
 
 
-void Mesh::FBonePose::Unserialise(InputArchive& arch)
+void Mesh::FBonePose::Unserialise(InputArchive& Arch)
 {
-	arch >> BoneId;
-	arch >> BoneUsageFlags;
-	arch >> BoneTransform;
+	Arch >> BoneId;
+	Arch >> BoneUsageFlags;
+	Arch >> BoneTransform;
 }
 
 
@@ -1066,7 +1072,7 @@ void Mesh::Unserialise(InputArchive& Arch)
 }
 
 
-bool Mesh::IsSimilar(const Mesh& o, bool bCompareLayouts) const
+bool Mesh::IsSimilar(const Mesh& Other, bool bCompareLayouts) const
 {
 	// Some meshes are just vertex indices (masks) we don't consider them for similarity,
 	// because the kind of vertex channel data they store is the kind that is ignored.
@@ -1075,52 +1081,58 @@ bool Mesh::IsSimilar(const Mesh& o, bool bCompareLayouts) const
 		return false;
 	}
 
-	bool equal = IndexBuffers == o.IndexBuffers;
-	if (equal) equal = (ReferenceID==o.ReferenceID);
-	if (equal && bCompareLayouts) equal = (Layouts.Num() == o.Layouts.Num());
-	if (equal && Skeleton != o.Skeleton)
+	bool bEqual = IndexBuffers == Other.IndexBuffers;
+	bEqual = bEqual && (ReferenceID == Other.ReferenceID);
+	
+	if (bEqual && bCompareLayouts) 
 	{
-		if (Skeleton && o.Skeleton)
+		bEqual = (Layouts.Num() == Other.Layouts.Num());
+	}
+
+	if (bEqual && Skeleton != Other.Skeleton)
+	{
+		if (Skeleton && Other.Skeleton)
 		{
-			equal = (*Skeleton == *o.Skeleton);
+			bEqual = (*Skeleton == *Other.Skeleton);
 		}
 		else
 		{
-			equal = false;
+			bEqual = false;
 		}
 	}
 	
-	if (equal && PhysicsBody != o.PhysicsBody)
+	if (bEqual && PhysicsBody != Other.PhysicsBody)
 	{
-		if (PhysicsBody && o.PhysicsBody)
+		if (PhysicsBody && Other.PhysicsBody)
 		{
-			equal = (*PhysicsBody == *o.PhysicsBody);
+			bEqual = (*PhysicsBody == *Other.PhysicsBody);
 		}
 		else
 		{
-			equal = false;
+			bEqual = false;
 		}
 	}
 
-	if (equal) equal = (Surfaces == o.Surfaces);
-	if (equal) equal = (Tags == o.Tags);
+	bEqual = bEqual && (Surfaces == Other.Surfaces);
+	bEqual = bEqual && (Tags == Other.Tags);
 
 	// Special comparison for layouts
 	if (bCompareLayouts)
 	{
-		for (int32 i = 0; equal && i < Layouts.Num(); ++i)
+		const int32 NumLayouts = Layouts.Num();
+		for (int32 LayoutIndex = 0; bEqual && LayoutIndex < NumLayouts; ++LayoutIndex)
 		{
-			equal &= Layouts[i]->IsSimilar(*o.Layouts[i]);
+			bEqual &= Layouts[LayoutIndex]->IsSimilar(*Other.Layouts[LayoutIndex]);
 		}
 	}
 
 	// Special comparison for vertex buffers
-	if (equal)
+	if (bEqual)
 	{
-		equal = VertexBuffers.IsSimilarRobust(o.VertexBuffers, bCompareLayouts);
+		bEqual = VertexBuffers.IsSimilarRobust(Other.VertexBuffers, bCompareLayouts);
 	}
 
-	return equal;
+	return bEqual;
 
 }
 
@@ -1289,155 +1301,155 @@ void Mesh::CheckIntegrity() const
 }
 
 
-static bool StaticMeshFormatIdentify_Project( const Mesh* pM )
+static bool StaticMeshFormatIdentify_Project(const Mesh* InMesh)
 {
     // This format is used internally for the mesh project
-    bool res = true;
+    bool bResult = true;
 
     // The first vertex buffer must be texcoords(2f), position(3f), normal(3f)
     // all tightly packed
-    res &= pM->VertexBuffers.GetBufferCount()>=1;
+    bResult &= InMesh->VertexBuffers.GetBufferCount() >= 1;
 
-    if ( res )
+    if (bResult)
     {
-        res &= pM->VertexBuffers.m_buffers[0].m_channels.Num()==3;
+        bResult &= InMesh->VertexBuffers.Buffers[0].Channels.Num() == 3;
     }
 
-    if ( res )
+    if (bResult)
     {
-        const FMeshBufferChannel& chan = pM->VertexBuffers.m_buffers[0].m_channels[0];
+        const FMeshBufferChannel& Channel = InMesh->VertexBuffers.Buffers[0].Channels[0];
 
-        res &= chan.m_semantic == MBS_TEXCOORDS;
-        res &= chan.m_format == MBF_FLOAT32;
-        res &= chan.m_componentCount == 2;
+        bResult &= Channel.Semantic == MBS_TEXCOORDS;
+        bResult &= Channel.Format == MBF_FLOAT32;
+        bResult &= Channel.ComponentCount == 2;
         //we don't really care about the semantic index
-        //res &= chan.m_semanticIndex == 0;
-        res &= chan.m_offset == 0;
+        //bResult &= Channel.SemanticIndex == 0;
+        bResult &= Channel.Offset == 0;
     }
 
-    if ( res )
+    if (bResult)
     {
-        const FMeshBufferChannel& chan = pM->VertexBuffers.m_buffers[0].m_channels[1];
+        const FMeshBufferChannel& Channel = InMesh->VertexBuffers.Buffers[0].Channels[1];
 
-        res &= chan.m_semantic == MBS_POSITION;
-        res &= chan.m_format == MBF_FLOAT32;
-        res &= chan.m_componentCount == 3;
-        res &= chan.m_semanticIndex == 0;
-        res &= chan.m_offset == 8;
+        bResult &= Channel.Semantic == MBS_POSITION;
+        bResult &= Channel.Format == MBF_FLOAT32;
+        bResult &= Channel.ComponentCount == 3;
+        bResult &= Channel.SemanticIndex == 0;
+        bResult &= Channel.Offset == 8;
     }
 
-    if ( res )
+    if (bResult)
     {
-        const FMeshBufferChannel& chan = pM->VertexBuffers.m_buffers[0].m_channels[2];
+        const FMeshBufferChannel& Channel = InMesh->VertexBuffers.Buffers[0].Channels[2];
 
-        res &= chan.m_semantic == MBS_NORMAL;
-        res &= chan.m_format == MBF_FLOAT32;
-        res &= chan.m_componentCount == 3;
-        res &= chan.m_semanticIndex == 0;
-        res &= chan.m_offset == 20;
+        bResult &= Channel.Semantic == MBS_NORMAL;
+        bResult &= Channel.Format == MBF_FLOAT32;
+        bResult &= Channel.ComponentCount == 3;
+        bResult &= Channel.SemanticIndex == 0;
+        bResult &= Channel.Offset == 20;
     }
 
     // The first index buffer must be just index buffers u32
-    if ( res )
+    if (bResult)
     {
-        res &= pM->IndexBuffers.m_buffers[0].m_channels.Num()>=1;
+        bResult &= InMesh->IndexBuffers.Buffers[0].Channels.Num() >= 1;
     }
 
-    if ( res )
+    if (bResult)
     {
-        const FMeshBufferChannel& chan = pM->IndexBuffers.m_buffers[0].m_channels[0];
+        const FMeshBufferChannel& Channel = InMesh->IndexBuffers.Buffers[0].Channels[0];
 
-        res &= chan.m_semantic == MBS_VERTEXINDEX;
-        res &= chan.m_format == MBF_UINT32;
-        res &= chan.m_componentCount == 1;
-        res &= chan.m_semanticIndex == 0;
-        res &= chan.m_offset == 0;
+        bResult &= Channel.Semantic == MBS_VERTEXINDEX;
+        bResult &= Channel.Format == MBF_UINT32;
+        bResult &= Channel.ComponentCount == 1;
+        bResult &= Channel.SemanticIndex == 0;
+        bResult &= Channel.Offset == 0;
     }
 
-    return res;
+    return bResult;
 }
 
 
-static bool StaticMeshFormatIdentify_ProjectWrapping( const Mesh* pM )
+static bool StaticMeshFormatIdentify_ProjectWrapping(const Mesh* InMesh)
 {
     // This format is used internally for the mesh project
-    bool res = true;
+    bool bResult = true;
 
     // The first vertex buffer must be texcoords(2f), position(3f), normal(3f), layoutBlock(uint32_t)
     // all tightly packed
-    res &= pM->VertexBuffers.GetBufferCount()>=2;
+    bResult &= InMesh->VertexBuffers.GetBufferCount() >= 2;
 
-    if ( res )
+    if (bResult)
     {
-        res &= pM->VertexBuffers.m_buffers[0].m_channels.Num()==3;
+        bResult &= InMesh->VertexBuffers.Buffers[0].Channels.Num() == 3;
     }
 
-    if ( res )
+    if (bResult)
     {
-        const FMeshBufferChannel& chan = pM->VertexBuffers.m_buffers[0].m_channels[0];
+        const FMeshBufferChannel& Channel = InMesh->VertexBuffers.Buffers[0].Channels[0];
 
-        res &= chan.m_semantic == MBS_TEXCOORDS;
-        res &= chan.m_format == MBF_FLOAT32;
-        res &= chan.m_componentCount == 2;
+        bResult &= Channel.Semantic == MBS_TEXCOORDS;
+        bResult &= Channel.Format == MBF_FLOAT32;
+        bResult &= Channel.ComponentCount == 2;
         // we don't really care about the semantic index as long as there is only one
-        // res &= chan.m_semanticIndex == 0;
-        res &= chan.m_offset == 0;
+        // bResult &= Channel.SemanticIndex == 0;
+        bResult &= Channel.Offset == 0;
     }
 
-    if ( res )
+    if (bResult)
     {
-        const FMeshBufferChannel& chan = pM->VertexBuffers.m_buffers[0].m_channels[1];
+        const FMeshBufferChannel& Channel = InMesh->VertexBuffers.Buffers[0].Channels[1];
 
-        res &= chan.m_semantic == MBS_POSITION;
-        res &= chan.m_format == MBF_FLOAT32;
-        res &= chan.m_componentCount == 3;
-        res &= chan.m_semanticIndex == 0;
-        res &= chan.m_offset == 8;
+        bResult &= Channel.Semantic == MBS_POSITION;
+        bResult &= Channel.Format == MBF_FLOAT32;
+        bResult &= Channel.ComponentCount == 3;
+        bResult &= Channel.SemanticIndex == 0;
+        bResult &= Channel.Offset == 8;
     }
 
-    if ( res )
+    if (bResult)
     {
-        const FMeshBufferChannel& chan = pM->VertexBuffers.m_buffers[0].m_channels[2];
+        const FMeshBufferChannel& Channel = InMesh->VertexBuffers.Buffers[0].Channels[2];
 
-        res &= chan.m_semantic == MBS_NORMAL;
-        res &= chan.m_format == MBF_FLOAT32;
-        res &= chan.m_componentCount == 3;
-        res &= chan.m_semanticIndex == 0;
-        res &= chan.m_offset == 20;
+        bResult &= Channel.Semantic == MBS_NORMAL;
+        bResult &= Channel.Format == MBF_FLOAT32;
+        bResult &= Channel.ComponentCount == 3;
+        bResult &= Channel.SemanticIndex == 0;
+        bResult &= Channel.Offset == 20;
     }
 
 	// Block IDs
-    if ( res )
+    if (bResult)
     {
-        const FMeshBufferChannel& chan = pM->VertexBuffers.m_buffers[1].m_channels[0];
+        const FMeshBufferChannel& Channel = InMesh->VertexBuffers.Buffers[1].Channels[0];
 
-        res &= chan.m_semantic == MBS_LAYOUTBLOCK;
+        bResult &= Channel.Semantic == MBS_LAYOUTBLOCK;
 		// We don't care about the layout block id format. We need to support them all.
-        //res &= chan.m_format == MBF_UINT64;
-        res &= chan.m_componentCount == 1;
+        //bResult &= Channel.Format == MBF_UINT64;
+        bResult &= Channel.ComponentCount == 1;
         // we don't really care about the semantic index as long as there is only one
-        // res &= chan.m_semanticIndex == 0;
-        res &= chan.m_offset == 0;
+        // bResult &= Channel.SemanticIndex == 0;
+        bResult &= Channel.Offset == 0;
     }
 
     // The first index buffer must be just index buffers u32
-    if ( res )
+    if (bResult)
     {
-        res &= pM->IndexBuffers.m_buffers[0].m_channels.Num()>=1;
+        bResult &= InMesh->IndexBuffers.Buffers[0].Channels.Num() >= 1;
     }
 
-    if ( res )
+    if (bResult)
     {
-        const FMeshBufferChannel& chan = pM->IndexBuffers.m_buffers[0].m_channels[0];
+        const FMeshBufferChannel& Channel = InMesh->IndexBuffers.Buffers[0].Channels[0];
 
-        res &= chan.m_semantic == MBS_VERTEXINDEX;
-        res &= chan.m_format == MBF_UINT32;
-        res &= chan.m_componentCount == 1;
-        res &= chan.m_semanticIndex == 0;
-        res &= chan.m_offset == 0;
+        bResult &= Channel.Semantic == MBS_VERTEXINDEX;
+        bResult &= Channel.Format == MBF_UINT32;
+        bResult &= Channel.ComponentCount == 1;
+        bResult &= Channel.SemanticIndex == 0;
+        bResult &= Channel.Offset == 0;
     }
 
-    return res;
+    return bResult;
 }
 
 
@@ -1460,54 +1472,72 @@ void Mesh::ResetStaticFormatFlags() const
 
 namespace
 {
-    void LogBuffer( FString& out, const FMeshBufferSet& bufset, int32 BufferElementLimit)
+    void LogBuffer(FString& Out, const FMeshBufferSet& BufferSet, int32 BufferElementLimit)
     {
-        (void)out;
-        (void)bufset;
-
-		uint32 elemCount = bufset.m_elementCount;
-        out += "  Set with "
-                + FString::Printf(TEXT("%d"), bufset.m_buffers.Num())
-                + " buffers and "
-                + FString::Printf(TEXT("%d"), elemCount)
-                + " elements.\n";
+		uint32 ElemCount = BufferSet.ElementCount;
+        Out += FString::Printf(TEXT("  Set with %d buffers and %d elements\n"), BufferSet.Buffers.Num(), ElemCount);
 
-        for( const FMeshBuffer& buf : bufset.m_buffers )
+        for(const FMeshBuffer& Buffer : BufferSet.Buffers)
         {
-            const uint8* pData = buf.m_data.GetData();
-
-            out += "    Buffer with "+ FString::Printf(TEXT("%d"), buf.m_channels.Num())
-                    + " channels and "+ FString::Printf(TEXT("%d"), buf.m_elementSize)+" elementsize\n";
-            for( const FMeshBufferChannel& chan : buf.m_channels )
+            Out += FString::Printf(TEXT("    Buffer with %d channels and %d elementsize\n"), Buffer.Channels.Num(), Buffer.ElementSize);
+            
+			const uint8* DataPtr = Buffer.Data.GetData();
+            for (const FMeshBufferChannel& Channel : Buffer.Channels)
             {
-                out += "      Channel with format: "+ FString::Printf(TEXT("%d"), chan.m_format)
-                        + " semantic: "+ FString::Printf(TEXT("%d"), chan.m_semantic)
-                        + " " + FString::Printf(TEXT("%d"), chan.m_semanticIndex)
-                        + " components: " + FString::Printf(TEXT("%d"), chan.m_componentCount)
-                        + " offset: " + FString::Printf(TEXT("%d"), chan.m_offset)+"\n";
-                for( size_t e=0; e<elemCount && e<BufferElementLimit; ++e )
+				Out += FString::Printf(TEXT("      Channel with format: %d semantic: %d %d, components: %d, offset: %d\n"),
+							Channel.Format, Channel.Semantic, Channel.SemanticIndex, Channel.ComponentCount, Channel.Offset);
+                
+				for (int32 ElemIndex = 0; uint32(ElemIndex) < ElemCount && ElemIndex < BufferElementLimit; ++ElemIndex)
                 {
-                    const uint8* pElementData = pData+buf.m_elementSize*e;
-                    const uint8* pChanData = pElementData+chan.m_offset;
-                    out += "        ";
-                    for (int c=0; c<chan.m_componentCount; ++c)
+                    Out += "        ";
+                    
+                    const uint8* ChanDataPtr = DataPtr + Buffer.ElementSize*ElemIndex + Channel.Offset;
+					for (uint32 CompIndex = 0; CompIndex < Channel.ComponentCount; ++CompIndex)
                     {
-                        out += "\t";
-                        switch (chan.m_format)
+                        Out += "\t";
+                        
+						switch (Channel.Format)
                         {
                         case MBF_UINT32:
-                        case MBF_NUINT32: out += FString::Printf(TEXT("%d"), *(const uint32_t*)pChanData); pChanData +=4; break;
+                        case MBF_NUINT32:
+						{
+							Out += FString::Printf(TEXT("%d"), *(const uint32*)ChanDataPtr); 
+							ChanDataPtr += 4; 
+							break;
+						}
                         case MBF_UINT16:
-                        case MBF_NUINT16: out += FString::Printf(TEXT("%d"), *(const uint16*)pChanData); pChanData +=2; break;
+                        case MBF_NUINT16: 
+						{
+							Out += FString::Printf(TEXT("%d"), *(const uint16*)ChanDataPtr); 
+							ChanDataPtr += 2; 
+							break;
+						}
                         case MBF_UINT8:
-                        case MBF_NUINT8: out += FString::Printf(TEXT("%d"), *(const uint8_t*)pChanData); pChanData +=1; break;
-						case MBF_FLOAT32: out += FString::Printf(TEXT("%.3f"), *(const float*)pChanData); pChanData += 4; break;
-                        case MBF_FLOAT16: out += FString::Printf(TEXT("%d"), *(const uint16*)pChanData); pChanData +=2; break;
-                        default: break;
+                        case MBF_NUINT8: 
+						{
+							Out += FString::Printf(TEXT("%d"), *(const uint8*)ChanDataPtr); 
+							ChanDataPtr += 1; 
+							break;
+						}
+						case MBF_FLOAT32: 
+						{
+							Out += FString::Printf(TEXT("%.3f"), *(const float*)ChanDataPtr); 
+							ChanDataPtr += 4; 
+							break;
+						}
+                        case MBF_FLOAT16: 
+						{
+							Out += FString::Printf(TEXT("%d"), *(const uint16*)ChanDataPtr); 
+							ChanDataPtr += 2; 
+							break;
+						}
+                        default: 
+							break;
                         }
-                        out += ",";
+
+                        Out += ",";
                     }
-                    out += "\n";
+                    Out += "\n";
                 }
             }
         }
diff --git a/Engine/Plugins/Experimental/Mutable/Source/MutableRuntime/Private/MuR/MeshBufferSet.cpp b/Engine/Plugins/Experimental/Mutable/Source/MutableRuntime/Private/MuR/MeshBufferSet.cpp
--- a/Engine/Plugins/Experimental/Mutable/Source/MutableRuntime/Private/MuR/MeshBufferSet.cpp
+++ b/Engine/Plugins/Experimental/Mutable/Source/MutableRuntime/Private/MuR/MeshBufferSet.cpp
@@ -19,7 +19,6 @@ namespace mu
 	MUTABLE_IMPLEMENT_ENUM_SERIALISABLE(EMeshBufferSemantic);
 
 	
-	//---------------------------------------------------------------------------------------------
 	static FMeshBufferFormatData s_meshBufferFormatData[] = // MBF_COUNT entries
 	{
 		{ 0, 0 },
@@ -55,42 +54,32 @@ namespace mu
 
 	static_assert(sizeof(s_meshBufferFormatData) / sizeof(FMeshBufferFormatData) == int32(MBF_COUNT));
 
-	//---------------------------------------------------------------------------------------------
-	const FMeshBufferFormatData& GetMeshFormatData(EMeshBufferFormat format)
+	const FMeshBufferFormatData& GetMeshFormatData(EMeshBufferFormat Format)
 	{
-		check(format >= 0);
-		check(format < MBF_COUNT);
-		return s_meshBufferFormatData[format];
+		check(Format >= 0);
+		check(Format < MBF_COUNT);
+		return s_meshBufferFormatData[Format];
 	}
 
-
-	//-----------------------------------------------------------------------------------------
-	void FMeshBufferSet::Serialise(OutputArchive& arch) const
+	void FMeshBufferSet::Serialise(OutputArchive& Arch) const
 	{
-		arch << m_elementCount;
-		arch << m_buffers;
+		Arch << ElementCount;
+		Arch << Buffers;
 	}
 
-
-	//-----------------------------------------------------------------------------------------
-	void FMeshBufferSet::Unserialise(InputArchive& arch)
+	void FMeshBufferSet::Unserialise(InputArchive& Arch)
 	{
-		arch >> m_elementCount;
-		arch >> m_buffers;
+		Arch >> ElementCount;
+		Arch >> Buffers;
 	}
 
 
-	//---------------------------------------------------------------------------------------------
-	//---------------------------------------------------------------------------------------------
-	//---------------------------------------------------------------------------------------------
 	int32 FMeshBufferSet::GetElementCount() const
 	{
-		return m_elementCount;
+		return ElementCount;
 	}
 
-
-	//---------------------------------------------------------------------------------------------
-	void FMeshBufferSet::SetElementCount(int32 Count)
+	void FMeshBufferSet::SetElementCount(int32 Count, EMemoryInitPolicy MemoryInitPolicy)
 	{
 		check(Count >= 0);
 		LLM_SCOPE_BYNAME(TEXT("MutableRuntime"));
@@ -101,139 +90,149 @@ namespace mu
 		// allow it if no other allocation will be done.  
 		const EAllowShrinking AllowShrinking = (Count == 0) ? EAllowShrinking::Yes : EAllowShrinking::No;
 
-		for (FMeshBuffer& buf : m_buffers)
+		for (FMeshBuffer& Buffer : Buffers)
 		{
-			buf.m_data.SetNumUninitialized(buf.m_elementSize * Count, AllowShrinking);
+			if (MemoryInitPolicy == EMemoryInitPolicy::Uninitialized)
+			{
+				Buffer.Data.SetNumUninitialized(Buffer.ElementSize * Count, AllowShrinking);
+			}
+			else if (MemoryInitPolicy == EMemoryInitPolicy::Zeroed)
+			{
+				Buffer.Data.SetNumZeroed(Buffer.ElementSize * Count, AllowShrinking);
+			}
+			else
+			{
+				check(false);
+			}
 		}
 		
-		m_elementCount = Count;
+		ElementCount = Count;
 	}
 
 
-	//---------------------------------------------------------------------------------------------
 	int32 FMeshBufferSet::GetBufferCount() const
 	{
-		return m_buffers.Num();
+		return Buffers.Num();
 	}
 
-
-	//---------------------------------------------------------------------------------------------
-	void FMeshBufferSet::SetBufferCount(int32 count)
+	void FMeshBufferSet::SetBufferCount(int32 Count)
 	{
 		LLM_SCOPE_BYNAME(TEXT("MutableRuntime"));
-		m_buffers.SetNum(count);
+		Buffers.SetNum(Count);
 	}
 
-
-	//---------------------------------------------------------------------------------------------
-	int32 FMeshBufferSet::GetBufferChannelCount(int32 buffer) const
+	int32 FMeshBufferSet::GetBufferChannelCount(int32 BufferIndex) const
 	{
-		check(buffer >= 0 && buffer < m_buffers.Num());
-		if (buffer >= 0 && buffer < m_buffers.Num())
+		check(Buffers.IsValidIndex(BufferIndex));
+		
+		if (Buffers.IsValidIndex(BufferIndex))
 		{
-			return m_buffers[buffer].m_channels.Num();
+			return Buffers[BufferIndex].Channels.Num();
 		}
+
 		return 0;
 	}
 
-
-	//---------------------------------------------------------------------------------------------
-	void FMeshBufferSet::GetChannel
-	(
-		int32 buffer,
-		int32 channel,
-		EMeshBufferSemantic* pSemantic,
-		int32* pSemanticIndex,
-		EMeshBufferFormat* pFormat,
-		int32* pComponentCount,
-		int32* pOffset
-	) const
+	void FMeshBufferSet::GetChannel(
+			int32 BufferIndex,
+			int32 ChannelIndex,
+			EMeshBufferSemantic* SemanticPtr,
+			int32* SemanticIndexPtr,
+			EMeshBufferFormat* FormatPtr,
+			int32* ComponentCountPtr,
+			int32* OffsetPtr) const
 	{
-		check(buffer >= 0 && buffer < m_buffers.Num());
-		check(channel >= 0 && channel < m_buffers[buffer].m_channels.Num());
+		check(Buffers.IsValidIndex(BufferIndex));
+		check(Buffers[BufferIndex].Channels.IsValidIndex(ChannelIndex));
 
-		const FMeshBufferChannel& chan = m_buffers[buffer].m_channels[channel];
+		const FMeshBufferChannel& Channel = Buffers[BufferIndex].Channels[ChannelIndex];
 
-		if (pSemantic)
+		if (SemanticPtr)
 		{
-			*pSemantic = chan.m_semantic;
+			*SemanticPtr = Channel.Semantic;
 		}
 
-		if (pSemanticIndex)
+		if (SemanticIndexPtr)
 		{
-			*pSemanticIndex = chan.m_semanticIndex;
+			*SemanticIndexPtr = Channel.SemanticIndex;
 		}
 
-		if (pFormat)
+		if (FormatPtr)
 		{
-			*pFormat = chan.m_format;
+			*FormatPtr = Channel.Format;
 		}
 
-		if (pComponentCount)
+		if (ComponentCountPtr)
 		{
-			*pComponentCount = chan.m_componentCount;
+			*ComponentCountPtr = Channel.ComponentCount;
 		}
 
-		if (pOffset)
+		if (OffsetPtr)
 		{
-			*pOffset = chan.m_offset;
+			*OffsetPtr = Channel.Offset;
 		}
 	}
 
 
-	//---------------------------------------------------------------------------------------------
-	void FMeshBufferSet::SetBuffer
-	(
-		int32 buffer,
-		int32 elementSize,
-		int32 channelCount,
-		const EMeshBufferSemantic* pSemantics,
-		const int32* pSemanticIndices,
-		const EMeshBufferFormat* pFormats,
-		const int32* pComponentCount,
-		const int32* pOffsets
-	)
+	void FMeshBufferSet::SetBuffer(
+		int32 BufferIndex,
+		int32 ElementSize,
+		int32 ChannelCount,
+		const EMeshBufferSemantic* SemanticsPtr,
+		const int32* SemanticIndicesPtr,
+		const EMeshBufferFormat* FormatsPtr,
+		const int32* ComponentCountPtr,
+		const int32* OffsetsPtr,
+		EMemoryInitPolicy MemoryInitPolicy)
 	{
-		check(buffer >= 0 && buffer < m_buffers.Num());
+		check(Buffers.IsValidIndex(BufferIndex));
 		LLM_SCOPE_BYNAME(TEXT("MutableRuntime"));
 
-		FMeshBuffer& buf = m_buffers[buffer];
+		FMeshBuffer& Buffer = Buffers[BufferIndex];
 
 		int32 MinElemSize = 0;
-		buf.m_channels.SetNum(channelCount);
-		for (int c = 0; c < channelCount; ++c)
+		Buffer.Channels.SetNum(ChannelCount);
+		for (int32 ChannelIndex = 0; ChannelIndex < ChannelCount; ++ChannelIndex)
 		{
-			FMeshBufferChannel& chan = buf.m_channels[c];
-			chan.m_semantic = pSemantics ? pSemantics[c] : MBS_NONE;
-			chan.m_semanticIndex = pSemanticIndices ? pSemanticIndices[c] : 0;
-			chan.m_format = pFormats ? pFormats[c] : MBF_NONE;
-			chan.m_componentCount = pComponentCount ? ((uint16)pComponentCount[c]) : 0;
-			chan.m_offset = pOffsets ? ((uint8)pOffsets[c]) : 0;
+			FMeshBufferChannel& Channel = Buffer.Channels[ChannelIndex];
+			Channel.Semantic = SemanticsPtr ? SemanticsPtr[ChannelIndex] : MBS_NONE;
+			Channel.SemanticIndex = SemanticIndicesPtr ? SemanticIndicesPtr[ChannelIndex] : 0;
+			Channel.Format = FormatsPtr ? FormatsPtr[ChannelIndex] : MBF_NONE;
+			Channel.ComponentCount = ComponentCountPtr ? ((uint16)ComponentCountPtr[ChannelIndex]) : 0;
+			Channel.Offset = OffsetsPtr ? ((uint8)OffsetsPtr[ChannelIndex]) : 0;
 
-			int32 ThisChannelMinElemSize = chan.m_offset + chan.m_componentCount * GetMeshFormatData(chan.m_format).SizeInBytes;			
+			int32 ThisChannelMinElemSize = Channel.Offset + Channel.ComponentCount * GetMeshFormatData(Channel.Format).SizeInBytes;			
 			MinElemSize = FMath::Max(MinElemSize, ThisChannelMinElemSize);
 		}
 
 		// Set the user specified element size, or enlarge it if it was too small
-		buf.m_elementSize = FMath::Max(elementSize, MinElemSize);
+		Buffer.ElementSize = FMath::Max(ElementSize, MinElemSize);
 
 		// Update the buffer data
-		buf.m_data.SetNumUninitialized(buf.m_elementSize * m_elementCount, EAllowShrinking::No);
+		if (MemoryInitPolicy == EMemoryInitPolicy::Uninitialized)
+		{
+			Buffer.Data.SetNumUninitialized(Buffer.ElementSize * ElementCount, EAllowShrinking::No);
+		}
+		else if (MemoryInitPolicy == EMemoryInitPolicy::Zeroed)
+		{
+			Buffer.Data.SetNumZeroed(Buffer.ElementSize * ElementCount, EAllowShrinking::No);
+		}
+		else
+		{
+			check(false);
+		}
 	}
 
-	//---------------------------------------------------------------------------------------------
-	void FMeshBufferSet::SetBufferChannel
-	(
-		int32 buffer,
-		int32 channelIndex,
-		EMeshBufferSemantic semantic,
-		int32 semanticIndex,
-		EMeshBufferFormat format,
-		int32 componentCount,
-		int32 offset
-	)
+	void FMeshBufferSet::SetBufferChannel(
+		int32 BufferIndex,
+		int32 ChannelIndex,
+		EMeshBufferSemantic Semantic,
+		int32 SemanticIndex,
+		EMeshBufferFormat Format,
+		int32 ComponentCount,
+		int32 Offset)
 	{
-		if (!(buffer >= 0 && buffer < m_buffers.Num()))
+		if (!Buffers.IsValidIndex(BufferIndex))
 		{
 			check(false);
 			return;
@@ -241,124 +240,104 @@ namespace mu
 
 		LLM_SCOPE_BYNAME(TEXT("MutableRuntime"));
 
-		FMeshBuffer& buf = m_buffers[buffer];
+		FMeshBuffer& Buffer = Buffers[BufferIndex];
 
-		if (!(channelIndex >= 0 && channelIndex < buf.m_channels.Num()))
+		if (!Buffer.Channels.IsValidIndex(ChannelIndex))
 		{
 			check(false);
 			return;
 		}
 
-		FMeshBufferChannel& chan = buf.m_channels[channelIndex];
-		chan.m_semantic = semantic;
-		chan.m_semanticIndex = semanticIndex;
-		chan.m_format = format;
-		chan.m_componentCount = uint16(componentCount);
-		chan.m_offset = uint8(offset);
+		FMeshBufferChannel& Channel = Buffer.Channels[ChannelIndex];
+		Channel.Semantic = Semantic;
+		Channel.SemanticIndex = SemanticIndex;
+		Channel.Format = Format;
+		Channel.ComponentCount = uint16(ComponentCount);
+		Channel.Offset = uint8(Offset);
 	}
 
 
-	//---------------------------------------------------------------------------------------------
-	uint8* FMeshBufferSet::GetBufferData(int32 buffer)
+	uint8* FMeshBufferSet::GetBufferData(int32 BufferIndex)
 	{
-		check(buffer >= 0 && buffer < m_buffers.Num());
-		uint8* pResult = m_buffers[buffer].m_data.GetData();
-		return pResult;
+		check(Buffers.IsValidIndex(BufferIndex));
+		uint8* ResultPtr = Buffers[BufferIndex].Data.GetData();
+		return ResultPtr;
 	}
 
-
-	//---------------------------------------------------------------------------------------------
-	const uint8* FMeshBufferSet::GetBufferData(int32 buffer) const
+	const uint8* FMeshBufferSet::GetBufferData(int32 BufferIndex) const
 	{
-		check(buffer >= 0 && buffer < m_buffers.Num());
-		const uint8* pResult = m_buffers[buffer].m_data.GetData();
-		return pResult;
+		check(Buffers.IsValidIndex(BufferIndex));
+		const uint8* ResultPtr = Buffers[BufferIndex].Data.GetData();
+		return ResultPtr;
 	}
 
-
-	//---------------------------------------------------------------------------------------------
-	uint32 FMeshBufferSet::GetBufferDataSize(int32 buffer) const
+	uint32 FMeshBufferSet::GetBufferDataSize(int32 BufferIndex) const
 	{
-		check(buffer >= 0 && buffer < m_buffers.Num());
-		const uint32 Result = m_buffers[buffer].m_data.Num();
+		check(Buffers.IsValidIndex(BufferIndex));
+		const uint32 Result = Buffers[BufferIndex].Data.Num();
 
 #if WITH_EDITOR
-		int32 Size = m_buffers[buffer].m_elementSize * m_elementCount;
+		int32 Size = Buffers[BufferIndex].ElementSize * ElementCount;
 		ensure(Size == Result);
 #endif
 
 		return Result;
 	}
 
-
-	//---------------------------------------------------------------------------------------------
-	void FMeshBufferSet::FindChannel
-	(
-		EMeshBufferSemantic semantic, int32 semanticIndex,
-		int32* pBuffer, int32* pChannel
-	) const
+	void FMeshBufferSet::FindChannel(EMeshBufferSemantic Semantic, int32 SemanticIndex, int32* BufferPtr, int32* ChannelPtr) const
 	{
-		check(pBuffer && pChannel);
+		check(BufferPtr && ChannelPtr);
 
-		*pBuffer = -1;
-		*pChannel = -1;
+		*BufferPtr = -1;
+		*ChannelPtr = -1;
 
-		for (int32 b = 0; b < m_buffers.Num(); ++b)
+		const int32 NumBuffers = Buffers.Num();
+		for (int32 BufferIndex = 0; BufferIndex < NumBuffers; ++BufferIndex)
 		{
-			for (int32 c = 0; c < m_buffers[b].m_channels.Num(); ++c)
+			const int32 NumChannels = Buffers[BufferIndex].Channels.Num();
+			for (int32 ChannelIndex = 0; ChannelIndex < NumChannels; ++ChannelIndex)
 			{
-				if (m_buffers[b].m_channels[c].m_semantic == semantic &&
-					m_buffers[b].m_channels[c].m_semanticIndex == semanticIndex)
+				if (Buffers[BufferIndex].Channels[ChannelIndex].Semantic == Semantic &&
+					Buffers[BufferIndex].Channels[ChannelIndex].SemanticIndex == SemanticIndex)
 				{
-					*pBuffer = b;
-					*pChannel = c;
+					*BufferPtr = BufferIndex;
+					*ChannelPtr = ChannelIndex;
+
 					return;
 				}
 			}
 		}
 	}
 
-
-	//---------------------------------------------------------------------------------------------
-	int32 FMeshBufferSet::GetElementSize(int32 buffer) const
+	int32 FMeshBufferSet::GetElementSize(int32 BufferIndex) const
 	{
-		check(buffer >= 0 && buffer < m_buffers.Num());
+		check(Buffers.IsValidIndex(BufferIndex));
 
-		return m_buffers[buffer].m_elementSize;
+		return Buffers[BufferIndex].ElementSize;
 	}
 
-
-	//---------------------------------------------------------------------------------------------
-	int32 FMeshBufferSet::GetChannelOffset
-	(
-		int32 buffer,
-		int32 channel
-	) const
+	int32 FMeshBufferSet::GetChannelOffset(int32 BufferIndex, int32 ChannelIndex) const
 	{
-		check(buffer >= 0 && buffer < m_buffers.Num());
-		check(channel >= 0 && channel < m_buffers[buffer].m_channels.Num());
+		check(Buffers.IsValidIndex(BufferIndex));
+		check(Buffers[BufferIndex].Channels.IsValidIndex(ChannelIndex));
 
-		return m_buffers[buffer].m_channels[channel].m_offset;
+		return Buffers[BufferIndex].Channels[ChannelIndex].Offset;
 	}
 
-
-	//---------------------------------------------------------------------------------------------
-	void FMeshBufferSet::AddBuffer(const FMeshBufferSet& Other, int32 buffer)
+	void FMeshBufferSet::AddBuffer(const FMeshBufferSet& Other, int32 BufferIndex)
 	{
 		check(GetElementCount() == Other.GetElementCount());
 
 		LLM_SCOPE_BYNAME(TEXT("MutableRuntime"));
 
-		m_buffers.Add(Other.m_buffers[buffer]);
+		Buffers.Add(Other.Buffers[BufferIndex]);
 	}
 
-
-	//---------------------------------------------------------------------------------------------
 	void FMeshBufferSet::RemoveBuffer(int32 BufferIndex)
 	{
-		if (BufferIndex >= 0 && BufferIndex < m_buffers.Num())
+		if (Buffers.IsValidIndex(BufferIndex))
 		{
-			m_buffers.RemoveAt(BufferIndex);
+			Buffers.RemoveAt(BufferIndex);
 		}
 		else
 		{
@@ -366,19 +345,17 @@ namespace mu
 		}
 	}
 
-
-	//---------------------------------------------------------------------------------------------
 	bool FMeshBufferSet::HasSameFormat(const FMeshBufferSet& Other) const
 	{
-		int32 bc = GetBufferCount();
-		if (Other.GetBufferCount() != bc)
+		const int32 BufferCount = GetBufferCount();
+		if (Other.GetBufferCount() != BufferCount)
 		{
 			return false;
 		}
 
-		for (int32 b = 0; b < bc; ++b)
+		for (int32 BufferIndex = 0; BufferIndex < BufferCount; ++BufferIndex)
 		{
-			if (!HasSameFormat(b, Other, b))
+			if (!HasSameFormat(BufferIndex, Other, BufferIndex))
 			{
 				return false;
 			}
@@ -387,81 +364,74 @@ namespace mu
 		return true;
 	}
 
-
-	//---------------------------------------------------------------------------------------------
 	bool FMeshBufferSet::HasSameFormat(int32 ThisBufferIndex, const FMeshBufferSet& Other, int32 OtherBufferIndex) const
 	{
-		return m_buffers[ThisBufferIndex].HasSameFormat(Other.m_buffers[OtherBufferIndex]);
+		return Buffers[ThisBufferIndex].HasSameFormat(Other.Buffers[OtherBufferIndex]);
 	}
 
-
-	//---------------------------------------------------------------------------------------------
 	int32 FMeshBufferSet::GetDataSize() const
 	{
-		int32 res = 0;
+		int32 Result = 0;
 
-		for (int32 b = 0; b < m_buffers.Num(); ++b)
+		const int32 NumBuffers = Buffers.Num();
+		for (int32 BufferIndex = 0; BufferIndex < NumBuffers; ++BufferIndex)
 		{
-			res += m_buffers[b].m_elementSize * m_elementCount;
+			Result += Buffers[BufferIndex].ElementSize * ElementCount;
 		}
 
-		return res;
+		return Result;
 	}
 
 	int32 FMeshBufferSet::GetAllocatedSize() const
 	{
 		int32 ByteCount = 0;
 
-		for (int32 BufferIndex = 0; BufferIndex < m_buffers.Num(); ++BufferIndex)
+		const int32 NumBuffers = Buffers.Num();
+		for (int32 BufferIndex = 0; BufferIndex < NumBuffers; ++BufferIndex)
 		{
-			ByteCount += m_buffers[BufferIndex].m_data.GetAllocatedSize();
+			ByteCount += Buffers[BufferIndex].Data.GetAllocatedSize();
 		}
 
 		return ByteCount;
 	}
 
-	//-----------------------------------------------------------------------------------------
-	void FMeshBufferSet::CopyElement(uint32 fromIndex, uint32 toIndex)
+	void FMeshBufferSet::CopyElement(uint32 FromIndex, uint32 ToIndex)
 	{
-		check(fromIndex < m_elementCount);
-		check(toIndex < m_elementCount);
+		check(FromIndex < ElementCount);
+		check(ToIndex < ElementCount);
 
-		if (fromIndex != toIndex)
+		if (FromIndex != ToIndex)
 		{
-			for (FMeshBuffer& b : m_buffers)
+			for (FMeshBuffer& Buffer : Buffers)
 			{
-				FMemory::Memcpy(&b.m_data[b.m_elementSize * toIndex],
-					&b.m_data[b.m_elementSize * fromIndex],
-					b.m_elementSize);
+				FMemory::Memcpy(
+						&Buffer.Data[Buffer.ElementSize * ToIndex],
+						&Buffer.Data[Buffer.ElementSize * FromIndex],
+						Buffer.ElementSize);
 			}
 		}
 	}
 
-
-	//-----------------------------------------------------------------------------------------
-	bool FMeshBufferSet::IsSpecialBufferToIgnoreInSimilar(const FMeshBuffer& b) const
+	bool FMeshBufferSet::IsSpecialBufferToIgnoreInSimilar(const FMeshBuffer& Buffer) const
 	{
-		if (b.m_channels.Num() == 1
-			&&
-			b.m_channels[0].m_semantic == MBS_VERTEXINDEX)
+		if (Buffer.Channels.Num() == 1 && Buffer.Channels[0].Semantic == MBS_VERTEXINDEX)
 		{
 			return true;
 		}
-		if (b.m_channels.Num() == 1
-			&&
-			b.m_channels[0].m_semantic == MBS_LAYOUTBLOCK)
+
+		if (Buffer.Channels.Num() == 1 && Buffer.Channels[0].Semantic == MBS_LAYOUTBLOCK)
 		{
 			return true;
 		}
+
 		return false;
 	}
 
-	//-----------------------------------------------------------------------------------------
 	bool FMeshBufferSet::IsSimilarRobust(const FMeshBufferSet& Other, bool bCompareUVs) const
 	{
 		MUTABLE_CPUPROFILER_SCOPE(FMeshBufferSet::IsSimilarRobust);
 
-		if (m_elementCount != Other.m_elementCount)
+		if (ElementCount != Other.ElementCount)
 		{
 			return false;
 		}
@@ -473,31 +443,31 @@ namespace mu
 			return true;
 		}
 
-		const int32 ThisNumBuffers = m_buffers.Num();
-		const int32 OtherNumBuffers = Other.m_buffers.Num();
+		const int32 ThisNumBuffers = Buffers.Num();
+		const int32 OtherNumBuffers = Other.Buffers.Num();
 
 		int32 I = 0, J = 0;
 		while (I < ThisNumBuffers && J < OtherNumBuffers)
 		{
-			if (IsSpecialBufferToIgnoreInSimilar(m_buffers[I]))
+			if (IsSpecialBufferToIgnoreInSimilar(Buffers[I]))
 			{
 				++I;
 				continue;
 			}
 
-			if (IsSpecialBufferToIgnoreInSimilar(Other.m_buffers[J]))
+			if (IsSpecialBufferToIgnoreInSimilar(Other.Buffers[J]))
 			{
 				++J;
 				continue;
 			}
 
-			const int32 ThisNumChannels = m_buffers[I].m_channels.Num();
-			const int32 OtherNumChannels = m_buffers[J].m_channels.Num();
+			const int32 ThisNumChannels = Buffers[I].Channels.Num();
+			const int32 OtherNumChannels = Buffers[J].Channels.Num();
 
-			const FMeshBuffer& ThisBuffer = m_buffers[I];
-			const FMeshBuffer& OtherBuffer = Other.m_buffers[J];
+			const FMeshBuffer& ThisBuffer = Buffers[I];
+			const FMeshBuffer& OtherBuffer = Other.Buffers[J];
 
-			if (!(ThisBuffer.m_channels == OtherBuffer.m_channels && ThisBuffer.m_elementSize==OtherBuffer.m_elementSize))
+			if (!(ThisBuffer.Channels == OtherBuffer.Channels && ThisBuffer.ElementSize == OtherBuffer.ElementSize))
 			{
 				return false;
 			}
@@ -509,7 +479,7 @@ namespace mu
 				MUTABLE_CPUPROFILER_SCOPE(FastCompare);
 
 				// This buffer can be directly compared
-				if (FMemory::Memcmp(ThisBuffer.m_data.GetData(), OtherBuffer.m_data.GetData(), ThisBuffer.m_data.Num()) != 0)
+				if (FMemory::Memcmp(ThisBuffer.Data.GetData(), OtherBuffer.Data.GetData(), ThisBuffer.Data.Num()) != 0)
 				{
 					return false;
 				}
@@ -518,21 +488,21 @@ namespace mu
 			{
 				MUTABLE_CPUPROFILER_SCOPE(SlowCompare);
 
-				for (uint32 Elem = 0; Elem < m_elementCount; ++Elem)
+				for (uint32 Elem = 0; Elem < ElementCount; ++Elem)
 				{
 					for (int32 C = 0; C < ThisNumChannels; ++C)
 					{
-						if (!bCompareUVs && ThisBuffer.m_channels[C].m_semantic == MBS_TEXCOORDS)
+						if (!bCompareUVs && ThisBuffer.Channels[C].Semantic == MBS_TEXCOORDS)
 						{
 							continue;
 						}
 
-						const SIZE_T SizeA = GetMeshFormatData(ThisBuffer.m_channels[C].m_format).SizeInBytes * ThisBuffer.m_channels[C].m_componentCount;
-						const SIZE_T SizeB = GetMeshFormatData(OtherBuffer.m_channels[C].m_format).SizeInBytes * OtherBuffer.m_channels[C].m_componentCount;
+						const SIZE_T SizeA = GetMeshFormatData(ThisBuffer.Channels[C].Format).SizeInBytes * ThisBuffer.Channels[C].ComponentCount;
+						const SIZE_T SizeB = GetMeshFormatData(OtherBuffer.Channels[C].Format).SizeInBytes * OtherBuffer.Channels[C].ComponentCount;
 						check(SizeA == SizeB);
 
-						const uint8* BuffA = ThisBuffer.m_data.GetData() + Elem * ThisBuffer.m_elementSize + ThisBuffer.m_channels[C].m_offset;
-						const uint8* BuffB = OtherBuffer.m_data.GetData() + Elem * OtherBuffer.m_elementSize + OtherBuffer.m_channels[C].m_offset;
+						const uint8* BuffA = ThisBuffer.Data.GetData() + Elem * ThisBuffer.ElementSize + ThisBuffer.Channels[C].Offset;
+						const uint8* BuffB = OtherBuffer.Data.GetData() + Elem * OtherBuffer.ElementSize + OtherBuffer.Channels[C].Offset;
 
 						if (FMemory::Memcmp(BuffA, BuffB, SizeA) != 0)
 						{
@@ -549,7 +519,7 @@ namespace mu
 		// Whatever buffers are left should be irrelevant
 		while (I < ThisNumBuffers)
 		{
-			if (!IsSpecialBufferToIgnoreInSimilar(m_buffers[I]))
+			if (!IsSpecialBufferToIgnoreInSimilar(Buffers[I]))
 			{
 				return false;
 			}
@@ -558,7 +528,7 @@ namespace mu
 
 		while (J < OtherNumBuffers)
 		{
-			if (!IsSpecialBufferToIgnoreInSimilar(Other.m_buffers[J]))
+			if (!IsSpecialBufferToIgnoreInSimilar(Other.Buffers[J]))
 			{
 				return false;
 			}
@@ -568,37 +538,39 @@ namespace mu
 		return true;
 	}
 
-	//-----------------------------------------------------------------------------------------
 	bool FMeshBufferSet::IsSimilar(const FMeshBufferSet& Other) const
 	{
 		MUTABLE_CPUPROFILER_SCOPE(FMeshBufferSet::IsSimilar);
 
-		if (m_elementCount != Other.m_elementCount) return false;
+		if (ElementCount != Other.ElementCount)
+		{
+			return false;
+		}
 
 		// Compare all buffers except the vertex index channel, which should always be alone in
 		// the last buffer
 		int32 Index = 0;
 		int32 OtherIndex = 0;
 
-		const int32 ThisNumBuffers = m_buffers.Num();
-		const int32 OtherNumBuffers = Other.m_buffers.Num();
+		const int32 ThisNumBuffers = Buffers.Num();
+		const int32 OtherNumBuffers = Other.Buffers.Num();
 
 		while (Index < ThisNumBuffers && OtherIndex < OtherNumBuffers)
 		{
 			// Is it a special buffer that we should ignore?
-			if (IsSpecialBufferToIgnoreInSimilar(m_buffers[Index]))
+			if (IsSpecialBufferToIgnoreInSimilar(Buffers[Index]))
 			{
 				++Index;
 				continue;
 			}
 
-			if (IsSpecialBufferToIgnoreInSimilar(Other.m_buffers[OtherIndex]))
+			if (IsSpecialBufferToIgnoreInSimilar(Other.Buffers[OtherIndex]))
 			{
 				++OtherIndex;
 				continue;
 			}
 
-			if (!(m_buffers[Index] == Other.m_buffers[OtherIndex]))
+			if (!(Buffers[Index] == Other.Buffers[OtherIndex]))
 			{
 				return false;
 			}
@@ -609,7 +581,7 @@ namespace mu
 		// Whatever buffers are left should be irrelevant
 		while (Index < ThisNumBuffers)
 		{
-			if (!IsSpecialBufferToIgnoreInSimilar(m_buffers[Index]))
+			if (!IsSpecialBufferToIgnoreInSimilar(Buffers[Index]))
 			{
 				return false;
 			}
@@ -618,7 +590,7 @@ namespace mu
 
 		while (OtherIndex < OtherNumBuffers)
 		{
-			if (!IsSpecialBufferToIgnoreInSimilar(Other.m_buffers[OtherIndex]))
+			if (!IsSpecialBufferToIgnoreInSimilar(Other.Buffers[OtherIndex]))
 			{
 				return false;
 			}
@@ -628,57 +600,52 @@ namespace mu
 		return true;
 	}
 
-
-	//-----------------------------------------------------------------------------------------
 	void FMeshBufferSet::ResetBufferIndices()
 	{
-		int32 currentIndices[MBS_COUNT];
-		memset(currentIndices, 0, sizeof(currentIndices));
-
-		for (FMeshBuffer& b : m_buffers)
+		int32 CurrentIndices[MBS_COUNT] = {0};
+		for (FMeshBuffer& Buffer : Buffers)
 		{
-			for (FMeshBufferChannel& c : b.m_channels)
+			for (FMeshBufferChannel& Channel : Buffer.Channels)
 			{
-				c.m_semanticIndex = currentIndices[c.m_semantic];
-				currentIndices[c.m_semantic]++;
+				Channel.SemanticIndex = CurrentIndices[Channel.Semantic];
+				CurrentIndices[Channel.Semantic]++;
 			}
 		}
 	}
 
-
-	//-----------------------------------------------------------------------------------------
-	void FMeshBufferSet::UpdateOffsets(int32 b)
+	void FMeshBufferSet::UpdateOffsets(int32 BufferIndex)
 	{
-		uint32 offset = 0;
-		for (FMeshBufferChannel& c : m_buffers[b].m_channels)
+		uint32 Offset = 0;
+		for (FMeshBufferChannel& Channel : Buffers[BufferIndex].Channels)
 		{
-			if (c.m_offset < offset)
+			if (Channel.Offset < Offset)
 			{
-				c.m_offset = offset;
+				Channel.Offset = Offset;
 			}
 			else
 			{
-				offset = c.m_offset;
+				Offset = Channel.Offset;
 			}
-			offset += c.m_componentCount * GetMeshFormatData(c.m_format).SizeInBytes;
+
+			Offset += Channel.ComponentCount * GetMeshFormatData(Channel.Format).SizeInBytes;
 		}
 
-		if (m_buffers[b].m_elementSize < offset)
+		if (Buffers[BufferIndex].ElementSize < Offset)
 		{
-			m_buffers[b].m_elementSize = offset;
+			Buffers[BufferIndex].ElementSize = Offset;
 		}
 	}
 
-
 	bool FMeshBufferSet::HasAnySemanticWithDifferentFormat(EMeshBufferSemantic Semantic, EMeshBufferFormat ExpectedFormat) const
 	{
-		for (const FMeshBuffer& Buffer : m_buffers)
+		for (const FMeshBuffer& Buffer : Buffers)
 		{
-			for (int32 ChannelIndex = 0; ChannelIndex < Buffer.m_channels.Num(); ++ChannelIndex)
+			const int32 NumChannels = Buffer.Channels.Num();
+			for (int32 ChannelIndex = 0; ChannelIndex < NumChannels; ++ChannelIndex)
 			{
-				if (Buffer.m_channels[ChannelIndex].m_semantic == Semantic)
+				if (Buffer.Channels[ChannelIndex].Semantic == Semantic)
 				{
-					if (Buffer.m_channels[ChannelIndex].m_format != ExpectedFormat)
+					if (Buffer.Channels[ChannelIndex].Format != ExpectedFormat)
 					{
 						return true;
 					}
@@ -689,22 +656,17 @@ namespace mu
 		return false;
 	}
 
-	
-	//-----------------------------------------------------------------------------------------
-	void FMeshBuffer::Serialise(OutputArchive& arch) const
+	void FMeshBuffer::Serialise(OutputArchive& Arch) const
 	{
-		arch << m_channels;
-		arch << m_data;
-		arch << m_elementSize;
+		Arch << Channels;
+		Arch << Data;
+		Arch << ElementSize;
 	}
 
-	
-	//-----------------------------------------------------------------------------------------
-	void FMeshBuffer::Unserialise(InputArchive& arch)
+	void FMeshBuffer::Unserialise(InputArchive& Arch)
 	{
-		arch >> m_channels;
-		arch >> m_data;
-		arch >> m_elementSize;
+		Arch >> Channels;
+		Arch >> Data;
+		Arch >> ElementSize;
 	}
-	
 }
diff --git a/Engine/Plugins/Experimental/Mutable/Source/MutableRuntime/Private/MuR/OpImageProject.cpp b/Engine/Plugins/Experimental/Mutable/Source/MutableRuntime/Private/MuR/OpImageProject.cpp
--- a/Engine/Plugins/Experimental/Mutable/Source/MutableRuntime/Private/MuR/OpImageProject.cpp
+++ b/Engine/Plugins/Experimental/Mutable/Source/MutableRuntime/Private/MuR/OpImageProject.cpp
@@ -778,7 +778,7 @@ namespace Private
 		float FadeEndCos = FMath::Cos(FadeEnd);
 
 		// Calculate the culled flag
-		EMeshBufferFormat LayoutBlockType = pMesh->VertexBuffers.m_buffers[1].m_channels[0].m_format;
+		EMeshBufferFormat LayoutBlockType = pMesh->VertexBuffers.Buffers[1].Channels[0].Format;
 		switch (LayoutBlockType)
 		{
 		case MBF_UINT64:
@@ -2173,10 +2173,10 @@ void MeshProject_Optimised_Wrapping(const Mesh* pMesh,
 		// Get the vertices
     int vertexCount = pMesh->GetVertexCount();
 	check(pMesh->GetVertexBuffers().GetElementSize(0) == sizeof(FOptimizedVertex));
-	check(pMesh->GetVertexBuffers().m_buffers.Num() == 2);
+	check(pMesh->GetVertexBuffers().Buffers.Num() == 2);
 	const FOptimizedVertex* pVertices = reinterpret_cast<const FOptimizedVertex*>(pMesh->GetVertexBuffers().GetBufferData(0));
 
-	EMeshBufferFormat LayoutBlockType = pMesh->VertexBuffers.m_buffers[1].m_channels[0].m_format;
+	EMeshBufferFormat LayoutBlockType = pMesh->VertexBuffers.Buffers[1].Channels[0].Format;
 
     // Get the indices
     check(pMesh->GetIndexBuffers().GetElementSize(0) == 4);
diff --git a/Engine/Plugins/Experimental/Mutable/Source/MutableRuntime/Private/MuR/OpMeshApplyLayout.h b/Engine/Plugins/Experimental/Mutable/Source/MutableRuntime/Private/MuR/OpMeshApplyLayout.h
--- a/Engine/Plugins/Experimental/Mutable/Source/MutableRuntime/Private/MuR/OpMeshApplyLayout.h
+++ b/Engine/Plugins/Experimental/Mutable/Source/MutableRuntime/Private/MuR/OpMeshApplyLayout.h
@@ -67,8 +67,8 @@ namespace mu
 		}
 
         check( layoutBuffer>=0 && layoutChannel>=0 );
-        check( pApplied->GetVertexBuffers().m_buffers[layoutBuffer].m_channels.Num()==1 );
-        check( pApplied->GetVertexBuffers().m_buffers[layoutBuffer].m_channels[layoutChannel].m_componentCount==1 );
+        check( pApplied->GetVertexBuffers().Buffers[layoutBuffer].Channels.Num()==1 );
+        check( pApplied->GetVertexBuffers().Buffers[layoutBuffer].Channels[layoutChannel].ComponentCount==1 );
 
         const uint16* pLayoutData = reinterpret_cast<const uint16*>( pApplied->GetVertexBuffers().GetBufferData( layoutBuffer ) );
 		UntypedMeshBufferIterator ItLayoutData(pApplied->GetVertexBuffers(), MBS_LAYOUTBLOCK, texCoordsSet);
diff --git a/Engine/Plugins/Experimental/Mutable/Source/MutableRuntime/Private/MuR/OpMeshClipDeform.h b/Engine/Plugins/Experimental/Mutable/Source/MutableRuntime/Private/MuR/OpMeshClipDeform.h
--- a/Engine/Plugins/Experimental/Mutable/Source/MutableRuntime/Private/MuR/OpMeshClipDeform.h
+++ b/Engine/Plugins/Experimental/Mutable/Source/MutableRuntime/Private/MuR/OpMeshClipDeform.h
@@ -121,15 +121,15 @@ namespace mu
 		
 		// Copy buffers skipping binding data.
 		FMeshBufferSet& ResultBuffers = Result->GetVertexBuffers();
-		ResultBuffers.m_elementCount = VB.m_elementCount;
+		ResultBuffers.ElementCount = VB.ElementCount;
 		// Remove one element to the number of buffers if BarycentricDataBuffer found. 
-		ResultBuffers.m_buffers.SetNum( FMath::Max( 0, VB.m_buffers.Num() - int32(BarycentricDataBuffer >= 0) ) );
+		ResultBuffers.Buffers.SetNum( FMath::Max( 0, VB.Buffers.Num() - int32(BarycentricDataBuffer >= 0) ) );
 		
-		for (int32 B = 0, R = 0; B < VB.m_buffers.Num(); ++B)
+		for (int32 B = 0, R = 0; B < VB.Buffers.Num(); ++B)
 		{
 			if (B != BarycentricDataBuffer)
 			{
-				ResultBuffers.m_buffers[R++] = VB.m_buffers[B];
+				ResultBuffers.Buffers[R++] = VB.Buffers[B];
 			}
 		}
 
diff --git a/Engine/Plugins/Experimental/Mutable/Source/MutableRuntime/Private/MuR/OpMeshClipMorphPlane.h b/Engine/Plugins/Experimental/Mutable/Source/MutableRuntime/Private/MuR/OpMeshClipMorphPlane.h
--- a/Engine/Plugins/Experimental/Mutable/Source/MutableRuntime/Private/MuR/OpMeshClipMorphPlane.h
+++ b/Engine/Plugins/Experimental/Mutable/Source/MutableRuntime/Private/MuR/OpMeshClipMorphPlane.h
@@ -140,9 +140,9 @@ namespace mu
 			VertexInfo.SetNum(VertexCount);
 			//int32 FirstCount = pBase->GetVertexBuffers().GetElementCount();
 
-			for (int32 BufferIndex = 0; BufferIndex < pBase->GetVertexBuffers().m_buffers.Num(); ++BufferIndex)
+			for (int32 BufferIndex = 0; BufferIndex < pBase->GetVertexBuffers().Buffers.Num(); ++BufferIndex)
 			{
-				const FMeshBuffer& Buffer = pBase->GetVertexBuffers().m_buffers[BufferIndex];
+				const FMeshBuffer& Buffer = pBase->GetVertexBuffers().Buffers[BufferIndex];
 
 				int32 ElemSize = pBase->GetVertexBuffers().GetElementSize(BufferIndex);
 				//int32 FirstSize = FirstCount * ElemSize;
@@ -172,7 +172,7 @@ namespace mu
 							{
 								for (int32 i = 0; i < NumVertices; ++i)
 								{
-									const uint8* pD = &Buffer.m_data[i*ElemSize + Offset];
+									const uint8* pD = &Buffer.Data[i*ElemSize + Offset];
 									for (int32 j = 0; j < Components; ++j)
 									{
 										VertexInfo[i].BoneIndices.Add(pD[j]);
@@ -186,7 +186,7 @@ namespace mu
 							{
 								for (int32 i = 0; i < NumVertices; ++i)
 								{
-									const uint16* pD = reinterpret_cast<const uint16*>(&Buffer.m_data[i*ElemSize + Offset]);
+									const uint16* pD = reinterpret_cast<const uint16*>(&Buffer.Data[i*ElemSize + Offset]);
 
 									for (int32 j = 0; j < Components; ++j)
 									{
@@ -201,7 +201,7 @@ namespace mu
 							{
 								for (int32 i = 0; i < NumVertices; ++i)
 								{
-									const uint32* pD = reinterpret_cast<const uint32*>(&Buffer.m_data[i*ElemSize + Offset]);
+									const uint32* pD = reinterpret_cast<const uint32*>(&Buffer.Data[i*ElemSize + Offset]);
 
 									for (int32 j = 0; j < Components; ++j)
 									{
@@ -228,7 +228,7 @@ namespace mu
 							{
 								for (int32 i = 0; i < NumVertices; ++i)
 								{
-									const uint8* pD = &Buffer.m_data[i*ElemSize + Offset];
+									const uint8* pD = &Buffer.Data[i*ElemSize + Offset];
 
 									for (int32 j = 0; j < Components; ++j)
 									{
@@ -244,7 +244,7 @@ namespace mu
 							{
 								for (int32 i = 0; i < NumVertices; ++i)
 								{
-									const uint16* pD = reinterpret_cast<const uint16*>(&Buffer.m_data[i*ElemSize + Offset]);
+									const uint16* pD = reinterpret_cast<const uint16*>(&Buffer.Data[i*ElemSize + Offset]);
 
 									for (int32 j = 0; j < Components; ++j)
 									{
@@ -260,7 +260,7 @@ namespace mu
 							{
 								for (int32 i = 0; i < NumVertices; ++i)
 								{
-									const uint32* pD = reinterpret_cast<const uint32*>(&Buffer.m_data[i*ElemSize + Offset]);
+									const uint32* pD = reinterpret_cast<const uint32*>(&Buffer.Data[i*ElemSize + Offset]);
 
 									for (int32 j = 0; j < Components; ++j)
 									{
@@ -274,7 +274,7 @@ namespace mu
                             {
 								for (int32 i = 0; i < NumVertices; ++i)
 								{
-									const float* pD = reinterpret_cast<const float*>(&Buffer.m_data[i*ElemSize + Offset]);
+									const float* pD = reinterpret_cast<const float*>(&Buffer.Data[i*ElemSize + Offset]);
 
 									for (int32 j = 0; j < Components; ++j)
 									{
diff --git a/Engine/Plugins/Experimental/Mutable/Source/MutableRuntime/Private/MuR/OpMeshClipWithMesh.cpp b/Engine/Plugins/Experimental/Mutable/Source/MutableRuntime/Private/MuR/OpMeshClipWithMesh.cpp
--- a/Engine/Plugins/Experimental/Mutable/Source/MutableRuntime/Private/MuR/OpMeshClipWithMesh.cpp
+++ b/Engine/Plugins/Experimental/Mutable/Source/MutableRuntime/Private/MuR/OpMeshClipWithMesh.cpp
@@ -422,12 +422,12 @@ namespace mu { namespace
         uint32 DestVertexCount = pBase->GetVertexCount();
 
         const FMeshBufferSet& MBSPriv2 = pBase->GetVertexBuffers();
-        for (int32 b = 0; b < MBSPriv2.m_buffers.Num(); ++b)
+        for (int32 b = 0; b < MBSPriv2.Buffers.Num(); ++b)
         {
-            for (int32 c = 0; c < MBSPriv2.m_buffers[b].m_channels.Num(); ++c)
+            for (int32 c = 0; c < MBSPriv2.Buffers[b].Channels.Num(); ++c)
             {
-                EMeshBufferSemantic Sem = MBSPriv2.m_buffers[b].m_channels[c].m_semantic;
-                int32 SemIndex = MBSPriv2.m_buffers[b].m_channels[c].m_semanticIndex;
+                EMeshBufferSemantic Sem = MBSPriv2.Buffers[b].Channels[c].Semantic;
+                int32 SemIndex = MBSPriv2.Buffers[b].Channels[c].SemanticIndex;
 
                 UntypedMeshBufferIteratorConst It(pBase->GetVertexBuffers(), Sem, SemIndex);
 
@@ -520,19 +520,19 @@ namespace mu { namespace
 
 		// Now go through all vertices in the mesh and record whether they are inside or outside of the ClipMesh
 		const FMeshBufferSet& MBSPriv = Base->GetVertexBuffers();
-		for (int32 b = 0; b < MBSPriv.m_buffers.Num(); ++b)
+		for (int32 b = 0; b < MBSPriv.Buffers.Num(); ++b)
 		{
-			const TArray<mu::FMeshBufferChannel>& Channels = MBSPriv.m_buffers[b].m_channels;
+			const TArray<mu::FMeshBufferChannel>& Channels = MBSPriv.Buffers[b].Channels;
 
 			for (int32 c = 0; c < Channels.Num(); ++c)
 			{
-				EMeshBufferSemantic Sem = Channels[c].m_semantic;
+				EMeshBufferSemantic Sem = Channels[c].Semantic;
 				if (Sem != MBS_TEXCOORDS)
 				{
 					continue;
 				}
 
-				int32 SemIndex = Channels[c].m_semanticIndex;
+				int32 SemIndex = Channels[c].SemanticIndex;
 				if (SemIndex != LayoutIndex)
 				{
 					continue;
diff --git a/Engine/Plugins/Experimental/Mutable/Source/MutableRuntime/Private/MuR/OpMeshFormat.cpp b/Engine/Plugins/Experimental/Mutable/Source/MutableRuntime/Private/MuR/OpMeshFormat.cpp
--- a/Engine/Plugins/Experimental/Mutable/Source/MutableRuntime/Private/MuR/OpMeshFormat.cpp
+++ b/Engine/Plugins/Experimental/Mutable/Source/MutableRuntime/Private/MuR/OpMeshFormat.cpp
@@ -33,57 +33,57 @@ namespace mu
 		}
 
 		// This happens in the debugger
-		if (ResultBuffer.m_channels.IsEmpty() || ResultBuffer.m_elementSize==0)
+		if (ResultBuffer.Channels.IsEmpty() || ResultBuffer.ElementSize==0)
 		{
 			return;
 		}
 
-		check(ResultBuffer.m_elementSize);
-		int32 ResultElementCount = ResultBuffer.m_data.Num() / ResultBuffer.m_elementSize;
+		check(ResultBuffer.ElementSize);
+		int32 ResultElementCount = ResultBuffer.Data.Num() / ResultBuffer.ElementSize;
 		check(SourceElementCount + ResultOffsetElements <= ResultElementCount);
 
 		// For every channel in this buffer
-		for (int32 ResultChannelIndex = 0; ResultChannelIndex < ResultBuffer.m_channels.Num(); ++ResultChannelIndex)
+		for (int32 ResultChannelIndex = 0; ResultChannelIndex < ResultBuffer.Channels.Num(); ++ResultChannelIndex)
 		{
-			FMeshBufferChannel& ResultChannel = ResultBuffer.m_channels[ResultChannelIndex];
+			FMeshBufferChannel& ResultChannel = ResultBuffer.Channels[ResultChannelIndex];
 
 			// Find this channel in the source mesh
 			int32 SourceBufferIndex;
 			int32 SourceChannelIndex;
-			Source.FindChannel(ResultChannel.m_semantic, ResultChannel.m_semanticIndex, &SourceBufferIndex, &SourceChannelIndex);
+			Source.FindChannel(ResultChannel.Semantic, ResultChannel.SemanticIndex, &SourceBufferIndex, &SourceChannelIndex);
 
-			int32 resultElemSize = ResultBuffer.m_elementSize;
-			int32 resultComponents = ResultChannel.m_componentCount;
-			int32 resultChannelSize = GetMeshFormatData(ResultChannel.m_format).SizeInBytes * resultComponents;
-			uint8* pResultBuf = ResultBuffer.m_data.GetData() + ResultOffsetElements*ResultBuffer.m_elementSize;
-			pResultBuf += ResultChannel.m_offset;
+			int32 resultElemSize = ResultBuffer.ElementSize;
+			int32 resultComponents = ResultChannel.ComponentCount;
+			int32 resultChannelSize = GetMeshFormatData(ResultChannel.Format).SizeInBytes * resultComponents;
+			uint8* pResultBuf = ResultBuffer.Data.GetData() + ResultOffsetElements*ResultBuffer.ElementSize;
+			pResultBuf += ResultChannel.Offset;
 
 			// Case 1: Special semantics that may be implicit or relative
 			//---------------------------------------------------------------------------------
-			if (bHasSpecialSemantics && ResultChannel.m_semantic == MBS_VERTEXINDEX)
+			if (bHasSpecialSemantics && ResultChannel.Semantic == MBS_VERTEXINDEX)
 			{
 				bool bHasVertexIndices = (SourceBufferIndex >=0 );
 				if (bHasVertexIndices)
 				{
-					check(SourceChannelIndex == 0 && Source.m_buffers[SourceBufferIndex].m_channels.Num() == 1);
+					check(SourceChannelIndex == 0 && Source.Buffers[SourceBufferIndex].Channels.Num() == 1);
 
-					const FMeshBuffer& SourceBuffer = Source.m_buffers[SourceBufferIndex];
-					const FMeshBufferChannel& SourceChannel = SourceBuffer.m_channels[SourceChannelIndex];
+					const FMeshBuffer& SourceBuffer = Source.Buffers[SourceBufferIndex];
+					const FMeshBufferChannel& SourceChannel = SourceBuffer.Channels[SourceChannelIndex];
 
-					bool bHasSameFormat = SourceChannel.m_format == ResultChannel.m_format;
+					bool bHasSameFormat = SourceChannel.Format == ResultChannel.Format;
 					if (bHasSameFormat)
 					{
 						check(SourceChannel == ResultChannel);
-						FMemory::Memcpy(pResultBuf,SourceBuffer.m_data.GetData(),SourceBuffer.m_data.Num());
+						FMemory::Memcpy(pResultBuf,SourceBuffer.Data.GetData(),SourceBuffer.Data.Num());
 					}
 					else
 					{
 						// Relative vertex IDs
 						check(IDPrefix);
-						check(SourceChannel.m_format == MBF_UINT32);
-						const uint32* SourceData = reinterpret_cast<const uint32*>(SourceBuffer.m_data.GetData());
+						check(SourceChannel.Format == MBF_UINT32);
+						const uint32* SourceData = reinterpret_cast<const uint32*>(SourceBuffer.Data.GetData());
 
-						check(ResultChannel.m_format == MBF_UINT64);
+						check(ResultChannel.Format == MBF_UINT64);
 						uint64* ResultData = reinterpret_cast<uint64*>(pResultBuf);
 
 						for (int32 i = 0; i < SourceElementCount; ++i)
@@ -100,7 +100,7 @@ namespace mu
 					// Implicit IDs
 					check(IDPrefix);
 
-					check(ResultChannel.m_format == MBF_UINT64);
+					check(ResultChannel.Format == MBF_UINT64);
 					uint64* ResultData = reinterpret_cast<uint64*>(pResultBuf);
 
 					for (int32 VertexIndex = 0; VertexIndex < SourceElementCount; ++VertexIndex)
@@ -114,28 +114,28 @@ namespace mu
 				continue;
 			}
 
-			else if (bHasSpecialSemantics && ResultChannel.m_semantic == MBS_LAYOUTBLOCK)
+			else if (bHasSpecialSemantics && ResultChannel.Semantic == MBS_LAYOUTBLOCK)
 			{
 				bool bHasBlockIds = (SourceBufferIndex >= 0);
 
 				if (bHasBlockIds)
 				{
-					const FMeshBuffer& SourceBuffer = Source.m_buffers[SourceBufferIndex];
-					const FMeshBufferChannel& SourceChannel = SourceBuffer.m_channels[SourceChannelIndex];
+					const FMeshBuffer& SourceBuffer = Source.Buffers[SourceBufferIndex];
+					const FMeshBufferChannel& SourceChannel = SourceBuffer.Channels[SourceChannelIndex];
 
-					bool bHasSameFormat = SourceChannel.m_format == ResultChannel.m_format;
+					bool bHasSameFormat = SourceChannel.Format == ResultChannel.Format;
 					if (bHasSameFormat)
 					{
 						check(SourceChannel==ResultChannel);
-						FMemory::Memcpy(pResultBuf, SourceBuffer.m_data.GetData(), SourceBuffer.m_data.Num());
+						FMemory::Memcpy(pResultBuf, SourceBuffer.Data.GetData(), SourceBuffer.Data.Num());
 					}
 					else
 					{
 						// Relative vertex IDs
-						check(SourceChannel.m_format == MBF_UINT16);
-						const uint16* SourceData = reinterpret_cast<const uint16*>(SourceBuffer.m_data.GetData());
+						check(SourceChannel.Format == MBF_UINT16);
+						const uint16* SourceData = reinterpret_cast<const uint16*>(SourceBuffer.Data.GetData());
 
-						check(ResultChannel.m_format == MBF_UINT64);
+						check(ResultChannel.Format == MBF_UINT64);
 						uint64* ResultData = reinterpret_cast<uint64*>(pResultBuf);
 
 						for (int32 i = 0; i < SourceElementCount; ++i)
@@ -168,11 +168,11 @@ namespace mu
 
 				// If we have to add colour channels, we will add them as white, to be neutral.
 				// \todo: normal channels also should have special values.
-				if (ResultChannel.m_semantic == MBS_COLOUR)
+				if (ResultChannel.Semantic == MBS_COLOUR)
 				{
 					generated = true;
 
-					switch (ResultChannel.m_format)
+					switch (ResultChannel.Format)
 					{
 					case MBF_FLOAT32:
 					{
@@ -253,9 +253,9 @@ namespace mu
 					&sourceFormat, &sourceComponents,
 					&sourceOffset
 				);
-				check(sourceSemantic == ResultChannel.m_semantic
+				check(sourceSemantic == ResultChannel.Semantic
 					&&
-					sourceSemanticIndex == ResultChannel.m_semanticIndex);
+					sourceSemanticIndex == ResultChannel.SemanticIndex);
 
 				int32 sourceElemSize = Source.GetElementSize(SourceBufferIndex);
 				const uint8* pSourceBuf = Source.GetBufferData(SourceBufferIndex);
@@ -264,13 +264,13 @@ namespace mu
 				// Copy element by element
 				for (int32 v = 0; v < SourceElementCount; ++v)
 				{
-					if (ResultChannel.m_format == sourceFormat && resultComponents == sourceComponents)
+					if (ResultChannel.Format == sourceFormat && resultComponents == sourceComponents)
 					{
 						FMemory::Memcpy(pResultBuf, pSourceBuf, resultChannelSize);
 					}
-					else if (ResultChannel.m_format == MBF_PACKEDDIR8_W_TANGENTSIGN
+					else if (ResultChannel.Format == MBF_PACKEDDIR8_W_TANGENTSIGN
 						||
-						ResultChannel.m_format == MBF_PACKEDDIRS8_W_TANGENTSIGN)
+						ResultChannel.Format == MBF_PACKEDDIRS8_W_TANGENTSIGN)
 					{
 						check(sourceComponents >= 3);
 						check(resultComponents == 4);
@@ -283,7 +283,7 @@ namespace mu
 								ConvertData
 								(
 									i,
-									pResultBuf, ResultChannel.m_format,
+									pResultBuf, ResultChannel.Format,
 									pSourceBuf, sourceFormat
 								);
 							}
@@ -295,15 +295,15 @@ namespace mu
 
 						// Look for the full tangent space
 						int32 tanXBuf, tanXChan, tanYBuf, tanYChan, tanZBuf, tanZChan;
-						Source.FindChannel(MBS_TANGENT, ResultChannel.m_semanticIndex, &tanXBuf, &tanXChan);
-						Source.FindChannel(MBS_BINORMAL, ResultChannel.m_semanticIndex, &tanYBuf, &tanYChan);
-						Source.FindChannel(MBS_NORMAL, ResultChannel.m_semanticIndex, &tanZBuf, &tanZChan);
+						Source.FindChannel(MBS_TANGENT, ResultChannel.SemanticIndex, &tanXBuf, &tanXChan);
+						Source.FindChannel(MBS_BINORMAL, ResultChannel.SemanticIndex, &tanYBuf, &tanYChan);
+						Source.FindChannel(MBS_NORMAL, ResultChannel.SemanticIndex, &tanZBuf, &tanZChan);
 
 						if (tanXBuf >= 0 && tanYBuf >= 0 && tanZBuf >= 0)
 						{
-							UntypedMeshBufferIteratorConst xIt(Source, MBS_TANGENT, ResultChannel.m_semanticIndex);
-							UntypedMeshBufferIteratorConst yIt(Source, MBS_BINORMAL, ResultChannel.m_semanticIndex);
-							UntypedMeshBufferIteratorConst zIt(Source, MBS_NORMAL, ResultChannel.m_semanticIndex);
+							UntypedMeshBufferIteratorConst xIt(Source, MBS_TANGENT, ResultChannel.SemanticIndex);
+							UntypedMeshBufferIteratorConst yIt(Source, MBS_BINORMAL, ResultChannel.SemanticIndex);
+							UntypedMeshBufferIteratorConst zIt(Source, MBS_NORMAL, ResultChannel.SemanticIndex);
 
 							xIt += v;
 							yIt += v;
@@ -312,11 +312,11 @@ namespace mu
 							FMatrix44f Mat(xIt.GetAsVec3f(), yIt.GetAsVec3f(), zIt.GetAsVec3f(), FVector3f(0, 0, 0));
 
 							uint8 sign = 0;
-							if (ResultChannel.m_format == MBF_PACKEDDIR8_W_TANGENTSIGN)
+							if (ResultChannel.Format == MBF_PACKEDDIR8_W_TANGENTSIGN)
 							{
 								sign = Mat.RotDeterminant() < 0 ? 0 : 255;
 							}
-							else if (ResultChannel.m_format == MBF_PACKEDDIRS8_W_TANGENTSIGN)
+							else if (ResultChannel.Format == MBF_PACKEDDIRS8_W_TANGENTSIGN)
 							{
 								sign = Mat.RotDeterminant() < 0 ? -128 : 127;
 							}
@@ -338,7 +338,7 @@ namespace mu
 								ConvertData
 								(
 									i,
-									pResultBuf, ResultChannel.m_format,
+									pResultBuf, ResultChannel.Format,
 									pSourceBuf, sourceFormat
 								);
 							}
@@ -347,8 +347,8 @@ namespace mu
 								// Add zeros. TODO: Warning?
 								FMemory::Memzero
 								(
-									pResultBuf + GetMeshFormatData(ResultChannel.m_format).SizeInBytes * i,
-									GetMeshFormatData(ResultChannel.m_format).SizeInBytes
+									pResultBuf + GetMeshFormatData(ResultChannel.Format).SizeInBytes * i,
+									GetMeshFormatData(ResultChannel.Format).SizeInBytes
 								);
 							}
 						}
@@ -359,7 +359,7 @@ namespace mu
 						// TODO: Optimise
 						if (sourceSemantic == MBS_BONEWEIGHTS)
 						{
-							if (ResultChannel.m_format == MBF_NUINT8)
+							if (ResultChannel.Format == MBF_NUINT8)
 							{
 								uint8* pData = (uint8*)pResultBuf;
 								uint8 accum = 0;
@@ -370,7 +370,7 @@ namespace mu
 								pData[0] += 255 - accum;
 							}
 
-							else if (ResultChannel.m_format == MBF_NUINT16)
+							else if (ResultChannel.Format == MBF_NUINT16)
 							{
 								uint16* pData = (uint16*)pResultBuf;
 								uint16 accum = 0;
@@ -471,7 +471,7 @@ namespace mu
 		Result.SetElementCount(vCount);
 		for (int32 b = 0; b < Result.GetBufferCount(); ++b)
 		{
-			MeshFormatBuffer(Source, Result.m_buffers[b], 0, bIsVertexBuffer, IDPrefix);
+			MeshFormatBuffer(Source, Result.Buffers[b], 0, bIsVertexBuffer, IDPrefix);
 		}
 
 
@@ -511,8 +511,8 @@ namespace mu
 						else
 						{
 							// Replace the buffer
-							check(Result.m_buffers[AlreadyExistingBufIndex].m_channels.Num()==1);
-							Result.m_buffers[AlreadyExistingBufIndex] = Source.m_buffers[b];
+							check(Result.Buffers[AlreadyExistingBufIndex].Channels.Num()==1);
+							Result.Buffers[AlreadyExistingBufIndex] = Source.Buffers[b];
 						}
 					}
 				}
@@ -569,17 +569,17 @@ namespace mu
 				const int32 ChannelCount = VertexBuffers.GetBufferChannelCount(BufferIndex);
 				for (int32 ChannelIndex = 0; ChannelIndex < ChannelCount; ++ChannelIndex)
 				{
-					const FMeshBufferChannel& Channel = VertexBuffers.m_buffers[BufferIndex].m_channels[ChannelIndex];
+					const FMeshBufferChannel& Channel = VertexBuffers.Buffers[BufferIndex].Channels[ChannelIndex];
 
-					if (Channel.m_semantic == MBS_BONEINDICES)
+					if (Channel.Semantic == MBS_BONEINDICES)
 					{
 						int32 resultBuf = 0;
 						int32 resultChan = 0;
 						FMeshBufferSet& formBuffs = Result->GetVertexBuffers();
-						formBuffs.FindChannel(MBS_BONEINDICES, Channel.m_semanticIndex, &resultBuf, &resultChan);
+						formBuffs.FindChannel(MBS_BONEINDICES, Channel.SemanticIndex, &resultBuf, &resultChan);
 						if (resultBuf >= 0)
 						{
-							UntypedMeshBufferIteratorConst it(VertexBuffers, MBS_BONEINDICES, Channel.m_semanticIndex);
+							UntypedMeshBufferIteratorConst it(VertexBuffers, MBS_BONEINDICES, Channel.SemanticIndex);
 							int32 maxBoneIndex = 0;
 							for (int32 v = 0; v < VertexBuffers.GetElementCount(); ++v)
 							{
@@ -594,7 +594,7 @@ namespace mu
 								++it;
 							}
 
-							EMeshBufferFormat& format = formBuffs.m_buffers[resultBuf].m_channels[resultChan].m_format;
+							EMeshBufferFormat& format = formBuffs.Buffers[resultBuf].Channels[resultChan].Format;
 							if (maxBoneIndex > 0xffff && (format == MBF_UINT8 || format == MBF_UINT16))
 							{
 								format = MBF_UINT32;
@@ -740,26 +740,26 @@ namespace mu
 
 				// \todo: This is a generic innefficient way
 				FMeshBufferSet NewVertexBuffers;
-				NewVertexBuffers.m_buffers = VertexBuffers.m_buffers;
+				NewVertexBuffers.Buffers = VertexBuffers.Buffers;
 
-				for (FMeshBuffer& Buffer: NewVertexBuffers.m_buffers)
+				for (FMeshBuffer& Buffer: NewVertexBuffers.Buffers)
 				{
 					int32 OffsetDelta = 0;
-					for (FMeshBufferChannel& Channel : Buffer.m_channels)
+					for (FMeshBufferChannel& Channel : Buffer.Channels)
 					{
-						Channel.m_offset += OffsetDelta;
+						Channel.Offset += OffsetDelta;
 
-						if (Channel.m_semanticIndex == SemanticIndex
+						if (Channel.SemanticIndex == SemanticIndex
 							&&
-							(Channel.m_semantic == MBS_BONEWEIGHTS || Channel.m_semantic == MBS_BONEINDICES)
+							(Channel.Semantic == MBS_BONEWEIGHTS || Channel.Semantic == MBS_BONEINDICES)
 							)
 						{
-							Channel.m_componentCount = RealInfluences;
-							OffsetDelta -= (BufferInfluences - RealInfluences) * GetMeshFormatData(Channel.m_format).SizeInBytes;
+							Channel.ComponentCount = RealInfluences;
+							OffsetDelta -= (BufferInfluences - RealInfluences) * GetMeshFormatData(Channel.Format).SizeInBytes;
 						}
 					}
 
-					Buffer.m_elementSize += OffsetDelta;
+					Buffer.ElementSize += OffsetDelta;
 				}
 
 				FormatBufferSet( VertexBuffers, NewVertexBuffers, true, false, true);
diff --git a/Engine/Plugins/Experimental/Mutable/Source/MutableRuntime/Private/MuR/OpMeshMorph.h b/Engine/Plugins/Experimental/Mutable/Source/MutableRuntime/Private/MuR/OpMeshMorph.h
--- a/Engine/Plugins/Experimental/Mutable/Source/MutableRuntime/Private/MuR/OpMeshMorph.h
+++ b/Engine/Plugins/Experimental/Mutable/Source/MutableRuntime/Private/MuR/OpMeshMorph.h
@@ -256,9 +256,9 @@ namespace mu
 			for (int32 ChannelIndex = 0; ChannelIndex < ChannelsNum; ++ChannelIndex)
 			{
 				const FMeshBufferSet& MBSPriv = RefTarget->GetVertexBuffers();
-				const FMeshBufferChannel& Channel = MBSPriv.m_buffers[MorphBufferDataChannel].m_channels[ChannelIndex];
-				EMeshBufferSemantic Sem = Channel.m_semantic;
-				int32 SemIndex = Channel.m_semanticIndex;
+				const FMeshBufferChannel& Channel = MBSPriv.Buffers[MorphBufferDataChannel].Channels[ChannelIndex];
+				EMeshBufferSemantic Sem = Channel.Semantic;
+				int32 SemIndex = Channel.SemanticIndex;
 			
 				if (Sem == MBS_NORMAL && bBaseHasNormals)
 				{
diff --git a/Engine/Plugins/Experimental/Mutable/Source/MutableRuntime/Private/MuR/OpMeshOptimizeSkinning.h b/Engine/Plugins/Experimental/Mutable/Source/MutableRuntime/Private/MuR/OpMeshOptimizeSkinning.h
--- a/Engine/Plugins/Experimental/Mutable/Source/MutableRuntime/Private/MuR/OpMeshOptimizeSkinning.h
+++ b/Engine/Plugins/Experimental/Mutable/Source/MutableRuntime/Private/MuR/OpMeshOptimizeSkinning.h
@@ -46,16 +46,16 @@ namespace mu
 		
 		// Iterate all vertex buffers and check if the BoneIndices buffers have the desired format
 		const FMeshBufferSet& InMeshVertexBuffers = InMesh->GetVertexBuffers();
-		for (int32 VertexBufferIndex = 0; !bRequiresFormatChange && VertexBufferIndex < InMeshVertexBuffers.m_buffers.Num(); ++VertexBufferIndex)
+		for (int32 VertexBufferIndex = 0; !bRequiresFormatChange && VertexBufferIndex < InMeshVertexBuffers.Buffers.Num(); ++VertexBufferIndex)
 		{
-			const FMeshBuffer& Buffer = InMeshVertexBuffers.m_buffers[VertexBufferIndex];
+			const FMeshBuffer& Buffer = InMeshVertexBuffers.Buffers[VertexBufferIndex];
 
 			const int32 ChannelsCount = InMeshVertexBuffers.GetBufferChannelCount(VertexBufferIndex);
 			for (int32 ChannelIndex = 0; ChannelIndex < ChannelsCount; ++ChannelIndex)
 			{
-				if (Buffer.m_channels[ChannelIndex].m_semantic == MBS_BONEINDICES)
+				if (Buffer.Channels[ChannelIndex].Semantic == MBS_BONEINDICES)
 				{
-					bRequiresFormatChange = Buffer.m_channels[ChannelIndex].m_format != DesiredBoneIndexFormat;
+					bRequiresFormatChange = Buffer.Channels[ChannelIndex].Format != DesiredBoneIndexFormat;
 					break;
 				}
 			}
@@ -82,12 +82,12 @@ namespace mu
 
 		mu::FMeshBufferSet& VertexBuffers = Result->GetVertexBuffers();
 
-		VertexBuffers.m_buffers.Reserve(VertexBuffersCount);
+		VertexBuffers.Buffers.Reserve(VertexBuffersCount);
 		VertexBuffers.SetElementCount(ElementCount);
 
 		for (int32 BufferIndex = 0; BufferIndex < VertexBuffersCount; ++BufferIndex)
 		{
-			const FMeshBuffer& SourceBuffer = InMeshVertexBuffers.m_buffers[BufferIndex];
+			const FMeshBuffer& SourceBuffer = InMeshVertexBuffers.Buffers[BufferIndex];
 			const int32 ChannelsCount = InMeshVertexBuffers.GetBufferChannelCount(BufferIndex);
 
 			int32 BoneIndexChannelIndex = INDEX_NONE;
@@ -107,26 +107,26 @@ namespace mu
 			// Copy and fix channel details
 			for (int32 ChannelIndex = 0; ChannelIndex < ChannelsCount; ++ChannelIndex)
 			{
-				const FMeshBufferChannel& Channel = SourceBuffer.m_channels[ChannelIndex];
-				EMeshBufferFormat Format = Channel.m_format;
+				const FMeshBufferChannel& Channel = SourceBuffer.Channels[ChannelIndex];
+				EMeshBufferFormat Format = Channel.Format;
 
-				if (Channel.m_semantic == MBS_BONEINDICES)
+				if (Channel.Semantic == MBS_BONEINDICES)
 				{
-					SourceBoneIndexFormat = Channel.m_format;
+					SourceBoneIndexFormat = Channel.Format;
 					Format = mu::MBF_UINT8;
 					BoneIndexChannelIndex = ChannelIndex;
 				}
 
-				Semantics.Add(Channel.m_semantic);
-				SemanticIndices.Add(Channel.m_semanticIndex);
+				Semantics.Add(Channel.Semantic);
+				SemanticIndices.Add(Channel.SemanticIndex);
 				Formats.Add(Format);
-				Components.Add(Channel.m_componentCount);
+				Components.Add(Channel.ComponentCount);
 				Offsets.Add(AuxOffset);
 
 				const int32 FormatSize = GetMeshFormatData(Format).SizeInBytes;
 
 				ElementSize += FormatSize;
-				AuxOffset += FormatSize * Channel.m_componentCount;
+				AuxOffset += FormatSize * Channel.ComponentCount;
 			}
 
 			// Copy buffers
@@ -137,7 +137,7 @@ namespace mu
 				VertexBuffers.SetBuffer(BufferIndex, ElementSize, ChannelsCount, Semantics.GetData(), SemanticIndices.GetData(), Formats.GetData(), Components.GetData(), Offsets.GetData());
 				uint8* Data = VertexBuffers.GetBufferData(BufferIndex);
 
-				const uint8* SourceData = (const uint8*)SourceBuffer.m_data.GetData();
+				const uint8* SourceData = (const uint8*)SourceBuffer.Data.GetData();
 				const int32 SourceBoneIndexSize = GetMeshFormatData(SourceBoneIndexFormat).SizeInBytes;
 
 				const int32 BoneIndexSize = GetMeshFormatData(mu::MBF_UINT8).SizeInBytes;
diff --git a/Engine/Plugins/Experimental/Mutable/Source/MutableRuntime/Private/MuR/OpMeshReshape.h b/Engine/Plugins/Experimental/Mutable/Source/MutableRuntime/Private/MuR/OpMeshReshape.h
--- a/Engine/Plugins/Experimental/Mutable/Source/MutableRuntime/Private/MuR/OpMeshReshape.h
+++ b/Engine/Plugins/Experimental/Mutable/Source/MutableRuntime/Private/MuR/OpMeshReshape.h
@@ -763,18 +763,18 @@ namespace mu
 	
 		FMeshBufferSet& ResultBuffers = VerticesReshapeMesh->GetVertexBuffers();
 
-		check(ResultBuffers.m_buffers.Num() == 0);
+		check(ResultBuffers.Buffers.Num() == 0);
 
 		// Copy buffers skipping binding data. 
-		ResultBuffers.m_elementCount = VB.m_elementCount;
+		ResultBuffers.ElementCount = VB.ElementCount;
 		// Remove one element to the number of buffers if BarycentricDataBuffer found. 
-		ResultBuffers.m_buffers.SetNum(FMath::Max(0, VB.m_buffers.Num() - int32(BarycentricDataBuffer >= 0)));
+		ResultBuffers.Buffers.SetNum(FMath::Max(0, VB.Buffers.Num() - int32(BarycentricDataBuffer >= 0)));
 		
-		for (int32 B = 0, R = 0; B < VB.m_buffers.Num(); ++B)
+		for (int32 B = 0, R = 0; B < VB.Buffers.Num(); ++B)
 		{
 			if (B != BarycentricDataBuffer)
 			{
-				ResultBuffers.m_buffers[R++] = VB.m_buffers[B];
+				ResultBuffers.Buffers[R++] = VB.Buffers[B];
 			}
 		}
 
diff --git a/Engine/Plugins/Experimental/Mutable/Source/MutableRuntime/Private/MuR/OpMeshTransform.h b/Engine/Plugins/Experimental/Mutable/Source/MutableRuntime/Private/MuR/OpMeshTransform.h
--- a/Engine/Plugins/Experimental/Mutable/Source/MutableRuntime/Private/MuR/OpMeshTransform.h
+++ b/Engine/Plugins/Experimental/Mutable/Source/MutableRuntime/Private/MuR/OpMeshTransform.h
@@ -31,13 +31,13 @@ namespace mu
 		FMatrix44f transformIT = transform.Inverse().GetTransposed();
 
         const FMeshBufferSet& MBSPriv = Result->GetVertexBuffers();
-        for ( int32 b=0; b<MBSPriv.m_buffers.Num(); ++b )
+        for ( int32 b=0; b<MBSPriv.Buffers.Num(); ++b )
         {
 
-            for ( int32 c=0; c<MBSPriv.m_buffers[b].m_channels.Num(); ++c )
+            for ( int32 c=0; c<MBSPriv.Buffers[b].Channels.Num(); ++c )
             {
-                EMeshBufferSemantic sem = MBSPriv.m_buffers[b].m_channels[c].m_semantic;
-                int semIndex = MBSPriv.m_buffers[b].m_channels[c].m_semanticIndex;
+                EMeshBufferSemantic sem = MBSPriv.Buffers[b].Channels[c].Semantic;
+                int semIndex = MBSPriv.Buffers[b].Channels[c].SemanticIndex;
 
                 UntypedMeshBufferIterator it( Result->GetVertexBuffers(), sem, semIndex );
 
diff --git a/Engine/Plugins/Experimental/Mutable/Source/MutableRuntime/Public/MuR/Mesh.h b/Engine/Plugins/Experimental/Mutable/Source/MutableRuntime/Public/MuR/Mesh.h
--- a/Engine/Plugins/Experimental/Mutable/Source/MutableRuntime/Public/MuR/Mesh.h
+++ b/Engine/Plugins/Experimental/Mutable/Source/MutableRuntime/Public/MuR/Mesh.h
@@ -187,18 +187,18 @@ namespace mu
 		/** Create a new empty mesh that repreents an external resource mesh. */
 		static Ptr<Mesh> CreateAsReference(uint32 ID, bool bForceLoad);
 
-        //! Deep clone this mesh.
+        /** Deep clone this mesh. */
         Ptr<Mesh> Clone() const;
 		
-		// Clone with flags allowing to not include some parts in the cloned mesh
+		/** Clone with flags allowing to not include some parts in the cloned mesh */
 		Ptr<Mesh> Clone(EMeshCopyFlags Flags) const;
 
-		// Copy form another mesh.
+		/** Copy form another mesh. */
 		void CopyFrom(const Mesh& From, EMeshCopyFlags Flags = EMeshCopyFlags::AllFlags);
 
-        //! Serialisation
-        static void Serialise( const Mesh* p, OutputArchive& arch );
-        static Ptr<Mesh> StaticUnserialise( InputArchive& arch );
+        /** Serialisation */
+        static void Serialise(const Mesh* InMesh, OutputArchive& Arch );
+        static Ptr<Mesh> StaticUnserialise(InputArchive& Arch);
 
 		// Resource interface
 		int32 GetDataSize() const override;
@@ -222,30 +222,33 @@ namespace mu
         //!
         int32 GetIndexCount() const;
 
-        //! Index buffers. They are owned by this mesh.
+        /** Index buffers. They are owned by this mesh. */
         FMeshBufferSet& GetIndexBuffers();
         const FMeshBufferSet& GetIndexBuffers() const;
 
         //
         int32 GetVertexCount() const;
 
-        //! Vertex buffers. They are owned by this mesh.
+        /** Vertex buffers. They are owned by this mesh. */
         FMeshBufferSet& GetVertexBuffers();
         const FMeshBufferSet& GetVertexBuffers() const;
 
-        //
         int32 GetFaceCount() const;
 
-        //! Get the number of surfaces defined in this mesh. Surfaces are buffer-contiguous mesh
-        //! fragments that share common properties (usually material)
+        /**
+		 * Get the number of surfaces defined in this mesh. Surfaces are buffer-contiguous mesh
+         * fragments that share common properties (usually material)
+		 */
         int32 GetSurfaceCount() const;
         void GetSurface(int32 SurfaceIndex,
                         int32& OutFirstVertex, int32& OutVertexCount,
                         int32& OutFirstIndex, int32& OutIndexCount,
 						int32& OutFirstBone, int32& OutBoneCount) const;
 
-        //! Return an internal id that can be used to match mesh surfaces and instance surfaces.
-        //! Only valid for meshes that are part of instances.
+        /**
+		 * Return an internal id that can be used to match mesh surfaces and instance surfaces.
+         * Only valid for meshes that are part of instances.
+		 */
         uint32 GetSurfaceId(int32 SurfaceIndex) const;
 
         //! \}
@@ -358,14 +361,16 @@ namespace mu
         //! \}
 
 
-        //! Get an internal identifier used to reference this mesh in operations like deferred
-        //! mesh building, or instance updating.
+        /** 
+		 * Get an internal identifier used to reference this mesh in operations like deferred
+         * mesh building, or instance updating.
+		 */
         uint32 GetId() const;
 
 
     protected:
 
-        //! Forbidden. Manage with the Ptr<> template.
+        /** Forbidden. Manage with the Ptr<> template. */
 		~Mesh() {}
 
     
@@ -377,7 +382,8 @@ namespace mu
 		/** Non-persistent internal id unique for a mesh generated for a specific state and parameter values. */
 		mutable uint32 InternalId = 0;
 
-		/** This is bit - mask on the EMeshFlags enumeration, marking what static formats are compatible with this one and other properties. 
+		/** 
+		 * This is bit - mask on the EMeshFlags enumeration, marking what static formats are compatible with this one and other properties. 
 		 * It should be reset after any operation that modifies the format.
 		 */
 		mutable EMeshFlags Flags = EMeshFlags::None;
@@ -385,41 +391,43 @@ namespace mu
 		/** Only valid if the right flags are set, this identifies a referenced mesh. */
 		uint32 ReferenceID = 0;
 
-		/** Prefix for the unique IDs related to this mesh (vertices and layout blocks). Useful if the mesh stores them in an implicit, or relative way. 
-		* See MeshVertexIdIterator for details.
-		*/
+		/** 
+		 * Prefix for the unique IDs related to this mesh (vertices and layout blocks). Useful if the mesh stores them in an implicit, or relative way. 
+		 * See MeshVertexIdIterator for details.
+		 */
 		uint32 MeshIDPrefix = 0;
 
-		//!
 		FMeshBufferSet VertexBuffers;
 
-		//!
 		FMeshBufferSet IndexBuffers;
 
-		//! Additional buffers used for temporary or custom data in different algorithms.
+		/** Additional buffers used for temporary or custom data in different algorithms. */
 		TArray<TPair<EMeshBufferType, FMeshBufferSet>> AdditionalBuffers;
 
 		TArray<FMeshSurface> Surfaces;
 
-		// Externally provided SkeletonIDs of the skeletons required by this mesh.
+		/** Externally provided SkeletonIDs of the skeletons required by this mesh. */
 		TArray<uint32> SkeletonIDs;
 
-		//! This skeleton and physics body are not owned and may be used by other meshes, so it cannot be modified
-		//! once the mesh has been fully created.
+		/** 
+		 * This skeleton and physics body are not owned and may be used by other meshes, so it cannot be modified
+		 * once the mesh has been fully created.
+		 */
 		Ptr<const Skeleton> Skeleton;
 		Ptr<const PhysicsBody> PhysicsBody;
 
-		//! Additional physics bodies referenced by the mesh that don't merge.
+		/** Additional physics bodies referenced by the mesh that don't merge. */
 		TArray<Ptr<const mu::PhysicsBody>> AdditionalPhysicsBodies;
 
-		//! Texture Layout blocks attached to this mesh. They are const because they could be shared with
-		//! other meshes, so they need to be cloned and replaced if a modification is needed.
+		/** 
+		 * Texture Layout blocks attached to this mesh. They are const because they could be shared with
+		 * other meshes, so they need to be cloned and replaced if a modification is needed.
+		 */
 		TArray<Ptr<const Layout>> Layouts;		
 
-		//!
 		TArray<FString> Tags;
 
-		// Opaque handle to external resources.
+		/** Opaque handle to external resources. */
 		TArray<uint64> StreamedResources;
 
 		struct FBonePose
@@ -433,27 +441,26 @@ namespace mu
 			inline void Serialise(OutputArchive& arch) const;
 			inline void Unserialise(InputArchive& arch);
 
-			//!
 			inline bool operator==(const FBonePose& Other) const
 			{
 				return BoneUsageFlags == Other.BoneUsageFlags && BoneId == Other.BoneId;
 			}
 		};
-		// This is the pose used by this mesh fragment, used to update the transforms of the final skeleton
-		// taking into consideration the meshes being used.
+
+		/** 
+		 * This is the pose used by this mesh fragment, used to update the transforms of the final skeleton
+		 * taking into consideration the meshes being used.
+		 */
 		TMemoryTrackedArray<FBonePose> BonePoses;
 
-		// Array containing the bonemaps of all surfaces in the mesh.
+		/** Array containing the bonemaps of all surfaces in the mesh. */
 		TArray<FBoneName> BoneMap;
 
-		//!
 		inline void Serialise(OutputArchive& arch) const;
 
-        //!
 		inline void Unserialise(InputArchive& arch);
 
 
-        //!
 		inline bool operator==(const Mesh& o) const
 		{
 			bool bEqual = true;
@@ -509,49 +516,56 @@ namespace mu
 			return bEqual;
 		}
 
-		//! Compare the mesh with another one, but ignore internal data like generated vertex
-		//! indices.
-		bool IsSimilar(const Mesh& o, bool bCompareLayouts) const;
+		/** 
+		 * Compare the mesh with another one, but ignore internal data like generated vertex
+		 * indices.
+		 */
+		bool IsSimilar(const Mesh& Other, bool bCompareLayouts) const;
 
 
-		//! Make a map from the vertices in this mesh to thefirst matching vertex of the given
-		//! mesh. If non is found, the index is set to -1.
+		/**
+		 * Make a map from the vertices in this mesh to thefirst matching vertex of the given
+		 * mesh. If non is found, the index is set to -1.
+		 */
 		struct FVertexMatchMap
 		{
-			//! One for every vertex
+			/** One for every vertex */
 			TArray<int32> FirstMatch;
 
-			//! The matches of every vertex in a sequence
+			/** The matches of every vertex in a sequence */
 			TArray<int32> Matches;
 
-			//!
-			bool DoMatch(int32 v, int32 ov) const;
+			bool DoMatch(int32 Vertex, int32 OtherVertex) const;
 		};
 
-		void GetVertexMap( const Mesh& other, FVertexMatchMap& vertexMap, float tolerance = 1e-3f) const;
+		void GetVertexMap(const Mesh& Other, FVertexMatchMap& VertexMap, float Tolerance = 1e-3f) const;
 
-		//! Compare the vertex attributes to check if they match.
+		/** Compare the vertex attributes to check if they match. */
 		UE::Math::TIntVector3<uint32> GetFaceVertexIndices(int32 f) const;
 
-		//! Return true if the given mesh has the same vertex and index formats, and in the same
-		//! buffer structure.
-		bool HasCompatibleFormat(const Mesh* pOther) const;
+		/** 
+		 * Return true if the given mesh has the same vertex and index formats, and in the same
+		 * buffer structure.
+		 */
+		bool HasCompatibleFormat(const Mesh* Other) const;
 
-		//! Update the flags identifying the mesh format as some of the optimised formats.
+		/** Update the flags identifying the mesh format as some of the optimised formats. */
 		void ResetStaticFormatFlags() const;
 
-		//! Create the surface data if not present.
+		/** Create the surface data if not present. */
 		void EnsureSurfaceData();
 
-		//! Check mesh buffer data for possible inconsistencies
+		/** Check mesh buffer data for possible inconsistencies */
 		void CheckIntegrity() const;
 
-		//! Change the buffer descriptions so that all buffer indices start at 0 and are in the
-		//! same order than memory.
+		/** 
+		 * Change the buffer descriptions so that all buffer indices start at 0 and are in the
+		 * same order than memory.
+		 */
 		void ResetBufferIndices();
 
-		//! Debug: get a text representation of the mesh
-		void Log(FString& out, int32 VertrexLimit);
+		/** Debug: get a text representation of the mesh */
+		void Log(FString& Out, int32 VertrexLimit);
     };
 
 
diff --git a/Engine/Plugins/Experimental/Mutable/Source/MutableRuntime/Public/MuR/MeshBufferSet.h b/Engine/Plugins/Experimental/Mutable/Source/MutableRuntime/Public/MuR/MeshBufferSet.h
--- a/Engine/Plugins/Experimental/Mutable/Source/MutableRuntime/Public/MuR/MeshBufferSet.h
+++ b/Engine/Plugins/Experimental/Mutable/Source/MutableRuntime/Public/MuR/MeshBufferSet.h
@@ -17,8 +17,7 @@ namespace mu::MemoryCounters
 
 namespace mu
 {
-	//! Supported formats for the elements in mesh buffers.
-	//! \ingroup runtime
+	/** Supported formats for the elements in mesh buffers. **/
 	typedef enum
 	{
 
@@ -33,30 +32,34 @@ namespace mu
 		MBF_INT16,
 		MBF_INT32,
 
-		//! Integers interpreted as being in the range 0.0f to 1.0f
+		/** Integers interpreted as being in the range 0.0f to 1.0f */
 		MBF_NUINT8,
 		MBF_NUINT16,
 		MBF_NUINT32,
 
-		//! Integers interpreted as being in the range -1.0f to 1.0f
+		/** Integers interpreted as being in the range -1.0f to 1.0f */
 		MBF_NINT8,
 		MBF_NINT16,
 		MBF_NINT32,
 
-        //! Packed 1 to -1 value using multiply+add (128 is almost zero). Use 8-bit unsigned ints.
+        /** Packed 1 to -1 value using multiply+add (128 is almost zero). Use 8-bit unsigned ints. */
         MBF_PACKEDDIR8,
 
-        //! Same as MBF_PACKEDDIR8, with the w component replaced with the sign of the determinant
-        //! of the vertex basis to define the orientation of the tangent space in UE4 format.
-        //! Use 8-bit unsigned ints.
+        /** 
+		 * Same as MBF_PACKEDDIR8, with the w component replaced with the sign of the determinant
+         * of the vertex basis to define the orientation of the tangent space in UE4 format.
+         * Use 8-bit unsigned ints.
+		*/
         MBF_PACKEDDIR8_W_TANGENTSIGN,
 
-        //! Packed 1 to -1 value using multiply+add (128 is almost zero). Use 8-bit signed ints.
+        /** Packed 1 to -1 value using multiply+add (128 is almost zero). Use 8-bit signed ints. */
         MBF_PACKEDDIRS8,
 
-        //! Same as MBF_PACKEDDIRS8, with the w component replaced with the sign of the determinant
-        //! of the vertex basis to define the orientation of the tangent space in UE4 format.
-        //! Use 8-bit signed ints.
+        /** 
+		 * Same as MBF_PACKEDDIRS8, with the w component replaced with the sign of the determinant
+         * of the vertex basis to define the orientation of the tangent space in UE4 format.
+         * Use 8-bit signed ints.
+		 */
         MBF_PACKEDDIRS8_W_TANGENTSIGN,
 
 		MBF_FLOAT64,
@@ -71,7 +74,7 @@ namespace mu
 
 	} EMeshBufferFormat;
 
-	//!
+	/** */
 	struct FMeshBufferFormatData
 	{
 		/** Size per component in bytes. */
@@ -81,20 +84,19 @@ namespace mu
 		uint8 MaxValueBits;
 	};
 
-	MUTABLERUNTIME_API const FMeshBufferFormatData& GetMeshFormatData(EMeshBufferFormat format);
+	MUTABLERUNTIME_API const FMeshBufferFormatData& GetMeshFormatData(EMeshBufferFormat Format);
 
 
-	//! Semantics of the mesh buffers
-	//! \ingroup runtime
+	/** Semantics of the mesh buffers */
 	typedef enum
 	{
 
 		MBS_NONE,
 
-		//! For index buffers, and mesh morphs
+		/** For index buffers, and mesh morphs */
 		MBS_VERTEXINDEX,
 
-		//! Standard vertex semantics
+		/** Standard vertex semantics */
 		MBS_POSITION,
 		MBS_NORMAL,
 		MBS_TANGENT,
@@ -104,71 +106,72 @@ namespace mu
 		MBS_BONEWEIGHTS,
 		MBS_BONEINDICES,
 
-		//! Internal semantic indicating what layout block each vertex belongs to.
-		//! It can be safely ignored if present in meshes returned by the system.
-		//! It will never be in the same buffer that other vertex semantics.
+		/**
+		 * Internal semantic indicating what layout block each vertex belongs to.
+		 * It can be safely ignored if present in meshes returned by the system.
+		 * It will never be in the same buffer that other vertex semantics.
+		 */
 		MBS_LAYOUTBLOCK,
 
 		MBS_CHART_DEPRECATED,
 
-		//! To let users define channels with semantics unknown to the system.
-		//! These channels will never be transformed, and the per-vertex or per-index data will be
-		//! simply copied.
+		/** 
+		 * To let users define channels with semantics unknown to the system.
+		 * These channels will never be transformed, and the per-vertex or per-index data will be
+		 * simply copied.
+		 */
 		MBS_OTHER,
 
-        //! Sign to define the orientation of the tangent space.
+        /** Sign to define the orientation of the tangent space. **/
         MBS_TANGENTSIGN_DEPRECATED,
 
-		//! Semantics usefule for mesh binding.
+		/** Semantics usefule for mesh binding. */
 		MBS_TRIANGLEINDEX,
 		MBS_BARYCENTRICCOORDS,
 		MBS_DISTANCE,
 
-		//! Semantics useful for alternative skin weight profiles.
+		/** Semantics useful for alternative skin weight profiles. */
 		MBS_ALTSKINWEIGHT,
 
-		//! Utility
+		/** Utility */
 		MBS_COUNT,
 
         _MBS_FORCE32BITS = 0xFFFFFFFF
 	} EMeshBufferSemantic;
 
 
-	//!
+	/** */
 	struct FMeshBufferChannel
 	{
 		FMeshBufferChannel()
 		{
-			m_semantic = MBS_NONE;
-			m_format = MBF_NONE;
-			m_semanticIndex = 0;
-			m_offset = 0;
-			m_componentCount = 0;
+			Semantic = MBS_NONE;
+			Format = MBF_NONE;
+			SemanticIndex = 0;
+			Offset = 0;
+			ComponentCount = 0;
 		}
 
-		//!
-		EMeshBufferSemantic m_semantic;
+		EMeshBufferSemantic Semantic;
 
-		//!
-		EMeshBufferFormat m_format;
+		EMeshBufferFormat Format;
 
-		//! Index of the semantic, in case there are more than one of this type.
-		int32 m_semanticIndex;
+		/** Index of the semantic, in case there are more than one of this type. */
+		int32 SemanticIndex;
 
-		//! Offset in bytes from the begining of a buffer element
-		uint16 m_offset;
+		/** Offset in bytes from the begining of a buffer element */
+		uint16 Offset;
 
-		//! Number of components of the type in m_format for every value in the channel
-		uint16 m_componentCount;
+		/** Number of components of the type in Format for every value in the channel */
+		uint16 ComponentCount;
 
-		//!
-		inline bool operator==(const FMeshBufferChannel& o) const
+		inline bool operator==(const FMeshBufferChannel& Other) const
 		{
-			return (m_semantic == o.m_semantic) &&
-				(m_format == o.m_format) &&
-				(m_semanticIndex == o.m_semanticIndex) &&
-				(m_offset == o.m_offset) &&
-				(m_componentCount == o.m_componentCount);
+			return (Semantic == Other.Semantic) &&
+				(Format == Other.Format) &&
+				(SemanticIndex == Other.SemanticIndex) &&
+				(Offset == Other.Offset) &&
+				(ComponentCount == Other.ComponentCount);
 		}
 
 	};
@@ -179,36 +182,28 @@ namespace mu
 		template<typename Type>
 		using TMemoryTrackedArray = TArray<Type, FDefaultMemoryTrackingAllocator<MemoryCounters::FMeshMemoryCounter>>;
 
-		//!
-		TArray<FMeshBufferChannel> m_channels;
+		TArray<FMeshBufferChannel> Channels;
+		TMemoryTrackedArray<uint8> Data;
+		uint32 ElementSize = 0;
 
-		//!
-		TMemoryTrackedArray<uint8> m_data;
+		void Serialise(mu::OutputArchive& Arch) const;
+		inline void Unserialise(mu::InputArchive& Arch);
 
-		//!
-		uint32 m_elementSize = 0;
-
-		//!
-		void Serialise(mu::OutputArchive& arch) const;
-
-		//!
-		inline void Unserialise(mu::InputArchive& arch);
-
-		//!
-		inline bool operator==(const FMeshBuffer& o) const
+		inline bool operator==(const FMeshBuffer& Other) const
 		{
-			bool equal = (m_channels == o.m_channels);
-			if (equal) equal = (m_elementSize == o.m_elementSize);
-			if (equal) equal = (m_data == o.m_data);
-			return equal;
+			bool bEqual = (Channels == Other.Channels);
+			bEqual = bEqual && (ElementSize == Other.ElementSize);
+			bEqual = bEqual && (Data == Other.Data);
+
+			return bEqual;
 		}
 
 		/** Return true if the buffer has any channel with the passed semantic. */
 		inline bool HasSemantic(EMeshBufferSemantic Semantic) const
 		{
-			for ( const FMeshBufferChannel& Channel : m_channels)
+			for (const FMeshBufferChannel& Channel : Channels)
 			{
-				if (Channel.m_semantic == Semantic)
+				if (Channel.Semantic == Semantic)
 				{
 					return true;
 				}
@@ -218,204 +213,219 @@ namespace mu
 
 		inline bool HasSameFormat(const FMeshBuffer& Other) const
 		{
-			return (m_channels == Other.m_channels && m_elementSize == Other.m_elementSize);
+			return (Channels == Other.Channels && ElementSize == Other.ElementSize);
 		}
 
 		inline bool HasPadding() const
 		{
 			uint32 ActualElementSize = 0;
-			for (const FMeshBufferChannel& Channel : m_channels)
+			for (const FMeshBufferChannel& Channel : Channels)
 			{
-				ActualElementSize += Channel.m_componentCount * GetMeshFormatData(Channel.m_format).SizeInBytes;
+				ActualElementSize += Channel.ComponentCount * GetMeshFormatData(Channel.Format).SizeInBytes;
 			}
-			check(ActualElementSize<=m_elementSize);
-			return ActualElementSize < m_elementSize;
+			check(ActualElementSize <= ElementSize);
+			return ActualElementSize < ElementSize;
 		}
 	};
 
 
-	//! Set of buffers storing mesh element data. Elements can be vertices, indices or faces.
+	/** Set of buffers storing mesh element data. Elements can be vertices, indices or faces. */
 	class MUTABLERUNTIME_API FMeshBufferSet
 	{
 	public:
 
-		//!
-		uint32 m_elementCount = 0;
+		uint32 ElementCount = 0;
+		TArray<FMeshBuffer> Buffers;
 
-		//!
-		TArray<FMeshBuffer> m_buffers;
+		void Serialise(OutputArchive& Arch) const;
+		void Unserialise(InputArchive& Arch);
 
-		//!
-		void Serialise(OutputArchive& arch) const;
-
-		//!
-		void Unserialise(InputArchive& arch);
-
-		//!
-		inline bool operator==(const FMeshBufferSet& o) const
+		inline bool operator==(const FMeshBufferSet& Other) const
 		{
-			return (m_elementCount == o.m_elementCount) &&
-				(m_buffers == o.m_buffers);
+			return ElementCount == Other.ElementCount && Buffers == Other.Buffers;
 		}
 
 	public:
 
-		//! Get the number of elements in the buffers.
+		/** Get the number of elements in the buffers */
 		int32 GetElementCount() const;
 
-		//! Set the number of vertices in the mesh. This will resize the vertex buffers keeping the
-		//! previous data when possible. New data is undefined.
-		void SetElementCount( int32 );
+		/** 
+		* Set the number of vertices in the mesh. This will resize the vertex buffers keeping the
+		* previous data when possible. New data content is defined by MemoryInitPolicy.
+		*/
+		void SetElementCount(int32 Count, EMemoryInitPolicy MemoryInitPolicy = EMemoryInitPolicy::Uninitialized);
 
-		//! Get the size in bytes of a buffer element.
-		//! \param buffer index of the buffer from 0 to GetBufferCount()-1
-		int32 GetElementSize( int32 buffer ) const;
+		/**
+		* Get the size in bytes of a buffer element.
+		* @param buffer index of the buffer from 0 to GetBufferCount()-1
+		*/
+		int32 GetElementSize(int32 Buffer) const;
 
-		//! Get the number of vertex buffers in the mesh
+		/** Get the number of vertex buffers in the mesh */
 		int32 GetBufferCount() const;
 
-		//! Set the number of vertex buffers in the mesh.
-		void SetBufferCount( int32 );
-
-		//! Get the number of channels in a vertex buffer.
-		//! \param buffer index of the vertex buffer from 0 to GetBufferCount()-1
-		int32 GetBufferChannelCount( int32 buffer ) const;
-
-		//! Get a channel of a buffer by index
-		//! \param buffer index of the vertex buffer from 0 to GetBufferCount()-1
-		//! \param channel index of the channel from 0 to GetBufferChannelCount( buffer )-1
-		//! \param[out] pSemantic semantic of the channel
-		//! \param[out] pSemanticIndex index of the semantic in case of having more than one of the
-		//!				same type.
-		//! \param[out] pFormat data format of the channel
-		//! \param[out] pComponentCount components of an element of the channel
-		//! \param[out] pOffset offset in bytes from the beginning of an element of the buffer
-		void GetChannel
-			(
-				int32 buffer,
-				int32 channel,
-				EMeshBufferSemantic* pSemantic,
-				int32* pSemanticIndex,
-				EMeshBufferFormat* pFormat,
-				int32* pComponentCount,
-				int32* pOffset
+		/** Set the number of vertex buffers in the mesh. */
+		void SetBufferCount(int32 Count);
+
+		/**
+		* Get the number of channels in a vertex buffer.
+		* \param buffer index of the vertex buffer from 0 to GetBufferCount()-1
+		*/
+		int32 GetBufferChannelCount(int32 BufferIndex) const;
+
+		/**
+		 * Get a channel of a buffer by index
+		 * \param buffer index of the vertex buffer from 0 to GetBufferCount()-1
+		 * \param channel index of the channel from 0 to GetBufferChannelCount( buffer )-1
+		 * \param[out] pSemantic semantic of the channel
+		 * \param[out] pSemanticIndex index of the semantic in case of having more than one of the
+		 *				same type.
+		 * \param[out] pFormat data format of the channel
+		 * \param[out] pComponentCount components of an element of the channel
+		 * \param[out] pOffset offset in bytes from the beginning of an element of the buffer
+		 */
+		void GetChannel(
+				int32 BufferIndex,
+				int32 ChannelIndex,
+				EMeshBufferSemantic* SemanticPtr,
+				int32* SemanticIndexPtr,
+				EMeshBufferFormat* FormatPtr,
+				int32* ComponentCountPtr,
+				int32* OffsetPtr
 			) const;
 
-		//! Set all the channels of a buffer
-		//! \param buffer index of the buffer from 0 to GetBufferCount()-1
-        //! \param elementSize sizei n bytes of a vertex element in this buffer
-        //! \param channelCount number of channels to set in the buffer
-        //! \param pSemantics buffer of channelCount semantics
-        //! \param pSemanticIndices buffer of indices for the semantic of every channel
-        //! \param pFormats buffer of channelCount formats
-        //! \param pComponentCounts buffer of channelCount component counts
-        //! \param pOffsets offsets in bytes of every particular channel inside the buffer element
-        void SetBuffer
-			(
-				int32 buffer,
-				int32 elementSize,
-				int32 channelCount,
-				const EMeshBufferSemantic* pSemantics=nullptr,
-				const int32* pSemanticIndices=nullptr,
-				const EMeshBufferFormat* pFormats=nullptr,
-				const int32* pComponentCounts=nullptr,
-				const int32* pOffsets=nullptr
-			);
-
-		//! Set one  channels of a buffer
-		//! \param buffer index of the buffer from 0 to GetBufferCount()-1
-        //! \param elementSize sizei n bytes of a vertex element in this buffer
-        //! \param channelIndex number of channels to set in the buffer
-        void SetBufferChannel
-			(
-				int32 buffer,
-				int32 channelIndex,
-				EMeshBufferSemantic semantic,
-				int32 semanticIndice,
-				EMeshBufferFormat format,
-				int32 componentCount,
-				int32 offset
+		/** 
+		 * Set all the channels of a buffer
+		 * \param buffer index of the buffer from 0 to GetBufferCount()-1
+         * \param elementSize sizei n bytes of a vertex element in this buffer
+         * \param channelCount number of channels to set in the buffer
+         * \param pSemantics buffer of channelCount semantics
+         * \param pSemanticIndices buffer of indices for the semantic of every channel
+         * \param pFormats buffer of channelCount formats
+         * \param pComponentCounts buffer of channelCount component counts
+         * \param pOffsets offsets in bytes of every particular channel inside the buffer element
+		 */ 
+        void SetBuffer(
+				int32 BufferIndex,
+				int32 ElementSize,
+				int32 ChannelCount,
+				const EMeshBufferSemantic* SemanticsPtr = nullptr,
+				const int32* SemanticIndicesPtr = nullptr,
+				const EMeshBufferFormat* FormatsPtr = nullptr,
+				const int32* ComponentCountsPtr = nullptr,
+				const int32* OffsetsPtr = nullptr,
+				EMemoryInitPolicy MemoryInitPolicy = EMemoryInitPolicy::Uninitialized);
+
+		/**
+		 * Set one  channels of a buffer
+		 * \param buffer index of the buffer from 0 to GetBufferCount()-1
+         * \param elementSize sizei n bytes of a vertex element in this buffer
+         * \param channelIndex number of channels to set in the buffer
+		 */
+        void SetBufferChannel(
+				int32 BufferIndex,
+				int32 ChannelIndex,
+				EMeshBufferSemantic Semantic,
+				int32 SemanticIndex,
+				EMeshBufferFormat Format,
+				int32 ComponentCount,
+				int32 Offset
 			);
 
-		//! Get a pointer to the object-owned data of a buffer.
-		//! Channel data is interleaved for every element and packed in the order it was set
-		//! without any padding.
-		//! \param buffer index of the buffer from 0 to GetBufferCount()-1
-		//! \todo Add padding support for better alignment of buffer elements.
-        uint8* GetBufferData( int32 buffer );
-		const uint8* GetBufferData(int32 buffer) const;
-		uint32 GetBufferDataSize( int32 buffer ) const;
-
-		//-----------------------------------------------------------------------------------------
-		// Utility methods
-		//-----------------------------------------------------------------------------------------
-
-		//! Find the index of a buffer channel by semantic and relative index inside the semantic.
-        //! \param semantic Semantic of the channel we are searching.
-        //! \param semanticIndex Index of the semantic of the channel we are searching. e.g. if we
-        //!         want the second set of texture coordinates, it should be 1.
-        //! \param[out] pBuffer -1 if the channel is not found, otherwise it will contain the index
-		//! of the buffer where the channel was found.
-		//! \param[out] pChannel -1 if the channel is not found, otherwise it will contain the
-		//! channel index of the channel inside the buffer returned at [buffer]
-		void FindChannel
-			(
-				EMeshBufferSemantic semantic,
-				int32 semanticIndex,
-				int32* pBuffer, int32* pChannel
-			) const;
-
-		//! Get the offset in bytes of the data of this channel inside an element data.
-		//! \param buffer index of the buffer from 0 to GetBufferCount()-1
-		//! \param channel index of the channel from 0 to GetBufferChannelCount( buffer )-1
-		int32 GetChannelOffset( int32 buffer, int32 channel ) const;
-
-		//! Add a new buffer by cloning a buffer from another set.
-		//! The buffers must have the same number of elements.
-		void AddBuffer( const FMeshBufferSet& Other, int32 buffer );
-
-		//! Return true if the formats of the two vertex buffers set match.
-		bool HasSameFormat( const FMeshBufferSet& Other ) const;
-
-		//! Remove the buffer at the specified position. This "invalidates" any buffer index that
-		//! was referencing buffers after the removed one.
-		void RemoveBuffer( int32 BufferIndex);
+		/** 
+		 * Get a pointer to the object-owned data of a buffer.
+		 * Channel data is interleaved for every element and packed in the order it was set
+		 * without any padding.
+		 * \param buffer index of the buffer from 0 to GetBufferCount()-1
+		 * \todo Add padding support for better alignment of buffer elements.
+		*/
+        uint8* GetBufferData(int32 Buffer);
+		const uint8* GetBufferData(int32 Buffer) const;
+		uint32 GetBufferDataSize(int32 Buffer) const;
+
+		/** Utility methods */
+
+		/** 
+		 * Find the index of a buffer channel by semantic and relative index inside the semantic.
+         * \param semantic Semantic of the channel we are searching.
+         * \param semanticIndex Index of the semantic of the channel we are searching. e.g. if we
+         *         want the second set of texture coordinates, it should be 1.
+         * \param[out] pBuffer -1 if the channel is not found, otherwise it will contain the index
+		 * of the buffer where the channel was found.
+		 * \param[out] pChannel -1 if the channel is not found, otherwise it will contain the
+		 * channel index of the channel inside the buffer returned at [buffer]
+		 */
+		void FindChannel(EMeshBufferSemantic Semantic, int32 SemanticIndex, int32* BufferPtr, int32* ChannelPtr) const;
+
+		/**
+		 * Get the offset in bytes of the data of this channel inside an element data.
+		 * \param buffer index of the buffer from 0 to GetBufferCount()-1
+		 * \param channel index of the channel from 0 to GetBufferChannelCount( buffer )-1
+		 */ 
+		int32 GetChannelOffset(int32 Buffer, int32 Channel) const;
+
+		/**
+		 * Add a new buffer by cloning a buffer from another set.
+		 * The buffers must have the same number of elements.
+		 */ 
+		void AddBuffer( const FMeshBufferSet& Other, int32 BufferIndex);
+
+		/** Return true if the formats of the two vertex buffers set match. **/
+		bool HasSameFormat(const FMeshBufferSet& Other) const;
+
+		/** 
+		 * Remove the buffer at the specified position. This "invalidates" any buffer index that
+		 * was referencing buffers after the removed one.
+		 */ 
+		void RemoveBuffer(int32 BufferIndex);
 
 	public:
 
-		//! Copy an element from one position to another, overwriting the other element.
-		//! Both positions must be valid, buffer size won't change.
-		void CopyElement(uint32 fromIndex, uint32 toIndex);
+		/**
+		 * Copy an element from one position to another, overwriting the other element.
+		 * Both positions must be valid, buffer size won't change.
+		 */ 
+		void CopyElement(uint32 FromIndex, uint32 ToIndex);
 
-		//! Compare the format of the two buffers at index buffer and return true if they match.
+		/** Compare the format of the two buffers at index buffer and return true if they match. **/
 		bool HasSameFormat(int32 ThisBufferIndex, const FMeshBufferSet& pOther, int32 OtherBufferIndex) const;
 
-		//! Get the total memory size of the buffers and this struct
+		/** Get the total memory size of the buffers and this struct */
 		int32 GetDataSize() const;
 
+		/** */
 		int32 GetAllocatedSize() const;
 
-		//! Compare the mesh buffer with another one, but ignore internal data like generated
-		//! vertex indices.
-		bool IsSpecialBufferToIgnoreInSimilar(const FMeshBuffer& b) const;
-
-		//! Compare the mesh buffer with another one, but ignore internal data like generated
-		//! vertex indices. Be aware this method compares the data byte by byte without checking
-		//! if the data belong to the buffer components and could give false negatives if unset 
-		//! padding data is present.
-		bool IsSimilar(const FMeshBufferSet& o) const;
-
-		//! Compare the mesh buffer with another one, but ignore internal data like generated
-		//! vertex indices. This version compares the data component-wise, skipping any memory
-		//! not specified in the buffer description.	
-		bool IsSimilarRobust(const FMeshBufferSet& Other,bool bCompareUVs) const;
-
-		//! Change the buffer descriptions so that all buffer indices start at 0 and are in the
-		//! same order than memory.
+		/** 
+		 * Compare the mesh buffer with another one, but ignore internal data like generated
+		 * vertex indices.
+		 */ 
+		bool IsSpecialBufferToIgnoreInSimilar(const FMeshBuffer& Buffer) const;
+
+		/** 
+		 * Compare the mesh buffer with another one, but ignore internal data like generated
+		 * vertex indices. Be aware this method compares the data byte by byte without checking
+		 * if the data belong to the buffer components and could give false negatives if unset 
+		 * padding data is present.
+		 */ 
+		bool IsSimilar(const FMeshBufferSet& Other) const;
+
+		/** 
+		 * Compare the mesh buffer with another one, but ignore internal data like generated
+		 * vertex indices. This version compares the data component-wise, skipping any memory
+		 * not specified in the buffer description.
+		 */
+		bool IsSimilarRobust(const FMeshBufferSet& Other, bool bCompareUVs) const;
+
+		/** 
+		 * Change the buffer descriptions so that all buffer indices start at 0 and are in the
+		 * same order than memory.
+		 */
 		void ResetBufferIndices();
 
-		//!
-		void UpdateOffsets(int32 b);
+		/** */
+		void UpdateOffsets(int32 BufferIndex);
 		
 		/** Check that all channels of a specific semantic use the provided format. */
 		bool HasAnySemanticWithDifferentFormat(EMeshBufferSemantic Semantic, EMeshBufferFormat ExpectedFormat) const;
diff --git a/Engine/Plugins/Experimental/Mutable/Source/MutableRuntime/Public/MuR/Types.h b/Engine/Plugins/Experimental/Mutable/Source/MutableRuntime/Public/MuR/Types.h
--- a/Engine/Plugins/Experimental/Mutable/Source/MutableRuntime/Public/MuR/Types.h
+++ b/Engine/Plugins/Experimental/Mutable/Source/MutableRuntime/Public/MuR/Types.h
@@ -14,6 +14,13 @@ namespace mu
 
 	/** Handle used to identify separate resources generated by mutable: images and meshes included in instances. */
     typedef uint64 FResourceID;
+
+	enum class EMemoryInitPolicy
+	{
+		Uninitialized,
+		Zeroed,
+	};
+
 }
 
 
diff --git a/Engine/Plugins/Experimental/Mutable/Source/MutableTools/Private/MuT/CodeGenerator_Mesh.cpp b/Engine/Plugins/Experimental/Mutable/Source/MutableTools/Private/MuT/CodeGenerator_Mesh.cpp
--- a/Engine/Plugins/Experimental/Mutable/Source/MutableTools/Private/MuT/CodeGenerator_Mesh.cpp
+++ b/Engine/Plugins/Experimental/Mutable/Source/MutableTools/Private/MuT/CodeGenerator_Mesh.cpp
@@ -354,12 +354,12 @@ namespace mu
 		check(TexCoordsBufferIndex >= 0);
 		check(TexCoordsChannelIndex >= 0);
 
-		const FMeshBufferChannel& TexCoordsChannel = Mesh->VertexBuffers.m_buffers[TexCoordsBufferIndex].m_channels[TexCoordsChannelIndex];
-		check(TexCoordsChannel.m_semantic== MBS_TEXCOORDS);
+		const FMeshBufferChannel& TexCoordsChannel = Mesh->VertexBuffers.Buffers[TexCoordsBufferIndex].Channels[TexCoordsChannelIndex];
+		check(TexCoordsChannel.Semantic == MBS_TEXCOORDS);
 
 		uint8* TexCoordData = Mesh->GetVertexBuffers().GetBufferData(TexCoordsBufferIndex);
 		int32 elemSize = Mesh->GetVertexBuffers().GetElementSize(TexCoordsBufferIndex);
-		int32 channelOffset = TexCoordsChannel.m_offset;
+		int32 channelOffset = TexCoordsChannel.Offset;
 		TexCoordData += channelOffset;
 
 		// Get a copy of the UVs as FVector2f to work with them. 
@@ -374,11 +374,11 @@ namespace mu
 			for (int32 VertexIndex = 0; VertexIndex < NumVertices; ++VertexIndex)
 			{
 				FVector2f& UV = TexCoords[VertexIndex];
-				if (TexCoordsChannel.m_format == MBF_FLOAT32)
+				if (TexCoordsChannel.Format == MBF_FLOAT32)
 				{
 					UV = *((FVector2f*)pVertices);
 				}
-				else if (TexCoordsChannel.m_format == MBF_FLOAT16)
+				else if (TexCoordsChannel.Format == MBF_FLOAT16)
 				{
 					const FFloat16* pUV = reinterpret_cast<const FFloat16*>(pVertices);
 					UV = FVector2f(float(pUV[0]), float(pUV[1]));
@@ -698,12 +698,12 @@ namespace mu
 				}
 
 				// Copy UVs
-				if (TexCoordsChannel.m_format == MBF_FLOAT32)
+				if (TexCoordsChannel.Format == MBF_FLOAT32)
 				{
 					FVector2f* pUV = reinterpret_cast<FVector2f*>(pVertices);
 					*pUV = *UV;
 				}
-				else if (TexCoordsChannel.m_format == MBF_FLOAT16)
+				else if (TexCoordsChannel.Format == MBF_FLOAT16)
 				{
 					FFloat16* pUV = reinterpret_cast<FFloat16*>(pVertices);
 					pUV[0] = FFloat16((*UV)[0]);
diff --git a/Engine/Plugins/Experimental/Mutable/Source/MutableTools/Private/MuT/DataPacker.cpp b/Engine/Plugins/Experimental/Mutable/Source/MutableTools/Private/MuT/DataPacker.cpp
--- a/Engine/Plugins/Experimental/Mutable/Source/MutableTools/Private/MuT/DataPacker.cpp
+++ b/Engine/Plugins/Experimental/Mutable/Source/MutableTools/Private/MuT/DataPacker.cpp
@@ -716,7 +716,7 @@ namespace mu
 
             if (!used)
             {
-                TArray<FMeshBuffer>& buffers = pMesh->GetVertexBuffers().m_buffers;
+                TArray<FMeshBuffer>& buffers = pMesh->GetVertexBuffers().Buffers;
                 buffers.RemoveAt(v);
             }
             else
