[mutable] Remove unnecessary concurrency synchronization. It is already implicitly safe because of the order of execution due to task dependencies.
#rnx

diff --git a/Engine/Plugins/Experimental/Mutable/Source/MutableTools/Private/MuT/CodeOptimiser.cpp b/Engine/Plugins/Experimental/Mutable/Source/MutableTools/Private/MuT/CodeOptimiser.cpp
--- a/Engine/Plugins/Experimental/Mutable/Source/MutableTools/Private/MuT/CodeOptimiser.cpp
+++ b/Engine/Plugins/Experimental/Mutable/Source/MutableTools/Private/MuT/CodeOptimiser.cpp
@@ -763,6 +763,8 @@ namespace mu
 
 		auto GetRequisites = [&ConstantSubgraphs](const Ptr<ASTOp>& SubgraphRoot, TArray< UE::Tasks::FTask, TInlineAllocator<8> >& OutRequisites)
 		{
+			MUTABLE_CPUPROFILER_SCOPE(ConstantGenerator_GetRequisites);
+
 			TArray< Ptr<ASTOp> > ScanRoots;
 			ScanRoots.Add(SubgraphRoot);
 			ASTOp::Traverse_TopDown_Unique_Imprecise(ScanRoots, [&SubgraphRoot, &OutRequisites, &ConstantSubgraphs](Ptr<ASTOp>& ChildNode)
@@ -790,17 +792,12 @@ namespace mu
 
 		if (bUseConcurrency)
 		{
-			/** Protect access to the original AST being optimized. */
-			FCriticalSection ASTAccessLock;
-
 			// Launch the tasks.
 			UE::Tasks::FTask LaunchTask = UE::Tasks::Launch(TEXT("ConstantGeneratorLaunchTasks"), 
-				[&ConstantSubgraphs, &GetRequisites, &ASTAccessLock, Pass, InOptions]()
+				[&ConstantSubgraphs, &GetRequisites, Pass, InOptions]()
 				{
 					MUTABLE_CPUPROFILER_SCOPE(ConstantGenerator_LaunchTasks);
 
-					FScopeLock Lock(&ASTAccessLock);
-
 					FImageOperator ImOp = FImageOperator::GetDefault(InOptions->ImageFormatFunc);
 
 					// Traverse list of constants to generate. It is ordered in a bottom-up way.
@@ -828,17 +825,21 @@ namespace mu
 							UE::Tasks::FTask ReferenceCompletion = InOptions->OptimisationOptions.ReferencedResourceProvider(ImageID, ResolveImage, bRunImmediatlyIfPossible);
 
 							UE::Tasks::FTask CompleteTask = UE::Tasks::Launch(TEXT("MutableResolveComplete"),
-								[SubgraphRoot, InOptions, ResolveImage, &ASTAccessLock]()
+								[SubgraphRoot, InOptions, ResolveImage]()
 								{
-									MUTABLE_CPUPROFILER_SCOPE(MutableResolveComplete);
-
-									Ptr<ASTOpConstantResource> ConstantOp = new ASTOpConstantResource;
-									ConstantOp->SourceDataDescriptor = SubgraphRoot->GetSourceDataDescriptor();
-									ConstantOp->Type = OP_TYPE::IM_CONSTANT;
-									ConstantOp->SetValue(ResolveImage->get(), InOptions->OptimisationOptions.DiskCacheContext);
-
+									Ptr<ASTOpConstantResource> ConstantOp;
 									{
-										FScopeLock Lock(&ASTAccessLock);
+										MUTABLE_CPUPROFILER_SCOPE(MutableResolveComplete_CreateConstant);
+										ConstantOp = new ASTOpConstantResource;
+										ConstantOp->Type = OP_TYPE::IM_CONSTANT;
+										{
+											MUTABLE_CPUPROFILER_SCOPE(GetSourceDataDescriptor);
+											ConstantOp->SourceDataDescriptor = SubgraphRoot->GetSourceDataDescriptor();
+										}
+										ConstantOp->SetValue(ResolveImage->get(), InOptions->OptimisationOptions.DiskCacheContext);
+									}
+									{
+										MUTABLE_CPUPROFILER_SCOPE(MutableResolveComplete_Replace);
 										ASTOp::Replace(SubgraphRoot, ConstantOp);
 									}
 								},
@@ -855,47 +856,22 @@ namespace mu
 							GetRequisites(SubgraphRoot,Requisites);
 
 							TUniquePtr<FConstantTask> Task(new FConstantTask(SubgraphRoot, InOptions, Pass));
-							FConstantTask* TaskPtr = Task.Get();
 
 							// Launch the preparation on the AST-modification pipe
-							UE::Tasks::FTask PrepareTask = UE::Tasks::Launch(TEXT("MutableConstantPrepare"), [TaskPtr, &ASTAccessLock]()
+							UE::Tasks::FTask CompleteTask = UE::Tasks::Launch(TEXT("MutableConstant"), [TaskPtr = MoveTemp(Task), ImOp]()
 								{
 									MUTABLE_CPUPROFILER_SCOPE(MutableConstantPrepare);
-									FScopeLock Lock(&ASTAccessLock);
 
 									// We need the clone because linking modifies ASTOp state and also to be safe for concurrency.
 									TaskPtr->SourceCloned = ASTOp::DeepClone(TaskPtr->Source);
-								},
-								Requisites,
-								LowLevelTasks::ETaskPriority::BackgroundHigh);
 
-							// Launch constant generation on any thread
-							UE::Tasks::FTask RunTask = UE::Tasks::Launch(TEXT("MutableConstantGeneration"), [TaskPtr, ImOp]()
-								{
 									TaskPtr->Run(ImOp);
-								},
-								PrepareTask,
-								LowLevelTasks::ETaskPriority::BackgroundHigh);
-
-							// Launch the completion on the AST-modification pipe
-							UE::Tasks::FTask CompleteTask = UE::Tasks::Launch(TEXT("MutableConstantComplete"), [TaskPtr = MoveTemp(Task), &ASTAccessLock]()
-								{
-									MUTABLE_CPUPROFILER_SCOPE(MutableConstantComplete);
-
-									{
-										MUTABLE_CPUPROFILER_SCOPE(Replace);
-										FScopeLock Lock(&ASTAccessLock);
 
-										ASTOp::Replace(TaskPtr->Source, TaskPtr->Result);
-										TaskPtr->Result = nullptr;
-									}
-
-									{
-										MUTABLE_CPUPROFILER_SCOPE(Release);
-										TaskPtr->Source = nullptr;
-									}
+									ASTOp::Replace(TaskPtr->Source, TaskPtr->Result);
+									TaskPtr->Result = nullptr;
+									TaskPtr->Source = nullptr;
 								},
-								RunTask,
+								Requisites,
 								LowLevelTasks::ETaskPriority::BackgroundHigh);
 
 							SubgraphCompletionEvent.AddPrerequisites(CompleteTask);
