[mutable] Added option to set component names from the generation settings menu of COs.

#rnx

diff --git a/Engine/Plugins/Experimental/Mutable/Source/CustomizableObjectEditor/Private/MuCOE/CustomizableObjectFactory.cpp b/Engine/Plugins/Experimental/Mutable/Source/CustomizableObjectEditor/Private/MuCOE/CustomizableObjectFactory.cpp
--- a/Engine/Plugins/Experimental/Mutable/Source/CustomizableObjectEditor/Private/MuCOE/CustomizableObjectFactory.cpp
+++ b/Engine/Plugins/Experimental/Mutable/Source/CustomizableObjectEditor/Private/MuCOE/CustomizableObjectFactory.cpp
@@ -24,6 +24,7 @@
 #include "Types/SlateEnums.h"
 #include "Widgets/Input/SButton.h"
 #include "Widgets/Input/SCheckBox.h"
+#include "Widgets/Input/SEditableTextBox.h"
 #include "Widgets/Input/SSpinBox.h"
 #include "Widgets/Input/STextComboBox.h"
 #include "Widgets/Layout/SBorder.h"
@@ -32,6 +33,9 @@
 
 #define LOCTEXT_NAMESPACE "CustomizableObjectFactory"
 
+static const FVector2D NewCOWindowsSize = FVector2D(300,340);
+static const FVector2D NewChildWindowsSize = FVector2D(300, 220);
+
 UCustomizableObjectFactory::UCustomizableObjectFactory()
 	: Super()
 {
@@ -126,11 +130,12 @@ UObject* UCustomizableObjectFactory::FactoryCreateNew(UClass* Class, UObject* In
 		}
 		else
 		{
-			check(CreationSettings.ReferenceSkeletalMeshes.Num() == CreationSettings.NumMeshComponents);
+			check(CreationSettings.ComponentsInfo.Num() == CreationSettings.NumMeshComponents);
 
-			for (int32 MeshIndex = 0; MeshIndex < CreationSettings.ReferenceSkeletalMeshes.Num(); ++MeshIndex)
+			for (int32 MeshIndex = 0; MeshIndex < CreationSettings.ComponentsInfo.Num(); ++MeshIndex)
 			{
-				ObjectPrivate->MutableMeshComponents[MeshIndex].ReferenceSkeletalMesh = CreationSettings.ReferenceSkeletalMeshes[MeshIndex].Get();
+				ObjectPrivate->MutableMeshComponents[MeshIndex].Name = CreationSettings.ComponentsInfo[MeshIndex].ComponentName;
+				ObjectPrivate->MutableMeshComponents[MeshIndex].ReferenceSkeletalMesh = CreationSettings.ComponentsInfo[MeshIndex].ReferenceSkeletalMesh.LoadSynchronous();
 			}
 		}
 	}
@@ -152,13 +157,13 @@ const FCustomizableObjectOptions FCustomizableObjectFactoryUI::ConstructFactoryU
 
 	// Default number of components for non-child objects
 	Options.NumMeshComponents = 1;
-	Options.ReferenceSkeletalMeshes.SetNum(1);
+	Options.ComponentsInfo.SetNum(1);
 	GenerateComponentOptions();
 
 	// Settings window
 	COSettingsWindow = SNew(SWindow)
 	.Title(LOCTEXT("CustomizableObjectFactoryptions", "New Costumizable Object"))
-	.ClientSize(FVector2D(300, 280))
+	.ClientSize(NewCOWindowsSize)
 	.SupportsMinimize(false)
 	.SupportsMaximize(false)
 	[
@@ -294,6 +299,7 @@ const FCustomizableObjectOptions FCustomizableObjectFactoryUI::ConstructFactoryU
 					]
 
 					+SVerticalBox::Slot()
+					.AutoHeight()
 					.Padding(5.0f, 5.0f, 10.0f, 0.0f)
 					[
 						SAssignNew(ComponentSelector, STextComboBox)
@@ -306,12 +312,33 @@ const FCustomizableObjectOptions FCustomizableObjectFactoryUI::ConstructFactoryU
 						.Font(IDetailLayoutBuilder::GetDetailFont())
 					]
 
+					+ SVerticalBox::Slot()
+					.Padding(0.0f, 15.0f, 0.0f, 0.0f)
+					.AutoHeight()
+					[
+						SNew(STextBlock)
+						.Visibility(this, &FCustomizableObjectFactoryUI::GetComponentWidgetsVisibility)
+						.Text(this, &FCustomizableObjectFactoryUI::GetSelectorWidgetText, true)
+						.Font(IDetailLayoutBuilder::GetDetailFont())
+					]
+
+					+ SVerticalBox::Slot()
+					.Padding(5.0f, 10.0f, 10.0f, 0.0f)
+					.AutoHeight()
+					[
+						SNew(SEditableTextBox)
+						.Visibility(this, &FCustomizableObjectFactoryUI::GetComponentWidgetsVisibility)
+						.Text(this, &FCustomizableObjectFactoryUI::GetComponentName)
+						.OnTextCommitted(this, &FCustomizableObjectFactoryUI::OnTextCommited)
+						.Font(IDetailLayoutBuilder::GetDetailFont())
+					]
+
 					+ SVerticalBox::Slot()
 					.Padding(0.0f, 15.0f, 0.0f, 0.0f)
 					[
 						SNew(STextBlock)
 						.Visibility(this, &FCustomizableObjectFactoryUI::GetComponentWidgetsVisibility)
-						.Text(this, &FCustomizableObjectFactoryUI::GetMeshSelectorWidgetText)
+						.Text(this, &FCustomizableObjectFactoryUI::GetSelectorWidgetText, false)
 						.Font(IDetailLayoutBuilder::GetDetailFont())
 					]
 
@@ -401,9 +428,9 @@ void FCustomizableObjectFactoryUI::OnCheckBoxChanged(ECheckBoxState State)
 	Options.ParentObject = nullptr;
 	Options.GroupNodeName.Empty();
 	Options.NumMeshComponents = Options.bIsChildObject ? 0 : 1;
-	Options.bIsChildObject ? Options.ReferenceSkeletalMeshes.Empty() : Options.ReferenceSkeletalMeshes.SetNum(1);
+	Options.bIsChildObject ? Options.ComponentsInfo.Empty() : Options.ComponentsInfo.SetNum(1);
 
-	FVector2D ClientSize = Options.bIsChildObject ? FVector2D(300, 220) : FVector2D(300, 280);
+	FVector2D ClientSize = Options.bIsChildObject ? NewChildWindowsSize : NewCOWindowsSize;
 	FVector2D WindowsSize = ClientSize * COSettingsWindow->GetDPIScaleFactor();
 	
 	COSettingsWindow->Resize(WindowsSize);
@@ -439,7 +466,7 @@ FText FCustomizableObjectFactoryUI::GetOKButtonTooltip() const
 	{
 		if (!IsConfigurationValid())
 		{
-			Tooltip = "Set the number of components that the Customizable Object will have. Then select a Reference Skeletal Mesh for each component.";
+			Tooltip = "Set the number of components that the Customizable Object will have. Then select a Reference Skeletal Mesh and a Name for each component.";
 		}
 	}
 
@@ -471,10 +498,10 @@ bool FCustomizableObjectFactoryUI::IsConfigurationValid() const
 	}
 	else
 	{
-		// Check if all components have a valid Skeletal Mesh assigned
+		// Check if all components have a valid Skeletal Mesh and name assigned
 		for (int32 CompIndex = 0; CompIndex < Options.NumMeshComponents; ++CompIndex)
 		{
-			if (!Options.ReferenceSkeletalMeshes.IsValidIndex(CompIndex) || !Options.ReferenceSkeletalMeshes[CompIndex].IsValid())
+			if (!Options.ComponentsInfo.IsValidIndex(CompIndex) || Options.ComponentsInfo[CompIndex].ReferenceSkeletalMesh.IsNull() || Options.ComponentsInfo[CompIndex].ComponentName.IsNone())
 			{
 				return false;
 			}
@@ -575,7 +602,7 @@ void FCustomizableObjectFactoryUI::OnNumComponentsChanged(int32 Value, ETextComm
 	if (Options.NumMeshComponents != Value)
 	{
 		Options.NumMeshComponents = Value;
-		Options.ReferenceSkeletalMeshes.SetNum(Value);
+		Options.ComponentsInfo.SetNum(Value);
 		GenerateComponentOptions();
 
 		if (Value > 0)
@@ -617,7 +644,7 @@ void FCustomizableObjectFactoryUI::OnSelectComponentComboBox(TSharedPtr<FString>
 }
 
 
-FText FCustomizableObjectFactoryUI::GetMeshSelectorWidgetText() const
+FText FCustomizableObjectFactoryUI::GetSelectorWidgetText(bool bIsName) const
 {
 	FString ComponentName;
 
@@ -626,43 +653,76 @@ FText FCustomizableObjectFactoryUI::GetMeshSelectorWidgetText() const
 		ComponentName = *ComponentSelector->GetSelectedItem();
 	}
 
-	FString Message = "Select " + ComponentName + " Skeletal Mesh:";
+	FString VariableName = bIsName ? " Name:" : " Skeletal Mesh:";
+	FString Message = "Select " + ComponentName + VariableName;
 
 	return FText::FromString(Message);
 }
 
 
-void FCustomizableObjectFactoryUI::OnPickedComponentSkeletalMesh(const FAssetData& SelectedAsset)
+int32 FCustomizableObjectFactoryUI::GetSelectedComponentIndex() const
 {
 	for (int32 ComponentIndex = 0; ComponentIndex < ComponentsOptions.Num(); ++ComponentIndex)
 	{
 		if (ComponentSelector.IsValid() && ComponentSelector->GetSelectedItem().IsValid()
 			&& *ComponentsOptions[ComponentIndex] == *ComponentSelector->GetSelectedItem()
-			&& Options.ReferenceSkeletalMeshes.IsValidIndex(ComponentIndex))
+			&& Options.ComponentsInfo.IsValidIndex(ComponentIndex))
 		{
-			if (USkeletalMesh* SkeletalMesh = Cast<USkeletalMesh>(SelectedAsset.GetAsset()))
-			{
-				Options.ReferenceSkeletalMeshes[ComponentIndex] = SkeletalMesh;
-			}
+			return ComponentIndex;
 		}
 	}
+
+	return -1;
 }
 
 
-FString FCustomizableObjectFactoryUI::GetSelectedComponentSkeletalMeshPath() const
+void FCustomizableObjectFactoryUI::OnPickedComponentSkeletalMesh(const FAssetData& SelectedAsset)
 {
-	for (int32 ComponentIndex = 0; ComponentIndex < ComponentsOptions.Num(); ++ComponentIndex)
+	int32 ComponentIndex = GetSelectedComponentIndex();
+	if (ComponentIndex != INDEX_NONE)
 	{
-		if (ComponentSelector.IsValid() && ComponentSelector->GetSelectedItem().IsValid()
-			&& *ComponentsOptions[ComponentIndex] == *ComponentSelector->GetSelectedItem()
-			&& Options.ReferenceSkeletalMeshes.IsValidIndex(ComponentIndex))
+		if (SelectedAsset.IsValid())
 		{
-			return Options.ReferenceSkeletalMeshes[ComponentIndex].IsValid() ? Options.ReferenceSkeletalMeshes[ComponentIndex]->GetPathName() : FString();
+			// we have to load the asset otherwise the asset thumbnail won't be visible
+			Options.ComponentsInfo[ComponentIndex].ReferenceSkeletalMesh = SelectedAsset.GetAsset();
 		}
 	}
+}
+
+
+FString FCustomizableObjectFactoryUI::GetSelectedComponentSkeletalMeshPath() const
+{
+	int32 ComponentIndex = GetSelectedComponentIndex();
+	FString SkeletalMeshPath;
+
+	if (ComponentIndex != INDEX_NONE && !Options.ComponentsInfo[ComponentIndex].ReferenceSkeletalMesh.IsNull())
+	{
+		SkeletalMeshPath = Options.ComponentsInfo[ComponentIndex].ReferenceSkeletalMesh.ToSoftObjectPath().ToString();
+	}
 
-	return FString();
+	return SkeletalMeshPath;
 }
 
 
+FText FCustomizableObjectFactoryUI::GetComponentName() const
+{
+	int32 ComponentIndex = GetSelectedComponentIndex();
+	if (ComponentIndex != INDEX_NONE)
+	{
+		return FText::FromName(Options.ComponentsInfo[ComponentIndex].ComponentName);
+	}
+
+	return FText();
+}
+
+
+void FCustomizableObjectFactoryUI::OnTextCommited(const FText& NewName, ETextCommit::Type CommitInfo)
+{
+	int32 ComponentIndex = GetSelectedComponentIndex();
+	if (ComponentIndex != INDEX_NONE)
+	{
+		Options.ComponentsInfo[ComponentIndex].ComponentName = FName(*NewName.ToString());
+	}
+}
+
 #undef LOCTEXT_NAMESPACE
diff --git a/Engine/Plugins/Experimental/Mutable/Source/CustomizableObjectEditor/Public/MuCOE/CustomizableObjectFactory.h b/Engine/Plugins/Experimental/Mutable/Source/CustomizableObjectEditor/Public/MuCOE/CustomizableObjectFactory.h
--- a/Engine/Plugins/Experimental/Mutable/Source/CustomizableObjectEditor/Public/MuCOE/CustomizableObjectFactory.h
+++ b/Engine/Plugins/Experimental/Mutable/Source/CustomizableObjectEditor/Public/MuCOE/CustomizableObjectFactory.h
@@ -42,8 +42,13 @@ struct FCustomizableObjectOptions
 	// Number of components
 	int32 NumMeshComponents = 1;
 
-	// Array with the skeletal meshes of each component
-	TArray<TWeakObjectPtr<USkeletalMesh>> ReferenceSkeletalMeshes;
+	struct FComponentInfo 
+	{
+		TSoftObjectPtr<USkeletalMesh> ReferenceSkeletalMesh;
+		FName ComponentName;
+	};
+
+	TArray<FComponentInfo> ComponentsInfo;
 };
 
 
@@ -100,7 +105,7 @@ public:
 
 	// Component Selectior Widget methods
 	EVisibility GetComponentWidgetsVisibility() const;
-	FText GetMeshSelectorWidgetText() const;
+	FText GetSelectorWidgetText(bool bIsName) const;
 	FString GetSelectedComponentSkeletalMeshPath() const;
 	int32 GetNumComponents() const;
 	void OnNumComponentsChanged(int32 Value, ETextCommit::Type CommitInfo);
@@ -109,6 +114,11 @@ public:
 	void GenerateComponentOptions();
 	bool IsComponentSelectorEnabled() const;
 
+	FText GetComponentName() const;
+	void OnTextCommited(const FText& NewName, ETextCommit::Type CommitInfo);
+
+	int32 GetSelectedComponentIndex() const;
+
 private:
 
 	// Pointer to the created widgets
