[mutable] Fixed layered material
+ Added new struct to identify material parameters in nodes. Now we take into account the material layer.

#rnx

diff --git a/Engine/Plugins/Experimental/Mutable/Source/CustomizableObject/Internal/MuCO/CustomizableObjectCustomVersion.h b/Engine/Plugins/Experimental/Mutable/Source/CustomizableObject/Internal/MuCO/CustomizableObjectCustomVersion.h
--- a/Engine/Plugins/Experimental/Mutable/Source/CustomizableObject/Internal/MuCO/CustomizableObjectCustomVersion.h
+++ b/Engine/Plugins/Experimental/Mutable/Source/CustomizableObject/Internal/MuCO/CustomizableObjectCustomVersion.h
@@ -126,6 +126,8 @@ struct CUSTOMIZABLEOBJECT_API FCustomizableObjectCustomVersion
 
 		NodeMaterialTypedImagePins,
 
+		FixedMultilayerMaterialIds,
+
 		// -----<new versions can be added above this line>-------------------------------------------------
 		VersionPlusOne,
 		LatestVersion = VersionPlusOne - 1
diff --git a/Engine/Plugins/Experimental/Mutable/Source/CustomizableObjectEditor/Private/MuCOE/GenerateMutableSource/GenerateMutableSourceGroupProjector.cpp b/Engine/Plugins/Experimental/Mutable/Source/CustomizableObjectEditor/Private/MuCOE/GenerateMutableSource/GenerateMutableSourceGroupProjector.cpp
--- a/Engine/Plugins/Experimental/Mutable/Source/CustomizableObjectEditor/Private/MuCOE/GenerateMutableSource/GenerateMutableSourceGroupProjector.cpp
+++ b/Engine/Plugins/Experimental/Mutable/Source/CustomizableObjectEditor/Private/MuCOE/GenerateMutableSource/GenerateMutableSourceGroupProjector.cpp
@@ -88,7 +88,7 @@ mu::NodeImagePtr GenerateMutableGroupProjection(const int32 NodeLOD, const int32
 				}
 				else
 				{
-					const FGuid ImageId = TypedNodeExt->GetParentMaterialNode()->GetParameterId(EMaterialParameterType::Texture, ImageIndex);
+					const FNodeMaterialParameterId ImageId = TypedNodeExt->GetParentMaterialNode()->GetParameterId(EMaterialParameterType::Texture, ImageIndex);
 					return TypedNodeExt->UsesImage(ImageId);
 				}
 			}();
diff --git a/Engine/Plugins/Experimental/Mutable/Source/CustomizableObjectEditor/Private/MuCOE/GenerateMutableSource/GenerateMutableSourceSurface.cpp b/Engine/Plugins/Experimental/Mutable/Source/CustomizableObjectEditor/Private/MuCOE/GenerateMutableSource/GenerateMutableSourceSurface.cpp
--- a/Engine/Plugins/Experimental/Mutable/Source/CustomizableObjectEditor/Private/MuCOE/GenerateMutableSource/GenerateMutableSourceSurface.cpp
+++ b/Engine/Plugins/Experimental/Mutable/Source/CustomizableObjectEditor/Private/MuCOE/GenerateMutableSource/GenerateMutableSourceSurface.cpp
@@ -512,7 +512,7 @@ mu::Ptr<mu::NodeSurface> GenerateMutableSourceSurface(const UEdGraphPin * Pin, F
 				mu::NodeImagePtr GroupProjectionImg;
 				UTexture2D* GroupProjectionReferenceTexture = nullptr;
 				const FString ImageName = TypedNodeMat->GetParameterName(EMaterialParameterType::Texture, ImageIndex).ToString();
-				const FGuid ImageId = TypedNodeMat->GetParameterId(EMaterialParameterType::Texture, ImageIndex);
+				const FNodeMaterialParameterId ImageId = TypedNodeMat->GetParameterId(EMaterialParameterType::Texture, ImageIndex);
 
 				FString MaterialImageId = FGroupProjectorImageInfo::GenerateId(TypedNodeMat, ImageIndex);
 				bool bShareProjectionTexturesBetweenLODs = false;
@@ -531,8 +531,9 @@ mu::Ptr<mu::NodeSurface> GenerateMutableSourceSurface(const UEdGraphPin * Pin, F
 						// Get the reference texture
 						UTexture2D* ReferenceTexture = nullptr;
 						{
+							//TODO(Max) UE-220247: Add support for multilayer materials
 							GenerationContext.CurrentMaterialTableParameter = ImageName;
-							GenerationContext.CurrentMaterialTableParameterId = ImageId.ToString();
+							GenerationContext.CurrentMaterialTableParameterId = ImageId.ParameterId.ToString();
 
 							ReferenceTexture = GroupProjectionImg.get() ? GroupProjectionReferenceTexture : nullptr;
 						
@@ -921,7 +922,7 @@ mu::Ptr<mu::NodeSurface> GenerateMutableSourceSurface(const UEdGraphPin * Pin, F
 			bool bVectorPinConnected = VectorPin && FollowInputPin(*VectorPin);
 
 			FString VectorName = TypedNodeMat->GetParameterName(EMaterialParameterType::Vector, VectorIndex).ToString();
-			FGuid VectorId = TypedNodeMat->GetParameterId(EMaterialParameterType::Vector, VectorIndex);
+			FNodeMaterialParameterId VectorId = TypedNodeMat->GetParameterId(EMaterialParameterType::Vector, VectorIndex);
 
 			if (bVectorPinConnected)
 			{				
@@ -949,7 +950,7 @@ mu::Ptr<mu::NodeSurface> GenerateMutableSourceSurface(const UEdGraphPin * Pin, F
 			bool bScalarPinConnected = ScalarPin && FollowInputPin(*ScalarPin);
 
 			FString ScalarName = TypedNodeMat->GetParameterName(EMaterialParameterType::Scalar, ScalarIndex).ToString();
-			FGuid ScalarId = TypedNodeMat->GetParameterId(EMaterialParameterType::Scalar, ScalarIndex);
+			FNodeMaterialParameterId ScalarId = TypedNodeMat->GetParameterId(EMaterialParameterType::Scalar, ScalarIndex);
 
 			if (bScalarPinConnected)
 			{
@@ -1190,7 +1191,7 @@ mu::Ptr<mu::NodeSurface> GenerateMutableSourceSurface(const UEdGraphPin * Pin, F
 				
 				if (!ImageNode) // Else if
 				{
-					const FGuid ImageId = ParentMaterialNode->GetParameterId(EMaterialParameterType::Texture, ImageIndex);
+					const FNodeMaterialParameterId ImageId = ParentMaterialNode->GetParameterId(EMaterialParameterType::Texture, ImageIndex);
 					
 					if (TypedNodeExt->UsesImage(ImageId))
 					{
@@ -1359,7 +1360,7 @@ mu::Ptr<mu::NodeSurface> GenerateMutableSourceSurface(const UEdGraphPin * Pin, F
 			SurfNode->Textures.SetNum(NumImages);
 			for (int32 ImageIndex = 0; ImageIndex < NumImages; ++ImageIndex)
 			{
-				const FGuid ImageId = ParentMaterialNode->GetParameterId(EMaterialParameterType::Texture, ImageIndex);
+				const FNodeMaterialParameterId ImageId = ParentMaterialNode->GetParameterId(EMaterialParameterType::Texture, ImageIndex);
 
 				if (TypedNodeEdit->UsesImage(ImageId))
 				{
diff --git a/Engine/Plugins/Experimental/Mutable/Source/CustomizableObjectEditor/Private/MuCOE/Nodes/CustomizableObjectNodeCopyMaterial.cpp b/Engine/Plugins/Experimental/Mutable/Source/CustomizableObjectEditor/Private/MuCOE/Nodes/CustomizableObjectNodeCopyMaterial.cpp
--- a/Engine/Plugins/Experimental/Mutable/Source/CustomizableObjectEditor/Private/MuCOE/Nodes/CustomizableObjectNodeCopyMaterial.cpp
+++ b/Engine/Plugins/Experimental/Mutable/Source/CustomizableObjectEditor/Private/MuCOE/Nodes/CustomizableObjectNodeCopyMaterial.cpp
@@ -270,7 +270,7 @@ int32 UCustomizableObjectNodeCopyMaterial::GetNumParameters(EMaterialParameterTy
 }
 
 
-FGuid UCustomizableObjectNodeCopyMaterial::GetParameterId(EMaterialParameterType Type, int32 ParameterIndex) const
+FNodeMaterialParameterId UCustomizableObjectNodeCopyMaterial::GetParameterId(EMaterialParameterType Type, int32 ParameterIndex) const
 {
 	UCustomizableObjectNodeMaterial* NodeMaterial = GetMaterialNode();
 	check(NodeMaterial);
@@ -306,7 +306,7 @@ FText UCustomizableObjectNodeCopyMaterial::GetParameterLayerName(EMaterialParame
 }
 
 
-bool UCustomizableObjectNodeCopyMaterial::HasParameter(const FGuid& ParameterId) const
+bool UCustomizableObjectNodeCopyMaterial::HasParameter(const FNodeMaterialParameterId& ParameterId) const
 {
 	if (UCustomizableObjectNodeMaterial* NodeMaterial = GetMaterialNode())
 	{
@@ -332,7 +332,7 @@ const UEdGraphPin* UCustomizableObjectNodeCopyMaterial::GetParameterPin(EMateria
 }
 
 
-UEdGraphPin* UCustomizableObjectNodeCopyMaterial::GetParameterPin(const FGuid& ParameterId) const
+UEdGraphPin* UCustomizableObjectNodeCopyMaterial::GetParameterPin(const FNodeMaterialParameterId& ParameterId) const
 {
 	if (UCustomizableObjectNodeMaterial* NodeMaterial = GetMaterialNode())
 	{
diff --git a/Engine/Plugins/Experimental/Mutable/Source/CustomizableObjectEditor/Private/MuCOE/Nodes/CustomizableObjectNodeEditMaterial.cpp b/Engine/Plugins/Experimental/Mutable/Source/CustomizableObjectEditor/Private/MuCOE/Nodes/CustomizableObjectNodeEditMaterial.cpp
--- a/Engine/Plugins/Experimental/Mutable/Source/CustomizableObjectEditor/Private/MuCOE/Nodes/CustomizableObjectNodeEditMaterial.cpp
+++ b/Engine/Plugins/Experimental/Mutable/Source/CustomizableObjectEditor/Private/MuCOE/Nodes/CustomizableObjectNodeEditMaterial.cpp
@@ -7,7 +7,6 @@
 #include "MuCOE/CustomizableObjectEditor_Deprecated.h"
 #include "MuCOE/CustomizableObjectLayout.h"
 #include "MuCOE/EdGraphSchema_CustomizableObject.h"
-#include "MuCOE/Nodes/CustomizableObjectNodeMaterial.h"
 
 class FCustomizableObjectNodeParentedMaterial;
 class UCustomizableObjectNodeRemapPins;
@@ -27,14 +26,14 @@ void UCustomizableObjectNodeEditMaterial::AllocateDefaultPins(UCustomizableObjec
 			if (ParentMaterialNode->IsImageMutableMode(ImageIndex))
 			{
 				UCustomizableObjectNodeEditMaterialPinEditImageData* PinEditImageData = NewObject<UCustomizableObjectNodeEditMaterialPinEditImageData>(this);
-				PinEditImageData->ImageId = ParentMaterialNode->GetParameterId(EMaterialParameterType::Texture, ImageIndex);
+				PinEditImageData->ImageParamId = ParentMaterialNode->GetParameterId(EMaterialParameterType::Texture, ImageIndex);
 			
 				const FName ImageName = ParentMaterialNode->GetParameterName(EMaterialParameterType::Texture, ImageIndex);
 				UEdGraphPin* PinImage = CustomCreatePin(EGPD_Input, Schema->PC_Image, ImageName, PinEditImageData);
 				PinImage->bHidden = true;
 				PinImage->bDefaultValueIsIgnored = true;
 
-				PinsParameter.Add(PinEditImageData->ImageId, FEdGraphPinReference(PinImage));
+				PinsParameterMap.Add(PinEditImageData->ImageParamId, FEdGraphPinReference(PinImage));
 
 				FString PinMaskName = ImageName.ToString() + " Mask";
 				UEdGraphPin* PinMask = CustomCreatePin(EGPD_Input, Schema->PC_Image, *PinMaskName);
@@ -50,7 +49,7 @@ void UCustomizableObjectNodeEditMaterial::AllocateDefaultPins(UCustomizableObjec
 }
 
 
-const UEdGraphPin* UCustomizableObjectNodeEditMaterial::GetUsedImageMaskPin(const FGuid& ImageId) const
+const UEdGraphPin* UCustomizableObjectNodeEditMaterial::GetUsedImageMaskPin(const FNodeMaterialParameterId& ImageId) const
 {
 	if (const UEdGraphPin* Pin = GetUsedImagePin(ImageId))
 	{
@@ -165,7 +164,7 @@ void UCustomizableObjectNodeEditMaterial::BackwardsCompatibleFixup()
 				}
 				
 				UCustomizableObjectNodeEditMaterialPinEditImageData*  PinEditImageData = NewObject<UCustomizableObjectNodeEditMaterialPinEditImageData>(this);
-				PinEditImageData->ImageId = FGuid::NewGuid();
+				PinEditImageData->ImageId_DEPRECATED = FGuid::NewGuid();
 				PinEditImageData->PinMask = PinMask;
 				
 				// Search for the Image Id the Edit pin was referring to.
@@ -174,7 +173,7 @@ void UCustomizableObjectNodeEditMaterial::BackwardsCompatibleFixup()
 				{
 					if (ParentMaterial->GetParameterName(EMaterialParameterType::Texture, ImageIndex).ToString() == Image.Name)
 					{
-						PinEditImageData->ImageId = ParentMaterial->GetParameterId(EMaterialParameterType::Texture, ImageIndex);
+						PinEditImageData->ImageId_DEPRECATED = ParentMaterial->GetParameterId(EMaterialParameterType::Texture, ImageIndex).ParameterId;
 						break;
 					}
 				}
@@ -194,7 +193,7 @@ void UCustomizableObjectNodeEditMaterial::BackwardsCompatibleFixup()
 		{
 			if (const UCustomizableObjectNodeEditMaterialPinEditImageData* PinData = Cast<UCustomizableObjectNodeEditMaterialPinEditImageData>(GetPinData(*Pin)))
 			{
-				PinsParameter.Add(PinData->ImageId, FEdGraphPinReference(Pin));
+				PinsParameter_DEPRECATED.Add(PinData->ImageId_DEPRECATED, FEdGraphPinReference(Pin));
 			}
 		}
 	}
@@ -208,6 +207,11 @@ void UCustomizableObjectNodeEditMaterial::BackwardsCompatibleFixup()
 	{
 		ReconstructNode();
 	}
+
+	if (CustomizableObjectCustomVersion < FCustomizableObjectCustomVersion::FixedMultilayerMaterialIds)
+	{
+		ReconstructNode();
+	}
 }
 
 
@@ -272,9 +276,9 @@ UCustomizableObjectNode& UCustomizableObjectNodeEditMaterial::GetNode()
 }
 
 
-TMap<FGuid, FEdGraphPinReference>& UCustomizableObjectNodeEditMaterial::GetPinsParameter()
+TMap<FNodeMaterialParameterId, FEdGraphPinReference>& UCustomizableObjectNodeEditMaterial::GetPinsParameter()
 {
-	return PinsParameter;	
+	return PinsParameterMap;	
 }
 
 
diff --git a/Engine/Plugins/Experimental/Mutable/Source/CustomizableObjectEditor/Private/MuCOE/Nodes/CustomizableObjectNodeEditMaterialBase.cpp b/Engine/Plugins/Experimental/Mutable/Source/CustomizableObjectEditor/Private/MuCOE/Nodes/CustomizableObjectNodeEditMaterialBase.cpp
--- a/Engine/Plugins/Experimental/Mutable/Source/CustomizableObjectEditor/Private/MuCOE/Nodes/CustomizableObjectNodeEditMaterialBase.cpp
+++ b/Engine/Plugins/Experimental/Mutable/Source/CustomizableObjectEditor/Private/MuCOE/Nodes/CustomizableObjectNodeEditMaterialBase.cpp
@@ -1,6 +1,7 @@
 // Copyright Epic Games, Inc. All Rights Reserved.
 
 #include "MuCOE/Nodes/CustomizableObjectNodeEditMaterialBase.h"
+#include "MuCOE/Nodes/CustomizableObjectNodeMaterialBase.h"
 
 #include "MuCOE/CustomizableObjectGraph.h"
 #include "MuCOE/ICustomizableObjectEditor.h"
diff --git a/Engine/Plugins/Experimental/Mutable/Source/CustomizableObjectEditor/Private/MuCOE/Nodes/CustomizableObjectNodeExtendMaterial.cpp b/Engine/Plugins/Experimental/Mutable/Source/CustomizableObjectEditor/Private/MuCOE/Nodes/CustomizableObjectNodeExtendMaterial.cpp
--- a/Engine/Plugins/Experimental/Mutable/Source/CustomizableObjectEditor/Private/MuCOE/Nodes/CustomizableObjectNodeExtendMaterial.cpp
+++ b/Engine/Plugins/Experimental/Mutable/Source/CustomizableObjectEditor/Private/MuCOE/Nodes/CustomizableObjectNodeExtendMaterial.cpp
@@ -66,18 +66,23 @@ void UCustomizableObjectNodeExtendMaterial::BackwardsCompatibleFixup()
 				{
 					if (ParentMaterial->GetParameterName(EMaterialParameterType::Texture, ImageIndex).ToString() == Image.Name)
 					{
-						ImageId = ParentMaterial->GetParameterId(EMaterialParameterType::Texture, ImageIndex);
+						ImageId = ParentMaterial->GetParameterId(EMaterialParameterType::Texture, ImageIndex).ParameterId;
 						break;
 					}
 				}
-                
-				PinsParameter.Add(ImageId, FEdGraphPinReference(ImagePin));                
+
+				PinsParameter_DEPRECATED.Add(ImageId, FEdGraphPinReference(ImagePin));                
 			}
 		}
-		
+
 		Images_DEPRECATED.Empty();
 		ReconstructNode();
 	}
+
+	if (CustomizableObjectCustomVersion < FCustomizableObjectCustomVersion::FixedMultilayerMaterialIds)
+	{
+		ReconstructNode();
+	}
 }
 
 
@@ -117,9 +122,9 @@ void UCustomizableObjectNodeExtendMaterial::AllocateDefaultPins(UCustomizableObj
 				UEdGraphPin* PinImage = CustomCreatePin(EGPD_Input, Schema->PC_Image, ImageName);
 				PinImage->bDefaultValueIsIgnored = true;
 
-				const FGuid ImageId = ParentMaterialNode->GetParameterId(EMaterialParameterType::Texture, ImageIndex);
+				const FNodeMaterialParameterId ImageId = ParentMaterialNode->GetParameterId(EMaterialParameterType::Texture, ImageIndex);
 				
-				PinsParameter.Add(ImageId, FEdGraphPinReference(PinImage));
+				PinsParameterMap.Add(ImageId, FEdGraphPinReference(PinImage));
 			}
 		}
 	}
@@ -243,9 +248,9 @@ UCustomizableObjectNode& UCustomizableObjectNodeExtendMaterial::GetNode()
 }
 
 
-TMap<FGuid, FEdGraphPinReference>& UCustomizableObjectNodeExtendMaterial::GetPinsParameter()
+TMap<FNodeMaterialParameterId, FEdGraphPinReference>& UCustomizableObjectNodeExtendMaterial::GetPinsParameter()
 {
-	return PinsParameter;
+	return PinsParameterMap;
 }
 
 
diff --git a/Engine/Plugins/Experimental/Mutable/Source/CustomizableObjectEditor/Private/MuCOE/Nodes/CustomizableObjectNodeMaterial.cpp b/Engine/Plugins/Experimental/Mutable/Source/CustomizableObjectEditor/Private/MuCOE/Nodes/CustomizableObjectNodeMaterial.cpp
--- a/Engine/Plugins/Experimental/Mutable/Source/CustomizableObjectEditor/Private/MuCOE/Nodes/CustomizableObjectNodeMaterial.cpp
+++ b/Engine/Plugins/Experimental/Mutable/Source/CustomizableObjectEditor/Private/MuCOE/Nodes/CustomizableObjectNodeMaterial.cpp
@@ -45,7 +45,7 @@ bool UCustomizableObjectNodeMaterialRemapPinsByName::Equal(const UCustomizableOb
 	const UCustomizableObjectNodeMaterialPinDataParameter* PinDataNewPin = Cast<UCustomizableObjectNodeMaterialPinDataParameter>(Node.GetPinData(NewPin));
 	if (PinDataOldPin && PinDataNewPin)
 	{
-		return PinDataOldPin->ParameterId == PinDataNewPin->ParameterId && (!OldPin.LinkedTo.Num() || OldPin.PinType == NewPin.PinType); // Pin type must match only if it was connected
+		return PinDataOldPin->MaterialParameterId == PinDataNewPin->MaterialParameterId && (!OldPin.LinkedTo.Num() || OldPin.PinType == NewPin.PinType); // Pin type must match only if it was connected
 	}
 	else
 	{
@@ -207,9 +207,9 @@ FName UCustomizableObjectNodeMaterial::GetImagePinMode(const UEdGraphPin& Pin) c
 
 int32 UCustomizableObjectNodeMaterial::GetImageUVLayoutFromMaterial(const int32 ImageIndex) const
 {
-	const FGuid ImageId = GetParameterId(EMaterialParameterType::Texture, ImageIndex);
+	const FNodeMaterialParameterId ImageId = GetParameterId(EMaterialParameterType::Texture, ImageIndex);
 
-	if (const int32 TextureCoordinate = GetExpressionTextureCoordinate(Material->GetMaterial(), ImageId);
+	if (const int32 TextureCoordinate = GetExpressionTextureCoordinate(Material->GetMaterial(), ImageId.ParameterId);
 		TextureCoordinate >= 0)
 	{
 		return TextureCoordinate;
@@ -227,7 +227,7 @@ int32 UCustomizableObjectNodeMaterial::GetImageUVLayoutFromMaterial(const int32
 	{
 		for (const TObjectPtr<UMaterialFunctionInterface>& Layer : *MaterialFunctionInterface)
 		{
-			if (const int32 TextureCoordinate = GetExpressionTextureCoordinate(Layer->GetPreviewMaterial()->GetMaterial(), ImageId); TextureCoordinate >= 0)
+			if (const int32 TextureCoordinate = GetExpressionTextureCoordinate(Layer->GetPreviewMaterial()->GetMaterial(), ImageId.ParameterId); TextureCoordinate >= 0)
 			{
 				return TextureCoordinate;
 			}	
@@ -389,7 +389,7 @@ void UCustomizableObjectNodeMaterial::BackwardsCompatibleFixup()
 			}
 
 			UCustomizableObjectNodeMaterialPinDataImage* PinData = NewObject<UCustomizableObjectNodeMaterialPinDataImage>(this);
-			PinData->ParameterId = FGuid::NewGuid();
+			PinData->ParameterId_DEPRECATED = FGuid::NewGuid();
 			PinData->ReferenceTexture = Image.ReferenceTexture;
 
 			// Find referenced Material Parameter
@@ -398,7 +398,7 @@ void UCustomizableObjectNodeMaterial::BackwardsCompatibleFixup()
 			{
 				if (GetParameterName(EMaterialParameterType::Texture, ParameterIndex).ToString() == Image.Name)
 				{
-					PinData->ParameterId = GetParameterId(EMaterialParameterType::Texture, ParameterIndex);
+					PinData->ParameterId_DEPRECATED = GetParameterId(EMaterialParameterType::Texture, ParameterIndex).ParameterId;
 
 					if (Image.UVLayout == -1)
 					{
@@ -438,7 +438,7 @@ void UCustomizableObjectNodeMaterial::BackwardsCompatibleFixup()
 			}
 			
 			UCustomizableObjectNodeMaterialPinDataVector* PinData = NewObject<UCustomizableObjectNodeMaterialPinDataVector>(this);
-			PinData->ParameterId = FGuid::NewGuid();
+			PinData->ParameterId_DEPRECATED = FGuid::NewGuid();
 			
 			// Find referenced Material Parameter
 			const int32 NumParameters = GetNumParameters(EMaterialParameterType::Vector);
@@ -446,7 +446,7 @@ void UCustomizableObjectNodeMaterial::BackwardsCompatibleFixup()
 			{
 				if (GetParameterName(EMaterialParameterType::Vector, ParameterIndex).ToString() == Vector.Name)
 				{
-					PinData->ParameterId = GetParameterId(EMaterialParameterType::Vector, ParameterIndex);
+					PinData->ParameterId_DEPRECATED = GetParameterId(EMaterialParameterType::Vector, ParameterIndex).ParameterId;
 					break;
 				}
 			}
@@ -464,7 +464,7 @@ void UCustomizableObjectNodeMaterial::BackwardsCompatibleFixup()
 			}
 			
 			UCustomizableObjectNodeMaterialPinDataScalar* PinData = NewObject<UCustomizableObjectNodeMaterialPinDataScalar>(this);
-			PinData->ParameterId = FGuid::NewGuid();
+			PinData->ParameterId_DEPRECATED = FGuid::NewGuid();
 
 			// Find referenced Material Parameter
 			const int32 NumParameters = GetNumParameters(EMaterialParameterType::Scalar);
@@ -472,7 +472,7 @@ void UCustomizableObjectNodeMaterial::BackwardsCompatibleFixup()
 			{
 				if (GetParameterName(EMaterialParameterType::Scalar, ParameterIndex).ToString() == Scalar.Name)
 				{
-					PinData->ParameterId = GetParameterId(EMaterialParameterType::Scalar, ParameterIndex);
+					PinData->ParameterId_DEPRECATED = GetParameterId(EMaterialParameterType::Scalar, ParameterIndex).ParameterId;
 					break;
 				}
 			}
@@ -509,7 +509,7 @@ void UCustomizableObjectNodeMaterial::BackwardsCompatibleFixup()
 					}
 				}(this);
 				
-				PinData->ParameterId = FGuid::NewGuid();
+				PinData->ParameterId_DEPRECATED = FGuid::NewGuid();
 				
 				AddPinData(*Pin, *PinData);
 			}
@@ -529,7 +529,7 @@ void UCustomizableObjectNodeMaterial::BackwardsCompatibleFixup()
 		{
 			if (const UCustomizableObjectNodeMaterialPinDataParameter* PinData = Cast<UCustomizableObjectNodeMaterialPinDataParameter>(GetPinData(*Pin)))
 			{
-				PinsParameter.Add(PinData->ParameterId, FEdGraphPinReference(Pin));
+				PinsParameter_DEPRECATED.Add(PinData->ParameterId_DEPRECATED, FEdGraphPinReference(Pin));
 			}
 		}
 	}
@@ -592,7 +592,7 @@ void UCustomizableObjectNodeMaterial::BackwardsCompatibleFixup()
 			{
 				if (const UEdGraphPin* ImagePin = GetParameterPin(EMaterialParameterType::Texture, ImageIndex))
 				{
-					FGuid ParameterId = GetParameterId(EMaterialParameterType::Texture, ImageIndex);
+					FGuid ParameterId = GetParameterId(EMaterialParameterType::Texture, ImageIndex).ParameterId;
 
 					TArray<FMaterialParameterInfo> TextureParameterInfo;
 					TArray<FGuid> TextureGuids;
@@ -677,6 +677,29 @@ void UCustomizableObjectNodeMaterial::BackwardsCompatibleFixup()
 			}
 		}
 	}
+
+	if (CustomizableObjectCustomVersion < FCustomizableObjectCustomVersion::FixedMultilayerMaterialIds)
+	{
+		if (Material->GetCachedExpressionData().bHasMaterialLayers)
+		{
+			// Needed since we can not get the layer index of repeated parameters
+			Super::ReconstructNode();
+		}
+		else
+		{
+			for (TMap<FGuid, FEdGraphPinReference>::TIterator It = PinsParameter_DEPRECATED.CreateIterator(); It; ++It)
+			{
+				PinsParameterMap.Add({ It.Key(), -1 }, It.Value());
+
+				// Move pin data id info to the new struct
+				if (UCustomizableObjectNodeMaterialPinDataParameter* PinData = Cast<UCustomizableObjectNodeMaterialPinDataParameter>(GetPinData(*It.Value().Get())))
+				{
+					PinData->MaterialParameterId.LayerIndex = INDEX_NONE;
+					PinData->MaterialParameterId.ParameterId = PinData->ParameterId_DEPRECATED;
+				}
+			}
+		}
+	}
 }
 
 
@@ -926,7 +949,7 @@ int32 UCustomizableObjectNodeMaterial::GetNumParameters(const EMaterialParameter
 }
 
 
-FGuid UCustomizableObjectNodeMaterial::GetParameterId(const EMaterialParameterType Type, const int32 ParameterIndex) const
+FNodeMaterialParameterId UCustomizableObjectNodeMaterial::GetParameterId(const EMaterialParameterType Type, const int32 ParameterIndex) const
 {
 	const FMaterialCachedExpressionData& Data = Material->GetCachedExpressionData();
 
@@ -934,11 +957,14 @@ FGuid UCustomizableObjectNodeMaterial::GetParameterId(const EMaterialParameterTy
 	{
 		if (Data.EditorOnlyData->EditorEntries[(int32)Type].EditorInfo.Num() != 0)
 		{
-			return Data.EditorOnlyData->EditorEntries[(int32)Type].EditorInfo[ParameterIndex].ExpressionGuid;
+			const FGuid ParameterId = Data.EditorOnlyData->EditorEntries[(int32)Type].EditorInfo[ParameterIndex].ExpressionGuid;
+			const int32 LayerIndex = GetParameterLayerIndex(Type, ParameterIndex);
+
+			return { ParameterId, LayerIndex };
 		}
 	}
 	
-	return FGuid();
+	return FNodeMaterialParameterId();
 }
 
 
@@ -1002,7 +1028,7 @@ FText UCustomizableObjectNodeMaterial::GetParameterLayerName(const EMaterialPara
 }
 
 
-bool UCustomizableObjectNodeMaterial::HasParameter(const FGuid& ParameterId) const
+bool UCustomizableObjectNodeMaterial::HasParameter(const FNodeMaterialParameterId& ParameterId) const
 {
 	if (!Material)
 	{
@@ -1022,9 +1048,12 @@ bool UCustomizableObjectNodeMaterial::HasParameter(const FGuid& ParameterId) con
 		for (TSet<FMaterialParameterInfo>::TConstIterator It(Entry.ParameterInfoSet); It; ++It)
 		{
 			const int32 IteratorIndex = It.GetId().AsInteger();
-			const FGuid& ParamGUid = Data.EditorOnlyData->EditorEntries[(int32)Type].EditorInfo[IteratorIndex].ExpressionGuid;
+			
+			const FGuid& ParamGuid = Data.EditorOnlyData->EditorEntries[(int32)Type].EditorInfo[IteratorIndex].ExpressionGuid;
+			const int32 LayerIndex = GetParameterLayerIndex(Type, IteratorIndex);
+			const FNodeMaterialParameterId ParamId = { ParamGuid, LayerIndex };
 
-			if (ParamGUid == ParameterId)
+			if (ParamId == ParameterId)
 			{
 				return true;
 			}
@@ -1037,15 +1066,15 @@ bool UCustomizableObjectNodeMaterial::HasParameter(const FGuid& ParameterId) con
 
 UEdGraphPin* UCustomizableObjectNodeMaterial::GetParameterPin(const EMaterialParameterType Type, const int32 ParameterIndex) const
 {
-	const FGuid ParameterId = GetParameterId(Type, ParameterIndex);
+	const FNodeMaterialParameterId ParameterId = GetParameterId(Type, ParameterIndex);
 	
 	return GetParameterPin(ParameterId);
 }
 
 
-UEdGraphPin* UCustomizableObjectNodeMaterial::GetParameterPin(const FGuid& ParameterId) const
+UEdGraphPin* UCustomizableObjectNodeMaterial::GetParameterPin(const FNodeMaterialParameterId& ParameterId) const
 {
-	if (const FEdGraphPinReference* Result = PinsParameter.Find(ParameterId))
+	if (const FEdGraphPinReference* Result = PinsParameterMap.Find(ParameterId))
 	{
 		return Result->Get();
 	}
@@ -1101,7 +1130,7 @@ TSharedPtr<IDetailsView> UCustomizableObjectNodeMaterial::CustomizePinDetails(co
 
 bool UCustomizableObjectNodeMaterial::CustomRemovePin(UEdGraphPin& Pin)
 {
-	for (TMap<FGuid, FEdGraphPinReference>::TIterator Iterator = PinsParameter.CreateIterator(); Iterator; ++Iterator)
+	for (TMap<FNodeMaterialParameterId, FEdGraphPinReference>::TIterator Iterator = PinsParameterMap.CreateIterator(); Iterator; ++Iterator)
 	{
 		if (Iterator->Value.Get() == &Pin)
 		{
@@ -1171,7 +1200,7 @@ bool UCustomizableObjectNodeMaterial::RealMaterialDataHasChanged() const
 	{
 		if (const UCustomizableObjectNodeMaterialPinDataParameter* PinData = Cast<UCustomizableObjectNodeMaterialPinDataParameter>(GetPinData(*Pin)))
 		{
-			if (!HasParameter(PinData->ParameterId) &&
+			if (!HasParameter(PinData->MaterialParameterId) &&
 				(FollowInputPin(*Pin) || !PinData->IsDefault()))
 			{
 				return true;
@@ -1250,7 +1279,7 @@ UCustomizableObjectNodeMaterialPinDataParameter* UCustomizableObjectNodeMaterial
 		check(false); // Parameter type not contemplated.
 	}
 
-	PinData->ParameterId = GetParameterId(Type, ParameterIndex);
+	PinData->MaterialParameterId = GetParameterId(Type, ParameterIndex);
 
 	return PinData;
 }
@@ -1294,7 +1323,7 @@ void UCustomizableObjectNodeMaterial::AllocateDefaultParameterPins(const EMateri
 		Pin->bHidden = true;
 		Pin->bDefaultValueIsIgnored = true;
 
-		PinsParameter.Add(PinData->ParameterId, FEdGraphPinReference(Pin));
+		PinsParameterMap.Add(PinData->MaterialParameterId, FEdGraphPinReference(Pin));
 	}
 }
 
diff --git a/Engine/Plugins/Experimental/Mutable/Source/CustomizableObjectEditor/Private/MuCOE/Nodes/CustomizableObjectNodeUseMaterial.cpp b/Engine/Plugins/Experimental/Mutable/Source/CustomizableObjectEditor/Private/MuCOE/Nodes/CustomizableObjectNodeUseMaterial.cpp
--- a/Engine/Plugins/Experimental/Mutable/Source/CustomizableObjectEditor/Private/MuCOE/Nodes/CustomizableObjectNodeUseMaterial.cpp
+++ b/Engine/Plugins/Experimental/Mutable/Source/CustomizableObjectEditor/Private/MuCOE/Nodes/CustomizableObjectNodeUseMaterial.cpp
@@ -16,7 +16,7 @@ bool FCustomizableObjectNodeUseMaterial::IsNodeOutDatedAndNeedsRefreshWork()
 	UCustomizableObjectNode& Node = GetNode();
 	const FCustomizableObjectNodeParentedMaterial& NodeParentedMaterial = GetNodeParentedMaterial();
 
-	const TMap<FGuid, FEdGraphPinReference>& PinsParameter = GetPinsParameter();
+	const TMap<FNodeMaterialParameterId, FEdGraphPinReference>& PinsParameter = GetPinsParameter();
 	
 	const bool bOutdated = [&]()
 	{
@@ -27,7 +27,7 @@ bool FCustomizableObjectNodeUseMaterial::IsNodeOutDatedAndNeedsRefreshWork()
 				return true;
 			}
 
-			for (const TTuple<FGuid, FEdGraphPinReference> Pair : PinsParameter)
+			for (const TTuple<FNodeMaterialParameterId, FEdGraphPinReference> Pair : PinsParameter)
 			{
 				const UEdGraphPin& Pin = *Pair.Value.Get();
 				if (!IsPinOrphan(Pin))
@@ -111,9 +111,9 @@ void FCustomizableObjectNodeUseMaterial::PinConnectionListChangedWork(UEdGraphPi
 
 void FCustomizableObjectNodeUseMaterial::CustomRemovePinWork(UEdGraphPin& Pin)
 {
-	TMap<FGuid, FEdGraphPinReference>& PinsParameter = GetPinsParameter();
+	TMap<FNodeMaterialParameterId, FEdGraphPinReference>& PinsParameter = GetPinsParameter();
 	
-	for(TMap<FGuid, FEdGraphPinReference>::TIterator It = PinsParameter.CreateIterator(); It; ++It)
+	for(TMap<FNodeMaterialParameterId, FEdGraphPinReference>::TIterator It = PinsParameter.CreateIterator(); It; ++It)
 	{
 		if (It.Value().Get() == &Pin) // We could improve performance if FEdGraphPinReference exposed the pin id.
 		{
@@ -130,13 +130,13 @@ const UCustomizableObjectNode& FCustomizableObjectNodeUseMaterial::GetNode() con
 }
 
 
-const TMap<FGuid, FEdGraphPinReference>& FCustomizableObjectNodeUseMaterial::GetPinsParameter() const
+const TMap<FNodeMaterialParameterId, FEdGraphPinReference>& FCustomizableObjectNodeUseMaterial::GetPinsParameter() const
 {
 	return const_cast<FCustomizableObjectNodeUseMaterial*>(this)->GetPinsParameter();
 }
 
 
-bool FCustomizableObjectNodeUseMaterial::UsesImage(const FGuid& ImageId) const
+bool FCustomizableObjectNodeUseMaterial::UsesImage(const FNodeMaterialParameterId& ImageId) const
 {
 	if (const UEdGraphPin* Pin = GetUsedImagePin(ImageId))
 	{
@@ -149,9 +149,9 @@ bool FCustomizableObjectNodeUseMaterial::UsesImage(const FGuid& ImageId) const
 }
 
 
-const UEdGraphPin* FCustomizableObjectNodeUseMaterial::GetUsedImagePin(const FGuid& ImageId) const
+const UEdGraphPin* FCustomizableObjectNodeUseMaterial::GetUsedImagePin(const FNodeMaterialParameterId& ImageId) const
 {
-	const TMap<FGuid, FEdGraphPinReference>& PinsParameter = GetPinsParameter();
+	const TMap<FNodeMaterialParameterId, FEdGraphPinReference>& PinsParameter = GetPinsParameter();
 
 	if (const FEdGraphPinReference* PinReference = PinsParameter.Find(ImageId))
 	{
diff --git a/Engine/Plugins/Experimental/Mutable/Source/CustomizableObjectEditor/Public/MuCOE/Nodes/CustomizableObjectNodeCopyMaterial.h b/Engine/Plugins/Experimental/Mutable/Source/CustomizableObjectEditor/Public/MuCOE/Nodes/CustomizableObjectNodeCopyMaterial.h
--- a/Engine/Plugins/Experimental/Mutable/Source/CustomizableObjectEditor/Public/MuCOE/Nodes/CustomizableObjectNodeCopyMaterial.h
+++ b/Engine/Plugins/Experimental/Mutable/Source/CustomizableObjectEditor/Public/MuCOE/Nodes/CustomizableObjectNodeCopyMaterial.h
@@ -39,13 +39,13 @@ public:
 	virtual TArray<FString> GetTags() const override;
 	virtual UEdGraphPin* GetMaterialAssetPin() const override;
 	virtual int32 GetNumParameters(EMaterialParameterType Type) const override;
-	virtual FGuid GetParameterId(EMaterialParameterType Type, int32 ParameterIndex) const override;
+	virtual FNodeMaterialParameterId GetParameterId(EMaterialParameterType Type, int32 ParameterIndex) const override;
 	virtual FName GetParameterName(EMaterialParameterType Type, int32 ParameterIndex) const override;
 	virtual int32 GetParameterLayerIndex(EMaterialParameterType Type, int32 ParameterIndex) const override;
 	virtual FText GetParameterLayerName(EMaterialParameterType Type, int32 ParameterIndex) const override;
-	virtual bool HasParameter(const FGuid& ParameterId) const override;
+	virtual bool HasParameter(const FNodeMaterialParameterId& ParameterId) const override;
 	virtual const UEdGraphPin* GetParameterPin(EMaterialParameterType Type, int32 ParameterIndex) const override;
-	virtual UEdGraphPin* GetParameterPin(const FGuid& ParameterId) const override;
+	virtual UEdGraphPin* GetParameterPin(const FNodeMaterialParameterId& ParameterId) const override;
 	virtual bool IsImageMutableMode(int32 ImageIndex) const override;
 	virtual bool IsImageMutableMode(const UEdGraphPin& Pin) const override;
 	virtual UTexture2D* GetImageReferenceTexture(int32 ImageIndex) const override;
diff --git a/Engine/Plugins/Experimental/Mutable/Source/CustomizableObjectEditor/Public/MuCOE/Nodes/CustomizableObjectNodeEditMaterial.h b/Engine/Plugins/Experimental/Mutable/Source/CustomizableObjectEditor/Public/MuCOE/Nodes/CustomizableObjectNodeEditMaterial.h
--- a/Engine/Plugins/Experimental/Mutable/Source/CustomizableObjectEditor/Public/MuCOE/Nodes/CustomizableObjectNodeEditMaterial.h
+++ b/Engine/Plugins/Experimental/Mutable/Source/CustomizableObjectEditor/Public/MuCOE/Nodes/CustomizableObjectNodeEditMaterial.h
@@ -4,11 +4,13 @@
 
 #include "MuCOE/Nodes/CustomizableObjectNodeEditLayoutBlocks.h"
 #include "MuCOE/Nodes/CustomizableObjectNodeUseMaterial.h"
+#include "MuCOE/Nodes/CustomizableObjectNodeMaterial.h"
 
 #include "CustomizableObjectNodeEditMaterial.generated.h"
 
 namespace ENodeTitleType { enum Type : int; }
 struct FCustomizableObjectNodeEditMaterialImage;
+struct FNodeMaterialParameterId;
 
 class FArchive;
 class FCustomizableObjectNodeParentedMaterial;
@@ -24,11 +26,16 @@ class CUSTOMIZABLEOBJECTEDITOR_API UCustomizableObjectNodeEditMaterialPinEditIma
 	GENERATED_BODY()
 
 public:
+
+	/** Parameter FGuid + Layer Index */
 	UPROPERTY()
-	FGuid ImageId;
+	FNodeMaterialParameterId ImageParamId;
 
 	UPROPERTY();
 	FEdGraphPinReference PinMask;
+
+	UPROPERTY()
+	FGuid ImageId_DEPRECATED;
 };
 
 
@@ -69,23 +76,27 @@ public:
 	// FCustomizableObjectNodeUseMaterial interface
 	virtual UCustomizableObjectNode& GetNode() override;
 	virtual FCustomizableObjectNodeParentedMaterial& GetNodeParentedMaterial() override;
-	virtual TMap<FGuid, FEdGraphPinReference>& GetPinsParameter() override;
+	virtual TMap<FNodeMaterialParameterId, FEdGraphPinReference>& GetPinsParameter() override;
 	
 	
 	/** Returns the Image mask pin of the given Image that will be edited.
 	 *
 	 * @returns Always returns a valid pin if EditsImage(const FGuid&) returns true. */
-	const UEdGraphPin* GetUsedImageMaskPin(const FGuid& ImageId) const;
+	const UEdGraphPin* GetUsedImageMaskPin(const FNodeMaterialParameterId& ImageId) const;
 
 	// Function to select all the layout blocks. Called when a parameter is reset (ParentMaterial, ParentLayoutIndex)
 	void SelectAllLayoutBlocks();
 
 private:
-	/** Relates a Parameter id to a Pin. Only used to improve performance. */
+
+	/** Relates a Parameter id (and layer if is a layered material) to a Pin. Only used to improve performance. */
 	UPROPERTY()
-	TMap<FGuid, FEdGraphPinReference> PinsParameter;
-	
+	TMap<FNodeMaterialParameterId, FEdGraphPinReference> PinsParameterMap;
+
 	// Deprecated properties
+	UPROPERTY()
+	TMap<FGuid, FEdGraphPinReference> PinsParameter_DEPRECATED;
+	
 	UPROPERTY()
 	TArray<FCustomizableObjectNodeEditMaterialImage> Images_DEPRECATED;
 	
diff --git a/Engine/Plugins/Experimental/Mutable/Source/CustomizableObjectEditor/Public/MuCOE/Nodes/CustomizableObjectNodeExtendMaterial.h b/Engine/Plugins/Experimental/Mutable/Source/CustomizableObjectEditor/Public/MuCOE/Nodes/CustomizableObjectNodeExtendMaterial.h
--- a/Engine/Plugins/Experimental/Mutable/Source/CustomizableObjectEditor/Public/MuCOE/Nodes/CustomizableObjectNodeExtendMaterial.h
+++ b/Engine/Plugins/Experimental/Mutable/Source/CustomizableObjectEditor/Public/MuCOE/Nodes/CustomizableObjectNodeExtendMaterial.h
@@ -63,7 +63,7 @@ public:
 
 	// FCustomizableObjectNodeUseMaterial interface
 	virtual FCustomizableObjectNodeParentedMaterial& GetNodeParentedMaterial() override;
-	virtual TMap<FGuid, FEdGraphPinReference>& GetPinsParameter() override;
+	virtual TMap<FNodeMaterialParameterId, FEdGraphPinReference>& GetPinsParameter() override;
 	virtual UEdGraphPin* OutputPin() const override;
 
 	// Own interface
@@ -78,11 +78,14 @@ private:
 	UPROPERTY()
 	FGuid ParentMaterialNodeId;
 	
-	/** Relates a Parameter id to a Pin. Only used to improve performance. */
-   	UPROPERTY()
-   	TMap<FGuid, FEdGraphPinReference> PinsParameter;
+	/** Relates a Parameter id (and layer if is a layered material) to a Pin. Only used to improve performance. */
+	UPROPERTY()
+	TMap<FNodeMaterialParameterId, FEdGraphPinReference> PinsParameterMap;
 
 	// Deprecated properties
+	UPROPERTY()
+	TMap<FGuid, FEdGraphPinReference> PinsParameter_DEPRECATED;
+
 	UPROPERTY()
 	TArray<FCustomizableObjectNodeExtendMaterialImage> Images_DEPRECATED;
 };
diff --git a/Engine/Plugins/Experimental/Mutable/Source/CustomizableObjectEditor/Public/MuCOE/Nodes/CustomizableObjectNodeMaterial.h b/Engine/Plugins/Experimental/Mutable/Source/CustomizableObjectEditor/Public/MuCOE/Nodes/CustomizableObjectNodeMaterial.h
--- a/Engine/Plugins/Experimental/Mutable/Source/CustomizableObjectEditor/Public/MuCOE/Nodes/CustomizableObjectNodeMaterial.h
+++ b/Engine/Plugins/Experimental/Mutable/Source/CustomizableObjectEditor/Public/MuCOE/Nodes/CustomizableObjectNodeMaterial.h
@@ -50,9 +50,14 @@ class CUSTOMIZABLEOBJECTEDITOR_API UCustomizableObjectNodeMaterialPinDataParamet
 	GENERATED_BODY()
 
 public:
+
+	/** Parameter id + layer index */
+	UPROPERTY()
+	FNodeMaterialParameterId MaterialParameterId;
+
 	/** Texture Parameter Id. */
 	UPROPERTY()
-	FGuid ParameterId;
+	FGuid ParameterId_DEPRECATED;
 
 	/** Returns true if all properties are in its default state. */
 	virtual bool IsDefault() const;
@@ -138,13 +143,13 @@ public:
 	virtual FName GetMeshComponentName() const override;
 	virtual TArray<FString> GetTags() const override;
 	virtual int32 GetNumParameters(EMaterialParameterType Type) const override;
-	virtual FGuid GetParameterId(EMaterialParameterType Type, int32 ParameterIndex) const override;
+	virtual FNodeMaterialParameterId GetParameterId(EMaterialParameterType Type, int32 ParameterIndex) const override;
 	virtual FName GetParameterName(EMaterialParameterType Type, int32 ParameterIndex) const override;
 	virtual int32 GetParameterLayerIndex(EMaterialParameterType Type, int32 ParameterIndex) const override;
 	virtual FText GetParameterLayerName(EMaterialParameterType Type, int32 ParameterIndex) const override;
-	virtual bool HasParameter(const FGuid& ParameterId) const override;
+	virtual bool HasParameter(const FNodeMaterialParameterId& ParameterId) const override;
 	virtual UEdGraphPin* GetParameterPin(EMaterialParameterType Type, int32 ParameterIndex) const override;
-	virtual UEdGraphPin* GetParameterPin(const FGuid& ParameterId) const override;
+	virtual UEdGraphPin* GetParameterPin(const FNodeMaterialParameterId& ParameterId) const override;
 	virtual bool IsImageMutableMode(int32 ImageIndex) const override;
 	virtual bool IsImageMutableMode(const UEdGraphPin& Pin) const override;
 	virtual UTexture2D* GetImageReferenceTexture(int32 ImageIndex) const override;
@@ -193,10 +198,10 @@ private:
 
 	static const TArray<EMaterialParameterType> ParameterTypes;
 
-	/** Relates a Parameter id (key) to a Pin (value). Only used to improve performance.
+	/** Relates a Parameter id (key) (and layer if is a layered material) to a Pin (value). Only used to improve performance.
 	  * If a deprecated pin and a non-deprecated pin have the same Parameter id, this the non-deprecated one prevails. */
 	UPROPERTY()
-	TMap<FGuid, FEdGraphPinReference> PinsParameter;
+	TMap<FNodeMaterialParameterId, FEdGraphPinReference> PinsParameterMap;
 	
 	/** Create the pin data of the given parameter type. */
 	UCustomizableObjectNodeMaterialPinDataParameter* CreatePinData(EMaterialParameterType Type, int32 ParameterIndex);
@@ -242,6 +247,9 @@ private:
 
 	UPROPERTY()
 	TArray<FCustomizableObjectNodeMaterialScalar> ScalarParams_DEPRECATED;
+
+	UPROPERTY()
+	TMap<FGuid, FEdGraphPinReference> PinsParameter_DEPRECATED;
 };
 
 
diff --git a/Engine/Plugins/Experimental/Mutable/Source/CustomizableObjectEditor/Public/MuCOE/Nodes/CustomizableObjectNodeMaterialBase.h b/Engine/Plugins/Experimental/Mutable/Source/CustomizableObjectEditor/Public/MuCOE/Nodes/CustomizableObjectNodeMaterialBase.h
--- a/Engine/Plugins/Experimental/Mutable/Source/CustomizableObjectEditor/Public/MuCOE/Nodes/CustomizableObjectNodeMaterialBase.h
+++ b/Engine/Plugins/Experimental/Mutable/Source/CustomizableObjectEditor/Public/MuCOE/Nodes/CustomizableObjectNodeMaterialBase.h
@@ -14,6 +14,33 @@ class UObject;
 
 DECLARE_MULTICAST_DELEGATE(FPostImagePinModeChangedDelegate)
 
+/** This struct helps us to identify a material parameter using its id and layer index in case of multimaterials
+* When a multilayer material has the same material in multiple layer the parameter Id is not enough to identify a parameter
+* we also need the its layer index.
+*/
+USTRUCT()
+struct FNodeMaterialParameterId
+{
+	GENERATED_BODY()
+
+	UPROPERTY()
+	FGuid ParameterId;
+
+	UPROPERTY()
+	int32 LayerIndex = INDEX_NONE;
+
+	bool operator==(const FNodeMaterialParameterId& Other) const = default;
+};
+
+
+inline uint32 GetTypeHash(const FNodeMaterialParameterId& Key)
+{
+	uint32 Hash = GetTypeHash(Key.ParameterId);
+	Hash = HashCombine(Hash, GetTypeHash(Key.LayerIndex));
+
+	return Hash;
+}
+
 
 /** Equivalent to mu::NodeSurface but with limitations. Currently only nodes that generate a mu::NodeSurfaceNew inherit this (NodeMaterial and NodeCopyMaterial).
  * Nodes that generate mu::NodeSurfaceEdit (NodeEditMaterial, NodeExtendMaterial), mu::NodeSurfaceSwitch (NodeSwitchMaterial)... are excluded.
@@ -55,7 +82,7 @@ public:
 	/** Returns the Material Parameter id.
 	 *
 	 * @param ParameterIndex Have to be valid. */
-	virtual FGuid GetParameterId(EMaterialParameterType Type, int32 ParameterIndex) const PURE_VIRTUAL(UCustomizableObjectNodeMaterial::GetParameterId, return {}; );
+	virtual FNodeMaterialParameterId GetParameterId(EMaterialParameterType Type, int32 ParameterIndex) const PURE_VIRTUAL(UCustomizableObjectNodeMaterial::GetParameterId, return {}; );
 	
 	/** Returns the Material Parameter name.
 	 *
@@ -74,7 +101,7 @@ public:
 	virtual FText GetParameterLayerName(EMaterialParameterType Type, int32 ParameterIndex) const PURE_VIRTUAL(UCustomizableObjectNodeMaterial::GetParameterLayerName, return {}; );
 
 	/** Returns true if the Material contains the given Material Parameter. */
-	virtual bool HasParameter(const FGuid& ParameterId) const PURE_VIRTUAL(UCustomizableObjectNodeMaterial::HasParameter, return {}; );
+	virtual bool HasParameter(const FNodeMaterialParameterId& ParameterId) const PURE_VIRTUAL(UCustomizableObjectNodeMaterial::HasParameter, return {}; );
 	
 	/** Get the pin for the given Material Parameter.
 	 * Not all parameters have pins.
@@ -88,7 +115,7 @@ public:
 	 *
 	 * @param ParameterId Material Parameter id.
 	 * @return Can return nullptr. */
-	virtual UEdGraphPin* GetParameterPin(const FGuid& ParameterId) const PURE_VIRTUAL(UCustomizableObjectNodeMaterial::GetParameterPin, return {}; );
+	virtual UEdGraphPin* GetParameterPin(const FNodeMaterialParameterId& ParameterId) const PURE_VIRTUAL(UCustomizableObjectNodeMaterial::GetParameterPin, return {}; );
 	
 	// --------------------
 	// IMAGES PARAMETERS
diff --git a/Engine/Plugins/Experimental/Mutable/Source/CustomizableObjectEditor/Public/MuCOE/Nodes/CustomizableObjectNodeUseMaterial.h b/Engine/Plugins/Experimental/Mutable/Source/CustomizableObjectEditor/Public/MuCOE/Nodes/CustomizableObjectNodeUseMaterial.h
--- a/Engine/Plugins/Experimental/Mutable/Source/CustomizableObjectEditor/Public/MuCOE/Nodes/CustomizableObjectNodeUseMaterial.h
+++ b/Engine/Plugins/Experimental/Mutable/Source/CustomizableObjectEditor/Public/MuCOE/Nodes/CustomizableObjectNodeUseMaterial.h
@@ -12,6 +12,7 @@ class UCustomizableObjectNode;
 class UEdGraphPin;
 struct FEdGraphPinReference;
 struct FGuid;
+struct FNodeMaterialParameterId;
 
 /** Common functionalities from EditMaterial and ExtendMaterial.
  *
@@ -22,10 +23,10 @@ public:
 	virtual ~FCustomizableObjectNodeUseMaterial() = default;
 	
 	/** Returns true if this node uses the given Image. */
-	bool UsesImage(const FGuid& ImageId) const;
+	bool UsesImage(const FNodeMaterialParameterId& ImageId) const;
 
 	/** Given an Image id, returns its Image pin. Can nullptr if the Image is not used. */
-	const UEdGraphPin* GetUsedImagePin(const FGuid& ImageId) const;
+	const UEdGraphPin* GetUsedImagePin(const FNodeMaterialParameterId& ImageId) const;
 
 protected:
 	// Begin must call functions
@@ -52,8 +53,8 @@ protected:
 	virtual FCustomizableObjectNodeParentedMaterial& GetNodeParentedMaterial() = 0;
 	
 	/** Return the PinsParameter map which this interface belongs to. */
-	virtual TMap<FGuid, FEdGraphPinReference>& GetPinsParameter() = 0;
-	virtual const TMap<FGuid, FEdGraphPinReference>& GetPinsParameter() const;
+	virtual TMap<FNodeMaterialParameterId, FEdGraphPinReference>& GetPinsParameter() = 0;
+	virtual const TMap<FNodeMaterialParameterId, FEdGraphPinReference>& GetPinsParameter() const;
 
 	/** Get the output material pin. */
 	virtual UEdGraphPin* OutputPin() const = 0;
