[mutable] Added a debug material option to customizable skeletal mesh actors. It can be set and enabled through blueprints.

#rnx

diff --git a/Engine/Plugins/Experimental/Mutable/Source/CustomizableObject/Private/MuCO/CustomizableSkeletalMeshActor.cpp b/Engine/Plugins/Experimental/Mutable/Source/CustomizableObject/Private/MuCO/CustomizableSkeletalMeshActor.cpp
--- a/Engine/Plugins/Experimental/Mutable/Source/CustomizableObject/Private/MuCO/CustomizableSkeletalMeshActor.cpp
+++ b/Engine/Plugins/Experimental/Mutable/Source/CustomizableObject/Private/MuCO/CustomizableSkeletalMeshActor.cpp
@@ -4,6 +4,8 @@
 
 #include "Components/SkeletalMeshComponent.h"
 #include "MuCO/CustomizableSkeletalComponent.h"
+#include "MuCO/CustomizableObject.h"
+#include "MuCO/CustomizableObjectSystem.h"
 
 #include UE_INLINE_GENERATED_CPP_BY_NAME(CustomizableSkeletalMeshActor)
 
@@ -53,7 +55,97 @@ void ACustomizableSkeletalMeshActor::AttachNewComponent()
 			}
 		}
 	}
+}
+
+
+UCustomizableObjectInstance* ACustomizableSkeletalMeshActor::GetComponentsCommonInstance()
+{
+	for (const UCustomizableSkeletalComponent* Component : CustomizableSkeletalComponents)
+	{
+		if (UCustomizableObjectInstance* COInstance = Component->CustomizableObjectInstance)
+		{
+			return COInstance;
+		}
+	}
+
+	return nullptr;
+}
+
+
+void ACustomizableSkeletalMeshActor::SetDebugMaterial(UMaterialInterface* InDebugMaterial)
+{
+	if (!InDebugMaterial)
+	{
+		return;
+	}
+
+	DebugMaterial = InDebugMaterial;
+}
+
+
+void ACustomizableSkeletalMeshActor::EnableDebugMaterial(bool bEnableDebugMaterial)
+{
+	bRemoveDebugMaterial = bDebugMaterialEnabled && !bEnableDebugMaterial;
+	bDebugMaterialEnabled = bEnableDebugMaterial;
+
+	if (UCustomizableObjectInstance* COInstance = GetComponentsCommonInstance())
+	{
+		//Bind Instance Update delegate to Actor
+		COInstance->UpdatedDelegate.AddUniqueDynamic(this, &ACustomizableSkeletalMeshActor::SwitchComponentsMaterials);
+		SwitchComponentsMaterials(COInstance);
+	}
+}
+
 
+void ACustomizableSkeletalMeshActor::SwitchComponentsMaterials(UCustomizableObjectInstance* Instance)
+{
+	if (!DebugMaterial)
+	{
+		return;
+	}
+
+	if (bDebugMaterialEnabled || bRemoveDebugMaterial)
+	{
+		UCustomizableObjectInstance* COInstance = GetComponentsCommonInstance();
+
+		if (!COInstance)
+		{
+			return;
+		}
+
+		for (int32 CompIndex = 0; CompIndex < SkeletalMeshComponents.Num(); ++CompIndex)
+		{
+			int32 NumMaterials = SkeletalMeshComponents[CompIndex]->GetNumMaterials();
+
+			if (bDebugMaterialEnabled)
+			{
+				for (int32 MatIndex = 0; MatIndex < NumMaterials; ++MatIndex)
+				{
+					SkeletalMeshComponents[CompIndex]->SetMaterial(MatIndex, DebugMaterial);
+				}
+			}
+			else // Remove debugmaterial
+			{
+				// check if original materials already overriden
+				const TArray<TObjectPtr<UMaterialInterface>>* OverrideMaterials = COInstance->GetOverrideMaterials(CompIndex);
+				const bool bUseOverrideMaterials = COInstance->GetCustomizableObject()->bEnableMeshCache && CVarEnableMeshCache.GetValueOnAnyThread();
+
+				if (bUseOverrideMaterials && OverrideMaterials->Num() > 0)
+				{
+					for (int32 MatIndex = 0; MatIndex < OverrideMaterials->Num(); ++MatIndex)
+					{
+						SkeletalMeshComponents[CompIndex]->SetMaterial(MatIndex, (*OverrideMaterials)[MatIndex]);
+					}
+				}
+				else
+				{
+					SkeletalMeshComponents[CompIndex]->EmptyOverrideMaterials();
+				}
+
+				bRemoveDebugMaterial = false;
+			}
+		}
+	}
 }
 
 #undef LOCTEXT_NAMESPACE
diff --git a/Engine/Plugins/Experimental/Mutable/Source/CustomizableObject/Public/MuCO/CustomizableSkeletalMeshActor.h b/Engine/Plugins/Experimental/Mutable/Source/CustomizableObject/Public/MuCO/CustomizableSkeletalMeshActor.h
--- a/Engine/Plugins/Experimental/Mutable/Source/CustomizableObject/Public/MuCO/CustomizableSkeletalMeshActor.h
+++ b/Engine/Plugins/Experimental/Mutable/Source/CustomizableObject/Public/MuCO/CustomizableSkeletalMeshActor.h
@@ -6,6 +6,7 @@
 
 #include "CustomizableSkeletalMeshActor.generated.h"
 
+class UCustomizableObjectInstance;
 class UCustomizableSkeletalComponent;
 
 class UObject;
@@ -24,6 +25,23 @@ public:
 
 	int32 GetNumComponents() { return CustomizableSkeletalComponents.Num(); }
 
+	/** Sets which material will be used for debugging purposes */
+	UFUNCTION(BlueprintCallable, Category = CustomizableSkeletalMeshActor)
+	void SetDebugMaterial(UMaterialInterface* InDebugMaterial);
+
+	/** Enables the debug material to all the componets of the actor */
+	UFUNCTION(BlueprintCallable, Category = CustomizableSkeletalMeshActor)
+	void EnableDebugMaterial(bool bEnableDebugMaterial);
+
+private:
+
+	/** Changes materials to all actor componets */
+	UFUNCTION()
+	void SwitchComponentsMaterials(UCustomizableObjectInstance* Instance);
+
+	/** return: the instance represented in all the actor components */
+	UCustomizableObjectInstance* GetComponentsCommonInstance();
+
 private:
 	
 	UPROPERTY(Category = CustomizableSkeletalMesh, VisibleAnywhere, BlueprintReadOnly, meta = (AllowPrivateAccess = "true"))
@@ -35,4 +53,13 @@ private:
 	// TODO: This is a temporary fix to not break the demos, we should update the demos to support the arrays of components
 	UPROPERTY(Category = CustomizableSkeletalMesh, VisibleAnywhere, BlueprintReadOnly, meta = (AllowPrivateAccess = "true"))
 	TObjectPtr<class UCustomizableSkeletalComponent> CustomizableSkeletalComponent;
+
+	/** Material used for debugging */
+	TObjectPtr<UMaterialInterface> DebugMaterial;
+
+	/** bool that determines if the debug material has been enabled */
+	bool bDebugMaterialEnabled = false;
+
+	/** bool used to remove the debug material whet it has been disabled */
+	bool bRemoveDebugMaterial = false;
 };
