[mutable] Added colored lines to mutable code and graph viewers

#rnx

diff --git a/Engine/Plugins/Experimental/Mutable/Source/CustomizableObjectEditor/Private/MuCOE/SMutableCodeViewer.cpp b/Engine/Plugins/Experimental/Mutable/Source/CustomizableObjectEditor/Private/MuCOE/SMutableCodeViewer.cpp
--- a/Engine/Plugins/Experimental/Mutable/Source/CustomizableObjectEditor/Private/MuCOE/SMutableCodeViewer.cpp
+++ b/Engine/Plugins/Experimental/Mutable/Source/CustomizableObjectEditor/Private/MuCOE/SMutableCodeViewer.cpp
@@ -35,6 +35,7 @@
 #include "Widgets/Input/SNumericEntryBox.h"
 #include "Widgets/Layout/SScrollBox.h"
 #include "Widgets/Views/STreeView.h"
+#include "MuCOE/Widgets/MutableExpanderArrow.h"
 
 
 #include "Widgets/Input/SSearchBox.h"
@@ -157,8 +158,7 @@ public:
 					+ SHorizontalBox::Slot()
 					.AutoWidth()
 					[
-						SNew(SExpanderArrow, SharedThis(this))
-						.ShouldDrawWires(true)
+						SNew(SMutableExpanderArrow, SharedThis(this))
 					]
 					
 					+ SHorizontalBox::Slot()
diff --git a/Engine/Plugins/Experimental/Mutable/Source/CustomizableObjectEditor/Private/MuCOE/SMutableGraphViewer.cpp b/Engine/Plugins/Experimental/Mutable/Source/CustomizableObjectEditor/Private/MuCOE/SMutableGraphViewer.cpp
--- a/Engine/Plugins/Experimental/Mutable/Source/CustomizableObjectEditor/Private/MuCOE/SMutableGraphViewer.cpp
+++ b/Engine/Plugins/Experimental/Mutable/Source/CustomizableObjectEditor/Private/MuCOE/SMutableGraphViewer.cpp
@@ -87,6 +87,7 @@
 #include "MuT/NodeImageInvertPrivate.h"
 #include "MuT/NodeImageSwizzlePrivate.h"
 #include "MuT/NodeImageMultiLayerPrivate.h"
+#include "Widgets/MutableExpanderArrow.h"
 
 class FExtender;
 class FReferenceCollector;
@@ -141,8 +142,7 @@ public:
 			.VAlign(VAlign_Center)
 			.AutoWidth()
 			[
-				SNew(SExpanderArrow, SharedThis(this))
-				.ShouldDrawWires(true)
+				SNew(SMutableExpanderArrow, SharedThis(this))
 			]
 
 			+ SHorizontalBox::Slot()
diff --git a/Engine/Plugins/Experimental/Mutable/Source/CustomizableObjectEditor/Private/MuCOE/Widgets/MutableExpanderArrow.cpp b/Engine/Plugins/Experimental/Mutable/Source/CustomizableObjectEditor/Private/MuCOE/Widgets/MutableExpanderArrow.cpp
new file mode 100644
--- /dev/null
+++ b/Engine/Plugins/Experimental/Mutable/Source/CustomizableObjectEditor/Private/MuCOE/Widgets/MutableExpanderArrow.cpp
@@ -0,0 +1,116 @@
+﻿// Copyright Epic Games, Inc. All Rights Reserved.
+
+#include "MutableExpanderArrow.h"
+#include "Widgets/SCompoundWidget.h"
+#include "Widgets/Views/SExpanderArrow.h"
+#include "Widgets/Views/ITableRow.h"
+
+
+void SMutableExpanderArrow::Construct(const FArguments& InArgs, const TSharedPtr<ITableRow>& TableRow)
+{
+	SExpanderArrow::Construct(SExpanderArrow::FArguments().ShouldDrawWires(true),TableRow);
+}
+
+
+int32 SMutableExpanderArrow::OnPaint(const FPaintArgs& Args, const FGeometry& AllottedGeometry, const FSlateRect& MyCullingRect, FSlateWindowElementList& OutDrawElements, int32 LayerId, const FWidgetStyle& InWidgetStyle, bool bParentEnabled) const
+{
+	static constexpr float WireThickness = 2.0f;
+	static constexpr float HalfWireThickness = WireThickness / 2.0f;
+	
+	static const FName NAME_VerticalBarBrush = TEXT("WhiteBrush");
+	const float Indent = IndentAmount.Get(10.f);
+	const FSlateBrush* VerticalBarBrush = (StyleSet == nullptr) ? nullptr : StyleSet->GetBrush(NAME_VerticalBarBrush);
+
+	if (ShouldDrawWires.Get() == true && VerticalBarBrush != nullptr)
+	{
+		const TSharedPtr<ITableRow> OwnerRow = OwnerRowPtr.Pin();
+		
+		// Draw vertical wires to indicate paths to parent nodes.
+		const TBitArray<>& NeedsWireByLevel = OwnerRow->GetWiresNeededByDepth();
+		const int32 NumLevels = NeedsWireByLevel.Num();
+		for (int32 Level = 0; Level < NumLevels; ++Level)
+		{
+			if (NeedsWireByLevel[Level])
+			{
+				const float CurrentIndent = Indent * Level;
+				FSlateDrawElement::MakeBox(
+					OutDrawElements,
+					LayerId,
+					AllottedGeometry.ToPaintGeometry(FVector2D(WireThickness, AllottedGeometry.Size.Y), FSlateLayoutTransform(FVector2D(CurrentIndent - 3.f, 0))),
+					VerticalBarBrush,
+					ESlateDrawEffect::None,
+					GetLevelColor(Level)
+				);
+			}
+		}
+
+		const float HalfCellHeight = 0.5f * AllottedGeometry.Size.Y;
+		
+		// For items that are the last expanded child in a list, we need to draw a special angle connector wire.
+		if (const bool bIsLastChild = OwnerRow->IsLastChild())
+		{
+			const float CurrentIndent = Indent * (NumLevels-1);
+			FSlateDrawElement::MakeBox(
+				OutDrawElements,
+				LayerId,
+				AllottedGeometry.ToPaintGeometry(FVector2D(WireThickness, HalfCellHeight + HalfWireThickness), FSlateLayoutTransform(FVector2D(CurrentIndent - 3.f, 0))),
+				VerticalBarBrush,
+				ESlateDrawEffect::None,
+				GetLevelColor(NumLevels -1)
+			);
+		}
+		
+		// If this item is expanded, we need to draw a 1/2-height the line down to its first child cell.
+		if ( const bool bItemAppearsExpanded = OwnerRow->IsItemExpanded() && OwnerRow->DoesItemHaveChildren())
+		{
+			const float CurrentIndent = Indent * NumLevels;
+			FSlateDrawElement::MakeBox(
+				OutDrawElements,
+				LayerId,
+				AllottedGeometry.ToPaintGeometry(FVector2D(WireThickness, HalfCellHeight+ HalfWireThickness), FSlateLayoutTransform(FVector2D(CurrentIndent - 3.f, HalfCellHeight- HalfWireThickness))),
+				VerticalBarBrush,
+				ESlateDrawEffect::None,
+				GetLevelColor(NumLevels)
+			);
+		}
+
+		
+		// Draw horizontal connector from parent wire to child.
+		if (NumLevels > 1)
+		{
+			float LeafDepth = OwnerRow->DoesItemHaveChildren() ? 10.f : 0.0f;
+			const float HorizontalWireStart = (NumLevels - 1)*Indent;
+			FSlateDrawElement::MakeBox(
+				OutDrawElements,
+				LayerId,
+				AllottedGeometry.ToPaintGeometry(
+					FVector2D(AllottedGeometry.Size.X - HorizontalWireStart - WireThickness - LeafDepth, WireThickness),
+					FSlateLayoutTransform(FVector2D(HorizontalWireStart + WireThickness - 3.f, 0.5f*(AllottedGeometry.Size.Y - WireThickness)))
+				),
+				VerticalBarBrush,
+				ESlateDrawEffect::None,
+				GetLevelColor(NumLevels - 1)
+			);
+		}	
+	}
+
+	LayerId = SCompoundWidget::OnPaint(Args, AllottedGeometry, MyCullingRect, OutDrawElements, LayerId, InWidgetStyle, bParentEnabled);
+	return LayerId;
+}
+
+
+FLinearColor SMutableExpanderArrow::GetLevelColor(const int32 Level) const
+{
+	constexpr float Alpha = 0.6;
+	constexpr int32 ArraySize = 4;
+	constexpr FLinearColor AvailableColors[ArraySize] =
+	{
+		FLinearColor(0.40784313725490196,0.5098039215686274,0.08627450980392157,Alpha),
+		FLinearColor(0.6549019607843137,0.3176470588235294,0.0784313725490196,Alpha),
+		FLinearColor(0.6666666666666666,0.10196078431372549,0.3803921568627451,Alpha),
+		FLinearColor(0.2,0.25098039215686274,0.6392156862745098, Alpha)
+	};
+	
+	const int ColorIndex = Level % ArraySize;
+	return AvailableColors[ColorIndex];
+}
\ No newline at end of file
diff --git a/Engine/Plugins/Experimental/Mutable/Source/CustomizableObjectEditor/Private/MuCOE/Widgets/MutableExpanderArrow.h b/Engine/Plugins/Experimental/Mutable/Source/CustomizableObjectEditor/Private/MuCOE/Widgets/MutableExpanderArrow.h
new file mode 100644
--- /dev/null
+++ b/Engine/Plugins/Experimental/Mutable/Source/CustomizableObjectEditor/Private/MuCOE/Widgets/MutableExpanderArrow.h
@@ -0,0 +1,30 @@
+﻿// Copyright Epic Games, Inc. All Rights Reserved.
+#pragma once
+
+#include "Widgets/Views/SExpanderArrow.h"
+
+/**
+ * Custom expander arrow object where the lines drawn will be colored using a circular coloring pattern based
+ * on the depth of the element to be drawn.
+ */
+class SMutableExpanderArrow final : public SExpanderArrow
+{
+public:
+	SLATE_BEGIN_ARGS( SMutableExpanderArrow ){}
+	SLATE_END_ARGS()
+	
+	void Construct( const FArguments& InArgs, const TSharedPtr<class ITableRow>& TableRow );
+	
+protected:
+	
+	virtual int32 OnPaint(const FPaintArgs& Args, const FGeometry& AllottedGeometry, const FSlateRect& MyCullingRect, FSlateWindowElementList& OutDrawElements, int32 LayerId, const FWidgetStyle& InWidgetStyle, bool bParentEnabled) const override;
+
+private:
+
+	/**
+	 * Provided a level value this method will return the color to be used to paint it.
+	 * @param Level The level whose draw color we want to know
+	 * @return The color to use during the drawing of the provided level index
+	 */
+	FLinearColor GetLevelColor(const int32 Level) const;
+};
