[mutable] Changed instance opening behavior for the Performance Analyzer from single to double-click

diff --git a/Engine/Plugins/Experimental/Mutable/Source/CustomizableObjectEditor/Private/MuCOE/CustomizableObjectEditorPerformanceAnalyzer.cpp b/Engine/Plugins/Experimental/Mutable/Source/CustomizableObjectEditor/Private/MuCOE/CustomizableObjectEditorPerformanceAnalyzer.cpp
--- a/Engine/Plugins/Experimental/Mutable/Source/CustomizableObjectEditor/Private/MuCOE/CustomizableObjectEditorPerformanceAnalyzer.cpp
+++ b/Engine/Plugins/Experimental/Mutable/Source/CustomizableObjectEditor/Private/MuCOE/CustomizableObjectEditorPerformanceAnalyzer.cpp
@@ -256,7 +256,7 @@ void SCustomizableObjectEditorPerformanceAnalyzer::Construct(const FArguments& I
 			SAssignNew(InstanceUpdatesListView,SListView<TSharedPtr<FInstanceUpdateDataElement>>)
 			.ListItemsSource(&InstanceUpdateElements)
 			.OnGenerateRow(this,&SCustomizableObjectEditorPerformanceAnalyzer::OnGenerateInstanceUpdateRow)
-			.OnMouseButtonClick(this, &SCustomizableObjectEditorPerformanceAnalyzer::OnInstancesRowClick)
+			.OnMouseButtonDoubleClick(this, &SCustomizableObjectEditorPerformanceAnalyzer::OnInstancesRowDoubleClick)
 			.SelectionMode(ESelectionMode::Single)
 			.IsFocusable(true)
 			.Orientation(Orient_Vertical)
@@ -607,7 +607,7 @@ TSharedRef<ITableRow> SCustomizableObjectEditorPerformanceAnalyzer::OnGenerateIn
 }
 
 
-void SCustomizableObjectEditorPerformanceAnalyzer::OnInstancesRowClick(TSharedPtr<FInstanceUpdateDataElement> InstanceUpdateDataElement)
+void SCustomizableObjectEditorPerformanceAnalyzer::OnInstancesRowDoubleClick(TSharedPtr<FInstanceUpdateDataElement> InstanceUpdateDataElement)
 {
 	if (InstanceUpdateDataElement && InstanceUpdateDataElement->Instance)
 	{
diff --git a/Engine/Plugins/Experimental/Mutable/Source/CustomizableObjectEditor/Private/MuCOE/CustomizableObjectEditorPerformanceAnalyzer.h b/Engine/Plugins/Experimental/Mutable/Source/CustomizableObjectEditor/Private/MuCOE/CustomizableObjectEditorPerformanceAnalyzer.h
--- a/Engine/Plugins/Experimental/Mutable/Source/CustomizableObjectEditor/Private/MuCOE/CustomizableObjectEditorPerformanceAnalyzer.h
+++ b/Engine/Plugins/Experimental/Mutable/Source/CustomizableObjectEditor/Private/MuCOE/CustomizableObjectEditorPerformanceAnalyzer.h
@@ -111,7 +111,7 @@ private:
 	
 	/** ListView methods */
 	TSharedRef<ITableRow> OnGenerateInstanceUpdateRow(TSharedPtr<FInstanceUpdateDataElement> InstanceUpdateDataElement, const TSharedRef<STableViewBase>& TableViewBase);
-	void OnInstancesRowClick(TSharedPtr<FInstanceUpdateDataElement> InstanceUpdateDataElement);
+	void OnInstancesRowDoubleClick(TSharedPtr<FInstanceUpdateDataElement> InstanceUpdateDataElement);
 	void OnInstanceUpdateListViewSort(EColumnSortPriority::Type ColumnPriority, const FName& ColumnId, EColumnSortMode::Type NewSortMode);
 	EColumnSortMode::Type GetColumnSortMode(FName ColumnName) const;
 
