Fixes for implicit FSoftObjectPath temporaries created when using direct construction of TSoftObjectPtr with FString.

diff --git a/Engine/Plugins/Experimental/Mutable/Source/CustomizableObject/Private/MuCO/CustomizableObject.cpp b/Engine/Plugins/Experimental/Mutable/Source/CustomizableObject/Private/MuCO/CustomizableObject.cpp
--- a/Engine/Plugins/Experimental/Mutable/Source/CustomizableObject/Private/MuCO/CustomizableObject.cpp
+++ b/Engine/Plugins/Experimental/Mutable/Source/CustomizableObject/Private/MuCO/CustomizableObject.cpp
@@ -509,7 +509,7 @@ void SerializeStreamedResources(FArchive& Ar, UObject* Object, TArray<FCustomiza
 					
 					FCustomizableObjectAssetUserData ResourceData;
 
-					TSoftObjectPtr<UAssetUserData> SoftAssetUserData(AssetUserDataPath);
+					TSoftObjectPtr<UAssetUserData> SoftAssetUserData = TSoftObjectPtr<UAssetUserData>(FSoftObjectPath(AssetUserDataPath));
 					ResourceData.AssetUserDataEditor = !SoftAssetUserData.IsNull() ? SoftAssetUserData.LoadSynchronous() : nullptr;
 
 					if (!ResourceData.AssetUserDataEditor)
@@ -2649,7 +2649,7 @@ FArchive& operator<<(FArchive& Ar, FAnimBpOverridePhysicsAssetsInfo& Info)
 		Ar << Info.PropertyIndex;
 
 		Info.AnimInstanceClass = TSoftClassPtr<UAnimInstance>(AnimInstanceClassPathString);
-		Info.SourceAsset = TSoftObjectPtr<UPhysicsAsset>(PhysicsAssetPathString);
+		Info.SourceAsset = TSoftObjectPtr<UPhysicsAsset>(FSoftObjectPath(PhysicsAssetPathString));
 	}
 
 	if (Ar.IsSaving())
@@ -2803,23 +2803,23 @@ FArchive& operator<<(FArchive& Ar, FMutableRefSkeletalMeshData& Data)
 
 		FString SkeletalMeshLODSettingsAssetPath;
 		Ar << SkeletalMeshLODSettingsAssetPath;
-		Data.SkeletalMeshLODSettings = TSoftObjectPtr<USkeletalMeshLODSettings>(SkeletalMeshLODSettingsAssetPath).LoadSynchronous();
+		Data.SkeletalMeshLODSettings = TSoftObjectPtr<USkeletalMeshLODSettings>(FSoftObjectPath(SkeletalMeshLODSettingsAssetPath)).LoadSynchronous();
 
 		FString SkeletonAssetPath;
 		Ar << SkeletonAssetPath;
-		Data.Skeleton = TSoftObjectPtr<USkeleton>(SkeletonAssetPath).LoadSynchronous();
+		Data.Skeleton = TSoftObjectPtr<USkeleton>(FSoftObjectPath(SkeletonAssetPath)).LoadSynchronous();
 
 		FString PhysicsAssetPath;
 		Ar << PhysicsAssetPath;
-		Data.PhysicsAsset = TSoftObjectPtr<UPhysicsAsset>(PhysicsAssetPath).LoadSynchronous();
+		Data.PhysicsAsset = TSoftObjectPtr<UPhysicsAsset>(FSoftObjectPath(PhysicsAssetPath)).LoadSynchronous();
 
 		FString PostProcessAnimInstAssetPath;
 		Ar << PostProcessAnimInstAssetPath;
-		Data.PostProcessAnimInst = TSoftClassPtr<UAnimInstance>(PostProcessAnimInstAssetPath).LoadSynchronous();
+		Data.PostProcessAnimInst = TSoftClassPtr<UAnimInstance>(FSoftObjectPath(PostProcessAnimInstAssetPath)).LoadSynchronous();
 
 		FString ShadowPhysicsAssetPath;
 		Ar << ShadowPhysicsAssetPath;
-		Data.ShadowPhysicsAsset = TSoftObjectPtr<UPhysicsAsset>(ShadowPhysicsAssetPath).LoadSynchronous();
+		Data.ShadowPhysicsAsset = TSoftObjectPtr<UPhysicsAsset>(FSoftObjectPath(ShadowPhysicsAssetPath)).LoadSynchronous();
 	}
 
 	Ar << Data.AssetUserData;
diff --git a/Engine/Plugins/Experimental/Mutable/Source/CustomizableObjectEditor/Private/MuCOE/CustomizableObjectCompiler.cpp b/Engine/Plugins/Experimental/Mutable/Source/CustomizableObjectEditor/Private/MuCOE/CustomizableObjectCompiler.cpp
--- a/Engine/Plugins/Experimental/Mutable/Source/CustomizableObjectEditor/Private/MuCOE/CustomizableObjectCompiler.cpp
+++ b/Engine/Plugins/Experimental/Mutable/Source/CustomizableObjectEditor/Private/MuCOE/CustomizableObjectCompiler.cpp
@@ -1178,7 +1178,7 @@ void FCustomizableObjectCompiler::CompileInternal(bool bAsync)
 
 			for (const FName& ReferencingObjectName : ReferencingObjectNames)
 			{
-				const TSoftObjectPtr SoftObjectPtr(ReferencingObjectName.ToString());
+				const TSoftObjectPtr<UObject> SoftObjectPtr = TSoftObjectPtr<UObject>(FSoftObjectPath(ReferencingObjectName.ToString()));
 
 				if (const UObject* ReferencingObject = SoftObjectPtr.LoadSynchronous())
 				{
diff --git a/Engine/Plugins/Experimental/Mutable/Source/CustomizableObjectEditor/Private/MuCOE/CustomizableObjectEditorModule.cpp b/Engine/Plugins/Experimental/Mutable/Source/CustomizableObjectEditor/Private/MuCOE/CustomizableObjectEditorModule.cpp
--- a/Engine/Plugins/Experimental/Mutable/Source/CustomizableObjectEditor/Private/MuCOE/CustomizableObjectEditorModule.cpp
+++ b/Engine/Plugins/Experimental/Mutable/Source/CustomizableObjectEditor/Private/MuCOE/CustomizableObjectEditorModule.cpp
@@ -512,7 +512,7 @@ void GetReferencingPackages(const UCustomizableObject& Object, TArray<FName>& Ob
 	TArray<FName> CustomizableObjectNames = ObjectNames;
 	for (const FName& CustomizableObjectName : CustomizableObjectNames)
 	{
-		const TSoftObjectPtr SoftObjectPtr(CustomizableObjectName.ToString());
+		const TSoftObjectPtr<UObject> SoftObjectPtr = TSoftObjectPtr<UObject>(FSoftObjectPath(CustomizableObjectName.ToString()));
 
 		const UCustomizableObject* ChildCustomizableObject = Cast<UCustomizableObject>(SoftObjectPtr.LoadSynchronous());
 		if (!ChildCustomizableObject)
@@ -560,7 +560,7 @@ bool FCustomizableObjectEditorModule::IsCompilationOutOfDate(const UCustomizable
 
 	for (const TTuple<FName, FGuid>& ParticipatingObject : Object.GetPrivate()->ParticipatingObjects)
 	{
-		TSoftObjectPtr SoftObjectPtr(ParticipatingObject.Key.ToString());
+		TSoftObjectPtr<UObject> SoftObjectPtr = TSoftObjectPtr<UObject>(FSoftObjectPath(ParticipatingObject.Key.ToString()));
 		if (SoftObjectPtr) // If loaded
 		{
 			PRAGMA_DISABLE_DEPRECATION_WARNINGS
@@ -619,7 +619,7 @@ bool FCustomizableObjectEditorModule::IsCompilationOutOfDate(const UCustomizable
 	
 	for (const FName& ObjectName : ReferencingObjectNames)
 	{
-		TSoftObjectPtr ReferencingObject(ObjectName.ToString()); 
+		TSoftObjectPtr<UObject> ReferencingObject = TSoftObjectPtr<UObject>(FSoftObjectPath(ObjectName.ToString())); 
 
 		if ((ReferencingObject && ReferencingObject->GetPackage()->IsDirty()) ||
 			!Object.GetPrivate()->ParticipatingObjects.Contains(ObjectName)) // Must be in the participating objects, if not it means it did not exist when compiling the object.
diff --git a/Engine/Plugins/FX/Niagara/Source/Niagara/Classes/NiagaraBakerOutput.h b/Engine/Plugins/FX/Niagara/Source/Niagara/Classes/NiagaraBakerOutput.h
--- a/Engine/Plugins/FX/Niagara/Source/Niagara/Classes/NiagaraBakerOutput.h
+++ b/Engine/Plugins/FX/Niagara/Source/Niagara/Classes/NiagaraBakerOutput.h
@@ -54,7 +54,7 @@ public:
 		const FString AssetFullName = GetAssetPath(PathFormat, FrameIndex);
 		const FString AssetName = FString(FPathViews::GetCleanFilename(AssetFullName));
 		const FString AssetPath = AssetFullName + "." + AssetName;
-		return TSoftObjectPtr<TObjectClass>(AssetPath).LoadSynchronous();
+		return TSoftObjectPtr<TObjectClass>(FSoftObjectPath(AssetPath)).LoadSynchronous();
 	}
 
 	static NIAGARA_API UObject* GetOrCreateAsset(const FString& PackagePath, UClass* ObjectClass, UClass* FactoryClass);
diff --git a/Engine/Plugins/PCG/Source/PCG/Private/Elements/PCGDataTableRowToParamData.cpp b/Engine/Plugins/PCG/Source/PCG/Private/Elements/PCGDataTableRowToParamData.cpp
--- a/Engine/Plugins/PCG/Source/PCG/Private/Elements/PCGDataTableRowToParamData.cpp
+++ b/Engine/Plugins/PCG/Source/PCG/Private/Elements/PCGDataTableRowToParamData.cpp
@@ -62,7 +62,7 @@ void UPCGDataTableRowToParamDataSettings::PostLoad()
 #if WITH_EDITOR
 	if (!PathOverride_DEPRECATED.IsEmpty())
 	{
-		DataTable = TSoftObjectPtr<UDataTable>(PathOverride_DEPRECATED);
+		DataTable = TSoftObjectPtr<UDataTable>(FSoftObjectPath(PathOverride_DEPRECATED));
 		PathOverride_DEPRECATED.Empty();
 	}
 #endif // WITH_EDITOR
diff --git a/Engine/Plugins/Runtime/RigVM/Source/RigVMDeveloper/Private/RigVMModel/RigVMBuildData.cpp b/Engine/Plugins/Runtime/RigVM/Source/RigVMDeveloper/Private/RigVMModel/RigVMBuildData.cpp
--- a/Engine/Plugins/Runtime/RigVM/Source/RigVMDeveloper/Private/RigVMModel/RigVMBuildData.cpp
+++ b/Engine/Plugins/Runtime/RigVM/Source/RigVMDeveloper/Private/RigVMModel/RigVMBuildData.cpp
@@ -32,7 +32,7 @@ TSoftObjectPtr<URigVMFunctionReferenceNode> FRigVMReferenceNodeData::GetReferenc
 {
 	if(ReferenceNodePtr.IsNull())
 	{
-		ReferenceNodePtr = TSoftObjectPtr<URigVMFunctionReferenceNode>(ReferenceNodePath);
+		ReferenceNodePtr = TSoftObjectPtr<URigVMFunctionReferenceNode>(FSoftObjectPath(ReferenceNodePath));
 	}
 	return ReferenceNodePtr;
 }
@@ -41,7 +41,7 @@ URigVMFunctionReferenceNode* FRigVMReferenceNodeData::GetReferenceNode()
 {
 	if(ReferenceNodePtr.IsNull())
 	{
-		ReferenceNodePtr = TSoftObjectPtr<URigVMFunctionReferenceNode>(ReferenceNodePath);
+		ReferenceNodePtr = TSoftObjectPtr<URigVMFunctionReferenceNode>(FSoftObjectPath(ReferenceNodePath));
 	}
 	if(!ReferenceNodePtr.IsValid())
 	{
@@ -404,7 +404,7 @@ void URigVMBuildData::RegisterFunctionReference(FRigVMReferenceNodeData InRefere
 	check(InReferenceNodeData.ReferencedFunctionIdentifier.GetNodeSoftPath().IsValid());
 
 	FString LibraryNodePath = InReferenceNodeData.ReferencedFunctionIdentifier.GetLibraryNodePath();
-	TSoftObjectPtr<URigVMLibraryNode> LibraryNodePtr = TSoftObjectPtr<URigVMLibraryNode>(LibraryNodePath);
+	TSoftObjectPtr<URigVMLibraryNode> LibraryNodePtr = TSoftObjectPtr<URigVMLibraryNode>(FSoftObjectPath(LibraryNodePath));
 
 	// Try to find a FunctionIdentifier with the same LibraryNodePath
 	bool bFound = false;
diff --git a/Engine/Source/Runtime/Engine/Private/LevelInstance/LevelInstanceSubsystem.cpp b/Engine/Source/Runtime/Engine/Private/LevelInstance/LevelInstanceSubsystem.cpp
--- a/Engine/Source/Runtime/Engine/Private/LevelInstance/LevelInstanceSubsystem.cpp
+++ b/Engine/Source/Runtime/Engine/Private/LevelInstance/LevelInstanceSubsystem.cpp
@@ -1495,11 +1495,11 @@ ILevelInstanceInterface* ULevelInstanceSubsystem::CreateLevelInstanceFrom(const
 		UBlueprint* NewBP = nullptr;
 		if (CreationParams.LevelPackageName.IsEmpty())
 		{
-			NewBP = FPackedLevelActorBuilder::CreatePackedLevelActorBlueprintWithDialog(TSoftObjectPtr<UBlueprint>(BPAssetPath), WorldPtr, bCompile);
+			NewBP = FPackedLevelActorBuilder::CreatePackedLevelActorBlueprintWithDialog(TSoftObjectPtr<UBlueprint>(FSoftObjectPath(BPAssetPath)), WorldPtr, bCompile);
 		}
 		else
 		{
-			NewBP = FPackedLevelActorBuilder::CreatePackedLevelActorBlueprint(TSoftObjectPtr<UBlueprint>(BPAssetPath), WorldPtr, bCompile);
+			NewBP = FPackedLevelActorBuilder::CreatePackedLevelActorBlueprint(TSoftObjectPtr<UBlueprint>(FSoftObjectPath(BPAssetPath)), WorldPtr, bCompile);
 		}
 				
 		if (NewBP)
diff --git a/Engine/Source/Runtime/Engine/Private/LevelStreaming.cpp b/Engine/Source/Runtime/Engine/Private/LevelStreaming.cpp
--- a/Engine/Source/Runtime/Engine/Private/LevelStreaming.cpp
+++ b/Engine/Source/Runtime/Engine/Private/LevelStreaming.cpp
@@ -2098,7 +2098,7 @@ void ULevelStreaming::RenameForPIE(int32 PIEInstanceID, bool bKeepWorldAssetName
 		FSoftObjectPath::AddPIEPackageName(PlayWorldStreamingPackageName);
 		if (bKeepWorldAssetName)
 		{
-			SetWorldAsset(TSoftObjectPtr<UWorld>(FString::Printf(TEXT("%s.%s"), *PlayWorldStreamingPackageName.ToString(), *FPackageName::ObjectPathToObjectName(WorldAsset.ToString()))));
+			SetWorldAsset(TSoftObjectPtr<UWorld>(FSoftObjectPath(FString::Printf(TEXT("%s.%s"), *PlayWorldStreamingPackageName.ToString(), *FPackageName::ObjectPathToObjectName(WorldAsset.ToString())))));
 		}
 		else
 		{
diff --git a/Engine/Source/Runtime/Engine/Private/PackedLevelActor/PackedLevelActorBuilder.cpp b/Engine/Source/Runtime/Engine/Private/PackedLevelActor/PackedLevelActorBuilder.cpp
--- a/Engine/Source/Runtime/Engine/Private/PackedLevelActor/PackedLevelActorBuilder.cpp
+++ b/Engine/Source/Runtime/Engine/Private/PackedLevelActor/PackedLevelActorBuilder.cpp
@@ -78,7 +78,7 @@ UBlueprint* FPackedLevelActorBuilder::CreatePackedLevelActorBlueprintWithDialog(
 
 	if (!SaveObjectPath.IsEmpty())
 	{
-		TSoftObjectPtr<UBlueprint> ExistingBPAsset(SaveObjectPath);
+		TSoftObjectPtr<UBlueprint> ExistingBPAsset = TSoftObjectPtr<UBlueprint>(FSoftObjectPath(SaveObjectPath));
 
 		if (UBlueprint* BP = ExistingBPAsset.LoadSynchronous())
 		{
@@ -432,7 +432,7 @@ bool FPackedLevelActorBuilder::CreateOrUpdateBlueprintFromPacked(APackedLevelAct
 		const bool bCompile = false;
 
 		FString AssetPath = PackagePath + AssetName + "." + AssetName;
-		BP = CreatePackedLevelActorBlueprintWithDialog(TSoftObjectPtr<UBlueprint>(AssetPath), InActor->GetWorldAsset(), bCompile);
+		BP = CreatePackedLevelActorBlueprintWithDialog(TSoftObjectPtr<UBlueprint>(FSoftObjectPath(AssetPath)), InActor->GetWorldAsset(), bCompile);
 	}
 
 	if (!BP || !BP->SimpleConstructionScript)
