[mutable] Fix crashes trying to convert deprecated data.
#rnx

diff --git a/Engine/Plugins/Experimental/Mutable/Source/CustomizableObjectEditor/Private/MuCOE/Nodes/CustomizableObjectNodeMaterial.cpp b/Engine/Plugins/Experimental/Mutable/Source/CustomizableObjectEditor/Private/MuCOE/Nodes/CustomizableObjectNodeMaterial.cpp
--- a/Engine/Plugins/Experimental/Mutable/Source/CustomizableObjectEditor/Private/MuCOE/Nodes/CustomizableObjectNodeMaterial.cpp
+++ b/Engine/Plugins/Experimental/Mutable/Source/CustomizableObjectEditor/Private/MuCOE/Nodes/CustomizableObjectNodeMaterial.cpp
@@ -680,7 +680,7 @@ void UCustomizableObjectNodeMaterial::BackwardsCompatibleFixup()
 
 	if (CustomizableObjectCustomVersion < FCustomizableObjectCustomVersion::FixedMultilayerMaterialIds)
 	{
-		if (Material->GetCachedExpressionData().bHasMaterialLayers)
+		if (Material && Material->GetCachedExpressionData().bHasMaterialLayers)
 		{
 			// Needed since we can not get the layer index of repeated parameters
 			Super::ReconstructNode();
@@ -692,7 +692,19 @@ void UCustomizableObjectNodeMaterial::BackwardsCompatibleFixup()
 				PinsParameterMap.Add({ It.Key(), -1 }, It.Value());
 
 				// Move pin data id info to the new struct
-				if (UCustomizableObjectNodeMaterialPinDataParameter* PinData = Cast<UCustomizableObjectNodeMaterialPinDataParameter>(GetPinData(*It.Value().Get())))
+				UEdGraphPin* GraphPin = It.Value().Get();
+				if (!GraphPin)
+				{
+					continue;
+				}
+
+				UCustomizableObjectNodePinData* GenericPinData = GetPinData(*GraphPin);
+				if (!GenericPinData)
+				{
+					continue;
+				}
+
+				if (UCustomizableObjectNodeMaterialPinDataParameter* PinData = Cast<UCustomizableObjectNodeMaterialPinDataParameter>(GenericPinData))
 				{
 					PinData->MaterialParameterId.LayerIndex = INDEX_NONE;
 					PinData->MaterialParameterId.ParameterId = PinData->ParameterId_DEPRECATED;
