[mutable] Added real-time morph support to table nodes.
#tests Editor
#rnx

diff --git a/Engine/Plugins/Experimental/Mutable/Source/CustomizableObjectEditor/Private/MuCOE/CustomizableObjectLayout.cpp b/Engine/Plugins/Experimental/Mutable/Source/CustomizableObjectEditor/Private/MuCOE/CustomizableObjectLayout.cpp
--- a/Engine/Plugins/Experimental/Mutable/Source/CustomizableObjectEditor/Private/MuCOE/CustomizableObjectLayout.cpp
+++ b/Engine/Plugins/Experimental/Mutable/Source/CustomizableObjectEditor/Private/MuCOE/CustomizableObjectLayout.cpp
@@ -114,7 +114,7 @@ void UCustomizableObjectLayout::GenerateBlocksFromUVs()
 			GenerationContext.MeshGenerationFlags.Push(ShapeFlags);
 
 			GenerationContext.ComponentInfos.Add(FMutableComponentInfo(FName(), SkeletalMesh));
-			MutableMesh = ConvertSkeletalMeshToMutable(SkeletalMesh, TSoftClassPtr<UAnimInstance>(), LOD, Material, LOD, Material, GenerationContext, Node);
+			MutableMesh = ConvertSkeletalMeshToMutable(SkeletalMesh, TSoftClassPtr<UAnimInstance>(), LOD, Material, LOD, Material, GenerationContext, Node, nullptr);
 		}
 		else if (UStaticMesh* StaticMesh = Cast<UStaticMesh>(Mesh))
 		{
diff --git a/Engine/Plugins/Experimental/Mutable/Source/CustomizableObjectEditor/Private/MuCOE/GenerateMutableSource/GenerateMutableSourceComponent.cpp b/Engine/Plugins/Experimental/Mutable/Source/CustomizableObjectEditor/Private/MuCOE/GenerateMutableSource/GenerateMutableSourceComponent.cpp
--- a/Engine/Plugins/Experimental/Mutable/Source/CustomizableObjectEditor/Private/MuCOE/GenerateMutableSource/GenerateMutableSourceComponent.cpp
+++ b/Engine/Plugins/Experimental/Mutable/Source/CustomizableObjectEditor/Private/MuCOE/GenerateMutableSource/GenerateMutableSourceComponent.cpp
@@ -77,7 +77,8 @@ mu::Ptr<mu::NodeComponent> GenerateMutableSourceComponent(const UEdGraphPin * Pi
 			FString MeshUniqueTags;
 			constexpr bool bIsReference = true;
 			TSoftClassPtr<UAnimInstance> AnimInstance;
-			mu::Ptr<mu::Mesh> MutableMesh = GenerateMutableMesh(SkeletalMesh, AnimInstance, 0, 0, 0, 0, MeshUniqueTags, GenerationContext, TypedComponentMesh, bIsReference);
+			mu::Ptr<mu::Mesh> MutableMesh = GenerateMutableMesh(SkeletalMesh, AnimInstance, 0, 0, 0, 0, MeshUniqueTags, 
+																GenerationContext, TypedComponentMesh, nullptr, bIsReference);
 
 			MeshNode->SetValue(MutableMesh);
 		}
diff --git a/Engine/Plugins/Experimental/Mutable/Source/CustomizableObjectEditor/Private/MuCOE/GenerateMutableSource/GenerateMutableSourceMesh.cpp b/Engine/Plugins/Experimental/Mutable/Source/CustomizableObjectEditor/Private/MuCOE/GenerateMutableSource/GenerateMutableSourceMesh.cpp
--- a/Engine/Plugins/Experimental/Mutable/Source/CustomizableObjectEditor/Private/MuCOE/GenerateMutableSource/GenerateMutableSourceMesh.cpp
+++ b/Engine/Plugins/Experimental/Mutable/Source/CustomizableObjectEditor/Private/MuCOE/GenerateMutableSource/GenerateMutableSourceMesh.cpp
@@ -682,7 +682,9 @@ namespace MutablePrivate
 }
 
 
-mu::MeshPtr ConvertSkeletalMeshToMutable(const USkeletalMesh* InSkeletalMesh, const TSoftClassPtr<UAnimInstance>& AnimBp, int32 LODIndexConnected, int32 SectionIndexConnected, int32 LODIndex, int32 SectionIndex, FMutableGraphGenerationContext& GenerationContext, const UCustomizableObjectNode* CurrentNode)
+mu::MeshPtr ConvertSkeletalMeshToMutable(const USkeletalMesh* InSkeletalMesh, const TSoftClassPtr<UAnimInstance>& AnimBp, int32 LODIndexConnected, 
+										 int32 SectionIndexConnected, int32 LODIndex, int32 SectionIndex, FMutableGraphGenerationContext& GenerationContext, 
+										 const UCustomizableObjectNode* CurrentNode, USkeletalMesh* TableReferenceSkeletalMesh)
 {
 	MUTABLE_CPUPROFILER_SCOPE(ConvertSkeletalMeshToMutable);
 
@@ -1160,8 +1162,9 @@ mu::MeshPtr ConvertSkeletalMeshToMutable(const USkeletalMesh* InSkeletalMesh, co
 		TArray<const UMorphTarget*> UsedMorphTargets;
 		UsedMorphTargets.Reserve(SkeletalMeshMorphTargets.Num());
 
-		const UCustomizableObjectNodeSkeletalMesh* NodeTyped = Cast<UCustomizableObjectNodeSkeletalMesh>(CurrentNode);
-		check(NodeTyped);
+		const UCustomizableObjectNodeSkeletalMesh* NodeTypedSkMesh = Cast<UCustomizableObjectNodeSkeletalMesh>(CurrentNode);
+		const UCustomizableObjectNodeTable* NodeTypedTable = Cast<UCustomizableObjectNodeTable>(CurrentNode);
+		check(NodeTypedSkMesh || (NodeTypedTable));
 
         // Add SkeletalMesh morphs to the usage override data structure.
         // Notice this will only be populated here, when compiling.
@@ -1192,21 +1195,36 @@ mu::MeshPtr ConvertSkeletalMeshToMutable(const USkeletalMesh* InSkeletalMesh, co
             TArray<FName> MorphTargetsNames;
             MorphTargetsNames.Reserve(SkeletalMeshMorphTargets.Num());
             
-            if (NodeTyped->bUseAllRealTimeMorphs)
-            {
-                for (const UMorphTarget* MorphTarget : SkeletalMeshMorphTargets)
-                {
-                    check(MorphTarget);
-                    MorphTargetsNames.Add(MorphTarget->GetFName());
-                }
-            } 
-            else
-            {
-                for (const FString& MorphName : NodeTyped->UsedRealTimeMorphTargetNames)
-                {     
-                    MorphTargetsNames.Emplace(*MorphName);
-                }
-            }
+			if (NodeTypedSkMesh)
+			{
+				if (NodeTypedSkMesh->bUseAllRealTimeMorphs)
+				{
+					for (const UMorphTarget* MorphTarget : SkeletalMeshMorphTargets)
+					{
+						check(MorphTarget);
+						MorphTargetsNames.Add(MorphTarget->GetFName());
+					}
+				}
+				else
+				{
+					for (const FString& MorphName : NodeTypedSkMesh->UsedRealTimeMorphTargetNames)
+					{
+						MorphTargetsNames.Emplace(*MorphName);
+					}
+				}
+			}
+			else if (NodeTypedTable && TableReferenceSkeletalMesh)
+			{
+				for (const UMorphTarget* MorphTarget : SkeletalMeshMorphTargets)
+				{
+					check(MorphTarget);
+
+					if (TableReferenceSkeletalMesh->FindMorphTarget(MorphTarget->GetFName()))
+					{
+						MorphTargetsNames.Add(MorphTarget->GetFName());
+					}
+				}
+			}
 
             return MorphTargetsNames;
         });
@@ -2428,7 +2446,9 @@ mu::MeshPtr ConvertStaticMeshToMutable(const UStaticMesh* StaticMesh, int32 LODI
 
 // Convert a Mesh constant to a mutable format. UniqueTags are the tags that make this Mesh unique that cannot be merged in the cache 
 // with the exact same Mesh with other tags
-mu::Ptr<mu::Mesh> GenerateMutableMesh(UObject * Mesh, const TSoftClassPtr<UAnimInstance>& AnimInstance, int32 LODIndexConnected, int32 SectionIndexConnected, int32 LODIndex, int32 SectionIndex, const FString& UniqueTags, FMutableGraphGenerationContext & GenerationContext, const UCustomizableObjectNode* CurrentNode, bool bIsReference)
+mu::Ptr<mu::Mesh> GenerateMutableMesh(UObject * Mesh, const TSoftClassPtr<UAnimInstance>& AnimInstance, int32 LODIndexConnected, int32 SectionIndexConnected, 
+									  int32 LODIndex, int32 SectionIndex, const FString& UniqueTags, FMutableGraphGenerationContext & GenerationContext, 
+									  const UCustomizableObjectNode* CurrentNode, USkeletalMesh* TableReferenceSkeletalMesh, bool bIsReference)
 {
 	// Get the mesh generation flags to use
 	EMutableMeshConversionFlags CurrentFlags = GenerationContext.MeshGenerationFlags.Last();
@@ -2447,7 +2467,7 @@ mu::Ptr<mu::Mesh> GenerateMutableMesh(UObject * Mesh, const TSoftClassPtr<UAnimI
 			}
 			else
 			{
-				MutableMesh = ConvertSkeletalMeshToMutable(SkeletalMesh, AnimInstance, LODIndexConnected, SectionIndexConnected, LODIndex, SectionIndex, GenerationContext, CurrentNode);
+				MutableMesh = ConvertSkeletalMeshToMutable(SkeletalMesh, AnimInstance, LODIndexConnected, SectionIndexConnected, LODIndex, SectionIndex, GenerationContext, CurrentNode, TableReferenceSkeletalMesh);
 
 				FSkeletalMeshModel* ImportedModel = SkeletalMesh->GetImportedModel();
 
@@ -2542,7 +2562,8 @@ mu::MeshPtr BuildMorphedMutableMesh(const UEdGraphPin* BaseSourcePin, const FStr
 		GetLODAndSectionForAutomaticLODs(GenerationContext, *Node, *SkeletalMesh, LODIndexConnected, SectionIndexConnected, LODIndex, SectionIndex, bOnlyConnectedLOD);
 		// Get the base mesh
 		constexpr bool bIsReference = false;
-		mu::Ptr<mu::Mesh> BaseSourceMesh = GenerateMutableMesh(SkeletalMesh, TSoftClassPtr<UAnimInstance>(), LODIndexConnected, SectionIndexConnected, LODIndex, SectionIndex, FString(), GenerationContext, Node, bIsReference);
+		mu::Ptr<mu::Mesh> BaseSourceMesh = GenerateMutableMesh(SkeletalMesh, TSoftClassPtr<UAnimInstance>(), LODIndexConnected, SectionIndexConnected, 
+															   LODIndex, SectionIndex, FString(), GenerationContext, Node, nullptr, bIsReference);
 		if (BaseSourceMesh)
 		{
 			// Clone it (it will probably be shared)
@@ -3400,7 +3421,8 @@ mu::NodeMeshPtr GenerateMutableSourceMesh(const UEdGraphPin* Pin,
 			FSkeletalMeshModel* ImportedModel = TypedNodeSkel->SkeletalMesh->GetImportedModel();
 			
 			constexpr bool bIsReference = false;
-			mu::Ptr<mu::Mesh> MutableMesh = GenerateMutableMesh(TypedNodeSkel->SkeletalMesh, TypedNodeSkel->AnimInstance, LODIndexConnected, SectionIndexConnected, LODIndex, SectionIndex, MeshUniqueTags, GenerationContext, TypedNodeSkel, bIsReference);
+			mu::Ptr<mu::Mesh> MutableMesh = GenerateMutableMesh(TypedNodeSkel->SkeletalMesh, TypedNodeSkel->AnimInstance, LODIndexConnected, SectionIndexConnected, 
+																LODIndex, SectionIndex, MeshUniqueTags, GenerationContext, TypedNodeSkel, nullptr, bIsReference);
 			if (MutableMesh)
 			{
 				MeshNode->SetValue(MutableMesh);
@@ -3587,7 +3609,8 @@ mu::NodeMeshPtr GenerateMutableSourceMesh(const UEdGraphPin* Pin,
 			check(SectionIndex < TypedNodeStatic->LODs[LODIndex].Materials.Num());
 
 			constexpr bool bIsReference = false;
-			mu::MeshPtr MutableMesh = GenerateMutableMesh(TypedNodeStatic->StaticMesh, TSoftClassPtr<UAnimInstance>(), LODIndex, SectionIndex, LODIndex, SectionIndex, FString(), GenerationContext, TypedNodeStatic, bIsReference);
+			mu::MeshPtr MutableMesh = GenerateMutableMesh(TypedNodeStatic->StaticMesh, TSoftClassPtr<UAnimInstance>(), LODIndex, SectionIndex, LODIndex, SectionIndex, 
+														  FString(), GenerationContext, TypedNodeStatic, nullptr, bIsReference);
 			if (MutableMesh)
 			{
 				MeshNode->SetValue(MutableMesh);
diff --git a/Engine/Plugins/Experimental/Mutable/Source/CustomizableObjectEditor/Private/MuCOE/GenerateMutableSource/GenerateMutableSourceMesh.h b/Engine/Plugins/Experimental/Mutable/Source/CustomizableObjectEditor/Private/MuCOE/GenerateMutableSource/GenerateMutableSourceMesh.h
--- a/Engine/Plugins/Experimental/Mutable/Source/CustomizableObjectEditor/Private/MuCOE/GenerateMutableSource/GenerateMutableSourceMesh.h
+++ b/Engine/Plugins/Experimental/Mutable/Source/CustomizableObjectEditor/Private/MuCOE/GenerateMutableSource/GenerateMutableSourceMesh.h
@@ -46,13 +46,15 @@ void GetLODAndSectionForAutomaticLODs(const FMutableGraphGenerationContext& Cont
  * @param LODIndex LOD we are generating. Will be different from LODIndexConnected only when using Automatic LOD From Mesh. 
  * @param SectionIndex Section we are generating. Will be different from SectionIndexConnected only when using Automatic LOD From Mesh.
  * @param GenerationContext 
- * @param CurrentNode 
+ * @param CurrentNode The node being compiled, can be a skeletal mesh or a table node.
+ * @param TableReferenceSkeletalMesh Only set if compiling a table node. It refers to the reference mesh of the skeletal mesh column being compiled.
  * @return Mutable Mesh. Nullptr if there has been an error. Empty mesh if the Skeletal Mesh does not contain the requested LOD + Section. */
 mu::Ptr<mu::Mesh> ConvertSkeletalMeshToMutable(const USkeletalMesh* InSkeletalMesh, const TSoftClassPtr<UAnimInstance>& AnimBp,
                                          int32 LODIndexConnected,  int32 SectionIndexConnected,
                                          int32 LODIndex, int32 SectionIndex,
                                          FMutableGraphGenerationContext& GenerationContext,
-                                         const UCustomizableObjectNode* CurrentNode);
+                                         const UCustomizableObjectNode* CurrentNode, 
+										 USkeletalMesh* TableReferenceSkeletalMesh);
 
 
 mu::Ptr<mu::Mesh> ConvertStaticMeshToMutable(const UStaticMesh* StaticMesh, int32 LODIndex, int32 SectionIndex,
@@ -69,12 +71,14 @@ mu::Ptr<mu::Mesh> ConvertStaticMeshToMutable(const UStaticMesh* StaticMesh, int3
  * @param SectionIndex Section we are generating. Will be different from SectionIndexConnected only when using Automatic LOD From Mesh.
  * @param MeshUniqueTags 
  * @param GenerationContext 
- * @param CurrentNode 
+ * @param CurrentNode The node being compiled, can be a skeletal mesh or a table node.
+ * @param TableReferenceSkeletalMesh Only set if compiling a table node. It refers to the reference mesh of the skeletal mesh column being compiled.
  * @return Mutable Mesh. Nullptr if there has been an error. Empty mesh if the Skeletal Mesh does not contain the requested LOD + Section. */
 mu::Ptr<mu::Mesh> GenerateMutableMesh(UObject* Mesh, const TSoftClassPtr<UAnimInstance>& AnimBp,
                                 int32 LODIndexConnected, int32 SectionIndexConnected,
                                 int32 LODIndex, int32 SectionIndex, const FString& MeshUniqueTags,
-                                FMutableGraphGenerationContext& GenerationContext, const UCustomizableObjectNode* CurrentNode, 
+                                FMutableGraphGenerationContext& GenerationContext, 
+								const UCustomizableObjectNode* CurrentNode, USkeletalMesh* TableReferenceSkeletalMesh,
 								bool bIsReference);
 
 
diff --git a/Engine/Plugins/Experimental/Mutable/Source/CustomizableObjectEditor/Private/MuCOE/GenerateMutableSource/GenerateMutableSourceTable.cpp b/Engine/Plugins/Experimental/Mutable/Source/CustomizableObjectEditor/Private/MuCOE/GenerateMutableSource/GenerateMutableSourceTable.cpp
--- a/Engine/Plugins/Experimental/Mutable/Source/CustomizableObjectEditor/Private/MuCOE/GenerateMutableSource/GenerateMutableSourceTable.cpp
+++ b/Engine/Plugins/Experimental/Mutable/Source/CustomizableObjectEditor/Private/MuCOE/GenerateMutableSource/GenerateMutableSourceTable.cpp
@@ -279,7 +279,8 @@ bool FillTableColumn(const UCustomizableObjectNodeTable* TableNode, mu::TablePtr
 
 			//TODO: Add AnimBp physics to Tables.
 			constexpr bool bIsReference = false;
-			mu::Ptr<mu::Mesh> MutableMesh = GenerateMutableMesh(SkeletalMesh, TSoftClassPtr<UAnimInstance>(), LODIndexConnected, SectionIndexConnected, LODIndex, SectionIndex, MeshUniqueTags, GenerationContext, TableNode, bIsReference);
+			mu::Ptr<mu::Mesh> MutableMesh = GenerateMutableMesh(SkeletalMesh, TSoftClassPtr<UAnimInstance>(), LODIndexConnected, SectionIndexConnected, 
+																LODIndex, SectionIndex, MeshUniqueTags, GenerationContext, TableNode, ReferenceSkeletalMesh, bIsReference);
 
 			if (MutableMesh)
 			{
@@ -386,7 +387,8 @@ bool FillTableColumn(const UCustomizableObjectNodeTable* TableNode, mu::TablePtr
 			}
 
 			constexpr bool bIsReference = false;
-			mu::MeshPtr MutableMesh = GenerateMutableMesh(StaticMesh, TSoftClassPtr<UAnimInstance>(), CurrentLOD, SectionIndex, CurrentLOD, SectionIndex, FString(), GenerationContext, TableNode, bIsReference);
+			mu::MeshPtr MutableMesh = GenerateMutableMesh(StaticMesh, TSoftClassPtr<UAnimInstance>(), CurrentLOD, SectionIndex, CurrentLOD, SectionIndex, 
+														  FString(), GenerationContext, TableNode, nullptr, bIsReference);
 
 			if (MutableMesh)
 			{
