Fixed type mismatch when assigning UObject* to TSoftObjectPtr<USkeletalMesh>.

diff --git a/Engine/Plugins/Experimental/Mutable/Source/CustomizableObjectEditor/Private/MuCOE/CustomizableObjectFactory.cpp b/Engine/Plugins/Experimental/Mutable/Source/CustomizableObjectEditor/Private/MuCOE/CustomizableObjectFactory.cpp
--- a/Engine/Plugins/Experimental/Mutable/Source/CustomizableObjectEditor/Private/MuCOE/CustomizableObjectFactory.cpp
+++ b/Engine/Plugins/Experimental/Mutable/Source/CustomizableObjectEditor/Private/MuCOE/CustomizableObjectFactory.cpp
@@ -679,12 +679,12 @@ int32 FCustomizableObjectFactoryUI::GetSelectedComponentIndex() const
 void FCustomizableObjectFactoryUI::OnPickedComponentSkeletalMesh(const FAssetData& SelectedAsset)
 {
 	int32 ComponentIndex = GetSelectedComponentIndex();
-	if (ComponentIndex != INDEX_NONE)
+	if (ComponentIndex != INDEX_NONE && SelectedAsset.IsValid())
 	{
-		if (SelectedAsset.IsValid())
+		if (USkeletalMesh* SkeletalMesh = Cast<USkeletalMesh>(SelectedAsset.GetAsset()))
 		{
 			// we have to load the asset otherwise the asset thumbnail won't be visible
-			Options.ComponentsInfo[ComponentIndex].ReferenceSkeletalMesh = SelectedAsset.GetAsset();
+			Options.ComponentsInfo[ComponentIndex].ReferenceSkeletalMesh = SkeletalMesh;
 		}
 	}
 }
