Fixed type mismatch when assigning UObject* to TSoftObjectPtr<UTextureCube>.

diff --git a/Engine/Plugins/Experimental/Mutable/Source/CustomizableObjectEditor/Private/MuCOE/SCustomizableObjectEditorAdvancedPreviewSettings.cpp b/Engine/Plugins/Experimental/Mutable/Source/CustomizableObjectEditor/Private/MuCOE/SCustomizableObjectEditorAdvancedPreviewSettings.cpp
--- a/Engine/Plugins/Experimental/Mutable/Source/CustomizableObjectEditor/Private/MuCOE/SCustomizableObjectEditorAdvancedPreviewSettings.cpp
+++ b/Engine/Plugins/Experimental/Mutable/Source/CustomizableObjectEditor/Private/MuCOE/SCustomizableObjectEditorAdvancedPreviewSettings.cpp
@@ -492,7 +492,7 @@ void SCustomizableObjectEditorAdvancedPreviewSettings::LoadProfileEnvironment()
 				LoadedObject = Redirector->DestinationObject;
 			}
 
-			DefaultSettings->Profiles[i].EnvironmentCubeMap = LoadedObject;
+			DefaultSettings->Profiles[i].EnvironmentCubeMap = CastChecked<UTextureCube>(LoadedObject);
 		}
 	}
 }
