[Mutable] Report in the Performance Analyzer the amount of triangles for each instance (initial generation).
- Updated LogBenchmarkUtil so it computes the total amount of triangles (for all updated components for all lods and sections)
- Added colors to better show the Update Result of the mesh update.
- Changed instance opening system to now use a Hyperlink instead of a row double-click.

#rnx

diff --git a/Engine/Plugins/Experimental/Mutable/Source/CustomizableObject/Internal/MuCO/LogBenchmarkUtil.h b/Engine/Plugins/Experimental/Mutable/Source/CustomizableObject/Internal/MuCO/LogBenchmarkUtil.h
--- a/Engine/Plugins/Experimental/Mutable/Source/CustomizableObject/Internal/MuCO/LogBenchmarkUtil.h
+++ b/Engine/Plugins/Experimental/Mutable/Source/CustomizableObject/Internal/MuCO/LogBenchmarkUtil.h
@@ -81,6 +81,8 @@ struct FInstanceUpdateStats
 	EUpdateResult UpdateResult = EUpdateResult::Error;
 
 	bool bLevelBegunPlay = false;
+
+	uint32 TriangleCount = 0;
 	
 	double QueueTime = 0.0;
 	double UpdateTime = 0.0;
diff --git a/Engine/Plugins/Experimental/Mutable/Source/CustomizableObject/Private/MuCO/LogBenchmarkUtil.cpp b/Engine/Plugins/Experimental/Mutable/Source/CustomizableObject/Private/MuCO/LogBenchmarkUtil.cpp
--- a/Engine/Plugins/Experimental/Mutable/Source/CustomizableObject/Private/MuCO/LogBenchmarkUtil.cpp
+++ b/Engine/Plugins/Experimental/Mutable/Source/CustomizableObject/Private/MuCO/LogBenchmarkUtil.cpp
@@ -57,7 +57,7 @@ TSharedPtr<FArchive> CreateFile()
 
 	const FString HeaderRow = TEXT(
 		"ID_CO;ID_COI;ID_UpdateType;ID_Descriptor;ID_UpdateResult;"											// Basic identifying data
-		"Context_LevelBegunPlay;"																	// The context of the update
+		"Context_LevelBegunPlay;TriangleCount;"															// The context of the update
 		"Time_Queue;Time_Update;Time_TaskGetMesh;Time_TaskLockCache;"										// Time and memory data ...
 		"Time_TaskGetImages;Time_TaskConvertResources;Time_TaskCallbacks;Memory_Update;Memory_Update_Real;"
 		"Time_TaskUpdateImage;Memory_TaskUpdateImage;Memory_TaskUpdateImage_Real");
@@ -311,6 +311,28 @@ void FLogBenchmarkUtil::FinishUpdateMesh(const TSharedRef<FUpdateContextPrivate>
 	{
 		Archive = CreateFile();
 	}
+
+	// Cache the amount of triangles of this instance
+	uint32 TriangleCount = 0;
+	for (int32 ComponentIndex = 0; ComponentIndex < Instance->GetNumComponents(); ComponentIndex++)
+	{
+		// Process the generated components (not null)
+		if (const USkeletalMesh* InstanceSkeletalMesh = Instance->GetSkeletalMesh(ComponentIndex))
+		{
+			const FSkeletalMeshRenderData* RenderData = InstanceSkeletalMesh->GetResourceForRendering();
+			check(RenderData);
+
+			// Add the amount of triangles for all LODs and all sections
+			for (const FSkeletalMeshLODRenderData& RenderDataObject : RenderData->LODRenderData)
+			{
+				for	(const FSkelMeshRenderSection& Section : RenderDataObject.RenderSections)
+				{
+					TriangleCount += Section.NumTriangles;
+				}
+			}
+		}
+	}
+
 	
 	// Identifying data
 	FInstanceUpdateStats UpdateData;
@@ -320,6 +342,7 @@ void FLogBenchmarkUtil::FinishUpdateMesh(const TSharedRef<FUpdateContextPrivate>
 	UpdateData.Descriptor = Context->GetCapturedDescriptor().ToString();
 	UpdateData.UpdateResult = Context->UpdateResult;
 	UpdateData.bLevelBegunPlay = Context->bLevelBegunPlay;
+	UpdateData.TriangleCount = TriangleCount;
 	UpdateData.QueueTime = Context->QueueTime * 1000;
 	UpdateData.UpdateTime = Context->UpdateTime * 1000;
 	UpdateData.TaskGetMeshTime = Context->TaskGetMeshTime * 1000;
@@ -329,15 +352,16 @@ void FLogBenchmarkUtil::FinishUpdateMesh(const TSharedRef<FUpdateContextPrivate>
 	UpdateData.TaskCallbacksTime = Context->TaskCallbacksTime * 1000;
 	UpdateData.UpdatePeakMemory = (Context->UpdateEndPeakBytes / 1024.0) / 1024.0;
 	UpdateData.UpdateRealPeakMemory = (Context->UpdateEndRealPeakBytes / 1024.0) / 1024.0;
-
+	
 	// todo: Find a better way of handling the construction of this string so we know the symmetry with the header row will not get lost when adding new elements
-	const FString UpdateString = FString::Printf(TEXT("%s;%s;%s;%s;%s;%s;%f;%f;%f;%f;%f;%f;%f;%f;%f"),
+	const FString UpdateString = FString::Printf(TEXT("%s;%s;%s;%s;%s;%s;%u;%f;%f;%f;%f;%f;%f;%f;%f;%f"),
 		*UpdateData.CustomizableObjectPathName,
 		*UpdateData.CustomizableObjectInstancePathName,
 		*UpdateData.UpdateType,
 		*UpdateData.Descriptor,
 		*StaticEnum<EUpdateResult>()->GetValueAsString(UpdateData.UpdateResult),
 		(UpdateData.bLevelBegunPlay ? TEXT("true") : TEXT("false")),
+		UpdateData.TriangleCount,
 		UpdateData.QueueTime,
 		UpdateData.UpdateTime,
 		UpdateData.TaskGetMeshTime,
@@ -375,7 +399,7 @@ void FLogBenchmarkUtil::FinishUpdateImage(const FString& CustomizableObjectPathN
 	
 
 	// todo: Find a better way of handling the construction of this string so we know the symmetry with the header row will not get lost when adding new elements
-	const FString UpdateString = FString::Printf(TEXT("%s;%s;%s;%s;;%s;;;;;;;;;;%f;%f;%f"),
+	const FString UpdateString = FString::Printf(TEXT("%s;%s;%s;%s;;%s;;;;;;;;;;;%f;%f;%f"),
 		*MipsUpdateData.CustomizableObjectPathName,
 		*MipsUpdateData.CustomizableObjectInstancePathName,
 		*MipsUpdateData.UpdateType,
diff --git a/Engine/Plugins/Experimental/Mutable/Source/CustomizableObjectEditor/Private/MuCOE/CustomizableObjectEditorPerformanceAnalyzer.cpp b/Engine/Plugins/Experimental/Mutable/Source/CustomizableObjectEditor/Private/MuCOE/CustomizableObjectEditorPerformanceAnalyzer.cpp
--- a/Engine/Plugins/Experimental/Mutable/Source/CustomizableObjectEditor/Private/MuCOE/CustomizableObjectEditorPerformanceAnalyzer.cpp
+++ b/Engine/Plugins/Experimental/Mutable/Source/CustomizableObjectEditor/Private/MuCOE/CustomizableObjectEditorPerformanceAnalyzer.cpp
@@ -21,6 +21,7 @@
 #include "Widgets/Input/SNumericEntryBox.h"
 #include "Widgets/Layout/SSeparator.h"
 #include "Widgets/Input/SButton.h"
+#include "Widgets/Input/SHyperlink.h"
 
 
 #define LOCTEXT_NAMESPACE "CustomizableObjectEditorPerformanceAnalyzer"
@@ -34,6 +35,8 @@ namespace InstanceUpdatesMainDataColumns
 	static const FName InstanceDescriptorColumnID(TEXT("Descriptor"));
 	static const FName InstanceUpdateResultColumnID(TEXT("UpdateResult"));
 
+	static const FName InstanceTriangleCount(TEXT("TriangleCount"));
+
 	// Initial Gen
 	static const FName QueueTimeColumnID (TEXT("QueueTime"));
 	static const FName UpdateTimeColumnID (TEXT("UpdateTime"));
@@ -61,6 +64,16 @@ void SInstanceUpdateDataRow::Construct(const FArguments& InArgs, const TSharedRe
 	);
 }
 
+
+void SInstanceUpdateDataRow::OnInstanceNameNavigation() const
+{
+	if (InstanceUpdateElement && InstanceUpdateElement->Instance)
+	{
+		GEditor->GetEditorSubsystem<UAssetEditorSubsystem>()->OpenEditorForAsset(InstanceUpdateElement->Instance.Get());
+	}
+}
+
+
 TSharedRef<SWidget> SInstanceUpdateDataRow::GenerateWidgetForColumn(const FName& InColumnName)
 {
 	if (!InstanceUpdateElement)
@@ -70,14 +83,21 @@ TSharedRef<SWidget> SInstanceUpdateDataRow::GenerateWidgetForColumn(const FName&
 	
 	if (InColumnName == InstanceUpdatesMainDataColumns::InstanceNameColumnID)
 	{
-		return SNew(STextBlock)
-			.Text(FText::FromString(InstanceUpdateElement->Instance->GetName()));
+		return SNew(SHyperlink)
+			.Style(FAppStyle::Get(), TEXT("NavigationHyperlink"))
+			.Text(FText::FromString(InstanceUpdateElement->Instance->GetName()))
+			.OnNavigate(this, &SInstanceUpdateDataRow::OnInstanceNameNavigation);
 	}
 	else if (InColumnName == InstanceUpdatesMainDataColumns::InstanceUpdateTypeID)
 	{
 		return SNew(STextBlock)
 			.Text(FText::FromString(InstanceUpdateElement->UpdateStats.UpdateType));
 	}
+	else if (InColumnName == InstanceUpdatesMainDataColumns::InstanceTriangleCount)
+	{
+		return SNew(STextBlock)
+			.Text(FText::FromString(FString::Printf(TEXT("%u"),InstanceUpdateElement->UpdateStats.TriangleCount)));
+	}
 	else if (InColumnName == InstanceUpdatesMainDataColumns::InstanceDescriptorColumnID)
 	{
 		return SNew(STextBlock)
@@ -85,8 +105,28 @@ TSharedRef<SWidget> SInstanceUpdateDataRow::GenerateWidgetForColumn(const FName&
 	}
 	else if (InColumnName == InstanceUpdatesMainDataColumns::InstanceUpdateResultColumnID)
 	{
+		FSlateColor TextColor;
+		switch (InstanceUpdateElement->UpdateStats.UpdateResult)
+		{
+		case EUpdateResult::Success:
+			TextColor = FSlateColor(FColor(70,207,120));
+			break;
+		case EUpdateResult::Warning:
+			TextColor = FSlateColor(FColor(250,226,7));
+			break;
+		case EUpdateResult::Error:
+		case EUpdateResult::ErrorOptimized:
+		case EUpdateResult::ErrorReplaced:
+		case EUpdateResult::ErrorDiscarded:
+		case EUpdateResult::Error16BitBoneIndex:
+			TextColor = FSlateColor(FColor(197,0,7));
+			break;
+		}
+		
 		return SNew(STextBlock)
-			.Text(UEnum::GetDisplayValueAsText( InstanceUpdateElement->UpdateStats.UpdateResult));
+			  .Text(UEnum::GetDisplayValueAsText( InstanceUpdateElement->UpdateStats.UpdateResult))
+			  .ColorAndOpacity(TextColor);
+		
 	}
 	else if (InColumnName == InstanceUpdatesMainDataColumns::QueueTimeColumnID)
 	{
@@ -256,7 +296,6 @@ void SCustomizableObjectEditorPerformanceAnalyzer::Construct(const FArguments& I
 			SAssignNew(InstanceUpdatesListView,SListView<TSharedPtr<FInstanceUpdateDataElement>>)
 			.ListItemsSource(&InstanceUpdateElements)
 			.OnGenerateRow(this,&SCustomizableObjectEditorPerformanceAnalyzer::OnGenerateInstanceUpdateRow)
-			.OnMouseButtonDoubleClick(this, &SCustomizableObjectEditorPerformanceAnalyzer::OnInstancesRowDoubleClick)
 			.SelectionMode(ESelectionMode::Single)
 			.IsFocusable(true)
 			.Orientation(Orient_Vertical)
@@ -282,6 +321,15 @@ void SCustomizableObjectEditorPerformanceAnalyzer::Construct(const FArguments& I
 				.OverflowPolicy(ETextOverflowPolicy::Ellipsis)
 				.FillWidth(1)
 
+				+ SHeaderRow::Column(InstanceUpdatesMainDataColumns::InstanceTriangleCount)
+				.DefaultLabel(FText(LOCTEXT("InstanceTriangleCountColumnLabel","Triangle Count")))
+				.DefaultTooltip(FText(LOCTEXT("InstanceTriangleCountColumnToolTip","Customizable Object Instance amount of triangles generated for all Components and all Lods")))
+				.OnSort(this, &SCustomizableObjectEditorPerformanceAnalyzer::OnInstanceUpdateListViewSort)
+				.SortMode(this, &SCustomizableObjectEditorPerformanceAnalyzer::GetColumnSortMode, InstanceUpdatesMainDataColumns::InstanceTriangleCount)
+				.HAlignCell(EHorizontalAlignment::HAlign_Right)
+				.OverflowPolicy(ETextOverflowPolicy::Ellipsis)
+				.FillWidth(1)
+				
 				+ SHeaderRow::Column(InstanceUpdatesMainDataColumns::QueueTimeColumnID)
 				.DefaultLabel(FText(LOCTEXT("QueueTimeColumnLabel", "Queue (ms)")))
 				.DefaultTooltip(FText(LOCTEXT("QueueTimeColumnToolTip", "Time spent before starting the update")))
@@ -607,15 +655,6 @@ TSharedRef<ITableRow> SCustomizableObjectEditorPerformanceAnalyzer::OnGenerateIn
 }
 
 
-void SCustomizableObjectEditorPerformanceAnalyzer::OnInstancesRowDoubleClick(TSharedPtr<FInstanceUpdateDataElement> InstanceUpdateDataElement)
-{
-	if (InstanceUpdateDataElement && InstanceUpdateDataElement->Instance)
-	{
-		GEditor->GetEditorSubsystem<UAssetEditorSubsystem>()->OpenEditorForAsset(InstanceUpdateDataElement->Instance.Get());
-	}
-}
-
-
 void SCustomizableObjectEditorPerformanceAnalyzer::OnInstanceUpdateListViewSort(EColumnSortPriority::Type ColumnPriority, const FName& ColumnId, EColumnSortMode::Type NewSortMode)
 {
 	CurrentSortColumn = ColumnId;
@@ -663,6 +702,21 @@ void SCustomizableObjectEditorPerformanceAnalyzer::OnInstanceUpdateListViewSort(
 			return ElementA.IsValid();
 		});
 	}
+	else if (ColumnId == InstanceUpdatesMainDataColumns::InstanceTriangleCount)
+	{
+		InstanceUpdateElements.StableSort([&](const TSharedPtr<FInstanceUpdateDataElement>& ElementA, const TSharedPtr<FInstanceUpdateDataElement>& ElementB)
+		{
+			if (ElementA.IsValid() && ElementB.IsValid())
+			{
+				const uint32 ValA = ElementA->UpdateStats.TriangleCount;
+				const uint32 ValB =  ElementB->UpdateStats.TriangleCount;
+
+				return NewSortMode == EColumnSortMode::Ascending ? ValA < ValB : ValA > ValB;
+			}
+
+			return ElementA.IsValid();
+		});
+	}
 	else if (ColumnId == InstanceUpdatesMainDataColumns::UpdateTimeColumnID)
 	{
 		InstanceUpdateElements.StableSort([&](const TSharedPtr<FInstanceUpdateDataElement>& ElementA, const TSharedPtr<FInstanceUpdateDataElement>& ElementB)
diff --git a/Engine/Plugins/Experimental/Mutable/Source/CustomizableObjectEditor/Private/MuCOE/CustomizableObjectEditorPerformanceAnalyzer.h b/Engine/Plugins/Experimental/Mutable/Source/CustomizableObjectEditor/Private/MuCOE/CustomizableObjectEditorPerformanceAnalyzer.h
--- a/Engine/Plugins/Experimental/Mutable/Source/CustomizableObjectEditor/Private/MuCOE/CustomizableObjectEditorPerformanceAnalyzer.h
+++ b/Engine/Plugins/Experimental/Mutable/Source/CustomizableObjectEditor/Private/MuCOE/CustomizableObjectEditorPerformanceAnalyzer.h
@@ -23,7 +23,7 @@ class UCustomizableObjectSystem;
 enum class ECustomizableObjectTextureCompression : uint8;
 
 /**
- * Backend object that represents the update of one instance. Currently only representes the data for the initial geneation of the instance update
+ * Backend object that represents the update of one instance. Currently only represents the data for the initial generation of the instance update
  */
 class FInstanceUpdateDataElement 
 {
@@ -32,7 +32,7 @@ public:
 	TStrongObjectPtr<UCustomizableObjectInstance> Instance = nullptr;
 
 	/** Index to aid in the sorting of the FInstanceUpdateDataElement in the SListView that contains them */
-	uint16 UpdateIndex = 0;
+	uint32 UpdateIndex = 0;
 	
 	/** Container with all the perf data in relation to the instance this object represents */
 	FInstanceUpdateStats UpdateStats;
@@ -52,6 +52,9 @@ private:
 	
 	/** Backend object whose data we are drawing */
 	TSharedPtr<FInstanceUpdateDataElement> InstanceUpdateElement;
+
+	/** Method invoked when the hyperlink showing the name of the instance is clicked */
+	void OnInstanceNameNavigation() const;
 	
 	virtual TSharedRef<SWidget> GenerateWidgetForColumn(const FName& InColumnName) override;
 };
@@ -111,7 +114,6 @@ private:
 	
 	/** ListView methods */
 	TSharedRef<ITableRow> OnGenerateInstanceUpdateRow(TSharedPtr<FInstanceUpdateDataElement> InstanceUpdateDataElement, const TSharedRef<STableViewBase>& TableViewBase);
-	void OnInstancesRowDoubleClick(TSharedPtr<FInstanceUpdateDataElement> InstanceUpdateDataElement);
 	void OnInstanceUpdateListViewSort(EColumnSortPriority::Type ColumnPriority, const FName& ColumnId, EColumnSortMode::Type NewSortMode);
 	EColumnSortMode::Type GetColumnSortMode(FName ColumnName) const;
 
