[mutable] Fixed problem of the version bridge where we were getting the asset from the compiled CO instead of the Root Object.
+ Removed some repeated code from the compiler.

#rnx

diff --git a/Engine/Plugins/Experimental/Mutable/Source/CustomizableObjectEditor/Private/MuCOE/CustomizableObjectCompiler.cpp b/Engine/Plugins/Experimental/Mutable/Source/CustomizableObjectEditor/Private/MuCOE/CustomizableObjectCompiler.cpp
--- a/Engine/Plugins/Experimental/Mutable/Source/CustomizableObjectEditor/Private/MuCOE/CustomizableObjectCompiler.cpp
+++ b/Engine/Plugins/Experimental/Mutable/Source/CustomizableObjectEditor/Private/MuCOE/CustomizableObjectCompiler.cpp
@@ -167,10 +167,13 @@ void FCustomizableObjectCompiler::Compile(const TSharedRef<FCompilationRequest>&
 		return;
 	}
 
-	if (CurrentObject->VersionBridge && !CurrentObject->VersionBridge->GetClass()->ImplementsInterface(UCustomizableObjectVersionBridgeInterface::StaticClass()))
+	UCustomizableObject* RootObject = GetRootObject(CurrentObject);
+	check(RootObject);
+
+	if (RootObject->VersionBridge && !RootObject->VersionBridge->GetClass()->ImplementsInterface(UCustomizableObjectVersionBridgeInterface::StaticClass()))
 	{
 		UE_LOG(LogMutable, Warning, TEXT("In Customizable Object [%s], the VersionBridge asset [%s] does not implement the required UCustomizableObjectVersionBridgeInterface."),
-			*CurrentObject->GetName(), *CurrentObject->VersionBridge.GetName());
+			*RootObject->GetName(), *RootObject->VersionBridge.GetName());
 		CompleteRequest(ECompilationStatePrivate::Completed, ECompilationResultPrivate::Errors);
 		return;
 	}
@@ -399,14 +402,14 @@ void FCustomizableObjectCompiler::ProcessChildObjectsRecursively(UCustomizableOb
 
 		if (ChildObject->VersionStruct.IsValid())
 		{
-			if (!GenerationContext.Object->VersionBridge)
+			if (!GenerationContext.RootVersionBridge)
 			{
 				UE_LOG(LogMutable, Warning, TEXT("The child Customizable Object [%s] defines its VersionStruct Property but its root CustomizableObject doesn't define the VersionBridge property. There's no way to verify the VersionStruct has to be included in this compilation, so the child CustomizableObject will be omitted."), 
 					*ChildObject->GetName());
 				continue;
 			}
 
-			ICustomizableObjectVersionBridgeInterface* CustomizableObjectVersionBridgeInterface = Cast<ICustomizableObjectVersionBridgeInterface>(GenerationContext.Object->VersionBridge);
+			ICustomizableObjectVersionBridgeInterface* CustomizableObjectVersionBridgeInterface = Cast<ICustomizableObjectVersionBridgeInterface>(GenerationContext.RootVersionBridge);
 
 			if (CustomizableObjectVersionBridgeInterface)
 			{
@@ -532,25 +535,40 @@ mu::NodeObjectPtr FCustomizableObjectCompiler::GenerateMutableRoot(
 	}
 
 	bool bMultipleBaseObjectsFound;
-	bOutIsRootObject = false;
 	UCustomizableObjectNodeObject* Root = GetRootNode(Object, bMultipleBaseObjectsFound);
 
 	if (bMultipleBaseObjectsFound)
 	{
-		ErrorMsg = LOCTEXT("MutlipleBase","Multiple base object nodes found. Only one will be used.");
+		ErrorMsg = LOCTEXT("MultipleBaseRoot","Multiple base object nodes found.");
 		return nullptr;
 	}
 
 	if (!Root)
 	{
-		ErrorMsg = LOCTEXT("NoBase","No base object node found. Object not built.");
+		ErrorMsg = LOCTEXT("NoRootBase","No base object node found. Object not built.");
 		return nullptr;
 	}
 
 	bOutIsRootObject = Root->ParentObject == nullptr;
 
-	UCustomizableObjectNodeObject* ActualRoot = Root;
-	UCustomizableObject* ActualRootObject = Object;
+	UCustomizableObject* ActualRootObject = GetRootObject(Object);
+	check(ActualRootObject);
+
+	GenerationContext.RootVersionBridge = ActualRootObject->VersionBridge;
+
+	UCustomizableObjectNodeObject* ActualRoot = GetRootNode(ActualRootObject, bMultipleBaseObjectsFound);
+
+	if (bMultipleBaseObjectsFound)
+	{
+		ErrorMsg = LOCTEXT("MultipleBaseActualRoot", "Multiple base object nodes found.");
+		return nullptr;
+	}
+
+	if (!ActualRoot)
+	{
+		ErrorMsg = LOCTEXT("NoActualRootBase", "No base object node found in root Customizable Object. Object not built.");
+		return nullptr;
+	}
 
 	ArrayAlreadyProcessedChild.Empty();
 
@@ -572,15 +590,6 @@ mu::NodeObjectPtr FCustomizableObjectCompiler::GenerateMutableRoot(
 		{
 			UE_LOG(LogMutable, Verbose, TEXT("PROFILE: [ %16.8f ] Begin search for children."), FPlatformTime::Seconds());
 
-			TArray<UCustomizableObject*> VisitedObjects;
-			ActualRootObject = Root->ParentObject ? GetFullGraphRootObject(Root, VisitedObjects) : Object;
-
-			if (Root->ParentObject != nullptr)
-			{
-				VisitedObjects.Empty();
-				ActualRoot = GetFullGraphRootNodeObject(ActualRoot, VisitedObjects);
-			}
-
 			// The object doesn't reference a root object but is a root object, look for all the objects that reference it and get their root nodes
 			FAssetRegistryModule& AssetRegistryModule = FModuleManager::LoadModuleChecked<FAssetRegistryModule>("AssetRegistry");
 			ProcessChildObjectsRecursively(ActualRootObject, GenerationContext);
@@ -638,16 +647,6 @@ mu::NodeObjectPtr FCustomizableObjectCompiler::GenerateMutableRoot(
 				GenerationContext.GroupIdToExternalNodeMap.Add(ArrayNodeObject[i]->ParentObjectGroupId, ArrayNodeObject[i]);
 			}
 		}
-
-		TArray<UCustomizableObject*> VisitedObjects;
-
-		if (Root->ParentObject != nullptr)
-		{
-			ActualRoot = GetFullGraphRootNodeObject(ActualRoot, VisitedObjects);
-		}
-
-		VisitedObjects.Empty();
-		ActualRootObject = Root->ParentObject ? GetFullGraphRootObject(Root, VisitedObjects) : Object;
 	}
 
 	// Ensure that the CO has a valid AutoLODStrategy on the ActualRoot.
diff --git a/Engine/Plugins/Experimental/Mutable/Source/CustomizableObjectEditor/Private/MuCOE/GenerateMutableSource/GenerateMutableSource.h b/Engine/Plugins/Experimental/Mutable/Source/CustomizableObjectEditor/Private/MuCOE/GenerateMutableSource/GenerateMutableSource.h
--- a/Engine/Plugins/Experimental/Mutable/Source/CustomizableObjectEditor/Private/MuCOE/GenerateMutableSource/GenerateMutableSource.h
+++ b/Engine/Plugins/Experimental/Mutable/Source/CustomizableObjectEditor/Private/MuCOE/GenerateMutableSource/GenerateMutableSource.h
@@ -899,6 +899,9 @@ struct FMutableGraphGenerationContext
 
 	/** Map to relate a Composite Data Table Row and its original DataTable */
 	TMap<UDataTable*,TMap<FName, TArray<UDataTable*>>> CompositeDataTableRowToOriginalDataTableMap;
+
+	/** Version Bridge of the root object */
+	TObjectPtr<UObject> RootVersionBridge;
 };
 
 /** Pin Data scope wrapper. Pops the pin data on scope exit. */
diff --git a/Engine/Plugins/Experimental/Mutable/Source/CustomizableObjectEditor/Private/MuCOE/GenerateMutableSource/GenerateMutableSourceTable.cpp b/Engine/Plugins/Experimental/Mutable/Source/CustomizableObjectEditor/Private/MuCOE/GenerateMutableSource/GenerateMutableSourceTable.cpp
--- a/Engine/Plugins/Experimental/Mutable/Source/CustomizableObjectEditor/Private/MuCOE/GenerateMutableSource/GenerateMutableSourceTable.cpp
+++ b/Engine/Plugins/Experimental/Mutable/Source/CustomizableObjectEditor/Private/MuCOE/GenerateMutableSource/GenerateMutableSourceTable.cpp
@@ -20,6 +20,7 @@
 #include "MuCOE/GenerateMutableSource/GenerateMutableSourceMesh.h"
 #include "MuCOE/Nodes/CustomizableObjectNodeTable.h"
 #include "MuCOE/Nodes/CustomizableObjectNodeAnimationPose.h"
+#include "MuCOE/GraphTraversal.h"
 #include "MuR/Mesh.h"
 #include "PhysicsEngine/PhysicsAsset.h"
 #include "Rendering/SkeletalMeshLODModel.h"
@@ -832,7 +833,7 @@ void RestrictRowContentByVersion( TArray<FName>& InOutRowNames, const UDataTable
 		return;
 	}
 
-	ICustomizableObjectVersionBridgeInterface* CustomizableObjectVersionBridgeInterface = Cast<ICustomizableObjectVersionBridgeInterface>(GenerationContext.Object->VersionBridge);
+	ICustomizableObjectVersionBridgeInterface* CustomizableObjectVersionBridgeInterface = Cast<ICustomizableObjectVersionBridgeInterface>(GenerationContext.RootVersionBridge);
 	if (!CustomizableObjectVersionBridgeInterface)
 	{
 		const FString Message = "Found a data table with at least a row with a Custom Version asset but the Root Object does not have a Version Bridge asset assigned.";
