[mutable] Modify some nodes to remove layout block references to other nodes and replace them with full layout definitions.
- The "Remove mesh" node now has its own layout unrelated to the one in the original mesh.
- The "Edit material" node now has its own layout, unrelated to the one in the mesh in the material being edited.
- Since layouts in "Edit Material" are now independent of actual layout blocks, a new mask needs to be generated and combined with existing texture patching masks.
- Fix image multiply operation precission problem
#rnx

diff --git a/Engine/Plugins/Experimental/Mutable/Source/CustomizableObject/Internal/MuCO/CustomizableObjectCustomVersion.h b/Engine/Plugins/Experimental/Mutable/Source/CustomizableObject/Internal/MuCO/CustomizableObjectCustomVersion.h
--- a/Engine/Plugins/Experimental/Mutable/Source/CustomizableObject/Internal/MuCO/CustomizableObjectCustomVersion.h
+++ b/Engine/Plugins/Experimental/Mutable/Source/CustomizableObject/Internal/MuCO/CustomizableObjectCustomVersion.h
@@ -128,6 +128,8 @@ struct CUSTOMIZABLEOBJECT_API FCustomizableObjectCustomVersion
 
 		FixedMultilayerMaterialIds,
 
+		UseUVRects, // Instead of pointing at layout blocks in the edited parent, define custom absolute UV layout rects.
+
 		// -----<new versions can be added above this line>-------------------------------------------------
 		VersionPlusOne,
 		LatestVersion = VersionPlusOne - 1
diff --git a/Engine/Plugins/Experimental/Mutable/Source/CustomizableObjectEditor/CustomizableObjectEditor.Build.cs b/Engine/Plugins/Experimental/Mutable/Source/CustomizableObjectEditor/CustomizableObjectEditor.Build.cs
--- a/Engine/Plugins/Experimental/Mutable/Source/CustomizableObjectEditor/CustomizableObjectEditor.Build.cs
+++ b/Engine/Plugins/Experimental/Mutable/Source/CustomizableObjectEditor/CustomizableObjectEditor.Build.cs
@@ -11,6 +11,7 @@ public class CustomizableObjectEditor : ModuleRules
 		ShortName = "MuCOE";
 
 		DefaultBuildSettings = BuildSettingsVersion.V2;
+		//bUseUnity = false;
 
 		PrivateIncludePathModuleNames.AddRange(
 			new string[] { 
diff --git a/Engine/Plugins/Experimental/Mutable/Source/CustomizableObjectEditor/Private/MuCOE/CustomizableObjectEditor.h b/Engine/Plugins/Experimental/Mutable/Source/CustomizableObjectEditor/Private/MuCOE/CustomizableObjectEditor.h
--- a/Engine/Plugins/Experimental/Mutable/Source/CustomizableObjectEditor/Private/MuCOE/CustomizableObjectEditor.h
+++ b/Engine/Plugins/Experimental/Mutable/Source/CustomizableObjectEditor/Private/MuCOE/CustomizableObjectEditor.h
@@ -261,8 +261,6 @@ public:
 	/** Getter of CustomizableObjectEditorAdvancedPreviewSettings */
 	TSharedPtr<class SCustomizableObjectEditorAdvancedPreviewSettings> GetCustomizableObjectEditorAdvancedPreviewSettings();
 
-	TSharedPtr<class SCustomizableObjectNodeLayoutBlocksEditor> GetLayoutBlocksEditor() { return LayoutBlocksEditor; }
-
 	/** Debug the object as a raw mutable data in the internal tools. */
 	void DebugObject() const;
 
@@ -424,10 +422,6 @@ private:
 	TSharedPtr<SGraphEditor> GraphEditor;
 	TSharedPtr<class IDetailsView> GraphNodeDetailsView;
 
-	/** This widget shows the layout blocks and they can be edited. It is kept alive for some problems re-registering toolbar commands.*/
-	TSharedPtr<class SCustomizableObjectNodeLayoutBlocksEditor> LayoutBlocksEditor;
-	TSharedPtr<class SCustomizableObjectNodeLayoutBlocksSelector> LayoutBlocksSelector;
-
 	/** Widget to select which node pins are visible. */
 	TSharedPtr<class SCustomizableObjectNodePinViewer> NodePinViewer;
 	
diff --git a/Engine/Plugins/Experimental/Mutable/Source/CustomizableObjectEditor/Private/MuCOE/GenerateMutableSource/GenerateMutableSourceLayout.cpp b/Engine/Plugins/Experimental/Mutable/Source/CustomizableObjectEditor/Private/MuCOE/GenerateMutableSource/GenerateMutableSourceLayout.cpp
--- a/Engine/Plugins/Experimental/Mutable/Source/CustomizableObjectEditor/Private/MuCOE/GenerateMutableSource/GenerateMutableSourceLayout.cpp
+++ b/Engine/Plugins/Experimental/Mutable/Source/CustomizableObjectEditor/Private/MuCOE/GenerateMutableSource/GenerateMutableSourceLayout.cpp
@@ -11,7 +11,7 @@
 #define LOCTEXT_NAMESPACE "CustomizableObjectEditor"
 
 
-mu::Ptr<mu::NodeLayout> GenerateMutableSourceLayout(const UEdGraphPin * Pin, FMutableGraphGenerationContext& GenerationContext, bool bLinkedToExtendMaterial)
+mu::Ptr<mu::NodeLayout> GenerateMutableSourceLayout(const UEdGraphPin * Pin, FMutableGraphGenerationContext& GenerationContext, bool bIgnoreLayoutWarning)
 {
 	check(Pin)
 	RETURN_ON_CYCLE(*Pin, GenerationContext)
@@ -53,7 +53,7 @@ mu::Ptr<mu::NodeLayout> GenerateMutableSourceLayout(const UEdGraphPin * Pin, FMu
 		}
 
 		bool bWasEmpty = false;
-		Result = CreateMutableLayoutNode(GenerationContext, TypedNodeBlocks->Layout, bLinkedToExtendMaterial,bWasEmpty);
+		Result = CreateMutableLayoutNode(GenerationContext, TypedNodeBlocks->Layout, bIgnoreLayoutWarning,bWasEmpty);
 		if (bWasEmpty)
 		{
 			FString msg = "Layout without any block found. A grid sized block will be used instead.";
diff --git a/Engine/Plugins/Experimental/Mutable/Source/CustomizableObjectEditor/Private/MuCOE/GenerateMutableSource/GenerateMutableSourceLayout.h b/Engine/Plugins/Experimental/Mutable/Source/CustomizableObjectEditor/Private/MuCOE/GenerateMutableSource/GenerateMutableSourceLayout.h
--- a/Engine/Plugins/Experimental/Mutable/Source/CustomizableObjectEditor/Private/MuCOE/GenerateMutableSource/GenerateMutableSourceLayout.h
+++ b/Engine/Plugins/Experimental/Mutable/Source/CustomizableObjectEditor/Private/MuCOE/GenerateMutableSource/GenerateMutableSourceLayout.h
@@ -9,7 +9,7 @@ class UEdGraphPin;
 struct FMutableGraphGenerationContext;
 
 
-extern mu::Ptr<mu::NodeLayout> GenerateMutableSourceLayout(const UEdGraphPin*, FMutableGraphGenerationContext&, bool bLinkedToExtendMaterial = false);
+extern mu::Ptr<mu::NodeLayout> GenerateMutableSourceLayout(const UEdGraphPin*, FMutableGraphGenerationContext&, bool bIgnoreLayoutWarning = false);
 
 extern mu::Ptr<mu::NodeLayout> CreateMutableLayoutNode(FMutableGraphGenerationContext& GenerationContext, const UCustomizableObjectLayout*, bool bIgnoreLayoutWarnings, bool& bWasEmpty);
 
diff --git a/Engine/Plugins/Experimental/Mutable/Source/CustomizableObjectEditor/Private/MuCOE/GenerateMutableSource/GenerateMutableSourceSurface.cpp b/Engine/Plugins/Experimental/Mutable/Source/CustomizableObjectEditor/Private/MuCOE/GenerateMutableSource/GenerateMutableSourceSurface.cpp
--- a/Engine/Plugins/Experimental/Mutable/Source/CustomizableObjectEditor/Private/MuCOE/GenerateMutableSource/GenerateMutableSourceSurface.cpp
+++ b/Engine/Plugins/Experimental/Mutable/Source/CustomizableObjectEditor/Private/MuCOE/GenerateMutableSource/GenerateMutableSourceSurface.cpp
@@ -26,6 +26,7 @@
 #include "MuCOE/GenerateMutableSource/GenerateMutableSourceImage.h"
 #include "MuCOE/GenerateMutableSource/GenerateMutableSourceMesh.h"
 #include "MuCOE/GenerateMutableSource/GenerateMutableSourceTable.h"
+#include "MuCOE/GenerateMutableSource/GenerateMutableSourceLayout.h"
 #include "MuCOE/GraphTraversal.h"
 #include "MuCOE/MutableUtils.h"
 #include "MuCOE/Nodes/CustomizableObjectNodeCopyMaterial.h"
@@ -1101,10 +1102,10 @@ mu::Ptr<mu::NodeSurface> GenerateMutableSourceSurface(const UEdGraphPin * Pin, F
 			}
 
 			// Parent, probably generated, will be retrieved from the cache
-			mu::NodeSurfacePtr ParentNode = GenerateMutableSourceSurface(ParentMaterialNode->OutputPin(), GenerationContext);
+			mu::Ptr<mu::NodeSurface> ParentNode = GenerateMutableSourceSurface(ParentMaterialNode->OutputPin(), GenerationContext);
 			SurfNode->Parent = ParentNode;
 
-			mu::NodeMeshPtr AddMeshNode;
+			mu::Ptr<mu::NodeMesh> AddMeshNode;
 			FMutableGraphMeshGenerationData MeshData;
 			if (const UEdGraphPin* ConnectedPin = FollowInputPin(*TypedNodeExt->AddMeshPin()))
 			{
@@ -1113,24 +1114,22 @@ mu::Ptr<mu::NodeSurface> GenerateMutableSourceSurface(const UEdGraphPin * Pin, F
 
 			if (AddMeshNode)
 			{
-				mu::NodeMeshPtr MeshPtr = AddMeshNode;
+				mu::Ptr<mu::NodeMesh> MeshPtr = AddMeshNode;
 
 				const TArray<UCustomizableObjectLayout*> Layouts = TypedNodeExt->GetLayouts();
 				
 				if (Layouts.Num())
 				{
-					//TODO: Implement support for multiple UV channels (e.g. Add warning for vertices which have a block in a layout but not in the other)
-					mu::NodeMeshFragmentPtr MeshFrag = new mu::NodeMeshFragment();
+					mu::Ptr<mu::NodeMeshFragment> MeshFrag = new mu::NodeMeshFragment();
 				
-					MeshFrag->SetMesh(MeshPtr);
-					MeshFrag->SetLayoutOrGroup(0);
-					MeshFrag->SetFragmentType(mu::NodeMeshFragment::FT_LAYOUT_BLOCKS);
-					MeshFrag->SetBlockCount(Layouts[0]->Blocks.Num());
-				
-					for (int i = 0; i < Layouts[0]->Blocks.Num(); ++i)
-					{
-						MeshFrag->SetBlock(i, i);
-					}
+					MeshFrag->SourceMesh = MeshPtr;
+					//TODO: Implement support for multiple UV channels (e.g. Add warning for vertices which have a block in a layout but not in the other)
+					MeshFrag->LayoutIndex = 0;
+
+					// For this case we don't want to create another layout: we will use the one defined in the mesh to be added since we want to add
+					// any block defined there.
+					//bool bWasEmpty = false;
+					//MeshFrag->Layout = CreateMutableLayoutNode(GenerationContext, Layouts[MeshFrag->LayoutIndex], true, bWasEmpty);
 				
 					MeshPtr = MeshFrag;
 				}
@@ -1270,9 +1269,7 @@ mu::Ptr<mu::NodeSurface> GenerateMutableSourceSurface(const UEdGraphPin * Pin, F
 		}
 		else
 		{
-			// Parent, probably generated, will be retrieved from the cache
-			mu::NodeSurfacePtr ParentNode = GenerateMutableSourceSurface(ParentMaterialNode->OutputPin(), GenerationContext);
-			SurfNode->Parent = ParentNode.get();
+			SurfNode->Parent = GenerateMutableSourceSurface(ParentMaterialNode->OutputPin(), GenerationContext);
 
 			const UEdGraphPin* BaseSourcePin = FindMeshBaseSource(*ParentMaterialNode->OutputPin(), false);
 
@@ -1283,57 +1280,20 @@ mu::Ptr<mu::NodeSurface> GenerateMutableSourceSurface(const UEdGraphPin * Pin, F
 			else
 			{
 				FMutableGraphMeshGenerationData DummyMeshData;
-				mu::NodeMeshPtr BaseSourceMesh = GenerateMutableSourceMesh(BaseSourcePin, GenerationContext, DummyMeshData, false, false);
-
-				mu::NodeMeshFragmentPtr MeshFrag = new mu::NodeMeshFragment();
-				MeshFrag->SetMesh(BaseSourceMesh);
-				MeshFrag->SetLayoutOrGroup(TypedNodeRemBlocks->ParentLayoutIndex);
-				MeshFrag->SetFragmentType(mu::NodeMeshFragment::FT_LAYOUT_BLOCKS);
+				mu::Ptr<mu::NodeMesh> SourceMesh = GenerateMutableSourceMesh(BaseSourcePin, GenerationContext, DummyMeshData, false, false);
+				bool bWasEmpty = false;
+				mu::Ptr<mu::NodeLayout> SourceLayout = CreateMutableLayoutNode(GenerationContext, TypedNodeRemBlocks->Layout, true, bWasEmpty);
+
+				mu::Ptr<mu::NodeMeshFragment> MeshFrag = new mu::NodeMeshFragment();
+				MeshFrag->SourceMesh = SourceMesh;
+				MeshFrag->Layout = SourceLayout;
+				MeshFrag->LayoutIndex = TypedNodeRemBlocks->ParentLayoutIndex;
 				MeshFrag->SetMessageContext(TypedNodeRemBlocks);
 
-				// Add the block indices
-				TArray<int> FoundBlockIndices;
-				const TArray<UCustomizableObjectLayout*> Layouts = ParentMaterialNode->GetLayouts();
-
-				if (!Layouts.IsValidIndex(TypedNodeRemBlocks->ParentLayoutIndex))
-				{
-					GenerationContext.Compiler->CompilerLog(LOCTEXT("RemoveMeshNode Invalid layout", "Node refers to an invalid texture layout. Add a layout to the selected parent material, then select some block(s)."), Node);
-					UE_LOG(LogMutable, Warning, TEXT("In object [%s] UCustomizableObjectNodeRemoveMeshBlocks refers to an invalid texture layout index %d. Parent node has %d layouts."),
-						*GenerationContext.Object->GetName(), TypedNodeRemBlocks->ParentLayoutIndex, Layouts.Num());
-				}
-				else
-				{
-					const UCustomizableObjectLayout* Layout = Layouts[TypedNodeRemBlocks->ParentLayoutIndex];
-
-					if (Cast<UCustomizableObjectNodeMaterial>(ParentMaterialNode))
-					{
-						for (const FGuid& Id : TypedNodeRemBlocks->BlockIds)
-						{
-							int BlockIndex = Layout->FindBlock(Id);
-							if (Layout->Blocks.IsValidIndex(BlockIndex))
-							{
-								FoundBlockIndices.Add(BlockIndex);
-							}
-							else
-							{
-								GenerationContext.Compiler->CompilerLog(LOCTEXT("RemoveMeshNode InvalidLayoutBlock", "Node refers to an invalid layout block."), Node);
-								UE_LOG(LogMutable, Warning, TEXT("In object [%s] UCustomizableObjectNodeRemoveMeshBlocks refers to an invalid layout block id %s. Parent node has %d blocks."),
-									*GenerationContext.Object->GetName(), *Id.ToString(), Layout->Blocks.Num());
-							}
-						}
-					}
-				}
-
-				MeshFrag->SetBlockCount(FoundBlockIndices.Num());
-				for (int i = 0; i < FoundBlockIndices.Num(); ++i)
-				{
-					MeshFrag->SetBlock(i, FoundBlockIndices[i]);
-				}
-
-				mu::NodePatchMeshPtr MeshPatch = new mu::NodePatchMesh();
+				mu::Ptr<mu::NodePatchMesh> MeshPatch = new mu::NodePatchMesh();
 				MeshPatch->SetRemove(MeshFrag.get());
-				SurfNode->Mesh = MeshPatch.get();
 				MeshPatch->SetMessageContext(Node);
+				SurfNode->Mesh = MeshPatch;
 			}
 
 			AddModifierToSharedSurface(GenerationContext, ParentMaterialNode, *TypedNodeRemBlocks);
@@ -1353,7 +1313,7 @@ mu::Ptr<mu::NodeSurface> GenerateMutableSourceSurface(const UEdGraphPin * Pin, F
 		else
 		{
 			// Parent, probably generated, will be retrieved from the cache
-			mu::NodeSurfacePtr ParentNode = GenerateMutableSourceSurface(ParentMaterialNode->OutputPin(), GenerationContext);
+			mu::Ptr<mu::NodeSurface> ParentNode = GenerateMutableSourceSurface(ParentMaterialNode->OutputPin(), GenerationContext);
 			SurfNode->Parent = ParentNode;
 
 			const int32 NumImages = ParentMaterialNode->GetNumParameters(EMaterialParameterType::Texture);
@@ -1372,68 +1332,32 @@ mu::Ptr<mu::NodeSurface> GenerateMutableSourceSurface(const UEdGraphPin * Pin, F
 					ImagePatchNode->SetMessageContext(Node);
 
 					// \todo: expose these two options?
-					ImagePatchNode->SetBlendType(mu::EBlendType::BT_BLEND);
-					ImagePatchNode->SetApplyToAlphaChannel(true);
+					ImagePatchNode->BlendType = mu::EBlendType::BT_BLEND;
+					ImagePatchNode->bApplyToAlpha=true;
 
 					// ReferenceTextureSize is used to limit the size of textures contributing to the final image.
 					const int32 ReferenceTextureSize = GetBaseTextureSize(GenerationContext, ParentMaterialNode, ImageIndex);
 
-					mu::NodeImagePtr ImageNode = GenerateMutableSourceImage(ConnectedImagePin, GenerationContext, ReferenceTextureSize);
-					ImagePatchNode->SetImage(ImageNode);
+					ImagePatchNode->Image = GenerateMutableSourceImage(ConnectedImagePin, GenerationContext, ReferenceTextureSize);
 
 					const UEdGraphPin* ImageMaskPin = TypedNodeEdit->GetUsedImageMaskPin(ImageId);
 					check(ImageMaskPin); // Ensured when reconstructing EditMaterial nodes. If it fails, something is wrong.
 					
 					if (const UEdGraphPin* ConnectedMaskPin = FollowInputPin(*ImageMaskPin))
 					{
-						mu::NodeImagePtr MaskNode = GenerateMutableSourceImage(ConnectedMaskPin, GenerationContext, ReferenceTextureSize);
-						ImagePatchNode->SetMask(MaskNode);
-					}
-
-					// Add the block indices
-					TArray<int> FoundBlockIndices;
-					const TArray<UCustomizableObjectLayout*> Layouts = ParentMaterialNode->GetLayouts();
-
-					if (!Layouts.IsValidIndex(TypedNodeEdit->ParentLayoutIndex))
-					{
-						GenerationContext.Compiler->CompilerLog(LOCTEXT("EditNode Invalid layout", "Node refers to an invalid texture layout."), Node);
-						UE_LOG(LogMutable, Warning, TEXT("In object [%s] UCustomizableObjectNodeEditMaterial refers to an invalid texture layout index %d. Parent node has %d layouts."),
-							*GenerationContext.Object->GetName(), TypedNodeEdit->ParentLayoutIndex, Layouts.Num());
-					}
-					else
-					{
-						if (const UCustomizableObjectNodeMaterial* ParentMaterial = Cast<UCustomizableObjectNodeMaterial>(ParentMaterialNode))
-						{
-							if (TypedNodeEdit->ParentLayoutIndex == ParentMaterial->GetImageUVLayout(ImageIndex))
-							{
-								const UCustomizableObjectLayout* Layout = Layouts[TypedNodeEdit->ParentLayoutIndex];
-
-								for (const FGuid& Id : TypedNodeEdit->BlockIds)
-								{
-									int BlockIndex = Layout->FindBlock(Id);
-									if (Layout->Blocks.IsValidIndex(BlockIndex))
-									{
-										FoundBlockIndices.Add(BlockIndex);
-									}
-									else
-									{
-										GenerationContext.Compiler->CompilerLog(LOCTEXT("EditNode InvalidLayoutBlock", "Node refers to an invalid layout block."), Node);
-										UE_LOG(LogMutable, Warning, TEXT("In object [%s] UCustomizableObjectNodeEditMaterial refers to an invalid layout block id %s. Parent node has %d blocks."),
-											*GenerationContext.Object->GetName(), *Id.ToString(), Layout->Blocks.Num());
-									}
-								}
-							}
-							else
-							{
-								GenerationContext.Compiler->CompilerLog(LOCTEXT("InvalidLayoutTexture", "Texture layout does not match the layout of Edit Material node or parent material changed."), Cast<UCustomizableObjectNode>(ConnectedImagePin->GetOwningNode()));
-							}
-						}
+						ImagePatchNode->Mask = GenerateMutableSourceImage(ConnectedMaskPin, GenerationContext, ReferenceTextureSize);
 					}
 
-					ImagePatchNode->SetBlockCount(FoundBlockIndices.Num());
-					for (int BlockIndex = 0; BlockIndex < FoundBlockIndices.Num(); ++BlockIndex)
+					// Add the blocks to patch
+					FIntPoint GridSize = TypedNodeEdit->Layout->GetGridSize();
+					FVector2f GridSizeF = FVector2f(GridSize);
+					ImagePatchNode->Blocks.Reserve(TypedNodeEdit->Layout->Blocks.Num());
+					for (const FCustomizableObjectLayoutBlock& LayoutBlock: TypedNodeEdit->Layout->Blocks)
 					{
-						ImagePatchNode->SetBlock(BlockIndex, FoundBlockIndices[BlockIndex]);
+						FBox2f Rect;
+						Rect.Min = FVector2f(LayoutBlock.Min) / GridSizeF;
+						Rect.Max = FVector2f(LayoutBlock.Max) / GridSizeF;
+						ImagePatchNode->Blocks.Add(Rect);
 					}
 
 					SurfNode->Textures[ImageIndex].Patch = ImagePatchNode;
diff --git a/Engine/Plugins/Experimental/Mutable/Source/CustomizableObjectEditor/Private/MuCOE/Nodes/CustomizableObjectNodeEditLayoutBlocks.cpp b/Engine/Plugins/Experimental/Mutable/Source/CustomizableObjectEditor/Private/MuCOE/Nodes/CustomizableObjectNodeEditLayoutBlocks.cpp
new file mode 100644
--- /dev/null
+++ b/Engine/Plugins/Experimental/Mutable/Source/CustomizableObjectEditor/Private/MuCOE/Nodes/CustomizableObjectNodeEditLayoutBlocks.cpp
@@ -0,0 +1,17 @@
+// Copyright Epic Games, Inc. All Rights Reserved.
+
+#include "MuCOE/Nodes/CustomizableObjectNodeEditLayoutBlocks.h"
+
+
+#define LOCTEXT_NAMESPACE "CustomizableObjectEditor"
+
+
+UCustomizableObjectNodeEditLayoutBlocks::UCustomizableObjectNodeEditLayoutBlocks()
+	: Super()
+{
+	Layout = CreateDefaultSubobject<UCustomizableObjectLayout>(FName("CustomizableObjectLayout"));
+	Layout->Blocks.Empty();
+}
+
+#undef LOCTEXT_NAMESPACE
+
diff --git a/Engine/Plugins/Experimental/Mutable/Source/CustomizableObjectEditor/Private/MuCOE/Nodes/CustomizableObjectNodeEditMaterial.cpp b/Engine/Plugins/Experimental/Mutable/Source/CustomizableObjectEditor/Private/MuCOE/Nodes/CustomizableObjectNodeEditMaterial.cpp
--- a/Engine/Plugins/Experimental/Mutable/Source/CustomizableObjectEditor/Private/MuCOE/Nodes/CustomizableObjectNodeEditMaterial.cpp
+++ b/Engine/Plugins/Experimental/Mutable/Source/CustomizableObjectEditor/Private/MuCOE/Nodes/CustomizableObjectNodeEditMaterial.cpp
@@ -94,8 +94,6 @@ void UCustomizableObjectNodeEditMaterial::SetParentNode(UCustomizableObject* Obj
 	Super::SetParentNode(Object, NodeId);
 	
 	PostSetParentNodeWork(Object, NodeId);
-
-	SelectAllLayoutBlocks();
 }
 
 
@@ -107,7 +105,7 @@ void UCustomizableObjectNodeEditMaterial::BackwardsCompatibleFixup()
 
 	// Convert deprecated node index list to the node id list.
 	if (CustomizableObjectCustomVersion < FCustomizableObjectCustomVersion::PostLoadToCustomVersion
-		&& BlockIds.Num() < Blocks_DEPRECATED.Num())
+		&& BlockIds_DEPRECATED.Num() < Blocks_DEPRECATED.Num())
 	{
 		UCustomizableObjectNodeMaterialBase* ParentMaterialNode = GetParentMaterialNode();
 
@@ -120,19 +118,19 @@ void UCustomizableObjectNodeEditMaterial::BackwardsCompatibleFixup()
 		}
 		else
 		{
-			UCustomizableObjectLayout* Layout = Layouts[ParentLayoutIndex];
+			UCustomizableObjectLayout* ParentLayout = Layouts[ParentLayoutIndex];
 
 			if (Cast<UCustomizableObjectNodeMaterial>(ParentMaterialNode))
 			{
-				for (int IndexIndex = BlockIds.Num(); IndexIndex < Blocks_DEPRECATED.Num(); ++IndexIndex)
+				for (int32 IndexIndex = BlockIds_DEPRECATED.Num(); IndexIndex < Blocks_DEPRECATED.Num(); ++IndexIndex)
 				{
-					int BlockIndex = Blocks_DEPRECATED[IndexIndex];
-					if (Layout->Blocks.IsValidIndex(BlockIndex) )
+					int32 BlockIndex = Blocks_DEPRECATED[IndexIndex];
+					if (ParentLayout->Blocks.IsValidIndex(BlockIndex) )
 					{
-						const FGuid Id = Layout->Blocks[BlockIndex].Id;
+						const FGuid Id = ParentLayout->Blocks[BlockIndex].Id;
 						if (Id.IsValid())
 						{
-							BlockIds.Add(Id);
+							BlockIds_DEPRECATED.Add(Id);
 						}
 						else
 						{
@@ -143,13 +141,65 @@ void UCustomizableObjectNodeEditMaterial::BackwardsCompatibleFixup()
 					else
 					{
 						UE_LOG(LogMutable, Warning, TEXT("[%s] UCustomizableObjectNodeEditMaterial refers to an invalid layout block index %d. Parent node has %d blocks."),
-							*GetOutermost()->GetName(), BlockIndex, Layout->Blocks.Num());
+							*GetOutermost()->GetName(), BlockIndex, ParentLayout->Blocks.Num());
 					}
 				}
 			}
 		}
 	}
 	
+	// Convert deprecated node id list to absolute rect list.
+	if (CustomizableObjectCustomVersion < FCustomizableObjectCustomVersion::UseUVRects)
+	{
+		// If we are here, it means this node was loaded from a version that didn't have it's own layout.
+		check(Layout->Blocks.IsEmpty());
+
+		UCustomizableObjectNodeMaterialBase* ParentMaterialNode = GetParentMaterialNode();
+
+		TArray<UCustomizableObjectLayout*> ParentLayouts = ParentMaterialNode->GetLayouts();
+
+		if (!ParentLayouts.IsValidIndex(ParentLayoutIndex))
+		{
+			UE_LOG(LogMutable, Warning, TEXT("[%s] UCustomizableObjectNodeRemoveMeshBlocks refers to an invalid texture layout index %d. Parent node has %d layouts."),
+				*GetOutermost()->GetName(), ParentLayoutIndex, ParentLayouts.Num());
+		}
+		else
+		{
+			UCustomizableObjectLayout* ParentLayout = ParentLayouts[ParentLayoutIndex];
+			FIntPoint GridSize = ParentLayout->GetGridSize();
+
+			Layout->SetGridSize(GridSize);
+
+			if (Cast<UCustomizableObjectNodeMaterial>(ParentMaterialNode))
+			{
+				for (const FGuid& BlockId : BlockIds_DEPRECATED)
+				{
+					for (const FCustomizableObjectLayoutBlock& ParentBlock : ParentLayout->Blocks)
+					{
+						if (ParentBlock.Id == BlockId)
+						{
+							FCustomizableObjectLayoutBlock NewBlock;
+							NewBlock = ParentBlock;
+
+							// Clear some unnecessary data.
+							NewBlock.bReduceBothAxes = false;
+							NewBlock.bReduceByTwo = false;
+							NewBlock.Priority = 0;
+
+							Layout->Blocks.Add(NewBlock);
+						}
+					}
+				}
+
+				if (Layout->Blocks.Num() != BlockIds_DEPRECATED.Num())
+				{
+					UE_LOG(LogMutable, Warning, TEXT("[%s] UCustomizableObjectNodeRemoveMeshBlocks refers to %d invalid layout block. It has been ignored during version upgrade."),
+						*GetOutermost()->GetName(), int32(BlockIds_DEPRECATED.Num() - Layout->Blocks.Num()));
+				}
+			}
+		}
+	}
+
 	if (CustomizableObjectCustomVersion < FCustomizableObjectCustomVersion::AutomaticNodeMaterial)
 	{
 		if (const UCustomizableObjectNodeMaterialBase* ParentMaterial = GetParentMaterialNode())
@@ -291,32 +341,8 @@ FCustomizableObjectNodeParentedMaterial& UCustomizableObjectNodeEditMaterial::Ge
 void UCustomizableObjectNodeEditMaterial::SetLayoutIndex(const int32 LayoutIndex)
 {
 	ParentLayoutIndex = LayoutIndex;
-	SelectAllLayoutBlocks();
-}
-
-
-void UCustomizableObjectNodeEditMaterial::SelectAllLayoutBlocks()
-{
-	BlockIds.Empty();
-
-	// On reset parent material we want to select all nodes 
-	if (UCustomizableObjectNodeMaterialBase* ParentMaterialNode = GetParentMaterialNode())
-	{
-		TArray<UCustomizableObjectLayout*> Layouts = ParentMaterialNode->GetLayouts();
-
-		if (Layouts.IsValidIndex(ParentLayoutIndex))
-		{
-			UCustomizableObjectLayout* Layout = Layouts[ParentLayoutIndex];
-
-			for (const FCustomizableObjectLayoutBlock& Block : Layout->Blocks)
-			{
-				BlockIds.AddUnique(Block.Id);
-			}
-		}
-	}
 }
 
 
-
 #undef LOCTEXT_NAMESPACE
 
diff --git a/Engine/Plugins/Experimental/Mutable/Source/CustomizableObjectEditor/Private/MuCOE/Nodes/CustomizableObjectNodeEditMaterialDetails.cpp b/Engine/Plugins/Experimental/Mutable/Source/CustomizableObjectEditor/Private/MuCOE/Nodes/CustomizableObjectNodeEditMaterialDetails.cpp
--- a/Engine/Plugins/Experimental/Mutable/Source/CustomizableObjectEditor/Private/MuCOE/Nodes/CustomizableObjectNodeEditMaterialDetails.cpp
+++ b/Engine/Plugins/Experimental/Mutable/Source/CustomizableObjectEditor/Private/MuCOE/Nodes/CustomizableObjectNodeEditMaterialDetails.cpp
@@ -6,10 +6,12 @@
 #include "DetailWidgetRow.h"
 #include "IDetailsView.h"
 #include "MuCOE/Nodes/CustomizableObjectNodeEditMaterial.h"
-#include "MuCOE/SCustomizableObjectNodeLayoutBlocksSelector.h"
+#include "MuCOE/Nodes/CustomizableObjectNodeMaterialBase.h"
+#include "MuCOE/SCustomizableObjectNodeLayoutBlocksEditor.h"
 #include "MuCOE/PinViewer/SPinViewer.h"
 #include "Widgets/Layout/SBox.h"
 #include "Widgets/Text/STextBlock.h"
+#include "UObject/Package.h"
 
 class FString;
 
@@ -37,24 +39,37 @@ void FCustomizableObjectNodeEditMaterialDetails::CustomizeDetails( IDetailLayout
 
 	if (Node)
 	{
-		UCustomizableObjectNodeEditLayoutBlocks* NodeEditLayout = Cast<UCustomizableObjectNodeEditLayoutBlocks>(Node);
-		if (NodeEditLayout)
-		{
-			// Add blocks selector
-			LayoutBlocksSelector = SNew(SCustomizableObjectNodeLayoutBlocksSelector);
+		// Add blocks selector
+		LayoutBlocksEditor = SNew(SCustomizableObjectNodeLayoutBlocksEditor);
 
-			BlocksCategory.AddCustomRow(LOCTEXT("BlocksDetails_BlockInstructions", "BlockInstructions"))
+		BlocksCategory.AddCustomRow(LOCTEXT("BlocksDetails_BlockInstructions", "BlockInstructions"))
+		[
+			SNew(SBox)
+			.HeightOverride(700.0f)
+			.WidthOverride(700.0f)
 			[
-				SNew(SBox)
-				.HeightOverride(700.0f)
-				.WidthOverride(700.0f)
-				[
-					LayoutBlocksSelector.ToSharedRef()
-				]
-			];
-
-			LayoutBlocksSelector->SetSelectedNode(NodeEditLayout);
+				LayoutBlocksEditor.ToSharedRef()
+			]
+		];
+
+		// Try to find the parent layout, because we want to show its UVs in the widget
+		UCustomizableObjectLayout* ParentLayout = nullptr;
+		if (UCustomizableObjectNodeMaterialBase* ParentMaterialNode = Node->GetParentMaterialNode())
+		{
+			TArray<UCustomizableObjectLayout*> Layouts = ParentMaterialNode->GetLayouts();
+
+			if (!Layouts.IsValidIndex(Node->ParentLayoutIndex))
+			{
+				UE_LOG(LogMutable, Warning, TEXT("[%s] UCustomizableObjectNodeEditMaterial refers to an invalid texture layout index %d. Parent node has %d layouts."),
+					*Node->GetOutermost()->GetName(), Node->ParentLayoutIndex, Layouts.Num());
+			}
+			else
+			{
+				ParentLayout = Layouts[Node->ParentLayoutIndex];
+			}
 		}
+
+		LayoutBlocksEditor->SetCurrentLayout(Node->Layout, ParentLayout);
 	}
 	else
 	{
diff --git a/Engine/Plugins/Experimental/Mutable/Source/CustomizableObjectEditor/Private/MuCOE/Nodes/CustomizableObjectNodeEditMaterialDetails.h b/Engine/Plugins/Experimental/Mutable/Source/CustomizableObjectEditor/Private/MuCOE/Nodes/CustomizableObjectNodeEditMaterialDetails.h
--- a/Engine/Plugins/Experimental/Mutable/Source/CustomizableObjectEditor/Private/MuCOE/Nodes/CustomizableObjectNodeEditMaterialDetails.h
+++ b/Engine/Plugins/Experimental/Mutable/Source/CustomizableObjectEditor/Private/MuCOE/Nodes/CustomizableObjectNodeEditMaterialDetails.h
@@ -9,7 +9,7 @@
 class FString;
 class IDetailLayoutBuilder;
 class UCustomizableObjectNodeEditMaterial;
-class SCustomizableObjectNodeLayoutBlocksSelector;
+class SCustomizableObjectNodeLayoutBlocksEditor;
 
 
 class FCustomizableObjectNodeEditMaterialDetails : public FCustomizableObjectNodeEditMaterialBaseDetails
@@ -25,5 +25,5 @@ private:
 	UCustomizableObjectNodeEditMaterial* Node;
 
 	// Layout block editor widget
-	TSharedPtr<SCustomizableObjectNodeLayoutBlocksSelector> LayoutBlocksSelector;
+	TSharedPtr<SCustomizableObjectNodeLayoutBlocksEditor> LayoutBlocksEditor;
 };
diff --git a/Engine/Plugins/Experimental/Mutable/Source/CustomizableObjectEditor/Private/MuCOE/Nodes/CustomizableObjectNodeRemoveMeshBlocks.cpp b/Engine/Plugins/Experimental/Mutable/Source/CustomizableObjectEditor/Private/MuCOE/Nodes/CustomizableObjectNodeRemoveMeshBlocks.cpp
--- a/Engine/Plugins/Experimental/Mutable/Source/CustomizableObjectEditor/Private/MuCOE/Nodes/CustomizableObjectNodeRemoveMeshBlocks.cpp
+++ b/Engine/Plugins/Experimental/Mutable/Source/CustomizableObjectEditor/Private/MuCOE/Nodes/CustomizableObjectNodeRemoveMeshBlocks.cpp
@@ -18,9 +18,11 @@ void UCustomizableObjectNodeRemoveMeshBlocks::BackwardsCompatibleFixup()
 {
 	Super::BackwardsCompatibleFixup();
 	
+	int32 CustomizableObjectCustomVersion = GetLinkerCustomVersion(FCustomizableObjectCustomVersion::GUID);
+
 	// Convert deprecated node index list to the node id list.
-	if (GetLinkerCustomVersion(FCustomizableObjectCustomVersion::GUID) < FCustomizableObjectCustomVersion::PostLoadToCustomVersion
-		&& BlockIds.Num() < Blocks_DEPRECATED.Num())
+	if (CustomizableObjectCustomVersion < FCustomizableObjectCustomVersion::PostLoadToCustomVersion
+		&& BlockIds_DEPRECATED.Num() < Blocks_DEPRECATED.Num())
 	{
 		if (UCustomizableObjectNodeMaterialBase* ParentMaterialNode = GetParentMaterialNode())
 		{
@@ -33,20 +35,20 @@ void UCustomizableObjectNodeRemoveMeshBlocks::BackwardsCompatibleFixup()
 			}
 			else if (UCustomizableObjectNodeMaterial* ParentMaterial = Cast<UCustomizableObjectNodeMaterial>(ParentMaterialNode))
 			{
-				UCustomizableObjectLayout* Layout = Layouts[ParentLayoutIndex];
+				UCustomizableObjectLayout* ParentLayout = Layouts[ParentLayoutIndex];
 
-				for (int IndexIndex = BlockIds.Num(); IndexIndex < Blocks_DEPRECATED.Num(); ++IndexIndex)
+				for (int IndexIndex = BlockIds_DEPRECATED.Num(); IndexIndex < Blocks_DEPRECATED.Num(); ++IndexIndex)
 				{
 					const int BlockIndex = Blocks_DEPRECATED[IndexIndex];
-					if (!Layout->Blocks.IsValidIndex(BlockIndex))
+					if (!ParentLayout->Blocks.IsValidIndex(BlockIndex))
 					{
 						UE_LOG(LogMutable, Warning, TEXT("[%s] UCustomizableObjectNodeRemoveMeshBlocks refers to an invalid layout block index %d. Parent node has %d blocks."),
-							*GetOutermost()->GetName(), BlockIndex, Layout->Blocks.Num());
+							*GetOutermost()->GetName(), BlockIndex, ParentLayout->Blocks.Num());
 
 						continue;
 					}
 					
-					const FGuid Id = Layout->Blocks[BlockIndex].Id;
+					const FGuid Id = ParentLayout->Blocks[BlockIndex].Id;
 					if (!Id.IsValid())
 					{
 						UE_LOG(LogMutable, Warning, TEXT("[%s] UCustomizableObjectNodeRemoveMeshBlocks refers to an valid layout block %d but that block doesn't have an id."),
@@ -55,7 +57,59 @@ void UCustomizableObjectNodeRemoveMeshBlocks::BackwardsCompatibleFixup()
 						continue;
 					}
 
-					BlockIds.Add(Id);
+					BlockIds_DEPRECATED.Add(Id);
+				}
+			}
+		}
+	}
+
+	// Convert deprecated node id list to absolute rect list.
+	if (CustomizableObjectCustomVersion < FCustomizableObjectCustomVersion::UseUVRects)
+	{
+		// If we are here, it means this node was loaded from a version that didn't have it's own layout.
+		check(Layout->Blocks.IsEmpty());
+
+		UCustomizableObjectNodeMaterialBase* ParentMaterialNode = GetParentMaterialNode();
+
+		TArray<UCustomizableObjectLayout*> ParentLayouts = ParentMaterialNode->GetLayouts();
+
+		if (!ParentLayouts.IsValidIndex(ParentLayoutIndex))
+		{
+			UE_LOG(LogMutable, Warning, TEXT("[%s] UCustomizableObjectNodeRemoveMeshBlocks refers to an invalid texture layout index %d. Parent node has %d layouts."),
+				*GetOutermost()->GetName(), ParentLayoutIndex, ParentLayouts.Num());
+		}
+		else
+		{
+			UCustomizableObjectLayout* ParentLayout = ParentLayouts[ParentLayoutIndex];
+			FIntPoint GridSize = ParentLayout->GetGridSize();
+
+			Layout->SetGridSize(GridSize);
+
+			if (Cast<UCustomizableObjectNodeMaterial>(ParentMaterialNode))
+			{
+				for (const FGuid& BlockId : BlockIds_DEPRECATED)
+				{
+					for (const FCustomizableObjectLayoutBlock& ParentBlock : ParentLayout->Blocks)
+					{
+						if (ParentBlock.Id == BlockId)
+						{
+							FCustomizableObjectLayoutBlock NewBlock;
+							NewBlock = ParentBlock;
+
+							// Clear some unnecessary data.
+							NewBlock.bReduceBothAxes = false;
+							NewBlock.bReduceByTwo = false;
+							NewBlock.Priority = 0;
+
+							Layout->Blocks.Add(NewBlock);
+						}
+					}
+				}
+
+				if (Layout->Blocks.Num() != BlockIds_DEPRECATED.Num())
+				{
+					UE_LOG(LogMutable, Warning, TEXT("[%s] UCustomizableObjectNodeRemoveMeshBlocks refers to %d invalid layout block. It has been ignored during version upgrade."),
+						*GetOutermost()->GetName(), int32(BlockIds_DEPRECATED.Num() - Layout->Blocks.Num()));
 				}
 			}
 		}
@@ -77,9 +131,6 @@ void UCustomizableObjectNodeRemoveMeshBlocks::PostEditChangeProperty(FPropertyCh
 	if (PropertyThatChanged 
 		&& (PropertyThatChanged->GetName() == TEXT("ParentMaterialObject") || PropertyThatChanged->GetName() == TEXT("ParentLayoutIndex")))
 	{
-		// Reset selected blocks
-		BlockIds.Empty();
-		
 		TSharedPtr<ICustomizableObjectEditor> Editor = GetGraphEditor();
 
 		if (Editor.IsValid())
@@ -140,15 +191,6 @@ bool UCustomizableObjectNodeRemoveMeshBlocks::IsNodeOutDatedAndNeedsRefresh()
 		{
 			Result = true;
 		}
-		else
-		{
-			UCustomizableObjectLayout* Layout = Layouts[ParentLayoutIndex];
-				
-			if (BlockIds.Num() > Layout->Blocks.Num())
-			{
-				Result = true;
-			}
-		}
 	}
 
 	// Remove previous compilation warnings
diff --git a/Engine/Plugins/Experimental/Mutable/Source/CustomizableObjectEditor/Private/MuCOE/Nodes/CustomizableObjectNodeRemoveMeshBlocksDetails.cpp b/Engine/Plugins/Experimental/Mutable/Source/CustomizableObjectEditor/Private/MuCOE/Nodes/CustomizableObjectNodeRemoveMeshBlocksDetails.cpp
--- a/Engine/Plugins/Experimental/Mutable/Source/CustomizableObjectEditor/Private/MuCOE/Nodes/CustomizableObjectNodeRemoveMeshBlocksDetails.cpp
+++ b/Engine/Plugins/Experimental/Mutable/Source/CustomizableObjectEditor/Private/MuCOE/Nodes/CustomizableObjectNodeRemoveMeshBlocksDetails.cpp
@@ -6,9 +6,11 @@
 #include "DetailWidgetRow.h"
 #include "IDetailsView.h"
 #include "MuCOE/Nodes/CustomizableObjectNodeRemoveMeshBlocks.h"
-#include "MuCOE/SCustomizableObjectNodeLayoutBlocksSelector.h"
+#include "MuCOE/Nodes/CustomizableObjectNodeMaterialBase.h"
+#include "MuCOE/SCustomizableObjectNodeLayoutBlocksEditor.h"
 #include "Widgets/Layout/SBox.h"
 #include "Widgets/Text/STextBlock.h"
+#include "UObject/Package.h"
 
 class FString;
 
@@ -37,24 +39,37 @@ void FCustomizableObjectNodeRemoveMeshBlocksDetails::CustomizeDetails( IDetailLa
 
 	if (Node)
 	{
-		UCustomizableObjectNodeEditLayoutBlocks* NodeEditLayout = Cast<UCustomizableObjectNodeEditLayoutBlocks>(Node);
-		if (NodeEditLayout)
-		{
-			// Add blocks selector
-			LayoutBlocksSelector = SNew(SCustomizableObjectNodeLayoutBlocksSelector);
+		// Add blocks selector
+		LayoutBlocksEditor = SNew(SCustomizableObjectNodeLayoutBlocksEditor);
 
-			BlocksCategory.AddCustomRow(LOCTEXT("BlocksDetails_BlockInstructions", "BlockInstructions"))
+		BlocksCategory.AddCustomRow(LOCTEXT("BlocksDetails_BlockInstructions", "BlockInstructions"))
+		[
+			SNew(SBox)
+			.HeightOverride(700.0f)
+			.WidthOverride(700.0f)
 			[
-				SNew(SBox)
-				.HeightOverride(700.0f)
-				.WidthOverride(700.0f)
-				[
-					LayoutBlocksSelector.ToSharedRef()
-				]
-			];
-
-			LayoutBlocksSelector->SetSelectedNode(NodeEditLayout);
+				LayoutBlocksEditor.ToSharedRef()
+			]
+		];
+
+		// Try to find the parent layout, because we want to show its UVs in the widget
+		UCustomizableObjectLayout* ParentLayout = nullptr;
+		if (UCustomizableObjectNodeMaterialBase* ParentMaterialNode = Node->GetParentMaterialNode())
+		{
+			TArray<UCustomizableObjectLayout*> Layouts = ParentMaterialNode->GetLayouts();
+
+			if (!Layouts.IsValidIndex(Node->ParentLayoutIndex))
+			{
+				UE_LOG(LogMutable, Warning, TEXT("[%s] UCustomizableObjectNodeRemoveMeshBlocks refers to an invalid texture layout index %d. Parent node has %d layouts."),
+					*Node->GetOutermost()->GetName(), Node->ParentLayoutIndex, Layouts.Num());
+			}
+			else
+			{
+				ParentLayout = Layouts[Node->ParentLayoutIndex];
+			}
 		}
+
+		LayoutBlocksEditor->SetCurrentLayout(Node->Layout, ParentLayout);
 	}
 	else
 	{
diff --git a/Engine/Plugins/Experimental/Mutable/Source/CustomizableObjectEditor/Private/MuCOE/Nodes/CustomizableObjectNodeRemoveMeshBlocksDetails.h b/Engine/Plugins/Experimental/Mutable/Source/CustomizableObjectEditor/Private/MuCOE/Nodes/CustomizableObjectNodeRemoveMeshBlocksDetails.h
--- a/Engine/Plugins/Experimental/Mutable/Source/CustomizableObjectEditor/Private/MuCOE/Nodes/CustomizableObjectNodeRemoveMeshBlocksDetails.h
+++ b/Engine/Plugins/Experimental/Mutable/Source/CustomizableObjectEditor/Private/MuCOE/Nodes/CustomizableObjectNodeRemoveMeshBlocksDetails.h
@@ -10,7 +10,7 @@
 class FString;
 class IDetailLayoutBuilder;
 class UCustomizableObjectNodeRemoveMeshBlocks;
-class SCustomizableObjectNodeLayoutBlocksSelector;
+class SCustomizableObjectNodeLayoutBlocksEditor;
 
 
 class FCustomizableObjectNodeRemoveMeshBlocksDetails : public FCustomizableObjectNodeEditMaterialBaseDetails
@@ -26,5 +26,5 @@ private:
 	UCustomizableObjectNodeRemoveMeshBlocks* Node;
 
 	// Layout block editor widget
-	TSharedPtr<SCustomizableObjectNodeLayoutBlocksSelector> LayoutBlocksSelector;
+	TSharedPtr<SCustomizableObjectNodeLayoutBlocksEditor> LayoutBlocksEditor;
 };
diff --git a/Engine/Plugins/Experimental/Mutable/Source/CustomizableObjectEditor/Private/MuCOE/Nodes/CustomizableObjectNodeTable.cpp b/Engine/Plugins/Experimental/Mutable/Source/CustomizableObjectEditor/Private/MuCOE/Nodes/CustomizableObjectNodeTable.cpp
--- a/Engine/Plugins/Experimental/Mutable/Source/CustomizableObjectEditor/Private/MuCOE/Nodes/CustomizableObjectNodeTable.cpp
+++ b/Engine/Plugins/Experimental/Mutable/Source/CustomizableObjectEditor/Private/MuCOE/Nodes/CustomizableObjectNodeTable.cpp
@@ -896,20 +896,6 @@ bool UCustomizableObjectNodeTable::IsPinRelevant(const UEdGraphPin* Pin) const
 }
 
 
-void UCustomizableObjectNodeTable::SetLayoutInLayoutEditor(UCustomizableObjectLayout* CurrentLayout)
-{
-	TSharedPtr<ICustomizableObjectEditor> Editor = GetGraphEditor();
-
-	if (TSharedPtr<FCustomizableObjectEditor> GraphEditor = StaticCastSharedPtr<FCustomizableObjectEditor>(GetGraphEditor()))
-	{
-		if (GraphEditor->GetLayoutBlocksEditor().IsValid())
-		{
-			GraphEditor->GetLayoutBlocksEditor()->SetCurrentLayout(CurrentLayout);
-		}
-	}
-}
-
-
 UTexture2D* UCustomizableObjectNodeTable::FindReferenceTextureParameter(const UEdGraphPin* Pin, FString ParameterImageName) const
 {
 	UMaterialInterface* Material = GetColumnDefaultAssetByType<UMaterialInterface>(Pin);
diff --git a/Engine/Plugins/Experimental/Mutable/Source/CustomizableObjectEditor/Private/MuCOE/SCustomizableObjectNodeLayoutBlocksEditor.cpp b/Engine/Plugins/Experimental/Mutable/Source/CustomizableObjectEditor/Private/MuCOE/SCustomizableObjectNodeLayoutBlocksEditor.cpp
--- a/Engine/Plugins/Experimental/Mutable/Source/CustomizableObjectEditor/Private/MuCOE/SCustomizableObjectNodeLayoutBlocksEditor.cpp
+++ b/Engine/Plugins/Experimental/Mutable/Source/CustomizableObjectEditor/Private/MuCOE/SCustomizableObjectNodeLayoutBlocksEditor.cpp
@@ -66,23 +66,25 @@ void SCustomizableObjectNodeLayoutBlocksEditor::Construct(const FArguments& InAr
 }
 
 
-void SCustomizableObjectNodeLayoutBlocksEditor::SetCurrentLayout(UCustomizableObjectLayout* Layout )
+void SCustomizableObjectNodeLayoutBlocksEditor::SetCurrentLayout( UCustomizableObjectLayout* Layout, UCustomizableObjectLayout* UVOverrideLayout )
 {
 	CurrentLayout = Layout;
 
+	UCustomizableObjectLayout* LayoutForUVs = UVOverrideLayout ? UVOverrideLayout : Layout;
+
 	// Try to locate the source mesh
 	TArray<FVector2f> UVs;
 	TArray<FVector2f> UnassignedUVs;
 
 	if (CurrentLayout)
 	{
-		CurrentLayout->GetUVChannel(UVs, CurrentLayout->GetUVChannel());
+		LayoutForUVs->GetUVChannel(UVs, LayoutForUVs->GetUVChannel());
 
 		UnassignedUVs = TArray<FVector2f>();
 		
 		if (CurrentLayout->UnassignedUVs.Num())
 		{
-			UnassignedUVs = CurrentLayout->UnassignedUVs[0];
+			UnassignedUVs = LayoutForUVs->UnassignedUVs[0];
 		}
 	}
 
diff --git a/Engine/Plugins/Experimental/Mutable/Source/CustomizableObjectEditor/Private/MuCOE/SCustomizableObjectNodeLayoutBlocksEditor.h b/Engine/Plugins/Experimental/Mutable/Source/CustomizableObjectEditor/Private/MuCOE/SCustomizableObjectNodeLayoutBlocksEditor.h
--- a/Engine/Plugins/Experimental/Mutable/Source/CustomizableObjectEditor/Private/MuCOE/SCustomizableObjectNodeLayoutBlocksEditor.h
+++ b/Engine/Plugins/Experimental/Mutable/Source/CustomizableObjectEditor/Private/MuCOE/SCustomizableObjectNodeLayoutBlocksEditor.h
@@ -47,8 +47,10 @@ public:
 	/** Binds commands associated with the viewport client. */
 	void BindCommands();
 
-	/**  */
-	void SetCurrentLayout( class UCustomizableObjectLayout* Layout );
+	/**  
+	* The optional UVOverrideLayout parameter can be speicifed to show different UVs in the widget instead of the ones in Layout.
+	*/
+	void SetCurrentLayout( class UCustomizableObjectLayout* Layout, UCustomizableObjectLayout* UVOverrideLayout=nullptr );
 
 private:
 
diff --git a/Engine/Plugins/Experimental/Mutable/Source/CustomizableObjectEditor/Private/MuCOE/SCustomizableObjectNodeLayoutBlocksSelector.cpp b/Engine/Plugins/Experimental/Mutable/Source/CustomizableObjectEditor/Private/MuCOE/SCustomizableObjectNodeLayoutBlocksSelector.cpp
deleted file mode 100644
diff --git a/Engine/Plugins/Experimental/Mutable/Source/CustomizableObjectEditor/Private/MuCOE/SCustomizableObjectNodeLayoutBlocksSelector.h b/Engine/Plugins/Experimental/Mutable/Source/CustomizableObjectEditor/Private/MuCOE/SCustomizableObjectNodeLayoutBlocksSelector.h
deleted file mode 100644
diff --git a/Engine/Plugins/Experimental/Mutable/Source/CustomizableObjectEditor/Private/MuCOE/SMutableGraphViewer.cpp b/Engine/Plugins/Experimental/Mutable/Source/CustomizableObjectEditor/Private/MuCOE/SMutableGraphViewer.cpp
--- a/Engine/Plugins/Experimental/Mutable/Source/CustomizableObjectEditor/Private/MuCOE/SMutableGraphViewer.cpp
+++ b/Engine/Plugins/Experimental/Mutable/Source/CustomizableObjectEditor/Private/MuCOE/SMutableGraphViewer.cpp
@@ -62,7 +62,7 @@
 #include "MuT/NodeModifierPrivate.h"
 #include "MuT/NodeImageFormatPrivate.h"
 #include "MuT/NodeMeshFormatPrivate.h"
-#include "MuT/NodePatchImagePrivate.h"
+#include "MuT/NodePatchImage.h"
 #include "MuT/NodeMeshConstantPrivate.h"
 #include "MuT/NodeModifierMeshClipMorphPlanePrivate.h"
 #include "MuT/NodePatchMeshPrivate.h"
@@ -77,7 +77,6 @@
 #include "MuT/NodeImageInterpolatePrivate.h"
 #include "MuT/NodeImagePlainColourPrivate.h"
 #include "MuT/NodeImageProjectPrivate.h"
-#include "MuT/NodeMeshFragmentPrivate.h"
 #include "MuT/NodeMeshMorphPrivate.h"
 #include "MuT/NodeScalarEnumParameterPrivate.h"
 #include "MuT/NodeScalarParameterPrivate.h"
@@ -476,9 +475,8 @@ void SMutableGraphViewer::GetChildrenForInfo(TSharedPtr<FMutableGraphTreeElement
 	else if (ParentNode->GetType() == mu::NodePatchImage::GetStaticType())
 	{
 		mu::NodePatchImage* PatchImageVar = StaticCast<mu::NodePatchImage*>(ParentNode);
-		mu::NodePatchImage::Private* Private = PatchImageVar->GetPrivate();
-		AddChildFunc(Private->m_pImage.get(), FString::Printf(TEXT("IMAGE")));
-		AddChildFunc(Private->m_pMask.get(), FString::Printf(TEXT("MASK")));
+		AddChildFunc(PatchImageVar->Image.get(), FString::Printf(TEXT("IMAGE")));
+		AddChildFunc(PatchImageVar->Mask.get(), FString::Printf(TEXT("MASK")));
 	}
 
 	else if (ParentNode->GetType() == mu::NodeModifierMeshClipMorphPlane::GetStaticType())
@@ -605,8 +603,7 @@ void SMutableGraphViewer::GetChildrenForInfo(TSharedPtr<FMutableGraphTreeElement
 	else if (ParentNode->GetType() == mu::NodeMeshFragment::GetStaticType())
 	{
 		mu::NodeMeshFragment* MeshFragmentVar = StaticCast<mu::NodeMeshFragment*>(ParentNode);
-		mu::NodeMeshFragment::Private* Private = MeshFragmentVar->GetPrivate();
-		AddChildFunc(Private->m_pMesh.get(), FString::Printf(TEXT("MESH")));
+		AddChildFunc(MeshFragmentVar->SourceMesh.get(), FString::Printf(TEXT("MESH")));
 	}
 
 	else if (ParentNode->GetType() == mu::NodeColourSampleImage::GetStaticType())
diff --git a/Engine/Plugins/Experimental/Mutable/Source/CustomizableObjectEditor/Public/MuCOE/CustomizableObjectLayout.h b/Engine/Plugins/Experimental/Mutable/Source/CustomizableObjectEditor/Public/MuCOE/CustomizableObjectLayout.h
--- a/Engine/Plugins/Experimental/Mutable/Source/CustomizableObjectEditor/Public/MuCOE/CustomizableObjectLayout.h
+++ b/Engine/Plugins/Experimental/Mutable/Source/CustomizableObjectEditor/Public/MuCOE/CustomizableObjectLayout.h
@@ -26,7 +26,7 @@ enum class ECustomizableObjectLayoutBlockReductionMethod : uint8
 {
 	// Layout blocks will be reduced by halves
 	Halve = 0 UMETA(DisplayName = "Reduce by Half"),
-	// LAyout blocks will be reduced by a grid unit
+	// Layout blocks will be reduced by a grid unit
 	Unitary = 1 UMETA(DisplayName = "Reduce by Unit")
 };
 
diff --git a/Engine/Plugins/Experimental/Mutable/Source/CustomizableObjectEditor/Public/MuCOE/Nodes/CustomizableObjectNodeEditLayoutBlocks.h b/Engine/Plugins/Experimental/Mutable/Source/CustomizableObjectEditor/Public/MuCOE/Nodes/CustomizableObjectNodeEditLayoutBlocks.h
--- a/Engine/Plugins/Experimental/Mutable/Source/CustomizableObjectEditor/Public/MuCOE/Nodes/CustomizableObjectNodeEditLayoutBlocks.h
+++ b/Engine/Plugins/Experimental/Mutable/Source/CustomizableObjectEditor/Public/MuCOE/Nodes/CustomizableObjectNodeEditLayoutBlocks.h
@@ -3,6 +3,7 @@
 #pragma once
 
 #include "MuCOE/Nodes/CustomizableObjectNodeEditMaterialBase.h"
+#include "MuCOE/CustomizableObjectLayout.h"
 
 #include "CustomizableObjectNodeEditLayoutBlocks.generated.h"
 
@@ -13,9 +14,20 @@ class CUSTOMIZABLEOBJECTEDITOR_API UCustomizableObjectNodeEditLayoutBlocks : pub
 public:
 	GENERATED_BODY()
 
-	// Selected blocks
+	UCustomizableObjectNodeEditLayoutBlocks();
+
+public:
+
+	/** Deprecated data from the time this node referred to other layout node blocks. Info is now in the Layout attribute. */
+	UPROPERTY()
+	TArray<FGuid> BlockIds_DEPRECATED;
+
+	/** Layout defining the blocks to edit. It is a standalone layout, unrelated to whatever is connected
+	 * to the UV channels of the mesh being edited.
+	 * Some properties in the UCustomizableObjectLayout, like shrinking strategy or priorities will not be relevant.
+	 */
 	UPROPERTY()
-	TArray<FGuid> BlockIds;
+	TObjectPtr<UCustomizableObjectLayout> Layout = nullptr;
 
 };
 
diff --git a/Engine/Plugins/Experimental/Mutable/Source/CustomizableObjectEditor/Public/MuCOE/Nodes/CustomizableObjectNodeEditMaterial.h b/Engine/Plugins/Experimental/Mutable/Source/CustomizableObjectEditor/Public/MuCOE/Nodes/CustomizableObjectNodeEditMaterial.h
--- a/Engine/Plugins/Experimental/Mutable/Source/CustomizableObjectEditor/Public/MuCOE/Nodes/CustomizableObjectNodeEditMaterial.h
+++ b/Engine/Plugins/Experimental/Mutable/Source/CustomizableObjectEditor/Public/MuCOE/Nodes/CustomizableObjectNodeEditMaterial.h
@@ -84,9 +84,6 @@ public:
 	 * @returns Always returns a valid pin if EditsImage(const FGuid&) returns true. */
 	const UEdGraphPin* GetUsedImageMaskPin(const FNodeMaterialParameterId& ImageId) const;
 
-	// Function to select all the layout blocks. Called when a parameter is reset (ParentMaterial, ParentLayoutIndex)
-	void SelectAllLayoutBlocks();
-
 private:
 
 	/** Relates a Parameter id (and layer if is a layered material) to a Pin. Only used to improve performance. */
@@ -104,8 +101,3 @@ private:
 	UPROPERTY()
 	TArray<int32> Blocks_DEPRECATED;
 };
-
-
-#if UE_ENABLE_INCLUDE_ORDER_DEPRECATED_IN_5_2
-#include "MuCOE/CustomizableObjectEditor_Deprecated.h"
-#endif
diff --git a/Engine/Plugins/Experimental/Mutable/Source/CustomizableObjectEditor/Public/MuCOE/Nodes/CustomizableObjectNodeTable.h b/Engine/Plugins/Experimental/Mutable/Source/CustomizableObjectEditor/Public/MuCOE/Nodes/CustomizableObjectNodeTable.h
--- a/Engine/Plugins/Experimental/Mutable/Source/CustomizableObjectEditor/Public/MuCOE/Nodes/CustomizableObjectNodeTable.h
+++ b/Engine/Plugins/Experimental/Mutable/Source/CustomizableObjectEditor/Public/MuCOE/Nodes/CustomizableObjectNodeTable.h
@@ -267,7 +267,6 @@ public:
 	void GetUVChannelForPin(const UEdGraphPin* Pin, TArray<FVector2f>& OutSegments, int32 UVChannel = 0) const;
 
 	// Methods to provide the Layouts to the Layout block editors
-	void SetLayoutInLayoutEditor(UCustomizableObjectLayout* CurrentLayout);
 	TArray<UCustomizableObjectLayout*> GetLayouts(const UEdGraphPin* Pin) const;
 
 	// Returns the name of the table column related to a pin
diff --git a/Engine/Plugins/Experimental/Mutable/Source/MutableRuntime/Internal/MuR/Operations.h b/Engine/Plugins/Experimental/Mutable/Source/MutableRuntime/Internal/MuR/Operations.h
--- a/Engine/Plugins/Experimental/Mutable/Source/MutableRuntime/Internal/MuR/Operations.h
+++ b/Engine/Plugins/Experimental/Mutable/Source/MutableRuntime/Internal/MuR/Operations.h
@@ -249,7 +249,7 @@ namespace mu
         //! Change the mesh format to match the format of another one.
         ME_FORMAT,
 
-        //! Extract a fragment of a mesh containing a specific layout block.
+        //! Extract a fragment of a mesh containing specific layout blocks.
         ME_EXTRACTLAYOUTBLOCK,
 
         //! Apply a transform in a 4x4 matrix to the geometry channels of the mesh
diff --git a/Engine/Plugins/Experimental/Mutable/Source/MutableRuntime/Private/MuR/CodeRunner.cpp b/Engine/Plugins/Experimental/Mutable/Source/MutableRuntime/Private/MuR/CodeRunner.cpp
--- a/Engine/Plugins/Experimental/Mutable/Source/MutableRuntime/Private/MuR/CodeRunner.cpp
+++ b/Engine/Plugins/Experimental/Mutable/Source/MutableRuntime/Private/MuR/CodeRunner.cpp
@@ -1896,7 +1896,7 @@ namespace mu
 
                 Ptr<const Mesh> Source = LoadMesh(FCacheAddress(source, item));
 
-                // Access with memcpy necessary for unaligned arm issues.
+                // Access with memcpy necessary for unaligned memory access issues.
                 uint64 blocks[512];
 				check(blockCount< 512);
 				FMemory::Memcpy(blocks, data, sizeof(uint64)*FMath::Min(512,int32(blockCount)));
@@ -1905,7 +1905,15 @@ namespace mu
 				{
 					Ptr<Mesh> Result = CreateMesh();
 					bool bOutSuccess;
-					MeshExtractLayoutBlock(Result.get(), Source.get(), layout, blockCount, blocks, bOutSuccess);
+
+					if (blockCount > 0)
+					{
+						MeshExtractLayoutBlock(Result.get(), Source.get(), layout, blockCount, blocks, bOutSuccess);
+					}
+					else
+					{
+						MeshExtractLayoutBlock(Result.get(), Source.get(), layout, bOutSuccess);
+					}
 
 					if (!bOutSuccess)
 					{
diff --git a/Engine/Plugins/Experimental/Mutable/Source/MutableRuntime/Private/MuR/OpImageBlend.h b/Engine/Plugins/Experimental/Mutable/Source/MutableRuntime/Private/MuR/OpImageBlend.h
--- a/Engine/Plugins/Experimental/Mutable/Source/MutableRuntime/Private/MuR/OpImageBlend.h
+++ b/Engine/Plugins/Experimental/Mutable/Source/MutableRuntime/Private/MuR/OpImageBlend.h
@@ -135,14 +135,14 @@ namespace mu
 	//---------------------------------------------------------------------------------------------
 	inline uint32 MultiplyChannelMasked(uint32 base, uint32 blended, uint32 mask)
 	{
-		uint32 multiply = (base * blended) >> 8;
-		uint32 masked = (((255 - mask) * base) + (mask * multiply)) >> 8;
+		uint32 multiply = (base * blended) / 255;
+		uint32 masked = (((255 - mask) * base) + (mask * multiply)) / 255;
 		return masked;
 	}
 
 	inline uint32 MultiplyChannel(uint32 base, uint32 blended)
 	{
-		uint32 multiply = (base * blended) >> 8;
+		uint32 multiply = (base * blended) / 255;
 		return multiply;
 	}
 
diff --git a/Engine/Plugins/Experimental/Mutable/Source/MutableRuntime/Private/MuR/OpMeshExtractLayoutBlock.h b/Engine/Plugins/Experimental/Mutable/Source/MutableRuntime/Private/MuR/OpMeshExtractLayoutBlock.h
--- a/Engine/Plugins/Experimental/Mutable/Source/MutableRuntime/Private/MuR/OpMeshExtractLayoutBlock.h
+++ b/Engine/Plugins/Experimental/Mutable/Source/MutableRuntime/Private/MuR/OpMeshExtractLayoutBlock.h
@@ -122,6 +122,7 @@ namespace mu
     }
 
 
+	/** Extract all vertices that have a layout block from the passed list on the given channel. */
     inline void MeshExtractLayoutBlock(Mesh* Result, const Mesh* Source,
                                            uint32 LayoutIndex,
                                            uint16 BlockCount,
@@ -201,4 +202,62 @@ namespace mu
         }
 	}
 
+	/** Extract all vertices that have a valid layout block on the given channel. */
+	inline void MeshExtractLayoutBlock(Mesh* Result, const Mesh* Source, uint32 LayoutIndex, bool& bOutSuccess)
+	{
+		check(Source);
+		bOutSuccess = true;
+
+		// TODO: Optimise
+		Result->CopyFrom(*Source);
+
+		UntypedMeshBufferIteratorConst itBlocks(Source->GetVertexBuffers(), MBS_LAYOUTBLOCK, LayoutIndex);
+
+		if (itBlocks.GetFormat() != MBF_NONE)
+		{
+			int32 ResultVertices = 0;
+			TArray<int32> OldToNew;
+			OldToNew.Init(-1, Source->GetVertexCount());
+			TArray<int32> NewToOld;
+			NewToOld.Reserve(Source->GetVertexCount());
+
+			if (itBlocks.GetFormat() == MBF_UINT16)
+			{
+				const uint16* pBlocks = reinterpret_cast<const uint16*>(itBlocks.ptr());
+				for (int32 i = 0; i < Source->GetVertexCount(); ++i)
+				{
+					uint16 VertexBlockRelative = pBlocks[i];
+
+					bool bFound = VertexBlockRelative != std::numeric_limits<uint16>::max();
+					if (bFound)
+					{
+						OldToNew[i] = ResultVertices++;
+						NewToOld.Add(i);
+					}
+				}
+			}
+			else if (itBlocks.GetFormat() == MBF_UINT64)
+			{
+				const uint64* pBlocks = reinterpret_cast<const uint64*>(itBlocks.ptr());
+				for (int32 i = 0; i < Source->GetVertexCount(); ++i)
+				{
+					uint64 VertexBlockId = pBlocks[i];
+
+					bool bFound = VertexBlockId != std::numeric_limits<uint64>::max();
+					if (bFound)
+					{
+						OldToNew[i] = ResultVertices++;
+						NewToOld.Add(i);
+					}
+				}
+			}
+			else
+			{
+				check(false);
+			}
+
+			MeshExtractFromVertices(Source, Result, OldToNew, NewToOld);
+		}
+	}
+
 }
diff --git a/Engine/Plugins/Experimental/Mutable/Source/MutableTools/Internal/MuT/NodeMeshFragmentPrivate.h b/Engine/Plugins/Experimental/Mutable/Source/MutableTools/Internal/MuT/NodeMeshFragmentPrivate.h
deleted file mode 100644
diff --git a/Engine/Plugins/Experimental/Mutable/Source/MutableTools/Internal/MuT/NodePatchImage.h b/Engine/Plugins/Experimental/Mutable/Source/MutableTools/Internal/MuT/NodePatchImage.h
--- a/Engine/Plugins/Experimental/Mutable/Source/MutableTools/Internal/MuT/NodePatchImage.h
+++ b/Engine/Plugins/Experimental/Mutable/Source/MutableTools/Internal/MuT/NodePatchImage.h
@@ -16,76 +16,40 @@ namespace mu
     class NodeImage;
  
 
-    //! Node that allows to modify an image from an object by blending other images on specific
-    //! layout blocks.
+    /** Node that allows to modify an image from an object by blending other images on specific layout blocks. */
 	class MUTABLETOOLS_API NodePatchImage : public Node
     {
-    public:
-
-        NodePatchImage();
-
-
-        //-----------------------------------------------------------------------------------------
-        // Node Interface
-        //-----------------------------------------------------------------------------------------
-
-        const FNodeType* GetType() const override;
-        static const FNodeType* GetStaticType();
-
-        //-----------------------------------------------------------------------------------------
-        // Own Interface
-        //-----------------------------------------------------------------------------------------
+	public:
 
-        //! Set the image that will be blended on the destination.
-        void SetImage(Ptr<NodeImage>);
+		/** Image to blend. */
+		Ptr<NodeImage> Image;
 
-        //! Get the image that will be blended on the destination.
-		Ptr<NodeImage> GetImage() const;
+		/** Optional mask controlling the blending area. */
+		Ptr<NodeImage> Mask;
 
-        //! Set the blending mask image.
-        void SetMask(Ptr<NodeImage>);
+		/** Rects in the parent layout homogeneous UV space to patch. */
+		TArray<FBox2f> Blocks;
 
-        //! Get the blending mask image.
-		Ptr<NodeImage> GetMask() const;
+		/** */
+		EBlendType BlendType = EBlendType::BT_BLEND;
 
-        //! Set the number of blocks in the layout that will be patched
-        void SetBlockCount( int32 );
-
-        //! Get the number of blocks in the layout that will be patched
-        int32 GetBlockCount() const;
-
-        //! Set a block of the layout that will be patched.
-        //! \param index is the index in the patching list, from 0 to GetBlockCount()-1
-        //! \param LayoutBlockIndex is the block index in the layout of the image were this patch operation will be applied.
-        void SetBlock( int32 index, int32 LayoutBlockIndex);
-
-        //! Get a block index of the layout that will be patched.
-        //! \param index is the index in the patching list, from 0 to GetBlockCount()-1
-        int32 GetBlock( int32 index ) const;
-
-        //! Set the blending operation to use to combine the pixels
-        EBlendType GetBlendType() const;
-        void SetBlendType(EBlendType);
-
-        //! Enable patching the alpha channel if present
-        bool GetApplyToAlphaChannel() const;
-        void SetApplyToAlphaChannel( bool );
+		/** Patch alpha channel as well? */
+		bool bApplyToAlpha = false;
 
+    public:
 
-        //-----------------------------------------------------------------------------------------
-        // Interface pattern
-        //-----------------------------------------------------------------------------------------
-        class Private;
-        Private* GetPrivate() const;
+		// Node Interface
+		virtual const FNodeType* GetType() const override { return GetStaticType(); }
+		static const FNodeType* GetStaticType() { return &StaticType; }
 
     protected:
 
         //! Forbidden. Manage with the Ptr<> template.
-        ~NodePatchImage();
+		~NodePatchImage() {}
 
-    private:
+	private:
 
-        Private* m_pD;
+		static FNodeType StaticType;
 
     };
 
diff --git a/Engine/Plugins/Experimental/Mutable/Source/MutableTools/Internal/MuT/NodePatchImagePrivate.h b/Engine/Plugins/Experimental/Mutable/Source/MutableTools/Internal/MuT/NodePatchImagePrivate.h
deleted file mode 100644
diff --git a/Engine/Plugins/Experimental/Mutable/Source/MutableTools/Private/MuT/ASTOpMeshExtractLayoutBlocks.cpp b/Engine/Plugins/Experimental/Mutable/Source/MutableTools/Private/MuT/ASTOpMeshExtractLayoutBlocks.cpp
--- a/Engine/Plugins/Experimental/Mutable/Source/MutableTools/Private/MuT/ASTOpMeshExtractLayoutBlocks.cpp
+++ b/Engine/Plugins/Experimental/Mutable/Source/MutableTools/Private/MuT/ASTOpMeshExtractLayoutBlocks.cpp
@@ -37,7 +37,7 @@ namespace mu
 		if (otherUntyped.GetOpType() == GetOpType())
 		{
 			const ASTOpMeshExtractLayoutBlocks* other = static_cast<const ASTOpMeshExtractLayoutBlocks*>(&otherUntyped);
-			return Source == other->Source && Layout == other->Layout && Blocks == other->Blocks;
+			return Source == other->Source && LayoutIndex == other->LayoutIndex && Blocks == other->Blocks;
 		}
 		return false;
 	}
@@ -47,7 +47,7 @@ namespace mu
 	{
 		Ptr<ASTOpMeshExtractLayoutBlocks> n = new ASTOpMeshExtractLayoutBlocks();
 		n->Source = mapChild(Source.child());
-		n->Layout = Layout;
+		n->LayoutIndex = LayoutIndex;
 		n->Blocks = Blocks;
 		return n;
 	}
@@ -84,7 +84,7 @@ namespace mu
 			AppendCode(program.m_byteCode, OP_TYPE::ME_EXTRACTLAYOUTBLOCK);
 			OP::ADDRESS sourceAt = Source ? Source->linkedAddress : 0;
 			AppendCode(program.m_byteCode, sourceAt);
-			AppendCode(program.m_byteCode, (uint16)Layout);
+			AppendCode(program.m_byteCode, (uint16)LayoutIndex);
 			AppendCode(program.m_byteCode, (uint16)Blocks.Num());
 
 			for (uint64 Id : Blocks)
diff --git a/Engine/Plugins/Experimental/Mutable/Source/MutableTools/Private/MuT/ASTOpMeshExtractLayoutBlocks.h b/Engine/Plugins/Experimental/Mutable/Source/MutableTools/Private/MuT/ASTOpMeshExtractLayoutBlocks.h
--- a/Engine/Plugins/Experimental/Mutable/Source/MutableTools/Private/MuT/ASTOpMeshExtractLayoutBlocks.h
+++ b/Engine/Plugins/Experimental/Mutable/Source/MutableTools/Private/MuT/ASTOpMeshExtractLayoutBlocks.h
@@ -24,9 +24,9 @@ namespace mu
 		ASTChild Source;
 
 		/** Layout to use to select the blocks. */
-		uint16 Layout = 0;
+		uint16 LayoutIndex = 0;
 
-		/** Block Ids to include in the resulting mesh. */
+		/** Block Ids to include in the resulting mesh. If this is empty all vertices with any valid block assigned will be included. */
 		TArray<uint64> Blocks;
 
 	public:
diff --git a/Engine/Plugins/Experimental/Mutable/Source/MutableTools/Private/MuT/CodeGenerator.cpp b/Engine/Plugins/Experimental/Mutable/Source/MutableTools/Private/MuT/CodeGenerator.cpp
--- a/Engine/Plugins/Experimental/Mutable/Source/MutableTools/Private/MuT/CodeGenerator.cpp
+++ b/Engine/Plugins/Experimental/Mutable/Source/MutableTools/Private/MuT/CodeGenerator.cpp
@@ -75,7 +75,7 @@
 #include "MuT/NodeObject.h"
 #include "MuT/NodeObjectGroupPrivate.h"
 #include "MuT/NodeObjectNew.h"
-#include "MuT/NodePatchImagePrivate.h"
+#include "MuT/NodePatchImage.h"
 #include "MuT/NodePatchMesh.h"
 #include "MuT/NodePrivate.h"
 #include "MuT/NodeRange.h"
@@ -389,9 +389,9 @@ namespace mu
 	}
 
 
-	Ptr<const Layout> CodeGenerator::AddLayout(Ptr<const NodeLayout> SourceLayout, uint32 MeshIDPrefix)
+	Ptr<const Layout> CodeGenerator::GenerateLayout(Ptr<const NodeLayout> SourceLayout, uint32 MeshIDPrefix)
 	{
-		Ptr<const Layout>* it = GeneratedLayouts.Find(SourceLayout.get());
+		Ptr<const Layout>* it = GeneratedLayouts.Find({ SourceLayout,MeshIDPrefix });
 
 		if (it)
 		{
@@ -422,69 +422,84 @@ namespace mu
 		}
 
 		check(GeneratedLayout->Blocks.IsEmpty() || GeneratedLayout->Blocks[0].Id != FLayoutBlock::InvalidBlockId);
-		GeneratedLayouts.Add(SourceLayout.get(), GeneratedLayout);
+		GeneratedLayouts.Add({ SourceLayout,MeshIDPrefix }, GeneratedLayout);
 
 		return GeneratedLayout;
 	}
 
 
-	Ptr<ASTOp> CodeGenerator::GenerateImageBlockPatch(Ptr<ASTOp> blockAd,
+	Ptr<ASTOp> CodeGenerator::GenerateImageBlockPatch(Ptr<ASTOp> InBlockOp,
 		const NodePatchImage* pPatch,
+		Ptr<Image> PatchMask,
 		Ptr<ASTOp> conditionAd,
 		const FImageGenerationOptions& ImageOptions )
 	{
 		// Blend operation
-		Ptr<ASTOp> blendAd;
+		Ptr<ASTOp> FinalOp;
 		{
 			MUTABLE_CPUPROFILER_SCOPE(PatchBlend);
 
-			Ptr<ASTOpImageLayer> op = new ASTOpImageLayer();
-			op->blendType = pPatch->GetPrivate()->m_blendType;
-			op->base = blockAd;
+			Ptr<ASTOpImageLayer> LayerOp = new ASTOpImageLayer();
+			LayerOp->blendType = pPatch->BlendType;
+			LayerOp->base = InBlockOp;
 
 			// When we patch from edit nodes, we want to apply it to all the channels.
-			// \todo: since we can choose the patch function, maybe we want to be able to
-			// select this as well.
-			op->Flags = pPatch->GetPrivate()->m_applyToAlpha
-				? OP::ImageLayerArgs::F_APPLY_TO_ALPHA
-				: 0;
-
-			NodeImage* pImage = pPatch->GetPrivate()->m_pImage.get();
-			Ptr<ASTOp> blend;
-			if (pImage)
+			// \todo: since we can choose the patch function, maybe we want to be able to select this as well.
+			LayerOp->Flags = pPatch->bApplyToAlpha ? OP::ImageLayerArgs::F_APPLY_TO_ALPHA : 0;
+
+			NodeImage* ImageNode = pPatch->Image.get();
+			Ptr<ASTOp> BlendOp;
+			if (ImageNode)
 			{
 				FImageGenerationResult BlendResult;
-				GenerateImage(ImageOptions, BlendResult, pImage);
-				blend = BlendResult.op;
+				GenerateImage(ImageOptions, BlendResult, ImageNode);
+				BlendOp = BlendResult.op;
 			}
 			else
 			{
-				blend = GenerateMissingImageCode(TEXT("Blend top image"), EImageFormat::IF_RGB_UBYTE, pPatch->GetMessageContext(), ImageOptions);
+				BlendOp = GenerateMissingImageCode(TEXT("Patch top image"), EImageFormat::IF_RGB_UBYTE, pPatch->GetMessageContext(), ImageOptions);
+			}
+			BlendOp = GenerateImageFormat(BlendOp, InBlockOp->GetImageDesc().m_format);
+			BlendOp = GenerateImageSize(BlendOp, ImageOptions.RectSize);
+			LayerOp->blend = BlendOp;
+
+			// Create the rect mask constant
+			Ptr<ASTOp> RectConstantOp;
+			{
+				Ptr<NodeImageConstant> pNode = new NodeImageConstant();
+				pNode->SetValue(PatchMask.get());
+
+				FImageGenerationOptions ConstantOptions;
+				FImageGenerationResult ConstantResult;
+				GenerateImage(ConstantOptions, ConstantResult, pNode);
+				RectConstantOp = ConstantResult.op;
 			}
-			blend = GenerateImageFormat(blend, blockAd->GetImageDesc().m_format);
-			blend = GenerateImageSize(blend, ImageOptions.RectSize);
-			op->blend = blend;
 
-			NodeImage* pMask = pPatch->GetPrivate()->m_pMask.get();
-			Ptr<ASTOp> mask;
-			if (pMask)
+			NodeImage* MaskNode = pPatch->Mask.get();
+			Ptr<ASTOp> MaskOp;
+			if (MaskNode)
 			{
+				// Combine the block rect mask with the user provided mask.
+
 				FImageGenerationResult MaskResult;
-				GenerateImage(ImageOptions, MaskResult, pMask);
-				mask = MaskResult.op;
+				GenerateImage(ImageOptions, MaskResult, MaskNode);
+				MaskOp = MaskResult.op;
+
+				Ptr<ASTOpImageLayer> PatchCombineOp = new ASTOpImageLayer;
+				PatchCombineOp->base = MaskOp;
+				PatchCombineOp->blend = RectConstantOp;
+				PatchCombineOp->blendType = EBlendType::BT_MULTIPLY;
+				MaskOp = PatchCombineOp;
 			}
 			else
 			{
-				// Set the argument default value: affect all pixels.
-				// TODO: Special operation code without mask
-				FImageGenerationOptions MissingMaskOptions;
-				mask = GeneratePlainImageCode(FVector4f(1,1,1,1), MissingMaskOptions);
+				MaskOp = RectConstantOp;
 			}
-			mask = GenerateImageFormat(mask, EImageFormat::IF_L_UBYTE);
-			mask = GenerateImageSize(mask, ImageOptions.RectSize);
-			op->mask = mask;
+			MaskOp = GenerateImageFormat(MaskOp, EImageFormat::IF_L_UBYTE);
+			MaskOp = GenerateImageSize(MaskOp, ImageOptions.RectSize);
+			LayerOp->mask = MaskOp;
 
-			blendAd = op;
+			FinalOp = LayerOp;
 		}
 
 		// Condition to enable this patch
@@ -494,20 +509,20 @@ namespace mu
 			{
 				Ptr<ASTOpConditional> op = new ASTOpConditional();
 				op->type = OP_TYPE::IM_CONDITIONAL;
-				op->no = blockAd;
-				op->yes = blendAd;
+				op->no = InBlockOp;
+				op->yes = FinalOp;
 				op->condition = conditionAd;
 				conditionalAd = op;
 			}
 
-			blockAd = conditionalAd;
+			FinalOp = conditionalAd;
 		}
 		else
 		{
-			blockAd = blendAd;
+			FinalOp = FinalOp;
 		}
 
-		return blockAd;
+		return FinalOp;
 	}
 
 
@@ -783,6 +798,50 @@ namespace mu
 	}
 
 
+	Ptr<Image> CodeGenerator::GenerateImageBlockPatchMask(const NodePatchImage* Patch, FIntPoint GridSize, int32 BlockPixelsX, int32 BlockPixelsY, box<UE::Math::TIntVector2<uint16>> RectInCells )
+	{
+		// Create a patching mask for the block
+		Ptr<Image> PatchMask;
+
+		FIntVector2 SourceTextureSize = { GridSize[0] * BlockPixelsX, GridSize[1] * BlockPixelsY };
+
+		FInt32Rect BlockRectInPixels;
+		BlockRectInPixels.Min = { RectInCells.min[0] * BlockPixelsX, RectInCells.min[1] * BlockPixelsY };
+		BlockRectInPixels.Max = { (RectInCells.min[0] + RectInCells.size[0]) * BlockPixelsX, (RectInCells.min[1] + RectInCells.size[1]) * BlockPixelsY };
+
+		for (const FBox2f& PatchRect : Patch->Blocks)
+		{
+			// Does the patch rect intersects the current block at all?
+			FInt32Rect PatchRectInPixels;
+			PatchRectInPixels.Min = { int32(PatchRect.Min[0] * SourceTextureSize[0]), int32(PatchRect.Min[1] * SourceTextureSize[1]) };
+			PatchRectInPixels.Max = { int32(PatchRect.Max[0] * SourceTextureSize[0]), int32(PatchRect.Max[1] * SourceTextureSize[1]) };
+
+			FInt32Rect BlockPatchRect = PatchRectInPixels;
+			BlockPatchRect.Clip(BlockRectInPixels);
+
+			if (BlockPatchRect.Area() > 0)
+			{
+				FInt32Point BlockSize = BlockRectInPixels.Size();
+				if (!PatchMask)
+				{
+					PatchMask = new mu::Image(BlockSize[0], BlockSize[1], 1, mu::EImageFormat::IF_L_UBYTE, mu::EInitializationType::Black);
+				}
+
+				uint8* Pixels = PatchMask->GetMipData(0);
+				FInt32Point BlockPatchOffset = BlockPatchRect.Min - BlockRectInPixels.Min;
+				FInt32Point BlockPatchSize = BlockPatchRect.Size();
+				for (int32 RowIndex = BlockPatchOffset[1]; RowIndex < BlockPatchOffset[1]+BlockPatchSize[1]; ++RowIndex)
+				{
+					uint8* RowPixels = Pixels + RowIndex * BlockSize[0] + BlockPatchOffset[0];
+					FMemory::Memset(RowPixels, 255, BlockPatchSize[0]);
+				}
+			}
+		}
+
+		return PatchMask;
+	}
+
+
     //---------------------------------------------------------------------------------------------
     void CodeGenerator::GenerateSurface( FSurfaceGenerationResult& result, 
 										 const FSurfaceGenerationOptions& Options,
@@ -1302,11 +1361,30 @@ namespace mu
 						for (int32 editIndex = 0; editIndex < edits.Num(); ++editIndex)
 						{
 							const FirstPassGenerator::FSurface::FEdit& e = edits[editIndex];
-							if (t < e.Node->Textures.Num())
+							if (t >= e.Node->Textures.Num())
 							{
-								if (const NodePatchImage* pPatch = e.Node->Textures[t].Patch.get())
+								continue;
+							}
+
+							if (const NodePatchImage* Patch = e.Node->Textures[t].Patch.get())
+							{
+								// Does the current block need to be patched?
+
+								// Ideally this should be the actual image size
+								constexpr int32 FakeLayoutSize = 256;
+
+								FIntPoint GridSize(FakeLayoutSize, FakeLayoutSize);
+								int32 BlockPixelsX = 1;
+								int32 BlockPixelsY = 1;
+								box< UE::Math::TIntVector2<uint16> > RectInCells;
+								RectInCells.min = { 0,0 };
+								RectInCells.size = { FakeLayoutSize ,FakeLayoutSize };
+
+								Ptr<Image> PatchMask = GenerateImageBlockPatchMask(Patch, GridSize, BlockPixelsX, BlockPixelsY, RectInCells);
+
+								if (PatchMask)
 								{
-									imageAd = GenerateImageBlockPatch(imageAd, pPatch, e.Condition, ImageOptions);
+									imageAd = GenerateImageBlockPatch(imageAd, Patch, PatchMask, e.Condition, ImageOptions);
 								}
 							}
 
@@ -1381,7 +1459,7 @@ namespace mu
 							//-------------------------------------
 
 							// Size of a layout block in pixels
-							FIntPoint grid = pLayout->GetGridSize();
+							FIntPoint GridSize = pLayout->GetGridSize();
 
 							EImageFormat FinalFormat = EImageFormat::IF_NONE;
 							int32 BlockPixelsX = 0;
@@ -1473,29 +1551,33 @@ namespace mu
 								FImageDesc BlockDesc = blockAd->GetImageDesc(bReturnBestOption, nullptr);
 
 								// Block in layout grid units (cells)
-								box< UE::Math::TIntVector2<uint16> > rectInCells;
-								rectInCells.min = pLayout->Blocks[BlockIndex].Min;
-								rectInCells.size = pLayout->Blocks[BlockIndex].Size;
+								box< UE::Math::TIntVector2<uint16> > RectInCells;
+								RectInCells.min = pLayout->Blocks[BlockIndex].Min;
+								RectInCells.size = pLayout->Blocks[BlockIndex].Size;
 
 								// If we don't know the size of a layout block in pixels, calculate it
-								UpdateBlockSize(BlockDesc, rectInCells.size);
+								UpdateBlockSize(BlockDesc, RectInCells.size);
 
 								// Even if we force the size afterwards, we need some size hint in some cases, like image projections.
 								ImageOptions.RectSize = UE::Math::TIntVector2<int32>(BlockDesc.m_size);
 
 								// Look for patches to this block
-								for (int32 editIndex = 0; editIndex < edits.Num(); ++editIndex)
+								for (const FirstPassGenerator::FSurface::FEdit& Edit : edits)
 								{
-									const FirstPassGenerator::FSurface::FEdit& e = edits[editIndex];
-									if (t < e.Node->Textures.Num())
+									if (t >= Edit.Node->Textures.Num())
 									{
-										if (const NodePatchImage* pPatch = e.Node->Textures[t].Patch.get())
+										continue;
+									}
+
+									if (const NodePatchImage* Patch = Edit.Node->Textures[t].Patch.get())
+									{
+										// Is the current block to be patched?
+										Ptr<Image> PatchMask = GenerateImageBlockPatchMask( Patch, GridSize, BlockPixelsX, BlockPixelsY, RectInCells );
+
+										// If there was any edit affecting this block
+										if (PatchMask)
 										{
-											// Is the current block to be patched?
-											if (pPatch->GetPrivate()->BlockIndices.Contains(BlockIndex))
-											{
-												blockAd = GenerateImageBlockPatch(blockAd, pPatch, e.Condition, ImageOptions);
-											}
+											blockAd = GenerateImageBlockPatch(blockAd, Patch, PatchMask, Edit.Condition, ImageOptions);
 										}
 									}
 								}
diff --git a/Engine/Plugins/Experimental/Mutable/Source/MutableTools/Private/MuT/CodeGenerator.h b/Engine/Plugins/Experimental/Mutable/Source/MutableTools/Private/MuT/CodeGenerator.h
--- a/Engine/Plugins/Experimental/Mutable/Source/MutableTools/Private/MuT/CodeGenerator.h
+++ b/Engine/Plugins/Experimental/Mutable/Source/MutableTools/Private/MuT/CodeGenerator.h
@@ -512,8 +512,11 @@ namespace mu
 		/** Evaluate if the image to generate is big enough to be split in separate operations and tiled afterwards. */
 		Ptr<ASTOp> ApplyTiling(Ptr<ASTOp> Source, UE::Math::TIntVector2<int32> Size, EImageFormat Format);
 
-		//!
-		Ptr<ASTOp> GenerateImageBlockPatch(Ptr<ASTOp> blockAd, const NodePatchImage* pPatch, Ptr<ASTOp> conditionAd, const FImageGenerationOptions& ImageOptions);
+		/** Generate a layout block-sized image with a mask including all pixels in the blocks defined in the patch node. */
+		Ptr<Image> GenerateImageBlockPatchMask(const NodePatchImage* Patch, FIntPoint GridSize, int32 BlockPixelsX, int32 BlockPixelsY, box<UE::Math::TIntVector2<uint16>> RectInCells);
+
+		/** Generate all the operations to apply the block patching on top of the BlockOp, and masking with PatchMask. */
+		Ptr<ASTOp> GenerateImageBlockPatch(Ptr<ASTOp> BlockOp, const NodePatchImage*, Ptr<Image> PatchMask, Ptr<ASTOp> ConditionOp, const FImageGenerationOptions&);
 
         //-----------------------------------------------------------------------------------------
         // Meshes
@@ -547,6 +550,9 @@ namespace mu
 			/** If true, UVs will be normalized. Normalize UVs should be done in cases where we operate with Images and Layouts */
 			bool bNormalizeUVs = false;
 
+			/** If true, assign vertices without layout to the first block. */
+			bool bEnsureAllVerticesHaveLayoutBlock = true;
+
 			/** If this has something the layouts in constant meshes will be ignored, because
 			* they are supposed to match some other set of layouts. If the array is empty, layouts
 			* are generated normally.
@@ -572,6 +578,7 @@ namespace mu
 				return State==Other.State 
 					&& bLayouts==Other.bLayouts 
 					&& bClampUVIslands == Other.bClampUVIslands && bNormalizeUVs == Other.bNormalizeUVs
+					&& bEnsureAllVerticesHaveLayoutBlock == Other.bEnsureAllVerticesHaveLayoutBlock
 					&& ActiveTags==Other.ActiveTags
 					&& OverrideLayouts ==Other.OverrideLayouts;
 			}
@@ -630,7 +637,21 @@ namespace mu
 
 		//! Map of layouts found in the code already generated. The map is from the source layout
 		//! node to the generated layout.
-		TMap<Ptr<const NodeLayout>, Ptr<const Layout>> GeneratedLayouts;
+		struct FGeneratedLayoutKey
+		{
+			Ptr<const NodeLayout> SourceLayout;
+			uint32 MeshIdPrefix;
+
+			friend FORCEINLINE uint32 GetTypeHash(const FGeneratedLayoutKey& InKey)
+			{
+				uint32 KeyHash = ::GetTypeHash(InKey.SourceLayout.get());
+				KeyHash = HashCombineFast(KeyHash, GetTypeHash(InKey.MeshIdPrefix));
+				return KeyHash;
+			}
+
+			FORCEINLINE bool operator==(const FGeneratedLayoutKey& Other) const = default;
+		};
+		TMap<FGeneratedLayoutKey, Ptr<const Layout>> GeneratedLayouts;
 
         void GenerateMesh(const FMeshGenerationOptions&, FMeshGenerationResult& result, const NodeMeshPtrConst&);
         void GenerateMesh_Constant(const FMeshGenerationOptions&, FMeshGenerationResult&, const NodeMeshConstant* );
@@ -658,7 +679,7 @@ namespace mu
 			bool bUseAbsoluteBlockIds);
 
 		//!
-		Ptr<const Layout> AddLayout(Ptr<const NodeLayout> SourceLayout, uint32 MeshIDPrefix);
+		Ptr<const Layout> GenerateLayout(Ptr<const NodeLayout> SourceLayout, uint32 MeshIDPrefix);
 
 		struct FExtensionDataGenerationResult
 		{
diff --git a/Engine/Plugins/Experimental/Mutable/Source/MutableTools/Private/MuT/CodeGenerator_Image.cpp b/Engine/Plugins/Experimental/Mutable/Source/MutableTools/Private/MuT/CodeGenerator_Image.cpp
--- a/Engine/Plugins/Experimental/Mutable/Source/MutableTools/Private/MuT/CodeGenerator_Image.cpp
+++ b/Engine/Plugins/Experimental/Mutable/Source/MutableTools/Private/MuT/CodeGenerator_Image.cpp
@@ -1449,7 +1449,7 @@ namespace mu
                 {
                     Ptr<ASTOpMeshExtractLayoutBlocks> eop = new ASTOpMeshExtractLayoutBlocks();
                     eop->Source = CurrentMeshToProjectOp;
-                    eop->Layout = node.m_layout;
+                    eop->LayoutIndex = node.m_layout;
 
                     eop->Blocks.Add(GeneratedLayoutBlockId);
 
diff --git a/Engine/Plugins/Experimental/Mutable/Source/MutableTools/Private/MuT/CodeGenerator_Mesh.cpp b/Engine/Plugins/Experimental/Mutable/Source/MutableTools/Private/MuT/CodeGenerator_Mesh.cpp
--- a/Engine/Plugins/Experimental/Mutable/Source/MutableTools/Private/MuT/CodeGenerator_Mesh.cpp
+++ b/Engine/Plugins/Experimental/Mutable/Source/MutableTools/Private/MuT/CodeGenerator_Mesh.cpp
@@ -54,7 +54,6 @@
 #include "MuT/NodeMeshFormat.h"
 #include "MuT/NodeMeshFormatPrivate.h"
 #include "MuT/NodeMeshFragment.h"
-#include "MuT/NodeMeshFragmentPrivate.h"
 #include "MuT/NodeMeshGeometryOperation.h"
 #include "MuT/NodeMeshGeometryOperationPrivate.h"
 #include "MuT/NodeMeshInterpolate.h"
@@ -684,16 +683,16 @@ namespace mu
 				}
 				else
 				{
-					// Map vertices without block to the first block.
+					// Map vertices without block
 					if (bUseAbsoluteBlockIds)
 					{
 						uint64* Ptr = reinterpret_cast<uint64*>(LayoutBufferPtr) + VertexIndex;
-						*Ptr = 0;
+						*Ptr = MeshOptions.bEnsureAllVerticesHaveLayoutBlock ? 0 : std::numeric_limits<uint64>::max();
 					}
 					else
 					{
 						uint16* Ptr = reinterpret_cast<uint16*>(LayoutBufferPtr) + VertexIndex;
-						*Ptr = 0;
+						*Ptr = MeshOptions.bEnsureAllVerticesHaveLayoutBlock ? 0 : std::numeric_limits<uint16>::max();
 					}
 				}
 
@@ -925,66 +924,36 @@ namespace mu
 	}
 
     //---------------------------------------------------------------------------------------------
-    void CodeGenerator::GenerateMesh_Fragment(const FMeshGenerationOptions& InOptions, FMeshGenerationResult& OutResult,
-		const NodeMeshFragment* FragmentNode )
+    void CodeGenerator::GenerateMesh_Fragment(const FMeshGenerationOptions& InOptions, FMeshGenerationResult& OutResult, const NodeMeshFragment* Node )
     {
-        NodeMeshFragment::Private& node = *FragmentNode->GetPrivate();
-
         FMeshGenerationResult BaseResult;
-        if ( node.m_pMesh )
+        if ( Node->SourceMesh )
         {
-			FMeshGenerationOptions BaseOptions = InOptions;
-			if (node.m_fragmentType == NodeMeshFragment::FT_LAYOUT_BLOCKS)
-			{
-				BaseOptions.bLayouts = true;
-			}
-            GenerateMesh( BaseOptions, BaseResult, node.m_pMesh );
-
-            if ( node.m_fragmentType==NodeMeshFragment::FT_LAYOUT_BLOCKS )
-            {
-                Ptr<ASTOpMeshExtractLayoutBlocks> op = new ASTOpMeshExtractLayoutBlocks();
-                OutResult.MeshOp = op;
+			Ptr<ASTOpMeshExtractLayoutBlocks> op = new ASTOpMeshExtractLayoutBlocks();
+			OutResult.MeshOp = op;
 
-                op->Source = BaseResult.MeshOp;
+			op->LayoutIndex = (uint16)Node->LayoutIndex;
 
-                if (BaseResult.GeneratedLayouts.Num()>node.LayoutOrGroup )
-                {
-                    const Layout* pLayout = BaseResult.GeneratedLayouts[node.LayoutOrGroup].Layout.get();
-                    op->Layout = (uint16)node.LayoutOrGroup;
-
-                    for ( int32 i=0; i<node.Blocks.Num(); ++i )
-                    {
-                        if (node.Blocks[i]>=0 && node.Blocks[i]<pLayout->Blocks.Num() )
-                        {
-                            uint64 bid = pLayout->Blocks[ node.Blocks[i] ].Id;
-                            op->Blocks.Add(bid);
-                        }
-                        else
-                        {
-                            ErrorLog->GetPrivate()->Add( "Internal layout block index error.",
-                                                            ELMT_ERROR, FragmentNode->GetMessageContext());
-                        }
-                    }
-                }
-                else
-                {
-                    // This argument is required
-                    ErrorLog->GetPrivate()->Add( "Missing layout in mesh fragment source.",
-                                                    ELMT_ERROR, FragmentNode->GetMessageContext());
-                }
-            }
+			// Generate the source mesh
+			FMeshGenerationOptions BaseOptions = InOptions;
+			BaseOptions.bLayouts = true;
+			BaseOptions.bEnsureAllVerticesHaveLayoutBlock = false;
 
-            else
-            {
-				check(false);
-            }
+			if (Node->Layout)
+			{
+				// Generate the layout with blocks to extract
+				Ptr<const Layout> Layout = GenerateLayout(Node->Layout, 0);
+				BaseOptions.OverrideLayouts.Empty();
+				BaseOptions.OverrideLayouts.Emplace(Layout, Node->Layout);
+			}
 
+            GenerateMesh( BaseOptions, BaseResult, Node->SourceMesh);
+            op->Source = BaseResult.MeshOp;
         }
         else
         {
             // This argument is required
-            ErrorLog->GetPrivate()->Add( "Mesh fragment source is not set.",
-                                            ELMT_ERROR, FragmentNode->GetMessageContext());
+            ErrorLog->GetPrivate()->Add( "Mesh fragment source is not set.", ELMT_ERROR, Node->GetMessageContext());
         }
 
         OutResult.BaseMeshOp = BaseResult.BaseMeshOp;
@@ -1375,11 +1344,7 @@ namespace mu
 			{
 				if (bIsOverridingLayouts)
 				{
-					for (int32 l = 0; l < DuplicateOf.Mesh->GetLayoutCount(); ++l)
-					{
-						FGeneratedLayout OverridingLayout = InOptions.OverrideLayouts[l];
-						OutResult.GeneratedLayouts.Add(OverridingLayout);
-					}
+					OutResult.GeneratedLayouts = InOptions.OverrideLayouts;
 				}
 				else
 				{
@@ -1442,7 +1407,7 @@ namespace mu
 
 						FGeneratedLayout GeneratedData;
 						GeneratedData.Source = LayoutNode;
-						GeneratedData.Layout = AddLayout(LayoutNode, MeshIDPrefix);
+						GeneratedData.Layout = GenerateLayout(LayoutNode, MeshIDPrefix);
 						const void* Context = InOptions.OverrideContext.Get(InNode->GetMessageContext());
 
 						bool bUseAbsoluteBlockIds = false;
diff --git a/Engine/Plugins/Experimental/Mutable/Source/MutableTools/Private/MuT/Node.cpp b/Engine/Plugins/Experimental/Mutable/Source/MutableTools/Private/MuT/Node.cpp
--- a/Engine/Plugins/Experimental/Mutable/Source/MutableTools/Private/MuT/Node.cpp
+++ b/Engine/Plugins/Experimental/Mutable/Source/MutableTools/Private/MuT/Node.cpp
@@ -28,6 +28,8 @@
 #include "MuT/NodeColourFromScalars.h"
 #include "MuT/NodeColourTable.h"
 #include "MuT/NodeScalarSwitch.h"
+#include "MuT/NodeMeshFragment.h"
+#include "MuT/NodePatchImage.h"
 
 namespace mu
 {
@@ -63,6 +65,10 @@ namespace mu
 	FNodeType NodeColourSampleImage::StaticType = FNodeType(Node::EType::ColorSampleImage, NodeColour::GetStaticType());
 	FNodeType NodeColourFromScalars::StaticType = FNodeType(Node::EType::ColorFromScalars, NodeColour::GetStaticType());
 
+	FNodeType NodeMeshFragment::StaticType = FNodeType(Node::EType::MeshFragment, NodeMesh::GetStaticType());
+
+	FNodeType NodePatchImage::StaticType = FNodeType(Node::EType::PatchImage, Node::GetStaticType());
+
 
 	FNodeType::FNodeType()
 	{
diff --git a/Engine/Plugins/Experimental/Mutable/Source/MutableTools/Private/MuT/NodeMeshFragment.cpp b/Engine/Plugins/Experimental/Mutable/Source/MutableTools/Private/MuT/NodeMeshFragment.cpp
deleted file mode 100644
diff --git a/Engine/Plugins/Experimental/Mutable/Source/MutableTools/Private/MuT/NodePatchImage.cpp b/Engine/Plugins/Experimental/Mutable/Source/MutableTools/Private/MuT/NodePatchImage.cpp
deleted file mode 100644
diff --git a/Engine/Plugins/Experimental/Mutable/Source/MutableTools/Public/MuT/NodeMeshFragment.h b/Engine/Plugins/Experimental/Mutable/Source/MutableTools/Public/MuT/NodeMeshFragment.h
--- a/Engine/Plugins/Experimental/Mutable/Source/MutableTools/Public/MuT/NodeMeshFragment.h
+++ b/Engine/Plugins/Experimental/Mutable/Source/MutableTools/Public/MuT/NodeMeshFragment.h
@@ -3,81 +3,43 @@
 #pragma once
 
 #include "MuR/Ptr.h"
-#include "MuR/RefCounted.h"
-#include "MuT/Node.h"
 #include "MuT/NodeMesh.h"
-
+#include "MuT/NodeLayout.h"
 
 namespace mu
 {
 
-	// Forward definitions
-	class NodeMeshFragment;
-	typedef Ptr<NodeMeshFragment> NodeMeshFragmentPtr;
-	typedef Ptr<const NodeMeshFragment> NodeMeshFragmentPtrConst;
-
-
-	//! This node selects a fragment of a mesh, by selecting some of its layout blocks.
-	//! \ingroup model
+	/** This node extracts a fragment of a mesh by some selection on a texture layout space. */
 	class MUTABLETOOLS_API NodeMeshFragment : public NodeMesh
 	{
 	public:
 
-		NodeMeshFragment();
-
-
-		//-----------------------------------------------------------------------------------------
-		// Node Interface
-		//-----------------------------------------------------------------------------------------
-
-		const FNodeType* GetType() const override;
-		static const FNodeType* GetStaticType();
-
-		//-----------------------------------------------------------------------------------------
-		// Own Interface
-		//-----------------------------------------------------------------------------------------
+		/** Mesh from which a fragment will be extracted. */
+		Ptr<NodeMesh> SourceMesh;
 
-        typedef enum {
-            FT_LAYOUT_BLOCKS = 0,
-            FT_FACE_GROUP_DEPRECATED
-        } FRAGMENT_TYPE;
+		/** Layout defining the blocks of mesh to extract by texture space. 
+		* If no layout is specified, all the blocks in the layout in the SourceMesh subgraph will be extracted.
+		*/
+		Ptr<NodeLayout> Layout;
 
-        //! Get the type of mesh fragment to extract.
-        FRAGMENT_TYPE GetFragmentType() const;
-        void SetFragmentType( FRAGMENT_TYPE type );
-
-		//!
-		NodeMeshPtr GetMesh() const;
-		void SetMesh( NodeMeshPtr );
-
-		//! Get the index of the layout to use to extract blocks.
-        int GetLayoutOrGroup() const;
-        void SetLayoutOrGroup( int layoutIndex );
-
-		//! Set the number of layout blocks to extract.
-		void SetBlockCount( int );
-
-		//! Get one of the layout block indices
-		int GetBlock( int i ) const;
-		void SetBlock( int i, int blockIndex );
+		/** Index of the UV channel in the source mesh to apply the Layout to decide what to extract. */
+		int32 LayoutIndex = 0;
 
+	public:
 
-		//-----------------------------------------------------------------------------------------
-		// Interface pattern
-		//-----------------------------------------------------------------------------------------
-		class Private;
-		Private* GetPrivate() const;
+		// Node Interface
+		virtual const FNodeType* GetType() const override { return GetStaticType(); }
+		static const FNodeType* GetStaticType() { return &StaticType; }
 
 	protected:
 
 		//! Forbidden. Manage with the Ptr<> template.
-		~NodeMeshFragment();
+		~NodeMeshFragment() {}
 
 	private:
 
-		Private* m_pD;
+		static FNodeType StaticType;
 
 	};
 
-
 }
