[mutable] Fix mac compilation issues.
#rnx

diff --git a/Engine/Plugins/Experimental/Mutable/Source/MutableTools/Private/MuT/CodeGenerator.cpp b/Engine/Plugins/Experimental/Mutable/Source/MutableTools/Private/MuT/CodeGenerator.cpp
--- a/Engine/Plugins/Experimental/Mutable/Source/MutableTools/Private/MuT/CodeGenerator.cpp
+++ b/Engine/Plugins/Experimental/Mutable/Source/MutableTools/Private/MuT/CodeGenerator.cpp
@@ -517,10 +517,6 @@ namespace mu
 
 			FinalOp = conditionalAd;
 		}
-		else
-		{
-			FinalOp = FinalOp;
-		}
 
 		return FinalOp;
 	}
diff --git a/Engine/Plugins/Experimental/Mutable/Source/MutableTools/Private/MuT/CodeGenerator_Mesh.cpp b/Engine/Plugins/Experimental/Mutable/Source/MutableTools/Private/MuT/CodeGenerator_Mesh.cpp
--- a/Engine/Plugins/Experimental/Mutable/Source/MutableTools/Private/MuT/CodeGenerator_Mesh.cpp
+++ b/Engine/Plugins/Experimental/Mutable/Source/MutableTools/Private/MuT/CodeGenerator_Mesh.cpp
@@ -944,7 +944,7 @@ namespace mu
 				// Generate the layout with blocks to extract
 				Ptr<const Layout> Layout = GenerateLayout(Node->Layout, 0);
 				BaseOptions.OverrideLayouts.Empty();
-				BaseOptions.OverrideLayouts.Emplace(Layout, Node->Layout);
+				BaseOptions.OverrideLayouts.Add({ Layout, Node->Layout });
 			}
 
             GenerateMesh( BaseOptions, BaseResult, Node->SourceMesh);
