[Mutable] Memory tracking allocator fixes.
 -Change tracking wrapping allocator so it more natvis friendly.
 -Fix aligned member functions compilaton when the base allocator does not support aligment.
 -Fix different counters created for each module. Allocations made in a module and frees in an other module would create a missmatch.

#rnx

diff --git a/Engine/Plugins/Experimental/Mutable/Source/MutableRuntime/Internal/MuR/SystemPrivate.h b/Engine/Plugins/Experimental/Mutable/Source/MutableRuntime/Internal/MuR/SystemPrivate.h
--- a/Engine/Plugins/Experimental/Mutable/Source/MutableRuntime/Internal/MuR/SystemPrivate.h
+++ b/Engine/Plugins/Experimental/Mutable/Source/MutableRuntime/Internal/MuR/SystemPrivate.h
@@ -23,8 +23,10 @@
 
 namespace mu::MemoryCounters
 {
-	struct FMemoryTrackerInternalMemoryCounterTag {};
-	using FMemoryTrackerInternalMemoryCounter = TMemoryCounter<FMemoryTrackerInternalMemoryCounterTag>;
+	struct MUTABLERUNTIME_API  FInternalMemoryCounter
+	{
+		static inline std::atomic<SSIZE_T> Counter {0};
+	};
 }
 
 namespace mu
@@ -217,7 +219,7 @@ namespace mu
 	template<class DATA>
 	class CodeContainer
 	{
-		using MemoryCounter = MemoryCounters::FMemoryTrackerInternalMemoryCounter;
+		using MemoryCounter = MemoryCounters::FInternalMemoryCounter;
 		using ArrayDataContainerType = TArray<DATA, FDefaultMemoryTrackingAllocator<MemoryCounter>>;
 		using MapDataContainerType = TMap<FCacheAddress, DATA, FDefaultMemoryTrackingSetAllocator<MemoryCounter>>;
 
@@ -449,7 +451,7 @@ namespace mu
 	{
 	public:
 
-		using AllocType = FDefaultMemoryTrackingAllocator<MemoryCounters::FMemoryTrackerInternalMemoryCounter>;
+		using AllocType = FDefaultMemoryTrackingAllocator<MemoryCounters::FInternalMemoryCounter>;
 
 		template<class Type, class Alloc = AllocType >
 		using TMemoryTrackedArray = TArray<Type, Alloc>;
@@ -1127,7 +1129,7 @@ namespace mu
     /** Struct to manage all the memory allocated for resources used during mutable operation. */
     struct FWorkingMemoryManager
     {	
-		using MemoryCounter = MemoryCounters::FMemoryTrackerInternalMemoryCounter;
+		using MemoryCounter = MemoryCounters::FInternalMemoryCounter;
 
 		template<class Type>
 		using TMemoryTrackedArray = TArray<Type, FDefaultMemoryTrackingAllocator<MemoryCounter>>;
diff --git a/Engine/Plugins/Experimental/Mutable/Source/MutableRuntime/Private/MuR/CodeRunner.h b/Engine/Plugins/Experimental/Mutable/Source/MutableRuntime/Private/MuR/CodeRunner.h
--- a/Engine/Plugins/Experimental/Mutable/Source/MutableRuntime/Private/MuR/CodeRunner.h
+++ b/Engine/Plugins/Experimental/Mutable/Source/MutableRuntime/Private/MuR/CodeRunner.h
@@ -25,8 +25,10 @@
 
 namespace mu::MemoryCounters
 {
-	struct FStreamingMemoryCounterTag {};
-	using  FStreamingMemoryCounter = TMemoryCounter<FStreamingMemoryCounterTag>;
+	struct MUTABLERUNTIME_API FStreamingMemoryCounter
+	{
+		static inline std::atomic<SSIZE_T> Counter {0};
+	};
 }
 
 namespace  mu
diff --git a/Engine/Plugins/Experimental/Mutable/Source/MutableRuntime/Private/MuR/System.cpp b/Engine/Plugins/Experimental/Mutable/Source/MutableRuntime/Private/MuR/System.cpp
--- a/Engine/Plugins/Experimental/Mutable/Source/MutableRuntime/Private/MuR/System.cpp
+++ b/Engine/Plugins/Experimental/Mutable/Source/MutableRuntime/Private/MuR/System.cpp
@@ -47,7 +47,6 @@ static FAutoConsoleVariableRef CVarEnableDetailedMemoryBudgetExceededLogging (
 	ECVF_Default);
 }
 
-
 namespace mu
 {
 	MUTABLE_IMPLEMENT_ENUM_SERIALISABLE(ETextureCompressionStrategy);
@@ -1129,7 +1128,7 @@ namespace mu
 		const SSIZE_T ImageAllocBytes	 = MemoryCounters::FImageMemoryCounter::Counter.load(std::memory_order_relaxed);
 		const SSIZE_T MeshAllocBytes     = MemoryCounters::FMeshMemoryCounter::Counter.load(std::memory_order_relaxed);
 		const SSIZE_T StreamAllocBytes   = MemoryCounters::FStreamingMemoryCounter::Counter.load(std::memory_order_relaxed);
-		const SSIZE_T InternalAllocBytes = MemoryCounters::FMemoryTrackerInternalMemoryCounter::Counter.load(std::memory_order_relaxed);
+		const SSIZE_T InternalAllocBytes = MemoryCounters::FInternalMemoryCounter::Counter.load(std::memory_order_relaxed);
 
 		SSIZE_T TotalBytes = ImageAllocBytes + MeshAllocBytes + StreamAllocBytes + InternalAllocBytes;
 
@@ -1217,7 +1216,7 @@ namespace mu
 		SSIZE_T TotalBytes = MemoryCounters::FImageMemoryCounter::Counter.load(std::memory_order_relaxed) + 
 						     MemoryCounters::FMeshMemoryCounter::Counter.load(std::memory_order_relaxed) +
 							 MemoryCounters::FStreamingMemoryCounter::Counter.load(std::memory_order_relaxed) +
-							 MemoryCounters::FMemoryTrackerInternalMemoryCounter::Counter.load(std::memory_order_relaxed);
+							 MemoryCounters::FInternalMemoryCounter::Counter.load(std::memory_order_relaxed);
 
 		return TotalBytes;
 	}
diff --git a/Engine/Plugins/Experimental/Mutable/Source/MutableRuntime/Public/MuR/ImageDataStorage.h b/Engine/Plugins/Experimental/Mutable/Source/MutableRuntime/Public/MuR/ImageDataStorage.h
--- a/Engine/Plugins/Experimental/Mutable/Source/MutableRuntime/Public/MuR/ImageDataStorage.h
+++ b/Engine/Plugins/Experimental/Mutable/Source/MutableRuntime/Public/MuR/ImageDataStorage.h
@@ -9,10 +9,14 @@
 #include "Containers/ArrayView.h"
 #include "Containers/StaticArray.h"
 
+#include <atomic>
+
 namespace mu::MemoryCounters
 {
-	struct FImageMemoryCounterTag {};
-	using FImageMemoryCounter = TMemoryCounter<FImageMemoryCounterTag>;
+	struct MUTABLERUNTIME_API FImageMemoryCounter
+	{
+		static inline std::atomic<SSIZE_T> Counter {0};
+	};
 }
 
 namespace mu
diff --git a/Engine/Plugins/Experimental/Mutable/Source/MutableRuntime/Public/MuR/MemoryTrackingAllocationPolicy.h b/Engine/Plugins/Experimental/Mutable/Source/MutableRuntime/Public/MuR/MemoryTrackingAllocationPolicy.h
--- a/Engine/Plugins/Experimental/Mutable/Source/MutableRuntime/Public/MuR/MemoryTrackingAllocationPolicy.h
+++ b/Engine/Plugins/Experimental/Mutable/Source/MutableRuntime/Public/MuR/MemoryTrackingAllocationPolicy.h
@@ -11,22 +11,19 @@
 
 namespace mu
 {
-	template<typename TagType>
-	struct TMemoryCounter
-	{
-		alignas(8) static inline std::atomic<SSIZE_T> Counter{ 0 };
-	};
-
-	template<typename Type>
-	struct TIsAMemoryCounter : std::false_type {};
-
-	template<typename TagType>
-	struct TIsAMemoryCounter<TMemoryCounter<TagType>> : std::true_type {};
+	
+	/**
+	 * CounterType is expected to be of the following form:
+	 * struct FCounterTypeName
+	 * {
+	 *     static inline std::atomic<SSIZE_T> Counter {0}; 		
+	 * };
+	 */
 
 	template<typename BaseAlloc, typename CounterType>
-	class FMemoryTrackingAllocatorWrapper
+	class TMemoryTrackingAllocatorWrapper
 	{
-		static_assert(TIsAMemoryCounter<CounterType>::value);
+		static_assert(std::is_same_v<decltype(CounterType::Counter.load()), SSIZE_T>, "CounterType::Counter must be signed.");
 
 	public:
 		using SizeType = typename BaseAlloc::SizeType;
@@ -34,7 +31,14 @@ namespace mu
 		enum { NeedsElementType = BaseAlloc::NeedsElementType };
 		enum { RequireRangeCheck = BaseAlloc::RequireRangeCheck };
 
-		class ForAnyElementType 
+		/** 
+		 * ForAnyElementType is privately inherited from the wrapped allocator ForAnyElementType so 
+		 * the base class members must be explicitly defined to compile. This way if any new method 
+		 * is added to the allocator interface, it forces its addition here. This is useful in case 
+		 * the new method needs to do some memory tracking, otherwise a simple using declaration may 
+		 * suffice.
+		 */
+		class ForAnyElementType : private BaseAlloc::ForAnyElementType 
 		{
 		public:
 			ForAnyElementType() 
@@ -56,14 +60,9 @@ namespace mu
 			ForAnyElementType(const ForAnyElementType&) = delete;
 			ForAnyElementType& operator=(const ForAnyElementType&) = delete;
 
-			FORCEINLINE decltype(auto) GetAllocation() const
-			{
-				return Base.GetAllocation();
-			}
-
 			FORCEINLINE void MoveToEmpty(ForAnyElementType& Other)
 			{
-				Base.MoveToEmpty(Other.Base);
+				BaseAlloc::ForAnyElementType::MoveToEmpty(Other);
 
 				CounterType::Counter.fetch_sub(AllocSize, std::memory_order_relaxed);
 
@@ -81,9 +80,10 @@ namespace mu
 				SIZE_T NumBytesPerElement
 			)
 			{
-				Base.ResizeAllocation(PreviousNumElements, NumElements, NumBytesPerElement);
+				BaseAlloc::ForAnyElementType::ResizeAllocation(PreviousNumElements, NumElements, NumBytesPerElement);
 
-				const SSIZE_T Differential = SSIZE_T(NumElements * NumBytesPerElement) - AllocSize;
+				const SSIZE_T AllocatedSize = (SSIZE_T)BaseAlloc::ForAnyElementType::GetAllocatedSize(NumElements, NumBytesPerElement); 
+				const SSIZE_T Differential = AllocatedSize - AllocSize;
 				const SSIZE_T PrevCounterValue = CounterType::Counter.fetch_add(Differential, std::memory_order_relaxed);
 				check(PrevCounterValue >= AllocSize);
 
@@ -91,19 +91,21 @@ namespace mu
 				FGlobalMemoryCounter::Update(Differential);
 #endif
 
-				AllocSize = NumElements * NumBytesPerElement;
+				AllocSize = AllocatedSize;
 			}
 
-			FORCEINLINE typename TEnableIf<TAllocatorTraits<BaseAlloc>::SupportsElementAlignment, void>::Type ResizeAllocation(
+			template<class T = BaseAlloc>
+			FORCEINLINE typename TEnableIf<TAllocatorTraits<T>::SupportsElementAlignment, void>::Type ResizeAllocation(
 				SizeType PreviousNumElements,
 				SizeType NumElements,
 				SIZE_T NumBytesPerElement,
 				uint32 AlignmentOfElement
 			)
 			{
-				Base.ResizeAllocation(PreviousNumElements, NumElements, NumBytesPerElement, AlignmentOfElement);
+				BaseAlloc::ForAnyElementType::ResizeAllocation(PreviousNumElements, NumElements, NumBytesPerElement, AlignmentOfElement);
 
-				const SSIZE_T Differential = SSIZE_T(NumElements * NumBytesPerElement) - AllocSize;
+				const SSIZE_T AllocatedSize = (SSIZE_T)BaseAlloc::ForAnyElementType::GetAllocatedSize(NumElements, NumBytesPerElement); 
+				const SSIZE_T Differential = AllocatedSize - AllocSize;
 				const SSIZE_T PrevCounterValue = CounterType::Counter.fetch_add(Differential, std::memory_order_relaxed);
 				check(PrevCounterValue >= AllocSize);
 
@@ -111,104 +113,35 @@ namespace mu
 				FGlobalMemoryCounter::Update(Differential);
 #endif
 
-				AllocSize = NumElements * NumBytesPerElement;
-			}
-
-			FORCEINLINE SizeType CalculateSlackReserve(
-				SizeType NumElements,
-				SIZE_T NumBytesPerElement
-			) const
-			{
-				return Base.CalculateSlackReserve(NumElements, NumBytesPerElement);
-			}
-
-			FORCEINLINE typename TEnableIf<TAllocatorTraits<BaseAlloc>::SupportsElementAlignment, SizeType>::Type CalculateSlackReserve(
-				SizeType NumElements,
-				SIZE_T NumBytesPerElement,
-				uint32 AlignmentOfElement
-			) const
-			{
-				return Base.CalculateSlackReserve(NumElements, NumBytesPerElement, AlignmentOfElement);
-			}
-
-			FORCEINLINE SizeType CalculateSlackShrink(
-				SizeType NumElements,
-				SizeType CurrentNumSlackElements,
-				SIZE_T NumBytesPerElement
-			) const
-			{
-				return Base.CalculateSlackShrink(NumElements, CurrentNumSlackElements, NumBytesPerElement);
-			}
-
-			FORCEINLINE typename TEnableIf<TAllocatorTraits<BaseAlloc>::SupportsElementAlignment, SizeType>::Type CalculateSlackShrink(
-				SizeType NumElements,
-				SizeType CurrentNumSlackElements,
-				SIZE_T NumBytesPerElement,
-				uint32 AlignmentOfElement
-			) const
-			{
-				return Base.CalculateSlackShrink(NumElements, CurrentNumSlackElements, NumBytesPerElement, AlignmentOfElement);
-			}
-
-			SizeType CalculateSlackGrow(
-				SizeType NumElements,
-				SizeType CurrentNumSlackElements,
-				SIZE_T NumBytesPerElement
-			) const
-			{
-				return Base.CalculateSlackGrow(NumElements, CurrentNumSlackElements, NumBytesPerElement);
-			}
-
-			FORCEINLINE typename TEnableIf<TAllocatorTraits<BaseAlloc>::SupportsElementAlignment, SizeType>::Type CalculateSlackGrow(
-				SizeType NumElements,
-				SizeType CurrentNumSlackElements,
-				SIZE_T NumBytesPerElement,
-				uint32 AlignmentOfElement
-			) const
-			{
-				return Base.CalculateSlackGrow(NumElements, CurrentNumSlackElements, NumBytesPerElement, AlignmentOfElement);
-			}
-
-			FORCEINLINE SIZE_T GetAllocatedSize(
-				SizeType NumAllocatedElements, 
-				SIZE_T NumBytesPerElement
-			) const
-			{
-				return Base.GetAllocatedSize(NumAllocatedElements, NumBytesPerElement);
+				AllocSize = AllocatedSize;
 			}
 
-			FORCEINLINE bool HasAllocation() const
-			{
-				return Base.HasAllocation();
-			}
-
-			FORCEINLINE SizeType GetInitialCapacity() const
-			{
-				return Base.GetInitialCapacity();
-			}
+			// Explicitly incorporate pass-through base allocator member functions.
+			using BaseAlloc::ForAnyElementType::GetAllocation;
+			using BaseAlloc::ForAnyElementType::CalculateSlackReserve;
+			using BaseAlloc::ForAnyElementType::CalculateSlackShrink;
+			using BaseAlloc::ForAnyElementType::CalculateSlackGrow;
+			using BaseAlloc::ForAnyElementType::GetAllocatedSize;
+			using BaseAlloc::ForAnyElementType::HasAllocation;
+			using BaseAlloc::ForAnyElementType::GetInitialCapacity;
 
 #if UE_ENABLE_ARRAY_SLACK_TRACKING
 			FORCEINLINE void SlackTrackerLogNum(SizeType NewNumUsed)
 			{
 				if constexpr (TAllocatorTraits<BaseAlloc>::SupportsSlackTracking)
 				{
-					Base.SlackTrackerLogNum(NewNumUsed);
+					BaseAlloc::ForElementType::SlackTrackerLogNum(NewNumUsed);	
 				}
 			}
 #endif
 
 		private:
-			typename BaseAlloc::ForAnyElementType Base;
 			SSIZE_T AllocSize = 0;
 		};
 
 		template<typename ElementType>
 		class ForElementType : public ForAnyElementType
 		{
-			// Some Allocators, e.g, TAlignedHeapAllocator do some static checks. Instanciate the BaseAlloc::ForElementType
-			// at compile time so those warnings are emitted. 
-			using BaseTypeInstantaitionType = decltype(DeclVal<typename BaseAlloc::template ForElementType<ElementType>>());
-
 		public:
 			ForElementType()
 			{
@@ -219,13 +152,14 @@ namespace mu
 				return (ElementType*)ForAnyElementType::GetAllocation();
 			}
 		};
+
 	};
 
 
-	// Default memory tarcking allocator needed for TArray and TMap.
+	/** Default memory tracking allocators needed for TArray and TMap. */
 
 	template<typename CounterType>
-	using FDefaultMemoryTrackingAllocator = FMemoryTrackingAllocatorWrapper<FDefaultAllocator, CounterType>;
+	using FDefaultMemoryTrackingAllocator = TMemoryTrackingAllocatorWrapper<FDefaultAllocator, CounterType>;
 
 	template<typename CounterType>
 	using FDefaultMemoryTrackingBitArrayAllocator = TInlineAllocator<4, FDefaultMemoryTrackingAllocator<CounterType>>;
@@ -245,11 +179,8 @@ namespace mu
 }
 
 template<typename BaseAlloc, typename Counter>
-struct TAllocatorTraits<mu::FMemoryTrackingAllocatorWrapper<BaseAlloc, Counter>> : public TAllocatorTraitsBase<mu::FMemoryTrackingAllocatorWrapper<BaseAlloc, Counter>>
+struct TAllocatorTraits<mu::TMemoryTrackingAllocatorWrapper<BaseAlloc, Counter>> : public TAllocatorTraits<BaseAlloc>
 {
-	enum { SupportsElementAlignment = TAllocatorTraits<BaseAlloc>::SupportsElementAlignment };
-	enum { SupportsSlackTracking = TAllocatorTraits<BaseAlloc>::SupportsSlackTracking };
 	enum { SupportsMoveFromOtherAllocator = false };
 };
 
-
diff --git a/Engine/Plugins/Experimental/Mutable/Source/MutableRuntime/Public/MuR/MeshBufferSet.h b/Engine/Plugins/Experimental/Mutable/Source/MutableRuntime/Public/MuR/MeshBufferSet.h
--- a/Engine/Plugins/Experimental/Mutable/Source/MutableRuntime/Public/MuR/MeshBufferSet.h
+++ b/Engine/Plugins/Experimental/Mutable/Source/MutableRuntime/Public/MuR/MeshBufferSet.h
@@ -11,8 +11,10 @@
 
 namespace mu::MemoryCounters
 {
-	struct FMeshMemoryCounterTag {};
-	using FMeshMemoryCounter = TMemoryCounter<FMeshMemoryCounterTag>;
+	struct MUTABLERUNTIME_API FMeshMemoryCounter
+	{
+		static inline std::atomic<SSIZE_T> Counter{0};
+	};
 }
 
 namespace mu
