[Mutable] Set dummy image when a reference image fails to convert to Mutable.
 - This can prevent crashes when data is corrupt or the source asset is missing.

#rnx

diff --git a/Engine/Plugins/Experimental/Mutable/Source/CustomizableObject/Private/MuCO/UnrealMutableImageProvider.cpp b/Engine/Plugins/Experimental/Mutable/Source/CustomizableObject/Private/MuCO/UnrealMutableImageProvider.cpp
--- a/Engine/Plugins/Experimental/Mutable/Source/CustomizableObject/Private/MuCO/UnrealMutableImageProvider.cpp
+++ b/Engine/Plugins/Experimental/Mutable/Source/CustomizableObject/Private/MuCO/UnrealMutableImageProvider.cpp
@@ -369,11 +369,13 @@ TTuple<UE::Tasks::FTask, TFunction<void()>> FUnrealMutableResourceProvider::GetR
 		// This could happen in the editor, because some source textures may have changed while updating.
 		// We just show a warning and move on. This cannot happen during cooks, so it is fine.
 		UE_LOG(LogMutable, Warning, TEXT("Failed to load some source texture data for image [%i]. Some textures may be corrupted."), Id);
+		
+		ResultCallback(CreateDummy());
+		return Invoke(TrivialReturn);
 	}
 
 	ResultCallback(Image);
 	return Invoke(TrivialReturn);
-	
 #else // WITH_EDITOR
 
 	// Not supported outside editor yet.
