[mutable] Improved readability of the triangle count of the Mutable Performance Analyzer

#rnx

diff --git a/Engine/Plugins/Experimental/Mutable/Source/CustomizableObjectEditor/Private/MuCOE/CustomizableObjectEditorPerformanceAnalyzer.cpp b/Engine/Plugins/Experimental/Mutable/Source/CustomizableObjectEditor/Private/MuCOE/CustomizableObjectEditorPerformanceAnalyzer.cpp
--- a/Engine/Plugins/Experimental/Mutable/Source/CustomizableObjectEditor/Private/MuCOE/CustomizableObjectEditorPerformanceAnalyzer.cpp
+++ b/Engine/Plugins/Experimental/Mutable/Source/CustomizableObjectEditor/Private/MuCOE/CustomizableObjectEditorPerformanceAnalyzer.cpp
@@ -96,7 +96,7 @@ TSharedRef<SWidget> SInstanceUpdateDataRow::GenerateWidgetForColumn(const FName&
 	else if (InColumnName == InstanceUpdatesMainDataColumns::InstanceTriangleCount)
 	{
 		return SNew(STextBlock)
-			.Text(FText::FromString(FString::Printf(TEXT("%u"),InstanceUpdateElement->UpdateStats.TriangleCount)));
+			.Text( FText::AsNumber(InstanceUpdateElement->UpdateStats.TriangleCount));
 	}
 	else if (InColumnName == InstanceUpdatesMainDataColumns::InstanceDescriptorColumnID)
 	{
@@ -126,7 +126,6 @@ TSharedRef<SWidget> SInstanceUpdateDataRow::GenerateWidgetForColumn(const FName&
 		return SNew(STextBlock)
 			  .Text(UEnum::GetDisplayValueAsText( InstanceUpdateElement->UpdateStats.UpdateResult))
 			  .ColorAndOpacity(TextColor);
-		
 	}
 	else if (InColumnName == InstanceUpdatesMainDataColumns::QueueTimeColumnID)
 	{
