[mutable] UV layout widget improvements:
- fix concurrency bug with unsafe access to UV render data.
- add more zoom levels and enable panning in any zoom level.
- add option to reset view zoom level and position.
- add additional grid lines for tiles outside the unit tile, X and Y axis, and tile limits.
- more enforcement of UE naming conventions.
- fix possible crashes with dynamic multi-component object previews.
#rnx

diff --git a/Engine/Plugins/Experimental/Mutable/Source/CustomizableObject/Private/MuCO/CustomizableObjectInstance.cpp b/Engine/Plugins/Experimental/Mutable/Source/CustomizableObject/Private/MuCO/CustomizableObjectInstance.cpp
--- a/Engine/Plugins/Experimental/Mutable/Source/CustomizableObject/Private/MuCO/CustomizableObjectInstance.cpp
+++ b/Engine/Plugins/Experimental/Mutable/Source/CustomizableObject/Private/MuCO/CustomizableObjectInstance.cpp
@@ -2319,12 +2319,13 @@ bool UCustomizableInstancePrivate::UpdateSkeletalMesh_PostBeginUpdate0(UCustomiz
 				}
 			}
 			
-			if (!OperationData->MeshChanged[ObjectComponentIndex])
+			if ( (!OperationData->MeshChanged.IsValidIndex(ObjectComponentIndex)) 
+				|| !OperationData->MeshChanged[ObjectComponentIndex] )
 			{
 				continue;
 			}
 
-			USkeletalMesh* SkeletalMesh = SkeletalMeshes[ObjectComponentIndex];
+			USkeletalMesh* SkeletalMesh = SkeletalMeshes.IsValidIndex(ObjectComponentIndex) ? SkeletalMeshes[ObjectComponentIndex] : nullptr;
 
 			if (!SkeletalMesh)
 			{
diff --git a/Engine/Plugins/Experimental/Mutable/Source/CustomizableObjectEditor/Private/MuCOE/CustomizableObjectLayout.cpp b/Engine/Plugins/Experimental/Mutable/Source/CustomizableObjectEditor/Private/MuCOE/CustomizableObjectLayout.cpp
--- a/Engine/Plugins/Experimental/Mutable/Source/CustomizableObjectEditor/Private/MuCOE/CustomizableObjectLayout.cpp
+++ b/Engine/Plugins/Experimental/Mutable/Source/CustomizableObjectEditor/Private/MuCOE/CustomizableObjectLayout.cpp
@@ -44,20 +44,20 @@ void UCustomizableObjectLayout::SetPackingStrategy(ECustomizableObjectTextureLay
 
 mu::EPackStrategy ConvertLayoutStrategy(const ECustomizableObjectTextureLayoutPackingStrategy LayoutPackStrategy)
 {
-	mu::EPackStrategy PackStrategy = mu::EPackStrategy::FIXED_LAYOUT;
+	mu::EPackStrategy PackStrategy = mu::EPackStrategy::Fixed;
 
 	switch (LayoutPackStrategy)
 	{
 	case ECustomizableObjectTextureLayoutPackingStrategy::Fixed:
-		PackStrategy = mu::EPackStrategy::FIXED_LAYOUT;
+		PackStrategy = mu::EPackStrategy::Fixed;
 		break;
 
 	case ECustomizableObjectTextureLayoutPackingStrategy::Resizable:
-		PackStrategy = mu::EPackStrategy::RESIZABLE_LAYOUT;
+		PackStrategy = mu::EPackStrategy::Resizeable;
 		break;
 
 	case ECustomizableObjectTextureLayoutPackingStrategy::Overlay:
-		PackStrategy = mu::EPackStrategy::OVERLAY_LAYOUT;
+		PackStrategy = mu::EPackStrategy::Overlay;
 		break;
 
 	default:
diff --git a/Engine/Plugins/Experimental/Mutable/Source/CustomizableObjectEditor/Private/MuCOE/GenerateMutableSource/GenerateMutableSourceLayout.cpp b/Engine/Plugins/Experimental/Mutable/Source/CustomizableObjectEditor/Private/MuCOE/GenerateMutableSource/GenerateMutableSourceLayout.cpp
--- a/Engine/Plugins/Experimental/Mutable/Source/CustomizableObjectEditor/Private/MuCOE/GenerateMutableSource/GenerateMutableSourceLayout.cpp
+++ b/Engine/Plugins/Experimental/Mutable/Source/CustomizableObjectEditor/Private/MuCOE/GenerateMutableSource/GenerateMutableSourceLayout.cpp
@@ -91,7 +91,7 @@ mu::Ptr<mu::NodeLayout> CreateMutableLayoutNode(FMutableGraphGenerationContext&
 	mu::EPackStrategy PackStrategy = ConvertLayoutStrategy(UnrealLayout->GetPackingStrategy());
 	LayoutNode->Strategy = PackStrategy;
 
-	LayoutNode->ReductionMethod = (UnrealLayout->GetBlockReductionMethod() == ECustomizableObjectLayoutBlockReductionMethod::Halve ? mu::EReductionMethod::HALVE_REDUCTION : mu::EReductionMethod::UNITARY_REDUCTION);
+	LayoutNode->ReductionMethod = (UnrealLayout->GetBlockReductionMethod() == ECustomizableObjectLayoutBlockReductionMethod::Halve ? mu::EReductionMethod::Halve : mu::EReductionMethod::Unitary);
 
 	if (bIgnoreLayoutWarnings)
 	{
diff --git a/Engine/Plugins/Experimental/Mutable/Source/CustomizableObjectEditor/Private/MuCOE/GenerateMutableSource/GenerateMutableSourceMesh.cpp b/Engine/Plugins/Experimental/Mutable/Source/CustomizableObjectEditor/Private/MuCOE/GenerateMutableSource/GenerateMutableSourceMesh.cpp
--- a/Engine/Plugins/Experimental/Mutable/Source/CustomizableObjectEditor/Private/MuCOE/GenerateMutableSource/GenerateMutableSourceMesh.cpp
+++ b/Engine/Plugins/Experimental/Mutable/Source/CustomizableObjectEditor/Private/MuCOE/GenerateMutableSource/GenerateMutableSourceMesh.cpp
@@ -3293,8 +3293,8 @@ mu::Ptr<mu::NodeLayout> CreateDefaultLayout()
 	mu::Ptr<mu::NodeLayout> LayoutNode = new mu::NodeLayout();
 	LayoutNode->Size = { GridSize, GridSize };
 	LayoutNode->MaxSize = { GridSize, GridSize };
-	LayoutNode->Strategy = mu::EPackStrategy::RESIZABLE_LAYOUT;
-	LayoutNode->ReductionMethod = mu::EReductionMethod::HALVE_REDUCTION;
+	LayoutNode->Strategy = mu::EPackStrategy::Resizeable;
+	LayoutNode->ReductionMethod = mu::EReductionMethod::Halve;
 	LayoutNode->Blocks.SetNum(1);
 	LayoutNode->Blocks[0].Min = { 0, 0 };
 	LayoutNode->Blocks[0].Size = { GridSize, GridSize };
diff --git a/Engine/Plugins/Experimental/Mutable/Source/CustomizableObjectEditor/Private/MuCOE/SCustomizableObjectLayoutGrid.cpp b/Engine/Plugins/Experimental/Mutable/Source/CustomizableObjectEditor/Private/MuCOE/SCustomizableObjectLayoutGrid.cpp
--- a/Engine/Plugins/Experimental/Mutable/Source/CustomizableObjectEditor/Private/MuCOE/SCustomizableObjectLayoutGrid.cpp
+++ b/Engine/Plugins/Experimental/Mutable/Source/CustomizableObjectEditor/Private/MuCOE/SCustomizableObjectLayoutGrid.cpp
@@ -79,7 +79,8 @@ private:
 };
 
 
-/** Custom Slate drawing element. Holds a copy of all information required to draw UVs. */
+/** Custom Slate drawing element. Holds a copy of all information required to draw UVs. 
+*/
 class FUVCanvasDrawer : public ICustomSlateElement
 {
 public:
@@ -151,20 +152,28 @@ void SCustomizableObjectLayoutGrid::Construct( const FArguments& InArgs )
 	OnSetReduceBlockByTwo = InArgs._OnSetReduceBlockByTwo;
 	OnSetBlockMask = InArgs._OnSetBlockMask;
 
-	UVCanvasDrawer = TSharedPtr<FUVCanvasDrawer>(new FUVCanvasDrawer());
-	UVCanvasDrawer->SetLayoutMode(Mode);
+	for (int32 BufferIndex=0;BufferIndex< UE_MUTABLE_UI_DRAWBUFFERS; ++BufferIndex)
+	{
+		UVCanvasDrawers[BufferIndex] = TSharedPtr<FUVCanvasDrawer>(new FUVCanvasDrawer());
+		UVCanvasDrawers[BufferIndex]->SetLayoutMode(Mode);
+	}
 }
 
 
 SCustomizableObjectLayoutGrid::~SCustomizableObjectLayoutGrid()
 {
 	// UVCanvasDrawer can only be destroyed after drawing the last command
-	ENQUEUE_RENDER_COMMAND(SafeDeletePreviewElement)(
-		[UVCanvasDrawer = UVCanvasDrawer](FRHICommandListImmediate& RHICmdList) mutable
-		{
-			UVCanvasDrawer.Reset();
-		}
-	);
+	for (int32 BufferIndex = 0; BufferIndex < UE_MUTABLE_UI_DRAWBUFFERS; ++BufferIndex)
+	{
+		TSharedPtr<class FUVCanvasDrawer> UVCanvasDrawer = UVCanvasDrawers[BufferIndex];
+
+		ENQUEUE_RENDER_COMMAND(SafeDeletePreviewElement)(
+			[UVCanvasDrawer](FRHICommandListImmediate& RHICmdList) mutable
+			{
+				UVCanvasDrawer.Reset();
+			}
+		);
+	}
 }
 
 
@@ -196,16 +205,17 @@ int32 SCustomizableObjectLayoutGrid::OnPaint(const FPaintArgs& Args, const FGeom
 	}
 
 	FVector2D OldSize = Size;
-	Size *= Zoom;
+
+	double ZoomFactor = FMath::Pow(2.0,double(Zoom-1));
+	Size *= ZoomFactor;
 
 	float AuxCellSize = Size.X / GridSizePoint.X;
 	
 	// Drawing Offsets
 	FVector2D Offset = FVector2D((AreaSize - Size).X / 2.0f, 0.0f);
-	FVector2D ZoomOffset = ((Size - OldSize) / 2.0f);
 	
 	// Drawing Origin
-	FVector2D Origin = BorderPadding + Offset + PaddingAmount - DistanceFromOrigin;
+	FVector2D Origin = BorderPadding + Offset + PaddingAmount;
 
 	// Setting Canvas Drawing Rectangles
 	FSlateRect SlateCanvasRect = AllottedGeometry.GetLayoutBoundingRect();
@@ -223,7 +233,7 @@ int32 SCustomizableObjectLayoutGrid::OnPaint(const FPaintArgs& Args, const FGeom
 		FMath::TruncToInt(FMath::Max(0.0f, ClippedCanvasRect.Right)),
 		FMath::TruncToInt(FMath::Max(0.0f, ClippedCanvasRect.Bottom)));
 
-	
+	TSharedPtr<class FUVCanvasDrawer> UVCanvasDrawer = UVCanvasDrawers[CurrentDrawBuffer];
 	UVCanvasDrawer->InitializeDrawingData(UVLayout, UnassignedUVLayoutVertices, Blocks.Get(), SelectedBlocks);
 	UVCanvasDrawer->Initialize(CanvasRect, ClippingRect, Origin * AllottedGeometry.Scale, Size * AllottedGeometry.Scale, GridSizePoint, AuxCellSize * AllottedGeometry.Scale);
 	FSlateDrawElement::MakeCustom(OutDrawElements, RetLayerId, UVCanvasDrawer);
@@ -235,7 +245,7 @@ int32 SCustomizableObjectLayoutGrid::OnPaint(const FPaintArgs& Args, const FGeom
 	};
 
 	// Drawing Multi-Selection rect
-	if (Mode == ELGM_Edit && Selecting)
+	if (Mode == ELGM_Edit && bIsSelecting)
 	{
 		TArray<FVector2D> SelectionSquarePoints;
 		SelectionSquarePoints.SetNum(2);
@@ -270,6 +280,8 @@ int32 SCustomizableObjectLayoutGrid::OnPaint(const FPaintArgs& Args, const FGeom
 
 void SCustomizableObjectLayoutGrid::Tick( const FGeometry& AllottedGeometry, const double InCurrentTime, const float InDeltaTime )
 {
+	CurrentDrawBuffer = FMath::Modulo(CurrentDrawBuffer+1, UE_MUTABLE_UI_DRAWBUFFERS);
+
 	const FVector2D BorderPadding = FVector2D(2,2);
 	const FVector2D AreaSize =  AllottedGeometry.Size - 2.0f * BorderPadding;
 	const float GridRatio = float(GridSize.Get().X)/float(GridSize.Get().Y);
@@ -286,11 +298,12 @@ void SCustomizableObjectLayoutGrid::Tick( const FGeometry& AllottedGeometry, con
 	}
 
 	FVector2D OldSize = Size;
-	Size *= Zoom;
+	double ZoomFactor = FMath::Pow(2.0, double(Zoom - 1));
+	Size *= ZoomFactor;
+
 	CellSize = Size.X/GridSize.Get().X;
 	FVector2D Offset = FVector2D((AreaSize - Size).X / 2.0f, 0.0f);
-	FVector2D ZoomOffset = (Size - OldSize) / 2.0f;
-	FVector2D Origin = BorderPadding + Offset + PaddingAmount - DistanceFromOrigin;
+	FVector2D Origin = BorderPadding + Offset + PaddingAmount;
 	DrawOrigin = Origin;
 
 	BlockRects.Empty();
@@ -313,28 +326,28 @@ void SCustomizableObjectLayoutGrid::Tick( const FGeometry& AllottedGeometry, con
 	}
 
 	// Update selection list
-	for (int i=0; i<SelectedBlocks.Num();)
+	for (int32 SelectedBlockIndex=0; SelectedBlockIndex <SelectedBlocks.Num();)
 	{
-		bool Found = false;
+		bool bFound = false;
 		for (const FCustomizableObjectLayoutBlock& Block : CurrentBlocks)
 		{
-			if (Block.Id == SelectedBlocks[i])
+			if (Block.Id == SelectedBlocks[SelectedBlockIndex])
 			{
-				Found = true;
+				bFound = true;
 			}
 		}
 
-		if ( !Found )
+		if ( !bFound )
 		{
-			SelectedBlocks.RemoveAt(i);
+			SelectedBlocks.RemoveAt(SelectedBlockIndex);
 		}
 		else
 		{
-			++i;
+			++SelectedBlockIndex;
 		}
 	}
 
-	if (Selecting)
+	if (bIsSelecting)
 	{
 		CalculateSelectionRect();
 	}
@@ -345,13 +358,15 @@ void SCustomizableObjectLayoutGrid::Tick( const FGeometry& AllottedGeometry, con
 
 FReply SCustomizableObjectLayoutGrid::OnMouseButtonDown( const FGeometry& MyGeometry, const FPointerEvent& MouseEvent )
 {
+	FReply Reply = FReply::Unhandled();
+
 	if (Mode == ELGM_Edit)
 	{
 		if (MouseEvent.GetEffectingButton() == EKeys::LeftMouseButton)
 		{
-			HasDragged = false;
-			Dragging = false;
-			Resizing = false;
+			bHasDragged = false;
+			bIsDragging = false;
+			bIsResizing = false;
 
 			// To know if we clicked on a block
 			bool ClickOnBlock = false;
@@ -385,12 +400,12 @@ FReply SCustomizableObjectLayoutGrid::OnMouseButtonDown( const FGeometry& MyGeom
 
 			if (SelectedBlocks.Num() && ClickOnBlock)
 			{
-				Dragging = true;
+				bIsDragging = true;
 				DragStart = Pos;
 
 				if (SelectedBlocks.Num() == 1 && SelectedBlockHandles.Contains(SelectedBlocks[0]))
 				{
-					Resizing = true;
+					bIsResizing = true;
 				}
 			}
 		}
@@ -404,6 +419,13 @@ FReply SCustomizableObjectLayoutGrid::OnMouseButtonDown( const FGeometry& MyGeom
 			const bool CloseAfterSelection = true;
 			FMenuBuilder MenuBuilder(CloseAfterSelection, NULL, TSharedPtr<FExtender>(), false, &FCoreStyle::Get(), false);
 
+			MenuBuilder.BeginSection("View", LOCTEXT("ViewActionsTitle", "View"));
+			{
+				FUIAction ResetViewAction(FExecuteAction::CreateSP(this, &SCustomizableObjectLayoutGrid::ResetView));
+				MenuBuilder.AddMenuEntry(LOCTEXT("ResetViewLabel", "Reset View"), LOCTEXT("ResetViewLabelTooltip", "Set the view to the unit UV space."), FSlateIcon(), ResetViewAction);
+			}
+			MenuBuilder.EndSection();
+
 			MenuBuilder.BeginSection("Block Management", LOCTEXT("BlockActionsTitle", "Block Actions"));
 			{
 				if (SelectedBlocks.Num())
@@ -492,18 +514,22 @@ FReply SCustomizableObjectLayoutGrid::OnMouseButtonDown( const FGeometry& MyGeom
 
 			FWidgetPath WidgetPath = MouseEvent.GetEventPath() != nullptr ? *MouseEvent.GetEventPath() : FWidgetPath();
 			FSlateApplication::Get().PushMenu(AsShared(), WidgetPath, MenuBuilder.MakeWidget(), FSlateApplication::Get().GetCursorPos(), FPopupTransitionEffect(FPopupTransitionEffect::ContextMenu));
+
+			Reply = FReply::Handled();
 		}
 		else if (MouseEvent.GetEffectingButton() == EKeys::MiddleMouseButton)
 		{
-			if (Zoom == 2)
-			{
-				Padding = true;
-				PaddingStart = MyGeometry.AbsoluteToLocal(MouseEvent.GetScreenSpacePosition());
-			}
+			bIsPadding = true;
+			PaddingStart = MyGeometry.AbsoluteToLocal(MouseEvent.GetScreenSpacePosition());
 		}
 	}
+	
+	if (!Reply.IsEventHandled())
+	{
+		Reply = SCompoundWidget::OnMouseButtonDown(MyGeometry, MouseEvent);
+	}
 
-	return SCompoundWidget::OnMouseButtonDown( MyGeometry, MouseEvent );
+	return Reply;
 }
 
 
@@ -562,6 +588,8 @@ UTexture2D* SCustomizableObjectLayoutGrid::GetBlockMaskValue() const
 
 FReply SCustomizableObjectLayoutGrid::OnMouseButtonUp( const FGeometry& MyGeometry, const FPointerEvent& MouseEvent )
 {
+	FReply Reply = FReply::Unhandled();
+
 	if (Mode == ELGM_Show)
 	{
 		return SCompoundWidget::OnMouseButtonUp(MyGeometry, MouseEvent);
@@ -569,8 +597,8 @@ FReply SCustomizableObjectLayoutGrid::OnMouseButtonUp( const FGeometry& MyGeomet
 
 	if ( MouseEvent.GetEffectingButton()==EKeys::LeftMouseButton)
 	{
-		Dragging = false;
-		Resizing = false;
+		bIsDragging = false;
+		bIsResizing = false;
 
 		// Left Shif is pressed for multi selection
 		bool bLeftShift = MouseEvent.GetModifierKeys().IsLeftShiftDown();
@@ -585,15 +613,15 @@ FReply SCustomizableObjectLayoutGrid::OnMouseButtonUp( const FGeometry& MyGeomet
 		PossibleSelectedBlocks.Reset();
 
 		// Reset selection if multi selection is not enabled
-		if (Mode == ELGM_Edit && !bLeftShift && !HasDragged)
+		if (Mode == ELGM_Edit && !bLeftShift && !bHasDragged)
 		{
 			// Only one selected block allowed in edit mode.
 			SelectedBlocks.Reset();
 		}
 
-		if (!Selecting)
+		if (!bIsSelecting)
 		{
-			if (!HasDragged)
+			if (!bHasDragged)
 			{
 				//Backward iteration to select the block rendered in front of the rest
 				const TArray<FCustomizableObjectLayoutBlock>& CurrentBlocks = Blocks.Get();
@@ -704,15 +732,24 @@ FReply SCustomizableObjectLayoutGrid::OnMouseButtonUp( const FGeometry& MyGeomet
 			SelectionChangedDelegate.ExecuteIfBound(SelectedBlocks);
 		}
 
-		HasDragged = false;
-		Selecting = false;
+		bHasDragged = false;
+		bIsSelecting = false;
 	}
 	else if (MouseEvent.GetEffectingButton() == EKeys::MiddleMouseButton)
 	{
-		Padding = false;
+		bIsPadding = false;
+	}
+	else if (MouseEvent.GetEffectingButton() == EKeys::RightMouseButton)
+	{
+		Reply = FReply::Handled();
+	}
+
+	if (!Reply.IsEventHandled())
+	{
+		Reply = SCompoundWidget::OnMouseButtonUp(MyGeometry, MouseEvent);
 	}
 
-	return SCompoundWidget::OnMouseButtonUp( MyGeometry, MouseEvent );
+	return Reply;
 }
 
 
@@ -729,22 +766,22 @@ FReply SCustomizableObjectLayoutGrid::OnMouseMove( const FGeometry& MyGeometry,
 	{
 		FVector2D Pos = MyGeometry.AbsoluteToLocal(MouseEvent.GetScreenSpacePosition());
 
-		if (Dragging && SelectedBlocks.Num())
+		if (bIsDragging && SelectedBlocks.Num())
 		{
 			FVector2D CellDelta = (Pos - DragStart) / CellSize;
 
-			int CellDeltaX = CellDelta.X;
-			int CellDeltaY = CellDelta.Y;
+			int32 CellDeltaX = CellDelta.X;
+			int32 CellDeltaY = CellDelta.Y;
 
 			DragStart += FVector2D(CellDeltaX * CellSize, CellDeltaY * CellSize);
 
 			if (CellDeltaX || CellDeltaY)
 			{
-				HasDragged = true;
+				bHasDragged = true;
 
 				const TArray<FCustomizableObjectLayoutBlock>& CurrentBlocks = Blocks.Get();
 
-				if (!Resizing)
+				if (!bIsResizing)
 				{
 					FIntRect TotalBlock;
 					bool bFirstBlock = true;
@@ -824,7 +861,7 @@ FReply SCustomizableObjectLayoutGrid::OnMouseMove( const FGeometry& MyGeometry,
 			}
 		}
 
-		if (!Selecting && !Dragging)
+		if (!bIsSelecting && !bIsDragging)
 		{
 			bool ClickOnBlock = false;
 
@@ -845,16 +882,16 @@ FReply SCustomizableObjectLayoutGrid::OnMouseMove( const FGeometry& MyGeometry,
 
 			if (!ClickOnBlock && (MouseDiference.X > MovementSensitivity || MouseDiference.Y > MovementSensitivity))
 			{
-				HasDragged = true;
-				Selecting = true;
+				bHasDragged = true;
+				bIsSelecting = true;
 			}
 		}
 	}
 	
-	if (!Dragging && !Resizing && SelectedBlocks.Num()==1)
+	if (!bIsDragging && !bIsResizing && SelectedBlocks.Num()==1)
 	{
 		const TArray<FCustomizableObjectLayoutBlock>& CurrentBlocks = Blocks.Get();
-		for (int i = CurrentBlocks.Num() - 1; i > -1; --i)
+		for (int32 i = CurrentBlocks.Num() - 1; i > -1; --i)
 		{
 			// Check for new created blocks
 			if (BlockRects.Contains(CurrentBlocks[i].Id) && SelectedBlocks.Contains(CurrentBlocks[i].Id))
@@ -862,17 +899,17 @@ FReply SCustomizableObjectLayoutGrid::OnMouseMove( const FGeometry& MyGeometry,
 				FVector2D Pos = MyGeometry.AbsoluteToLocal(MouseEvent.GetScreenSpacePosition());
 				if (MouseOnBlock(CurrentBlocks[i].Id, Pos, true))
 				{
-					ResizeCursor = true;
+					bIsResizeCursor = true;
 					break;
 				}
 			}
 
-			ResizeCursor = false;
+			bIsResizeCursor = false;
 		}
 	}
 
 	// In case we lose focus
-	if (Padding)
+	if (bIsPadding)
 	{
 		if (MouseEvent.IsMouseButtonDown(EKeys::MiddleMouseButton))
 		{
@@ -882,19 +919,19 @@ FReply SCustomizableObjectLayoutGrid::OnMouseMove( const FGeometry& MyGeometry,
 		}
 		else
 		{
-			Padding = false;
+			bIsPadding = false;
 		}
 	}
 
 	if (!MouseEvent.IsMouseButtonDown(EKeys::LeftMouseButton))
 	{
-		Selecting = false;
-		Dragging = false;
+		bIsSelecting = false;
+		bIsDragging = false;
 
-		if (Resizing)
+		if (bIsResizing)
 		{
-			ResizeCursor = false;
-			Resizing = false;
+			bIsResizeCursor = false;
+			bIsResizing = false;
 		}
 	}
 
@@ -904,29 +941,41 @@ FReply SCustomizableObjectLayoutGrid::OnMouseMove( const FGeometry& MyGeometry,
 
 FReply SCustomizableObjectLayoutGrid::OnMouseWheel(const FGeometry& MyGeometry, const FPointerEvent& MouseEvent)
 {
+	constexpr int32 MinZoomLevel = -2;
+	constexpr int32 MaxZoomLevel = 3;
+
 	if (Mode == ELGM_Edit)
 	{
+		double OldZoomFactor = FMath::Pow(2.0, double(Zoom - 1));
+		FVector2D UnzoomedPadding = PaddingAmount * (1.0f / OldZoomFactor);
+
 		if (MouseEvent.GetWheelDelta() > 0)
 		{
-			if (Zoom < 2)
+			int32 NewZoomLevel = FMath::Min(Zoom + 1, MaxZoomLevel);
+			if (Zoom != NewZoomLevel)
 			{
-				FVector2D GridCenter = DrawOrigin + (FVector2D((float)GridSize.Get().X, (float)GridSize.Get().Y) / 2.0f) * CellSize;
-				DistanceFromOrigin = CurrentMousePosition - GridCenter;
+				//FVector2D GridCenter = DrawOrigin + (FVector2D((float)GridSize.Get().X, (float)GridSize.Get().Y) / 2.0f) * CellSize;
+				//DistanceFromOrigin = CurrentMousePosition - GridCenter;
 
-				Zoom++;
+				Zoom = NewZoomLevel;
 			}
 		}
 		else
 		{
-			if (Zoom > 1)
+			int32 NewZoomLevel = FMath::Max(Zoom - 1, MinZoomLevel);
+			if (Zoom != NewZoomLevel)
 			{
-				DistanceFromOrigin = FVector2D::Zero();
-				PaddingAmount = FVector2D::Zero();
+				//DistanceFromOrigin = FVector2D::Zero();
+				//PaddingAmount = FVector2D::Zero();
 
-				Zoom--;
+				Zoom = NewZoomLevel;
 			}
 		}
 
+		double NewZoomFactor = FMath::Pow(2.0, double(Zoom - 1));
+		FVector2D RezoomedPadding = UnzoomedPadding * NewZoomFactor;
+		PaddingAmount = RezoomedPadding;
+
 		return FReply::Handled().SetUserFocus(SharedThis(this), EFocusCause::Mouse, true);
 	}
 
@@ -969,7 +1018,7 @@ FReply SCustomizableObjectLayoutGrid::OnKeyDown(const FGeometry& MyGeometry, con
 
 FCursorReply SCustomizableObjectLayoutGrid::OnCursorQuery(const FGeometry& MyGeometry, const FPointerEvent& CursorEvent) const
 {
-	if (ResizeCursor)
+	if (bIsResizeCursor)
 	{
 		return FCursorReply::Cursor(EMouseCursor::ResizeSouthEast);
 	}
@@ -1011,6 +1060,13 @@ void SCustomizableObjectLayoutGrid::DeleteSelectedBlocks()
 }
 
 
+void SCustomizableObjectLayoutGrid::ResetView()
+{
+	Zoom = 1;
+	PaddingAmount = { 0,0 };
+}
+
+
 void SCustomizableObjectLayoutGrid::GenerateNewBlock(FVector2D MousePosition)
 {
 	if (MousePosition.X > 0 && MousePosition.Y > 0 && MousePosition.X < GridSize.Get().X && MousePosition.Y < GridSize.Get().Y)
@@ -1298,119 +1354,201 @@ void FUVCanvasDrawer::SetLayoutMode(ELayoutGridMode Mode)
 
 void FUVCanvasDrawer::Draw_RenderThread(class FRHICommandListImmediate& RHICmdList, const void* InWindowBackBuffer, const FSlateCustomDrawParams& Params)
 {
-	if (Initialized)
+	if (!Initialized)
 	{
-		RenderTarget->SetRenderTargetTexture(*(FTextureRHIRef*)InWindowBackBuffer);
+		return;
+	}
 
-		FCanvas Canvas(RenderTarget, nullptr, FGameTime(), GMaxRHIFeatureLevel);
-		
-		Canvas.SetRenderTargetRect(RenderTarget->GetViewRect());
-		Canvas.SetRenderTargetScissorRect(RenderTarget->GetClippingRect());
+	RenderTarget->SetRenderTargetTexture(*(FTextureRHIRef*)InWindowBackBuffer);
 
-		// Num Lines
-		const uint32 NumEdges = UVLayout.Num() / 2;
-		const uint32 NumGridLines = GridSize.X + GridSize.Y + 2;
-		const uint32 NumUnasignedUVs = UnassignedUVs.Num() * 4;
+	FCanvas Canvas(RenderTarget, nullptr, FGameTime(), GMaxRHIFeatureLevel);
+		
+	Canvas.SetRenderTargetRect(RenderTarget->GetViewRect());
+	Canvas.SetRenderTargetScissorRect(RenderTarget->GetClippingRect());
+
+	// Number of tiles to render in each axis including the unit tile.
+	constexpr int32 NumTiles = 8;
+
+	// Num Lines
+	const uint32 NumAxisLines = 2;
+	const uint32 NumEdges = UVLayout.Num() / 2;
+	const uint32 NumUnitGridLines = GridSize.X + GridSize.Y + 2;
+	const uint32 NumTileLines = NumTiles + NumTiles + 2;
+	const uint32 NumExtendedGridLines = NumTiles * (GridSize.X + GridSize.Y) + 2;
+	const uint32 NumUnasignedUVs = UnassignedUVs.Num() * 4;
+
+	// Num Vertices and Triangles including blocks and unit tile quad.
+	const uint32 RectCount = (LayoutMode == ELayoutGridMode::ELGM_Edit) ? Blocks.Num() * 2 : Blocks.Num();
+	const uint32 NumVertices = RectCount*4 + 4;
+	const uint32 NumTriangles = RectCount*2 + 2;
+
+	const int32 NumLines = NumAxisLines + NumEdges + NumUnitGridLines + NumTileLines + NumExtendedGridLines + NumUnasignedUVs;
+
+	FBatchedElements* BatchedElements = Canvas.GetBatchedElements(FCanvas::ET_Line);
+	BatchedElements->AddReserveLines(NumLines);
+	BatchedElements->AddReserveVertices(NumVertices);
+	BatchedElements->AddReserveTriangles(NumTriangles, GWhiteTexture, ESimpleElementBlendMode::SE_BLEND_Translucent);
+
+	// Color Definitions
+	const FColor ExtendedGridLineColor = FColor(150, 150, 150, 32);
+	const FColor GridLineColor = FColor(150, 150, 150, 64);
+	const FColor TileLineColor = FColor(200, 200, 150, 48);
+	const FColor UVLineColor = FColor(255, 255, 255, 255);
+	const FColor UnassignedUVsColor = FColor::Yellow;
+	const FColor ResizeBlockColor = FColor(255, 96, 96, 255);
+	const FColor UnitTileBlockColor = FColor(96, 96, 96, 128);
+
+	const FHitProxyId HitProxyId = Canvas.GetHitProxyId();
+
+	int32 FullTilesSize = NumTiles * Size.X;
+	FVector2D TilesOrigin = Origin - FVector2D(FullTilesSize / 2, FullTilesSize / 2);
+
+	// Create lines as pairs of points
+	FVector LinePoints[2];
+
+	// Unit Tile
+	{
+		FRect2D TileBlock;
+		TileBlock.Min = FVector2f(Origin);
+		TileBlock.Size = FVector2f(Size);
 
-		// Num Vertices
-		const uint32 NumVertices = LayoutMode == ELayoutGridMode::ELGM_Edit ? Blocks.Num() * 8 : Blocks.Num() * 4;
+		DrawBlock(BatchedElements, HitProxyId, TileBlock, UnitTileBlockColor, nullptr);
+	}
 
-		// Num Triangles
-		const uint32 NumTriangles = LayoutMode == ELayoutGridMode::ELGM_Edit ? Blocks.Num() * 4 : Blocks.Num() * 2;
 
-		FBatchedElements* BatchedElements = Canvas.GetBatchedElements(FCanvas::ET_Line);
-		BatchedElements->AddReserveLines(NumEdges + NumGridLines + NumUnasignedUVs);
-		BatchedElements->AddReserveVertices(NumVertices);
-		BatchedElements->AddReserveTriangles(NumTriangles, GWhiteTexture, ESimpleElementBlendMode::SE_BLEND_Translucent);
+	// Drawing Extended Grid
+	{
+		// Vertical Lines
+		for (int32 LineIndex = 0; LineIndex < NumTiles*GridSize.X + 1; LineIndex++)
+		{
+			LinePoints[0] = FVector(TilesOrigin.X + LineIndex * CellSize, TilesOrigin.Y, 0.0f);
+			LinePoints[1] = FVector(TilesOrigin.X + LineIndex * CellSize, TilesOrigin.Y + FullTilesSize, 0.0f);
 
-		// Color Definitions
-		const FColor GridLineColor = FColor(150, 150, 150, 64);
-		const FColor UVLineColor = FColor(255, 255, 255, 255);
-		const FColor UnassignedUVsColor = FColor::Yellow;
-		const FColor ResizeBlockColor = FColor(255, 96, 96, 255);
+			BatchedElements->AddTranslucentLine(LinePoints[0], LinePoints[1], ExtendedGridLineColor, HitProxyId, 2.0f);
+		}
 
-		const FHitProxyId HitProxyId = Canvas.GetHitProxyId();
+		// Drawing Unit Grid Horizontal Lines
+		for (int32 LineIndex = 0; LineIndex < NumTiles*GridSize.Y + 1; LineIndex++)
+		{
+			LinePoints[0] = FVector(TilesOrigin.X, TilesOrigin.Y + LineIndex * CellSize, 0.0f);
+			LinePoints[1] = FVector(TilesOrigin.X + FullTilesSize, TilesOrigin.Y + LineIndex * CellSize, 0.0f);
 
-		// Create line points
-		FVector LinePoints[2];
+			BatchedElements->AddTranslucentLine(LinePoints[0], LinePoints[1], ExtendedGridLineColor, HitProxyId, 2.0f);
+		}
+	}
 
-		// Drawing Vertical Lines
+	// Drawing Unit Grid
+	{
+		// Vertical Lines
 		for (int32 LineIndex = 0; LineIndex < GridSize.X + 1; LineIndex++)
 		{
-			LinePoints[0] = FVector(FVector2D(Origin.X + LineIndex * CellSize, Origin.Y), 0.0f);
-			LinePoints[1] = FVector(FVector2D(Origin.X + LineIndex * CellSize, Origin.Y + Size.Y), 0.0f);
+			LinePoints[0] = FVector(Origin.X + LineIndex * CellSize, Origin.Y, 0.0f);
+			LinePoints[1] = FVector(Origin.X + LineIndex * CellSize, Origin.Y + Size.Y, 0.0f);
 
 			BatchedElements->AddTranslucentLine(LinePoints[0], LinePoints[1], GridLineColor, HitProxyId, 2.0f);
 		}
 
-		// Drawing Horizontal Lines
+		// Drawing Unit Grid Horizontal Lines
 		for (int32 LineIndex = 0; LineIndex < GridSize.Y + 1; LineIndex++)
 		{
-			LinePoints[0] = FVector(FVector2D(Origin.X, Origin.Y + LineIndex * CellSize), 0.0f);
-			LinePoints[1] = FVector(FVector2D(Origin.X + Size.X, Origin.Y + LineIndex * CellSize), 0.0f);
+			LinePoints[0] = FVector(Origin.X, Origin.Y + LineIndex * CellSize, 0.0f);
+			LinePoints[1] = FVector(Origin.X + Size.X, Origin.Y + LineIndex * CellSize, 0.0f);
 
 			BatchedElements->AddTranslucentLine(LinePoints[0], LinePoints[1], GridLineColor, HitProxyId, 2.0f);
 		}
+	}
 
-		// Drawing UV Lines
-		for (uint32 LineIndex = 0; LineIndex < NumEdges; ++LineIndex)
+	// Drawing Tiles 
+	{
+		// Vertical Lines
+		for (int32 LineIndex = 0; LineIndex < NumTiles + 1; LineIndex++)
 		{
-			LinePoints[0] = FVector(Origin + UVLayout[LineIndex * 2 + 0] * Size, 0.0f);
-			LinePoints[1] = FVector(Origin + UVLayout[LineIndex * 2 + 1] * Size, 0.0f);
+			LinePoints[0] = FVector(TilesOrigin.X + LineIndex * Size.X, TilesOrigin.Y, 0.0f);
+			LinePoints[1] = FVector(TilesOrigin.X + LineIndex * Size.X, TilesOrigin.Y + FullTilesSize, 0.0f);
 
-			BatchedElements->AddLine(LinePoints[0], LinePoints[1], UVLineColor, HitProxyId);
+			BatchedElements->AddTranslucentLine(LinePoints[0], LinePoints[1], TileLineColor, HitProxyId, 2.0f);
 		}
 
-		// Drawing Unassigned UVs
-		const FVector2D CrossSize = Size * 0.01;
-		for (const FVector2d& Vertex : UnassignedUVs)
+		// Horizontal Lines
+		for (int32 LineIndex = 0; LineIndex < NumTiles + 1; LineIndex++)
 		{
-			LinePoints[0] = FVector(Origin + FVector2D(Vertex) * Size + FVector2D(CrossSize), 0.0f);
-			LinePoints[1] = FVector(Origin + FVector2D(Vertex) * Size - FVector2D(CrossSize) * FVector2D(1.0f, -1.0f), 0.0f);
-			BatchedElements->AddLine(LinePoints[0], LinePoints[1], UVLineColor, HitProxyId);
+			LinePoints[0] = FVector(TilesOrigin.X, TilesOrigin.Y + LineIndex * Size.Y, 0.0f);
+			LinePoints[1] = FVector(TilesOrigin.X + FullTilesSize, TilesOrigin.Y + LineIndex * Size.Y, 0.0f);
 
-			LinePoints[0] = FVector(Origin + FVector2D(Vertex) * Size - FVector2D(CrossSize), 0.0f);
-			BatchedElements->AddLine(LinePoints[0], LinePoints[1], UVLineColor, HitProxyId);
+			BatchedElements->AddTranslucentLine(LinePoints[0], LinePoints[1], TileLineColor, HitProxyId, 2.0f);
+		}
+	}
 
-			LinePoints[1] = FVector(Origin + FVector2D(Vertex) * Size + FVector2D(CrossSize) * FVector2D(1.0f, -1.0f), 0.0f);
-			BatchedElements->AddLine(LinePoints[0], LinePoints[1], UVLineColor, HitProxyId);
+	// Axes
+	{
+		LinePoints[0] = FVector(TilesOrigin.X, Origin.Y, 0.0f);
+		LinePoints[1] = FVector(TilesOrigin.X + FullTilesSize, Origin.Y, 0.0f);
+		BatchedElements->AddTranslucentLine(LinePoints[0], LinePoints[1], FColor(255, 150, 150, 200), HitProxyId, 2.0f);
 
-			LinePoints[0] = FVector(Origin + FVector2D(Vertex) * Size + FVector2D(CrossSize), 0.0f);
-			BatchedElements->AddLine(LinePoints[0], LinePoints[1], UVLineColor, HitProxyId);
-		}
+		LinePoints[0] = FVector(Origin.X, TilesOrigin.Y, 0.0f);
+		LinePoints[1] = FVector(Origin.X, TilesOrigin.Y + FullTilesSize, 0.0f);
+		BatchedElements->AddTranslucentLine(LinePoints[0], LinePoints[1], FColor(150, 255, 150, 200), HitProxyId, 2.0f);
+	}
 
-		// Drawing Blocks
-		for (const FCustomizableObjectLayoutBlock& Block : Blocks)
-		{
-			const FColor SelectionBlockColor = SelectedBlocks.Contains(Block.Id) ? SelectedBlockColor : UnselectedBlockColor;
 
-			const FVector2f BlockMin(Block.Min);
-			const FVector2f BlockMax(Block.Max);
+	// Drawing UV Lines
+	for (uint32 LineIndex = 0; LineIndex < NumEdges; ++LineIndex)
+	{
+		LinePoints[0] = FVector(Origin + UVLayout[LineIndex * 2 + 0] * Size, 0.0f);
+		LinePoints[1] = FVector(Origin + UVLayout[LineIndex * 2 + 1] * Size, 0.0f);
 
-			// Selection Block
-			FRect2D SelectionBlock;
-			SelectionBlock.Min = FVector2f(Origin) + BlockMin * CellSize + CellSize * 0.1f;
-			SelectionBlock.Size = (BlockMax - BlockMin) * CellSize - CellSize * 0.2f;
+		BatchedElements->AddLine(LinePoints[0], LinePoints[1], UVLineColor, HitProxyId);
+	}
 
-			DrawBlock(BatchedElements, HitProxyId, SelectionBlock, SelectionBlockColor, Block.Mask.Get());
+	// Drawing Unassigned UVs
+	const FVector2D CrossSize = Size * 0.01;
+	for (const FVector2d& Vertex : UnassignedUVs)
+	{
+		LinePoints[0] = FVector(Origin + FVector2D(Vertex) * Size + FVector2D(CrossSize), 0.0f);
+		LinePoints[1] = FVector(Origin + FVector2D(Vertex) * Size - FVector2D(CrossSize) * FVector2D(1.0f, -1.0f), 0.0f);
+		BatchedElements->AddLine(LinePoints[0], LinePoints[1], UVLineColor, HitProxyId);
 
-			if (LayoutMode == ELayoutGridMode::ELGM_Edit)
-			{
-				// Resize Block
-				FRect2D ResizeBlock;;
-				float HandleRectSize = FMath::Log2(float(GridSize.X)) / 10.0f;
-				ResizeBlock.Size = FVector2f(CellSize) * HandleRectSize;
-				ResizeBlock.Min = SelectionBlock.Min + SelectionBlock.Size - ResizeBlock.Size;
+		LinePoints[0] = FVector(Origin + FVector2D(Vertex) * Size - FVector2D(CrossSize), 0.0f);
+		BatchedElements->AddLine(LinePoints[0], LinePoints[1], UVLineColor, HitProxyId);
 
-				DrawBlock(BatchedElements, HitProxyId, ResizeBlock, ResizeBlockColor);
-			}
-		}
+		LinePoints[1] = FVector(Origin + FVector2D(Vertex) * Size + FVector2D(CrossSize) * FVector2D(1.0f, -1.0f), 0.0f);
+		BatchedElements->AddLine(LinePoints[0], LinePoints[1], UVLineColor, HitProxyId);
 
-		Canvas.Flush_RenderThread(RHICmdList, true);
+		LinePoints[0] = FVector(Origin + FVector2D(Vertex) * Size + FVector2D(CrossSize), 0.0f);
+		BatchedElements->AddLine(LinePoints[0], LinePoints[1], UVLineColor, HitProxyId);
+	}
 
-		RenderTarget->ClearRenderTargetTexture();
-		RHICmdList.SetScissorRect(false, 0, 0, 0, 0);
+	// Drawing Blocks
+	for (const FCustomizableObjectLayoutBlock& Block : Blocks)
+	{
+		const FColor SelectionBlockColor = SelectedBlocks.Contains(Block.Id) ? SelectedBlockColor : UnselectedBlockColor;
+
+		const FVector2f BlockMin(Block.Min);
+		const FVector2f BlockMax(Block.Max);
+
+		// Selection Block
+		FRect2D SelectionBlock;
+		SelectionBlock.Min = FVector2f(Origin) + BlockMin * CellSize + CellSize * 0.1f;
+		SelectionBlock.Size = (BlockMax - BlockMin) * CellSize - CellSize * 0.2f;
+
+		DrawBlock(BatchedElements, HitProxyId, SelectionBlock, SelectionBlockColor, Block.Mask.Get());
+
+		if (LayoutMode == ELayoutGridMode::ELGM_Edit)
+		{
+			// Resize Block
+			FRect2D ResizeBlock;;
+			float HandleRectSize = FMath::Log2(float(GridSize.X)) / 10.0f;
+			ResizeBlock.Size = FVector2f(CellSize) * HandleRectSize;
+			ResizeBlock.Min = SelectionBlock.Min + SelectionBlock.Size - ResizeBlock.Size;
+
+			DrawBlock(BatchedElements, HitProxyId, ResizeBlock, ResizeBlockColor);
+		}
 	}
+
+	Canvas.Flush_RenderThread(RHICmdList, true);
+
+	RenderTarget->ClearRenderTargetTexture();
+	RHICmdList.SetScissorRect(false, 0, 0, 0, 0);
 }
 
 
diff --git a/Engine/Plugins/Experimental/Mutable/Source/CustomizableObjectEditor/Private/MuCOE/SCustomizableObjectLayoutGrid.h b/Engine/Plugins/Experimental/Mutable/Source/CustomizableObjectEditor/Private/MuCOE/SCustomizableObjectLayoutGrid.h
--- a/Engine/Plugins/Experimental/Mutable/Source/CustomizableObjectEditor/Private/MuCOE/SCustomizableObjectLayoutGrid.h
+++ b/Engine/Plugins/Experimental/Mutable/Source/CustomizableObjectEditor/Private/MuCOE/SCustomizableObjectLayoutGrid.h
@@ -5,6 +5,8 @@
 #include "MuCOE/CustomizableObjectLayout.h"
 #include "Widgets/SCompoundWidget.h"
 
+#define UE_MUTABLE_UI_DRAWBUFFERS 2
+
 class FPaintArgs;
 class FSlateRect;
 class FSlateWindowElementList;
@@ -106,6 +108,9 @@ public:
 	/** Calls the delegate to delete the selected blocks */
 	void DeleteSelectedBlocks();
 
+	/** Reset the view zoom and pan to show the unit UV space. */
+	void ResetView();
+
 	/** Generates a new block at mouse position */
 	void GenerateNewBlock(FVector2D MousePosition);
 
@@ -188,22 +193,22 @@ private:
 
 	/** Booleans needed for the Block Management */
 	/** Indicates when we have dragged the mouse after click */
-	bool HasDragged = false;
+	bool bHasDragged = false;
 
 	/** Indicates when we are dragging the mouse */
-	bool Dragging = false;
+	bool bIsDragging = false;
 	
 	/** Indicates when we are resizing a block */
-	bool Resizing = false;
+	bool bIsResizing = false;
 	
 	/** Indicates when we have to change the mouse cursor */
-	bool ResizeCursor = false;
+	bool bIsResizeCursor = false;
 	
 	/** Indicates when we are making a selection */
-	bool Selecting = false;
+	bool bIsSelecting = false;
 	
 	/** Indicates when we are padding */
-	bool Padding = false;
+	bool bIsPadding = false;
 
 	/** Position where the drag started */
 	FVector2D DragStart;
@@ -217,9 +222,6 @@ private:
 	/** Position where the padding started */
 	FVector2D PaddingStart;
 
-	/** Distance from the origin in the padding movement */
-	FVector2D DistanceFromOrigin = FVector2D::Zero();
-
 	/** Level of zoom */
 	int32 Zoom = 1;
 
@@ -232,6 +234,9 @@ private:
 	/** Current mouse position */
 	FVector2D CurrentMousePosition;
 
-	/** Custom Slate drawing element. Used to improve the UVs drawing performance. */
-	TSharedPtr<class FUVCanvasDrawer> UVCanvasDrawer;
+	/** Custom Slate drawing element. Used to improve the UVs drawing performance. 
+	* * This is multi-buffered because it is read and written simultaneously from the render and the game threads.
+	*/
+	int32 CurrentDrawBuffer = 0;
+	TSharedPtr<class FUVCanvasDrawer> UVCanvasDrawers[UE_MUTABLE_UI_DRAWBUFFERS];
 };
diff --git a/Engine/Plugins/Experimental/Mutable/Source/MutableRuntime/Private/MuR/OpLayoutPack.cpp b/Engine/Plugins/Experimental/Mutable/Source/MutableRuntime/Private/MuR/OpLayoutPack.cpp
--- a/Engine/Plugins/Experimental/Mutable/Source/MutableRuntime/Private/MuR/OpLayoutPack.cpp
+++ b/Engine/Plugins/Experimental/Mutable/Source/MutableRuntime/Private/MuR/OpLayoutPack.cpp
@@ -109,7 +109,7 @@ namespace mu
 
 	inline void ReductionOperation(uint16& BlockSize, EReductionMethod ReductionMethod, int32 bReduceByTwo)
 	{
-		if (ReductionMethod == EReductionMethod::UNITARY_REDUCTION)
+		if (ReductionMethod == EReductionMethod::Unitary)
 		{
 			int32 Reduction = (bReduceByTwo && BlockSize > 2) ? 2 : 1;
 			BlockSize -= Reduction;
@@ -351,7 +351,7 @@ namespace mu
 			scratch.positions[scratch.sorted[best].Index] = FIntVector2(bestX, bestLevel);
 			*maxY = FMath::Max(*maxY, uint16(bestLevel + scratch.sorted[best].size[1]) );
 
-			if (packStrategy == EPackStrategy::FIXED_LAYOUT && *maxY > layoutSizeY)
+			if (packStrategy == EPackStrategy::Fixed && *maxY > layoutSizeY)
 			{
 				fits = false;
 				break;
@@ -451,7 +451,7 @@ namespace mu
 
         // Grow until the area is big enough to fit all blocks. We always grow X first, because
         // in case we cannot pack everything, we will grow Y with the current horizon algorithm.
-		if (LayoutStrategy == EPackStrategy::RESIZABLE_LAYOUT)
+		if (LayoutStrategy == EPackStrategy::Resizeable)
 		{
 			maxX = FGenericPlatformMath::RoundUpToPowerOfTwo(maxX);
 			maxY = FGenericPlatformMath::RoundUpToPowerOfTwo(maxY);
@@ -592,7 +592,7 @@ namespace mu
 		// Sort blocks by height, area
 		scratch.sorted.Sort(CompareBlocks);
 
-		if(LayoutStrategy == EPackStrategy::FIXED_LAYOUT)
+		if(LayoutStrategy == EPackStrategy::Fixed)
 		{
 			if (usePriority)
 			{
@@ -620,7 +620,7 @@ namespace mu
 			// Try to pack everything
 			fits = SetPositions(bestY, layoutSizeY, &maxX, &maxY, scratch, LayoutStrategy);
 
-			if (!fits && LayoutStrategy == EPackStrategy::FIXED_LAYOUT)
+			if (!fits && LayoutStrategy == EPackStrategy::Fixed)
 			{
 				// Sort by priority before shrink
 				if (usePriority)
diff --git a/Engine/Plugins/Experimental/Mutable/Source/MutableRuntime/Public/MuR/Layout.h b/Engine/Plugins/Experimental/Mutable/Source/MutableRuntime/Public/MuR/Layout.h
--- a/Engine/Plugins/Experimental/Mutable/Source/MutableRuntime/Public/MuR/Layout.h
+++ b/Engine/Plugins/Experimental/Mutable/Source/MutableRuntime/Public/MuR/Layout.h
@@ -16,16 +16,16 @@ namespace mu
 	//! Types of layout packing strategies 
 	enum class EPackStrategy : uint32
 	{
-		RESIZABLE_LAYOUT,
-		FIXED_LAYOUT,
-		OVERLAY_LAYOUT
+		Resizeable,
+		Fixed,
+		Overlay
 	};
 
 	//! Types of layout reduction methods 
 	enum class EReductionMethod : uint32
 	{
-		HALVE_REDUCTION,	// Divide axis by 2
-		UNITARY_REDUCTION	// Reduces 1 block the axis 
+		Halve,	// Divide axis by 2
+		Unitary	// Reduces 1 block the axis 
 	};
 
 	/** */
@@ -103,9 +103,9 @@ namespace mu
 		TArray<FLayoutBlock> Blocks;
 
 		//! Packing strategy
-		EPackStrategy Strategy = EPackStrategy::RESIZABLE_LAYOUT;
+		EPackStrategy Strategy = EPackStrategy::Resizeable;
 
-		EReductionMethod ReductionMethod = EReductionMethod::HALVE_REDUCTION;
+		EReductionMethod ReductionMethod = EReductionMethod::Halve;
 
 	public:
 
diff --git a/Engine/Plugins/Experimental/Mutable/Source/MutableTools/Private/MuT/CodeGenerator.cpp b/Engine/Plugins/Experimental/Mutable/Source/MutableTools/Private/MuT/CodeGenerator.cpp
--- a/Engine/Plugins/Experimental/Mutable/Source/MutableTools/Private/MuT/CodeGenerator.cpp
+++ b/Engine/Plugins/Experimental/Mutable/Source/MutableTools/Private/MuT/CodeGenerator.cpp
@@ -1233,7 +1233,7 @@ namespace mu
 						}
 					}
 
-					bool bIsOverlayLayout = pLayout->GetLayoutPackingStrategy() == mu::EPackStrategy::OVERLAY_LAYOUT;
+					bool bIsOverlayLayout = pLayout->GetLayoutPackingStrategy() == mu::EPackStrategy::Overlay;
 
                     if (!bIsOverlayLayout && layoutOp)
                     {
diff --git a/Engine/Plugins/Experimental/Mutable/Source/MutableTools/Private/MuT/CodeGenerator_Mesh.cpp b/Engine/Plugins/Experimental/Mutable/Source/MutableTools/Private/MuT/CodeGenerator_Mesh.cpp
--- a/Engine/Plugins/Experimental/Mutable/Source/MutableTools/Private/MuT/CodeGenerator_Mesh.cpp
+++ b/Engine/Plugins/Experimental/Mutable/Source/MutableTools/Private/MuT/CodeGenerator_Mesh.cpp
@@ -301,22 +301,26 @@ namespace mu
 		TArray<int32> OverlappingBlocks;
 		for (int32 BlockIndex = 0; BlockIndex < NumBlocks; ++BlockIndex)
 		{
-			// Get the block rect
-			FImageSize Min = Layout->Blocks[BlockIndex].Min;
-			FImageSize Size = Layout->Blocks[BlockIndex].Size;
-
-			box<FVector2f>& BlockRect = BlockRects[BlockIndex];
-			BlockRect.min[0] = ((float)Min.X) / (float)Grid.X;
-			BlockRect.min[1] = ((float)Min.Y) / (float)Grid.Y;
-			BlockRect.size[0] = ((float)Size.X) / (float)Grid.X;
-			BlockRect.size[1] = ((float)Size.Y) / (float)Grid.Y;
-
 			bool bBlockHasMask = GeneratedLayout.Source->Blocks[BlockIndex].Mask.get() != nullptr;
 
 			// Fill the block index per cell array
 			// Ignore the block in this stage if it has a mask, because blocks with masks will very likely overlap other blocks
 			if (!bBlockHasMask)
 			{
+				// Get the block rect
+				FImageSize Min = Layout->Blocks[BlockIndex].Min;
+				FImageSize Size = Layout->Blocks[BlockIndex].Size;
+
+				// Clamp block to layout grid
+				Size.X = FMath::Min(Min.X + Size.X, Grid.X);
+				Size.Y = FMath::Min(Min.Y + Size.Y, Grid.Y);
+
+				box<FVector2f>& BlockRect = BlockRects[BlockIndex];
+				BlockRect.min[0] = ((float)Min.X) / (float)Grid.X;
+				BlockRect.min[1] = ((float)Min.Y) / (float)Grid.Y;
+				BlockRect.size[0] = ((float)Size.X) / (float)Grid.X;
+				BlockRect.size[1] = ((float)Size.Y) / (float)Grid.Y;
+
 				for (uint16 Y = Min.Y; Y < Min.Y + Size.Y; ++Y)
 				{
 					const uint16 PositionY = Y * Grid.X;
@@ -367,7 +371,7 @@ namespace mu
 			TexCoords.SetNumUninitialized(NumVertices);
 
 			bool bNonNormalizedUVs = false;
-			const bool bIsOverlayLayout = GeneratedLayout.Layout->GetLayoutPackingStrategy() == mu::EPackStrategy::OVERLAY_LAYOUT;
+			const bool bIsOverlayLayout = GeneratedLayout.Layout->GetLayoutPackingStrategy() == mu::EPackStrategy::Overlay;
 
 			const uint8* pVertices = TexCoordData;
 			for (int32 VertexIndex = 0; VertexIndex < NumVertices; ++VertexIndex)
diff --git a/Engine/Plugins/Experimental/Mutable/Source/MutableTools/Private/MuT/NodeLayout.cpp b/Engine/Plugins/Experimental/Mutable/Source/MutableTools/Private/MuT/NodeLayout.cpp
--- a/Engine/Plugins/Experimental/Mutable/Source/MutableTools/Private/MuT/NodeLayout.cpp
+++ b/Engine/Plugins/Experimental/Mutable/Source/MutableTools/Private/MuT/NodeLayout.cpp
@@ -71,7 +71,7 @@ namespace mu
 			LayoutNode = new NodeLayout;
 			LayoutNode->Size = { uint16(gridSizeX), uint16(gridSizeY) };
 			LayoutNode->MaxSize = { uint16(gridSizeX), uint16(gridSizeY) };
-			LayoutNode->Strategy = EPackStrategy::RESIZABLE_LAYOUT;
+			LayoutNode->Strategy = EPackStrategy::Resizeable;
 			
 			TArray<box<FIntVector2>> blocks;
 			
diff --git a/Engine/Plugins/Experimental/Mutable/Source/MutableTools/Public/MuT/NodeLayout.h b/Engine/Plugins/Experimental/Mutable/Source/MutableTools/Public/MuT/NodeLayout.h
--- a/Engine/Plugins/Experimental/Mutable/Source/MutableTools/Public/MuT/NodeLayout.h
+++ b/Engine/Plugins/Experimental/Mutable/Source/MutableTools/Public/MuT/NodeLayout.h
@@ -51,8 +51,8 @@ namespace mu
 		TArray<FSourceLayoutBlock> Blocks;
 
 		//! Packing strategy
-		EPackStrategy Strategy = EPackStrategy::RESIZABLE_LAYOUT;
-		EReductionMethod ReductionMethod = EReductionMethod::HALVE_REDUCTION;
+		EPackStrategy Strategy = EPackStrategy::Resizeable;
+		EReductionMethod ReductionMethod = EReductionMethod::Halve;
 
 		/** When compiling, ignore generated warnings from this LOD on.
 		* -1 means all warnings are generated.
diff --git a/Engine/Plugins/Experimental/Mutable/Source/MutableTools/Public/MuT/NodeSurfaceNew.h b/Engine/Plugins/Experimental/Mutable/Source/MutableTools/Public/MuT/NodeSurfaceNew.h
--- a/Engine/Plugins/Experimental/Mutable/Source/MutableTools/Public/MuT/NodeSurfaceNew.h
+++ b/Engine/Plugins/Experimental/Mutable/Source/MutableTools/Public/MuT/NodeSurfaceNew.h
@@ -15,8 +15,8 @@
 namespace mu
 {
 
-    //! This node makes a new Surface from several meshes and images.
-    class MUTABLETOOLS_API NodeSurfaceNew : public NodeSurface
+	/** This node makes a new Surface (Mesh Material Section) from a mesh and several material parameters like images, vectors and scalars. */
+	class MUTABLETOOLS_API NodeSurfaceNew : public NodeSurface
 	{
 	public:
 
@@ -37,9 +37,9 @@ namespace mu
 			FString Name;
 			FString MaterialName;
 			FString MaterialParameterName;
-			NodeImagePtr Image;
+			Ptr<NodeImage> Image;
 
-			// It could be negative, to indicate no layout.
+			/** Index of the layout transform to apply to this image. It could be negative, to indicate no layout transform. */
 			int8 LayoutIndex = 0;
 		};
 
