[Mutable] Added changes required for HairStrandsMutable plugin.

#rnx

diff --git a/Engine/Plugins/Experimental/Mutable/Source/CustomizableObject/Internal/MuCO/CustomizableObjectInstancePrivate.h b/Engine/Plugins/Experimental/Mutable/Source/CustomizableObject/Internal/MuCO/CustomizableObjectInstancePrivate.h
--- a/Engine/Plugins/Experimental/Mutable/Source/CustomizableObject/Internal/MuCO/CustomizableObjectInstancePrivate.h
+++ b/Engine/Plugins/Experimental/Mutable/Source/CustomizableObject/Internal/MuCO/CustomizableObjectInstancePrivate.h
@@ -33,6 +33,7 @@ struct FGeneratedMaterial;
 struct FGeneratedTexture;
 class UPhysicsAsset;
 class USkeleton;
+class UCustomizableObjectExtension;
 
 
 /** \param OnlyLOD: If not 0, extract and convert only one single LOD from the source image.
@@ -163,6 +164,19 @@ struct FAnimBpGeneratedPhysicsAssets
 };
 
 
+USTRUCT()
+struct FExtensionInstanceData
+{
+	GENERATED_BODY()
+ 
+	UPROPERTY()
+	TWeakObjectPtr<const UCustomizableObjectExtension> Extension;
+ 
+	UPROPERTY()
+	FInstancedStruct Data;
+};
+
+
 #if WITH_EDITOR
 DECLARE_MULTICAST_DELEGATE_OneParam(FObjectInstanceTransactedDelegate, const FTransactionObjectEvent&);
 #endif //WIT_EDITOR
@@ -354,6 +368,9 @@ public:
 	UPROPERTY(Transient, Category = Animation, editfixedsize, VisibleAnywhere)
 	TMap<TSubclassOf<UAnimInstance>, FAnimBpGeneratedPhysicsAssets> AnimBpPhysicsAssets;
 
+	UPROPERTY(Transient)
+	TArray<FExtensionInstanceData> ExtensionInstanceData;
+	
 	// The pass-through assets that will be loaded during an update
 	TArray<TSoftObjectPtr<const UTexture>> PassThroughTexturesToLoad;
 	TArray<TSoftObjectPtr<const UStreamableRenderAsset>> PassThroughMeshesToLoad;
diff --git a/Engine/Plugins/Experimental/Mutable/Source/CustomizableObject/Private/MuCO/CustomizableObjectInstance.cpp b/Engine/Plugins/Experimental/Mutable/Source/CustomizableObject/Private/MuCO/CustomizableObjectInstance.cpp
--- a/Engine/Plugins/Experimental/Mutable/Source/CustomizableObject/Private/MuCO/CustomizableObjectInstance.cpp
+++ b/Engine/Plugins/Experimental/Mutable/Source/CustomizableObject/Private/MuCO/CustomizableObjectInstance.cpp
@@ -18,6 +18,7 @@
 #include "MaterialDomain.h"
 #include "Materials/Material.h"
 #include "Materials/MaterialInstanceDynamic.h"
+#include "Misc/Optional.h"
 #include "Modules/ModuleManager.h"
 #include "Tasks/Task.h"
 
@@ -234,6 +235,8 @@ void UCustomizableInstancePrivate::InitCustomizableObjectData(const UCustomizabl
 	FCustomizableInstanceComponentData TemplateComponentData;
 	TemplateComponentData.LastMeshIdPerLOD.Init(MAX_uint64, MAX_MESH_LOD_COUNT);
 	ComponentsData.Init(TemplateComponentData, InCustomizableObject->GetComponentCount());
+
+	ExtensionInstanceData.Empty();
 }
 
 
@@ -2018,13 +2021,7 @@ bool UCustomizableInstancePrivate::UpdateSkeletalMesh_PostBeginUpdate0(UCustomiz
 
 		return false;
 	}
-
-	// None of the current meshes requires a mesh update. Continue to BuildMaterials
-	if (!bUpdateMeshes)
-	{
-		return true;
-	}
-
+	
 	TextureReuseCache.Empty(); // Sections may have changed, so invalidate the texture reuse cache because it's indexed by section
 
 	TArray<TObjectPtr<USkeletalMesh>> OldSkeletalMeshes = SkeletalMeshes;
@@ -2033,6 +2030,126 @@ bool UCustomizableInstancePrivate::UpdateSkeletalMesh_PostBeginUpdate0(UCustomiz
 
 	const FModelResources& ModelResources = CustomizableObject->GetPrivate()->GetModelResources();
 
+	// Collate the Extension Data on the instance into groups based on the extension that produced
+	// it, so that we only need to call extension functions such as OnSkeletalMeshCreated once for
+	// each extension.
+	TMap<const UCustomizableObjectExtension*, TArray<FInputPinDataContainer>> ExtensionToExtensionData;
+	{
+		const TArrayView<const UCustomizableObjectExtension* const> AllExtensions = ICustomizableObjectModule::Get().GetRegisteredExtensions();
+
+		// Pre-populate ExtensionToExtensionData with empty entries for all extensions.
+		//
+		// This ensures that extension functions such as OnSkeletalMeshCreated are called for each
+		// extension, even if they didn't produce any extension data.
+		Public->GetPrivate()->ExtensionInstanceData.Empty(AllExtensions.Num());
+		for (const UCustomizableObjectExtension* Extension : AllExtensions)
+		{
+			ExtensionToExtensionData.Add(Extension);
+		}
+
+		const TArrayView<const FRegisteredObjectNodeInputPin> ExtensionPins = ICustomizableObjectModule::Get().GetAdditionalObjectNodePins();
+
+		for (FInstanceUpdateData::FNamedExtensionData& ExtensionOutput : OperationData->InstanceUpdateData.ExtendedInputPins)
+		{
+			const FRegisteredObjectNodeInputPin* FoundPin =
+				Algo::FindBy(ExtensionPins, ExtensionOutput.Name, &FRegisteredObjectNodeInputPin::GlobalPinName);
+
+			if (!FoundPin)
+			{
+				// Failed to find the corresponding pin for this output
+				// 
+				// This may indicate that a plugin has been removed or renamed since the CO was compiled
+				UE_LOG(LogMutable, Error, TEXT("Failed to find Object node input pin with name %s"), *ExtensionOutput.Name.ToString());
+				continue;
+			}
+
+			const UCustomizableObjectExtension* Extension = FoundPin->Extension.Get();
+			if (!Extension)
+			{
+				// Extension is not loaded or not found
+				UE_LOG(LogMutable, Error, TEXT("Extension for Object node input pin %s is no longer valid"), *ExtensionOutput.Name.ToString());
+				continue;
+			}
+
+			if (ExtensionOutput.Data->Origin == mu::ExtensionData::EOrigin::Invalid)
+			{
+				// Null data was produced
+				//
+				// This can happen if a node produces an ExtensionData but doesn't initialize it
+				UE_LOG(LogMutable, Error, TEXT("Invalid data sent to Object node input pin %s"), *ExtensionOutput.Name.ToString());
+				continue;
+			}
+
+			// All registered extensions were added to the map above, so if the extension is still
+			// registered it should be found.
+			TArray<FInputPinDataContainer>* ContainerArray = ExtensionToExtensionData.Find(Extension);
+			if (!ContainerArray)
+			{
+				UE_LOG(LogMutable, Error, TEXT("Object node input pin %s received data for unregistered extension %s"),
+					*ExtensionOutput.Name.ToString(), *Extension->GetPathName());
+				continue;
+			}
+
+			const FCustomizableObjectResourceData* ReferencedExtensionData = nullptr;
+			switch (ExtensionOutput.Data->Origin)
+			{
+				case mu::ExtensionData::EOrigin::ConstantAlwaysLoaded:
+				{
+					check(CustomizableObject->GetPrivate()->GetAlwaysLoadedExtensionData().IsValidIndex(ExtensionOutput.Data->Index));
+					ReferencedExtensionData = &CustomizableObject->GetPrivate()->GetAlwaysLoadedExtensionData()[ExtensionOutput.Data->Index];
+				}
+				break;
+
+				case mu::ExtensionData::EOrigin::ConstantStreamed:
+				{
+					check(CustomizableObject->GetPrivate()->GetStreamedExtensionData().IsValidIndex(ExtensionOutput.Data->Index));
+
+					const FCustomizableObjectStreamedResourceData& StreamedData =
+						CustomizableObject->GetPrivate()->GetStreamedExtensionData()[ExtensionOutput.Data->Index];
+
+					if (!StreamedData.IsLoaded())
+					{
+						// The data should have been loaded as part of executing the CO program.
+						//
+						// This could indicate a bug in the streaming logic.
+						UE_LOG(LogMutable, Error, TEXT("Customizable Object produced a streamed extension data that is not loaded: %s"),
+							*StreamedData.GetPath().ToString());
+
+						continue;
+					}
+
+					ReferencedExtensionData = &StreamedData.GetLoadedData();
+				}
+				break;
+
+				default:
+					unimplemented();
+			}
+
+			check(ReferencedExtensionData);
+
+			ContainerArray->Emplace(FoundPin->InputPin, ReferencedExtensionData->Data);
+		}
+	}
+
+	// Give each extension the chance to generate Extension Instance Data
+	for (const TPair<const UCustomizableObjectExtension*, TArray<FInputPinDataContainer>>& Pair : ExtensionToExtensionData)
+	{
+		FInstancedStruct NewExtensionInstanceData = Pair.Key->GenerateExtensionInstanceData(Pair.Value);
+		if (NewExtensionInstanceData.IsValid())
+		{
+			FExtensionInstanceData& NewData = Public->GetPrivate()->ExtensionInstanceData.AddDefaulted_GetRef();
+			NewData.Extension = Pair.Key;
+			NewData.Data = MoveTemp(NewExtensionInstanceData);
+		}
+	}
+
+	// None of the current meshes requires a mesh update. Continue to BuildMaterials
+	if (!bUpdateMeshes)
+	{
+		return true;
+	}
+	
 	// Initialize the maximum number of SkeletalMeshes we could possibly have. 
 	SkeletalMeshes.Init(nullptr, CustomizableObject->GetComponentCount());
 
@@ -2209,93 +2326,9 @@ bool UCustomizableInstancePrivate::UpdateSkeletalMesh_PostBeginUpdate0(UCustomiz
 			break;
 		}
 
-		// Collate the extension data on the instance into groups based on the extension that
-		// produced it, so that we only need to call OnSkeletalMeshCreated once for each extension.
+		for (const TPair<const UCustomizableObjectExtension*, TArray<FInputPinDataContainer>>& Pair : ExtensionToExtensionData)
 		{
-			TMap<const UCustomizableObjectExtension*, TArray<FInputPinDataContainer>> ExtensionToExtensionData;
-
-			const TArrayView<const FRegisteredObjectNodeInputPin> ExtensionPins = ICustomizableObjectModule::Get().GetAdditionalObjectNodePins();
-
-			for (FInstanceUpdateData::FNamedExtensionData& ExtensionOutput : OperationData->InstanceUpdateData.ExtendedInputPins)
-			{
-				const FRegisteredObjectNodeInputPin* FoundPin =
-					Algo::FindBy(ExtensionPins, ExtensionOutput.Name, &FRegisteredObjectNodeInputPin::GlobalPinName);
-
-				if (!FoundPin)
-				{
-					// Failed to find the corresponding pin for this output
-					// 
-					// This may indicate that a plugin has been removed or renamed since the CO was compiled
-					UE_LOG(LogMutable, Error, TEXT("Failed to find Object node input pin with name %s"), *ExtensionOutput.Name.ToString());
-					continue;
-				}
-
-				const UCustomizableObjectExtension* Extension = FoundPin->Extension.Get();
-				if (!Extension)
-				{
-					// Extension is not loaded or not found
-					UE_LOG(LogMutable, Error, TEXT("Extension for Object node input pin %s is no longer valid"), *ExtensionOutput.Name.ToString());
-					continue;
-				}
-
-				if (ExtensionOutput.Data->Origin == mu::ExtensionData::EOrigin::Invalid)
-				{
-					// Null data was produced
-					//
-					// This can happen if a node produces an ExtensionData but doesn't initialize it
-					UE_LOG(LogMutable, Error, TEXT("Invalid data sent to Object node input pin %s"), *ExtensionOutput.Name.ToString());
-					continue;
-				}
-
-				TArray<FInputPinDataContainer>& ContainerArray = ExtensionToExtensionData.FindOrAdd(Extension);
-
-				const FCustomizableObjectResourceData* ReferencedExtensionData = nullptr;
-				switch (ExtensionOutput.Data->Origin)
-				{
-					case mu::ExtensionData::EOrigin::ConstantAlwaysLoaded:
-					{
-						check(CustomizableObject->GetPrivate()->GetAlwaysLoadedExtensionData().IsValidIndex(ExtensionOutput.Data->Index));
-						ReferencedExtensionData = &CustomizableObject->GetPrivate()->GetAlwaysLoadedExtensionData()[ExtensionOutput.Data->Index];
-					}
-					break;
-
-					case mu::ExtensionData::EOrigin::ConstantStreamed:
-					{
-						check(CustomizableObject->GetPrivate()->GetStreamedExtensionData().IsValidIndex(ExtensionOutput.Data->Index));
-						
-						const FCustomizableObjectStreamedResourceData& StreamedData =
-							CustomizableObject->GetPrivate()->GetStreamedExtensionData()[ExtensionOutput.Data->Index];
-
-						if (!StreamedData.IsLoaded())
-						{
-							// The data should have been loaded as part of executing the CO program.
-							//
-							// This could indicate a bug in the streaming logic.
-							UE_LOG(LogMutable, Error, TEXT("Customizable Object produced a streamed extension data that is not loaded: %s"),
-								*StreamedData.GetPath().ToString());
-
-							continue;
-						}
-
-						ReferencedExtensionData = &StreamedData.GetLoadedData();
-					}
-					break;
-
-					default:
-						unimplemented();
-				}
-
-				check(ReferencedExtensionData);
-		
-				ContainerArray.Emplace(FoundPin->InputPin, ReferencedExtensionData->Data);
-			}
-
-			// Now that we have an array of extension data for each extension, go through the extensions and
-			// give each one its data.
-			for (const TPair<const UCustomizableObjectExtension*, TArray<FInputPinDataContainer>>& Pair : ExtensionToExtensionData)
-			{
-				Pair.Key->OnSkeletalMeshCreated(Pair.Value, ObjectComponentIndex, SkeletalMesh);
-			}
+			Pair.Key->OnSkeletalMeshCreated(Pair.Value, ObjectComponentIndex, SkeletalMesh);
 		}
 	}
 
@@ -3503,7 +3536,7 @@ void UCustomizableInstancePrivate::InitSkeletalMeshData(const TSharedRef<FUpdate
 	SkeletalMesh->SetImportedBounds(RefSkeletalMeshData.Bounds);
 	SkeletalMesh->SetPostProcessAnimBlueprint(RefSkeletalMeshData.PostProcessAnimInst.Get());
 	SkeletalMesh->SetShadowPhysicsAsset(RefSkeletalMeshData.ShadowPhysicsAsset.Get());
-
+	
 	// Set Min LOD
 	SkeletalMesh->SetMinLod(FMath::Max(CustomizableObject.LODSettings.MinLOD.GetDefault(), (int32)FirstLODAvailable));
 	SkeletalMesh->SetQualityLevelMinLod(CustomizableObject.LODSettings.MinQualityLevelLOD);
@@ -7255,6 +7288,19 @@ void UCustomizableObjectInstance::ForEachAnimInstance(int32 ComponentIndex, FEac
 }
 
 
+FInstancedStruct UCustomizableObjectInstance::GetExtensionInstanceData(const UCustomizableObjectExtension* Extension) const
+{
+	const FExtensionInstanceData* FoundData = Algo::FindBy(PrivateData->ExtensionInstanceData, Extension, &FExtensionInstanceData::Extension);
+	if (FoundData)
+	{
+		return FoundData->Data;
+	}
+
+	// Data not found. Return an empty instance.
+	return FInstancedStruct();
+}
+
+
 TSet<UAssetUserData*> UCustomizableObjectInstance::GetMergedAssetUserData(int32 ComponentIndex) const
 {
 	UCustomizableInstancePrivate* PrivateInstanceData = GetPrivate();
diff --git a/Engine/Plugins/Experimental/Mutable/Source/CustomizableObject/Private/MuCO/CustomizableObjectInstanceUsage.cpp b/Engine/Plugins/Experimental/Mutable/Source/CustomizableObject/Private/MuCO/CustomizableObjectInstanceUsage.cpp
--- a/Engine/Plugins/Experimental/Mutable/Source/CustomizableObject/Private/MuCO/CustomizableObjectInstanceUsage.cpp
+++ b/Engine/Plugins/Experimental/Mutable/Source/CustomizableObject/Private/MuCO/CustomizableObjectInstanceUsage.cpp
@@ -14,6 +14,7 @@
 #include "GameFramework/Pawn.h"
 #include "Materials/MaterialInstanceDynamic.h"
 #include "MuCO/CustomizableObject.h"
+#include "MuCO/ICustomizableObjectModule.h"
 #include "Rendering/SkeletalMeshRenderData.h"
 #include "UObject/ObjectSaveContext.h"
 #include "Stats/Stats.h"
@@ -29,6 +30,11 @@ ETickableTickType UCustomizableObjectInstanceUsage::GetTickableTickType() const
 
 void UCustomizableObjectInstanceUsage::Callbacks() const
 {
+	for (const UCustomizableObjectExtension* Extension : ICustomizableObjectModule::Get().GetRegisteredExtensions())
+	{
+		Extension->OnCustomizableObjectInstanceUsageUpdated(this);
+	}
+	
 	if (CustomizableSkeletalComponent)
 	{
 		CustomizableSkeletalComponent->UpdatedDelegate.ExecuteIfBound();
diff --git a/Engine/Plugins/Experimental/Mutable/Source/CustomizableObject/Private/MuCO/CustomizableObjectModule.cpp b/Engine/Plugins/Experimental/Mutable/Source/CustomizableObject/Private/MuCO/CustomizableObjectModule.cpp
--- a/Engine/Plugins/Experimental/Mutable/Source/CustomizableObject/Private/MuCO/CustomizableObjectModule.cpp
+++ b/Engine/Plugins/Experimental/Mutable/Source/CustomizableObject/Private/MuCO/CustomizableObjectModule.cpp
@@ -31,9 +31,9 @@ public:
 	// ICustomizableObjectModule interface
 	FString GetPluginVersion() const override;
 	ECustomizableObjectNumBoneInfluences GetNumBoneInfluences() const override;
-	void RegisterExtension(TObjectPtr<const UCustomizableObjectExtension> Extension) override;
-	void UnregisterExtension(TObjectPtr<const UCustomizableObjectExtension> Extension) override;
-	TArrayView<const TObjectPtr<const UCustomizableObjectExtension>> GetRegisteredExtensions() const override;
+	void RegisterExtension(const UCustomizableObjectExtension* Extension) override;
+	void UnregisterExtension(const UCustomizableObjectExtension* Extension) override;
+	TArrayView<const UCustomizableObjectExtension* const> GetRegisteredExtensions() const override;
 	TArrayView<const FRegisteredCustomizableObjectPinType> GetExtendedPinTypes() const override;
 	TArrayView<const FRegisteredObjectNodeInputPin> GetAdditionalObjectNodePins() const override;
 
@@ -49,7 +49,7 @@ private:
 	// Ensure extensions aren't garbage collected
 	TArray<TStrongObjectPtr<const UCustomizableObjectExtension>> StrongExtensions;
 	// For returning from GetRegisteredExtensions
-	TArray<TObjectPtr<const UCustomizableObjectExtension>> Extensions;
+	TArray<const UCustomizableObjectExtension*> Extensions;
 
 	TArray<FRegisteredCustomizableObjectPinType> ExtendedPinTypes;
 	TArray<FRegisteredObjectNodeInputPin> AdditionalObjectNodePins;
@@ -135,8 +135,7 @@ ECustomizableObjectNumBoneInfluences FCustomizableObjectModule::GetNumBoneInflue
 	return ECustomizableObjectNumBoneInfluences::Four;
 }
 
-
-void FCustomizableObjectModule::RegisterExtension(TObjectPtr<const UCustomizableObjectExtension> Extension)
+void FCustomizableObjectModule::RegisterExtension(const UCustomizableObjectExtension* Extension)
 {
 	check(IsInGameThread());
 	
@@ -146,7 +145,7 @@ void FCustomizableObjectModule::RegisterExtension(TObjectPtr<const UCustomizable
 	RefreshExtensionData();
 }
 
-void FCustomizableObjectModule::UnregisterExtension(TObjectPtr<const UCustomizableObjectExtension> Extension)
+void FCustomizableObjectModule::UnregisterExtension(const UCustomizableObjectExtension* Extension)
 {
 	check(IsInGameThread());
 	
@@ -156,7 +155,7 @@ void FCustomizableObjectModule::UnregisterExtension(TObjectPtr<const UCustomizab
 	RefreshExtensionData();
 }
 
-TArrayView<const TObjectPtr<const UCustomizableObjectExtension>> FCustomizableObjectModule::GetRegisteredExtensions() const
+TArrayView<const UCustomizableObjectExtension* const> FCustomizableObjectModule::GetRegisteredExtensions() const
 {
 	check(IsInGameThread());
 	return MakeArrayView(Extensions);
@@ -179,7 +178,7 @@ void FCustomizableObjectModule::RefreshExtensionData()
 	ExtendedPinTypes.Reset();
 	AdditionalObjectNodePins.Reset();
 
-	for (const TObjectPtr<const UCustomizableObjectExtension>& Extension : Extensions)
+	for (const UCustomizableObjectExtension* Extension : Extensions)
 	{
 		for (const FCustomizableObjectPinType& PinType : Extension->GetPinTypes())
 		{
diff --git a/Engine/Plugins/Experimental/Mutable/Source/CustomizableObject/Private/MuCO/CustomizableObjectSystem.cpp b/Engine/Plugins/Experimental/Mutable/Source/CustomizableObject/Private/MuCO/CustomizableObjectSystem.cpp
--- a/Engine/Plugins/Experimental/Mutable/Source/CustomizableObject/Private/MuCO/CustomizableObjectSystem.cpp
+++ b/Engine/Plugins/Experimental/Mutable/Source/CustomizableObject/Private/MuCO/CustomizableObjectSystem.cpp
@@ -2084,6 +2084,21 @@ namespace impl
 			OperationData->RelevantParametersInProgress.Reset();
 		}
 
+		// Copy ExtensionData Object node input from the Instance to the InstanceUpdateData
+		for (int32 ExtensionDataIndex = 0; ExtensionDataIndex < OperationData->MutableInstance->GetExtensionDataCount(); ExtensionDataIndex++)
+		{
+			mu::Ptr<const mu::ExtensionData> ExtensionData;
+			FName Name;
+			OperationData->MutableInstance->GetExtensionData(ExtensionDataIndex, ExtensionData, Name);
+
+			check(ExtensionData);
+
+			FInstanceUpdateData::FNamedExtensionData& NewEntry = OperationData->InstanceUpdateData.ExtendedInputPins.AddDefaulted_GetRef();
+			NewEntry.Data = ExtensionData;
+			NewEntry.Name = Name;
+			check(NewEntry.Name != NAME_None);
+		}
+		
 #if WITH_EDITOR
 		const uint32 EndCycles = FPlatformTime::Cycles();
 		OperationData->MutableRuntimeCycles = EndCycles - StartCycles;
diff --git a/Engine/Plugins/Experimental/Mutable/Source/CustomizableObject/Public/MuCO/CustomizableObjectExtension.h b/Engine/Plugins/Experimental/Mutable/Source/CustomizableObject/Public/MuCO/CustomizableObjectExtension.h
--- a/Engine/Plugins/Experimental/Mutable/Source/CustomizableObject/Public/MuCO/CustomizableObjectExtension.h
+++ b/Engine/Plugins/Experimental/Mutable/Source/CustomizableObject/Public/MuCO/CustomizableObjectExtension.h
@@ -9,6 +9,7 @@
 #include "CustomizableObjectExtension.generated.h"
 
 class UCustomizableObjectInstance;
+class UCustomizableObjectInstanceUsage;
 class USkeletalMesh;
 class ACustomizableSkeletalMeshActor;
 
@@ -98,10 +99,25 @@ public:
 	 * @param InputPinData - The data for only the input pins *registered by this extension*. This
 	 * helps to enforce separation between the extensions, so that they don't depend on each other.
 	 * 
-	 * @param ComponentIndex - The component index of the Skeletal Mesh, for the case where the pin
+	 * @param ObjectComponentIndex - The component index of the Skeletal Mesh, for the case where the pin
 	 * data is associated with a particular component.
 	 * 
 	 * @param SkeletalMesh - The Skeletal Mesh that was created.
 	 */
 	virtual void OnSkeletalMeshCreated(const TArray<FInputPinDataContainer>& InputPinData, int32 ObjectComponentIndex, USkeletalMesh* SkeletalMesh) const {}
+	
+	/**
+	 * Note that the data registered here is completely independent of any Extension Data used in
+	 * the Customizable Object graph. Even though Extension Data and this Extension Instance Data
+	 * both use FInstancedStruct to box an extension-defined struct, there's no requirement that
+	 * they use the same struct type, so they may be completely unrelated.
+	 * 
+	 * Note that GetExtensionInstanceData returns the struct by value to ensure memory safety, so
+	 * the struct should ideally be small and cheap to copy. If you need to reference large data
+	 * from this struct, consider wrapping it in a UObject or referencing it via a TSharedPtr so
+	 * that the large data itself isn't being copied.
+	 */
+	virtual FInstancedStruct GenerateExtensionInstanceData(const TArray<FInputPinDataContainer>& InputPinData) const { return FInstancedStruct(); }
+	
+	virtual void OnCustomizableObjectInstanceUsageUpdated(const UCustomizableObjectInstanceUsage* Usage) const {}
 };
diff --git a/Engine/Plugins/Experimental/Mutable/Source/CustomizableObject/Public/MuCO/CustomizableObjectInstance.h b/Engine/Plugins/Experimental/Mutable/Source/CustomizableObject/Public/MuCO/CustomizableObjectInstance.h
--- a/Engine/Plugins/Experimental/Mutable/Source/CustomizableObject/Public/MuCO/CustomizableObjectInstance.h
+++ b/Engine/Plugins/Experimental/Mutable/Source/CustomizableObject/Public/MuCO/CustomizableObjectInstance.h
@@ -23,6 +23,7 @@ class UCustomizableObject;
 class UTexture2D;
 class FUpdateContextPrivate;
 class UCustomizableObjectInstanceUsage;
+class UCustomizableObjectExtension;
 class UMaterialInterface;
 struct FFrame;
 struct FGameplayTagContainer;
@@ -609,6 +610,9 @@ public:
 	/** Fix the given UAnimInstance instance. */
 	void AnimInstanceFixup(UAnimInstance* AnimInstance) const;
 
+	UFUNCTION(BlueprintCallable, Category = CustomizableObjectInstance)
+	FInstancedStruct GetExtensionInstanceData(const UCustomizableObjectExtension* Extension) const;
+	
 	/** See FCustomizableObjectInstanceDescriptor::SaveDescriptor. */
 	void SaveDescriptor(FArchive &CustomizableObjectDescriptor, bool bUseCompactDescriptor);
 
diff --git a/Engine/Plugins/Experimental/Mutable/Source/CustomizableObject/Public/MuCO/ICustomizableObjectModule.h b/Engine/Plugins/Experimental/Mutable/Source/CustomizableObject/Public/MuCO/ICustomizableObjectModule.h
--- a/Engine/Plugins/Experimental/Mutable/Source/CustomizableObject/Public/MuCO/ICustomizableObjectModule.h
+++ b/Engine/Plugins/Experimental/Mutable/Source/CustomizableObject/Public/MuCO/ICustomizableObjectModule.h
@@ -64,9 +64,9 @@ public:
 	 * These may only be called from the game thread
 	 */
 
-	virtual void RegisterExtension(TObjectPtr<const UCustomizableObjectExtension> Extension) = 0;
-	virtual void UnregisterExtension(TObjectPtr<const UCustomizableObjectExtension> Extension) = 0;
-	virtual TArrayView<const TObjectPtr<const UCustomizableObjectExtension>> GetRegisteredExtensions() const = 0;
+	virtual void RegisterExtension(const UCustomizableObjectExtension* Extension) = 0;
+	virtual void UnregisterExtension(const UCustomizableObjectExtension* Extension) = 0;
+	virtual TArrayView<const UCustomizableObjectExtension* const> GetRegisteredExtensions() const = 0;
 
 	/**
 	 * The results from these functions should only reference extensions that are still valid.
diff --git a/Engine/Plugins/Experimental/Mutable/Source/CustomizableObjectEditor/Private/MuCOE/CustomizableObjectEditor.cpp b/Engine/Plugins/Experimental/Mutable/Source/CustomizableObjectEditor/Private/MuCOE/CustomizableObjectEditor.cpp
--- a/Engine/Plugins/Experimental/Mutable/Source/CustomizableObjectEditor/Private/MuCOE/CustomizableObjectEditor.cpp
+++ b/Engine/Plugins/Experimental/Mutable/Source/CustomizableObjectEditor/Private/MuCOE/CustomizableObjectEditor.cpp
@@ -2194,8 +2194,10 @@ void FCustomizableObjectEditor::LogSearchResult(const UObject& Context, const FS
 
 void FCustomizableObjectEditor::OnUpdatePreviewInstance(UCustomizableObjectInstance* Instance)
 {
-	const bool bVisible = PreviewInstance->GetPrivate()->SkeletalMeshStatus == ESkeletalMeshStatus::Success;
-	Actor->GetRootComponent()->SetVisibility(bVisible, true);
+	if (PreviewInstance->GetPrivate()->SkeletalMeshStatus != ESkeletalMeshStatus::Success)
+	{
+		Actor->GetRootComponent()->SetVisibility(false, true);
+	}
 	Actor->GetRootComponent()->UpdateBounds();
 	
 	if (TextureAnalyzer.IsValid())
