[mutable] Fix bug compiling layouts with blocks that have masks.
#rnx

diff --git a/Engine/Plugins/Experimental/Mutable/Source/MutableTools/Private/MuT/CodeGenerator_Mesh.cpp b/Engine/Plugins/Experimental/Mutable/Source/MutableTools/Private/MuT/CodeGenerator_Mesh.cpp
--- a/Engine/Plugins/Experimental/Mutable/Source/MutableTools/Private/MuT/CodeGenerator_Mesh.cpp
+++ b/Engine/Plugins/Experimental/Mutable/Source/MutableTools/Private/MuT/CodeGenerator_Mesh.cpp
@@ -303,24 +303,20 @@ namespace mu
 		{
 			bool bBlockHasMask = GeneratedLayout.Source->Blocks[BlockIndex].Mask.get() != nullptr;
 
+			// Fill the block rect
+			FImageSize Min = Layout->Blocks[BlockIndex].Min;
+			FImageSize Size = Layout->Blocks[BlockIndex].Size;
+
+			box<FVector2f>& BlockRect = BlockRects[BlockIndex];
+			BlockRect.min[0] = float(Min.X) / float(Grid.X);
+			BlockRect.min[1] = float(Min.Y) / float(Grid.Y);
+			BlockRect.size[0] = float(Size.X) / float(Grid.X);
+			BlockRect.size[1] = float(Size.Y) / float(Grid.Y);
+
 			// Fill the block index per cell array
 			// Ignore the block in this stage if it has a mask, because blocks with masks will very likely overlap other blocks
 			if (!bBlockHasMask)
 			{
-				// Get the block rect
-				FImageSize Min = Layout->Blocks[BlockIndex].Min;
-				FImageSize Size = Layout->Blocks[BlockIndex].Size;
-
-				// Clamp block to layout grid
-				Size.X = FMath::Min(Min.X + Size.X, Grid.X);
-				Size.Y = FMath::Min(Min.Y + Size.Y, Grid.Y);
-
-				box<FVector2f>& BlockRect = BlockRects[BlockIndex];
-				BlockRect.min[0] = ((float)Min.X) / (float)Grid.X;
-				BlockRect.min[1] = ((float)Min.Y) / (float)Grid.Y;
-				BlockRect.size[0] = ((float)Size.X) / (float)Grid.X;
-				BlockRect.size[1] = ((float)Size.Y) / (float)Grid.Y;
-
 				for (uint16 Y = Min.Y; Y < Min.Y + Size.Y; ++Y)
 				{
 					const uint16 PositionY = Y * Grid.X;
