[Mutable] Fix real time morph surface generation in meshes coming from a table.
- Flag indicating the surface needs to allocate a buffer for real time morph targets was omitted.

#rnx

diff --git a/Engine/Plugins/Experimental/Mutable/Source/CustomizableObjectEditor/Private/MuCOE/GenerateMutableSource/GenerateMutableSourceMesh.cpp b/Engine/Plugins/Experimental/Mutable/Source/CustomizableObjectEditor/Private/MuCOE/GenerateMutableSource/GenerateMutableSourceMesh.cpp
--- a/Engine/Plugins/Experimental/Mutable/Source/CustomizableObjectEditor/Private/MuCOE/GenerateMutableSource/GenerateMutableSourceMesh.cpp
+++ b/Engine/Plugins/Experimental/Mutable/Source/CustomizableObjectEditor/Private/MuCOE/GenerateMutableSource/GenerateMutableSourceMesh.cpp
@@ -4258,8 +4258,7 @@ mu::NodeMeshPtr GenerateMutableSourceMesh(const UEdGraphPin* Pin,
 									MeshData.MinNumTriangles = ImportedModel->LODModels[LODIndex].Sections[SectionIndex].NumTriangles;
 
 									// With tables we have to use the default skeletal mesh to tell if the surface mesh needs morphs or clothing.
-									// For now tables will not have morphs as there isn't a way to select them.
-									MeshData.bHasRealTimeMorphs = false;
+									MeshData.bHasRealTimeMorphs = GenerationContext.Options.bRealTimeMorphTargetsEnabled && DefaultSkeletalMesh->GetMorphTargets().Num(); 
 									MeshData.bHasClothing = GenerationContext.Options.bClothingEnabled && ImportedModel->LODModels[LODIndex].HasClothData();
 								}
 							}
