[mutable] Upgraded tree scaling system for the Mutable Code Viewer tree view (debugger).
- The splitter will make the tree view expand horizontally.
- If the contents of the tree view do not fit in the parent allocated space, then a scrollbar will allow the horizontal navigation of the tree rows, allowing the inspection of infinitely deep operations.
- In some cases it can make the widget flicker.

#rnx

diff --git a/Engine/Plugins/Experimental/Mutable/Source/CustomizableObjectEditor/Private/MuCOE/SMutableCodeViewer.cpp b/Engine/Plugins/Experimental/Mutable/Source/CustomizableObjectEditor/Private/MuCOE/SMutableCodeViewer.cpp
--- a/Engine/Plugins/Experimental/Mutable/Source/CustomizableObjectEditor/Private/MuCOE/SMutableCodeViewer.cpp
+++ b/Engine/Plugins/Experimental/Mutable/Source/CustomizableObjectEditor/Private/MuCOE/SMutableCodeViewer.cpp
@@ -143,6 +143,7 @@ public:
 			// First coll showing operation name and type
 			+ SHorizontalBox::Slot()
 			.HAlign(EHorizontalAlignment::HAlign_Fill)
+			.AutoWidth()
 			[
 				SNew(SOverlay)
 
@@ -180,6 +181,7 @@ public:
 			
 			+ SHorizontalBox::Slot()
 			.HAlign(EHorizontalAlignment::HAlign_Left)
+			.AutoWidth()
 			[
 				SNew(SHorizontalBox)
 
@@ -345,7 +347,7 @@ void SMutableCodeViewer::Construct(const FArguments& InArgs, const TSharedPtr<mu
 {
 	// Min width allowed for the column. Needed to avoid having issues with the constants space being to small
 	// and then getting too tall on the y axis crashing the UI drawer.
-	constexpr float MinParametersCollWidth = 400;
+	constexpr float MinParametersCollWidth = 200;
 
 	SetCurrentModel(InMutableModel, InReferencedTextures);
 	
@@ -353,7 +355,9 @@ void SMutableCodeViewer::Construct(const FArguments& InArgs, const TSharedPtr<mu
 	ToolbarBuilder.SetLabelVisibility(EVisibility::Visible);
 	ToolbarBuilder.SetStyle(&FAppStyle::Get(), "SlimToolBar");
 
-	ToolbarBuilder.AddWidget(SNew(STextBlock).Text(FText::FromString(InArgs._DataTag)));
+	ToolbarBuilder.AddWidget(
+		SNew(STextBlock)
+		.Text(FText::FromString(InArgs._DataTag)));
 
 	TSharedRef<SScrollBar> TreeVertScrollBar =
 		SNew(SScrollBar).
@@ -371,37 +375,42 @@ void SMutableCodeViewer::Construct(const FArguments& InArgs, const TSharedPtr<mu
 		]
 		+ SVerticalBox::Slot()
 		.VAlign(VAlign_Fill)
+		.Padding(5,2)
 		[
 			SNew(SSplitter)
 			.Orientation(EOrientation::Orient_Horizontal)
 			
 			+ SSplitter::Slot()
 			.Value(0.35f)
+			.MinSize(520)
+			.Resizable(true)
 			[
 				SNew(SVerticalBox)
 
 				// Search box for tree operations
 				+ SVerticalBox::Slot()
 				.AutoHeight()
-				.HAlign(HAlign_Right)
+				.HAlign(HAlign_Left)
 				[
 					SNew(SHorizontalBox)
-
+				
 					// Search by name
 					+ SHorizontalBox::Slot()
 					.AutoWidth()
 					[
 						SNew(SHorizontalBox)
-
+				
 						+ SHorizontalBox::Slot()
 						.AutoWidth()
+						.VAlign(VAlign_Center)
 						[
 							SNew(STextBlock)
 							.Text(LOCTEXT("SelectedOperationByStringLabel","Search Operation by String :"))
 						]
 					
 						+ SHorizontalBox::Slot()
-						.AutoWidth()
+						.MaxWidth(250)
+						.VAlign(VAlign_Center)
 						[
 							SNew(SSearchBox)
 							.HintText(LOCTEXT("OperationToSearchHintText","Search OP"))
@@ -419,16 +428,18 @@ void SMutableCodeViewer::Construct(const FArguments& InArgs, const TSharedPtr<mu
 					.Padding(4,2)
 					[
 						SNew(SHorizontalBox)
-
+				
 						+ SHorizontalBox::Slot()
 						.AutoWidth()
+						.VAlign(VAlign_Center)
 						[
 							SNew(STextBlock)
 							.Text(LOCTEXT("OperationToSearchRegexLabel","Is RegEx?"))
 						]
-
+				
 						+ SHorizontalBox::Slot()
 						.AutoWidth()
+						.VAlign(VAlign_Center)
 						[
 							SNew(SCheckBox)
 							.OnCheckStateChanged(this,&SMutableCodeViewer::OnRegexToggleChanged)
@@ -439,32 +450,34 @@ void SMutableCodeViewer::Construct(const FArguments& InArgs, const TSharedPtr<mu
 				// Operation type filtering slot
 				+ SVerticalBox::Slot()
 				.AutoHeight()
-				.HAlign(HAlign_Right)
+				.HAlign(HAlign_Left)
 				[
 					// Box containing navigation elements
 					SNew(SVerticalBox)
-
+				
 					+ SVerticalBox::Slot()
 					.Padding(2,4)
 					.AutoHeight()
 					[
 						SNew(SHorizontalBox)
-	
+				
 						+ SHorizontalBox::Slot()
 						.AutoWidth()
 						[
 							SNew(SHorizontalBox)
-
+				
 							+ SHorizontalBox::Slot()
 							.AutoWidth()
+							.VAlign(VAlign_Center)
 							[
 								SNew(STextBlock)
 								.Text(LOCTEXT("SelectedOperationTypeLabel","Search Operation Type :"))
 							]
-
+				
 							// ComboBox used to select one or another Op_Type for tree navigation purposes
 							+SHorizontalBox::Slot()
 							.AutoWidth()
+							.VAlign(VAlign_Center)
 							[
 								SAssignNew(TargetedTypeSelector,SComboBox<TSharedPtr<const FMutableOperationElement>>)
 								.OptionsSource(&FoundModelOperationTypeElements)
@@ -478,29 +491,32 @@ void SMutableCodeViewer::Construct(const FArguments& InArgs, const TSharedPtr<mu
 								.OnSelectionChanged(this,&SMutableCodeViewer::OnNavigationSelectedOperationChanged)
 							]
 						]
-
+				
 						+ SHorizontalBox::Slot()
 						.Padding(4,0)
 						.AutoWidth()
+						.VAlign(VAlign_Center)
 						[
 							SNew(SButton)
 							.Text(LOCTEXT("GoToPreviousOperationButton"," < "))
 							.OnClicked(this,&SMutableCodeViewer::OnGoToPreviousOperationButtonPressed)
 							.IsEnabled(this,&SMutableCodeViewer::CanInteractWithPreviousOperationButton)
 						]
-
+				
 						+ SHorizontalBox::Slot()
 						.Padding(4,0)
 						.AutoWidth()
+						.VAlign(VAlign_Center)
 						[
 							SNew(STextBlock)
 							.Text(this,&SMutableCodeViewer::OnPrintNavigableObjectAddressesCount)
-							.Justification(ETextJustify::Right)
+							.Justification(ETextJustify::Center)
 						]
-		
+				
 						+ SHorizontalBox::Slot()
 						.Padding(4,0)
 						.AutoWidth()
+						.VAlign(VAlign_Center)
 						[
 							SNew(SButton)
 							.Text(LOCTEXT("GoToNextOperationButton"," > "))
@@ -521,14 +537,15 @@ void SMutableCodeViewer::Construct(const FArguments& InArgs, const TSharedPtr<mu
 						SNew(SHorizontalBox)
 
 						+ SHorizontalBox::Slot()
-						.FillWidth(1.0f)
+						.FillContentWidth(1)
 						[
-							SNew(SScrollBox)
+							SNew(SScrollBox)		
 							.Orientation(EOrientation::Orient_Horizontal)
 							.ConsumeMouseWheel(EConsumeMouseWheel::Never)
-
+							
 							+ SScrollBox::Slot()
 							.HAlign(HAlign_Fill)
+							.FillContentSize(1)
 							[
 								SAssignNew(TreeView, STreeView<TSharedPtr<FMutableCodeTreeElement>>)
 								.TreeItemsSource(&RootNodes)
@@ -548,8 +565,7 @@ void SMutableCodeViewer::Construct(const FArguments& InArgs, const TSharedPtr<mu
 
 									+ SHeaderRow::Column(MutableCodeTreeViewColumns::OperationsColumnID)
 										.DefaultLabel(LOCTEXT("Operation", "Operation"))
-										.ManualWidth(618.0f)
-				
+
 									+ SHeaderRow::Column(MutableCodeTreeViewColumns::AdditionalDataColumnID)
 										.DefaultLabel(LOCTEXT("OperationFlags", "Flags"))
 										.FixedWidth(50.0f)
@@ -595,7 +611,6 @@ void SMutableCodeViewer::Construct(const FArguments& InArgs, const TSharedPtr<mu
 								.Text(LOCTEXT("SkipMipsLabel", "Skip mips on generate :"))
 								.Visibility(this, &SMutableCodeViewer::IsMipSkipVisible)
 							]
-
 							+ SHorizontalBox::Slot()
 							[
 								SNew(SNumericEntryBox<int32>)
@@ -628,7 +643,6 @@ void SMutableCodeViewer::Construct(const FArguments& InArgs, const TSharedPtr<mu
 							&(MutableModel->GetPrivate()->m_program),
 							SharedThis(this))
 					]
-
 				]
 				+ SSplitter::Slot()
 				.Value(0.72f)
