[mutable] Fixed Api changes that made VillageDemo not compile
- Removed not used method that also exposed some data not required by external users)

diff --git a/Engine/Plugins/Experimental/Mutable/Source/CustomizableObject/Internal/MuCO/CustomizableObjectSystemPrivate.h b/Engine/Plugins/Experimental/Mutable/Source/CustomizableObject/Internal/MuCO/CustomizableObjectSystemPrivate.h
--- a/Engine/Plugins/Experimental/Mutable/Source/CustomizableObject/Internal/MuCO/CustomizableObjectSystemPrivate.h
+++ b/Engine/Plugins/Experimental/Mutable/Source/CustomizableObject/Internal/MuCO/CustomizableObjectSystemPrivate.h
@@ -820,6 +820,25 @@ public:
 #endif
 
 	int32 NumLODUpdatesLastTick = 0;
+
+
+public:
+	/**
+	 * Get to know if the settings used by the mutable syustem are optimzed for benchmarking operations or not
+	 * @return true if using benchmarking optimized settings, false otherwise.
+	 */
+	static bool IsUsingBenchmarkingSettings();
+
+	/**
+	 * Enable or disable the usage of benchmarking optimized settings
+	 * @param bUseBenchmarkingOptimizedSettings True to enable the usage of benchmarking settings, false to disable it.
+	 */
+	CUSTOMIZABLEOBJECT_API static void SetUsageOfBenchmarkingSettings(bool bUseBenchmarkingOptimizedSettings);
+
+private:
+
+	/** Flag that controls some of the settings used for the generation of instances. */
+	inline static bool bUseBenchmarkingSettings = false;
 };
 
 
diff --git a/Engine/Plugins/Experimental/Mutable/Source/CustomizableObject/Private/MuCO/CustomizableObjectSystem.cpp b/Engine/Plugins/Experimental/Mutable/Source/CustomizableObject/Private/MuCO/CustomizableObjectSystem.cpp
--- a/Engine/Plugins/Experimental/Mutable/Source/CustomizableObject/Private/MuCO/CustomizableObjectSystem.cpp
+++ b/Engine/Plugins/Experimental/Mutable/Source/CustomizableObject/Private/MuCO/CustomizableObjectSystem.cpp
@@ -3968,6 +3968,19 @@ void UCustomizableObjectSystemPrivate::UnCacheTextureParameters(const TArray<FCu
 }
 
 
+bool UCustomizableObjectSystemPrivate::IsUsingBenchmarkingSettings()
+{
+	return bUseBenchmarkingSettings;
+}
+
+
+void UCustomizableObjectSystemPrivate::SetUsageOfBenchmarkingSettings(bool bUseBenchmarkingOptimizedSettings)
+{
+	bUseBenchmarkingSettings = bUseBenchmarkingOptimizedSettings;
+}
+
+
+
 int32 UCustomizableObjectSystem::GetNumInstances() const
 {
 	int32 NumInstances;
@@ -4114,33 +4127,35 @@ void UCustomizableObjectSystem::SetReleaseMutableTexturesImmediately(bool bRelea
 }
 
 
-void UCustomizableObjectSystem::SetBenchmarkState(bool bIsEnabled)
+void UCustomizableObjectSystem::EnableBenchmark()
 {
-	bIsBenchmarking = bIsEnabled;
+	// Start reporting benchmarking data (log and .csv file)
+	FLogBenchmarkUtil::SetBenchmarkReportingStateOverride(true);
 }
 
 
-bool UCustomizableObjectSystem::IsBenchmarking()
+void UCustomizableObjectSystem::EndBenchmark()
 {
-	return bIsBenchmarking;
+	// Stop the reporting of benchmarking data
+	FLogBenchmarkUtil::SetBenchmarkReportingStateOverride(false);
 }
 
 
 bool UCustomizableObjectSystem::IsMeshCacheEnabled(bool bCheckCVarOnGameThread /** = false */)
 {
-	return bIsBenchmarking ? false : CVarEnableMeshCache.GetValueOnAnyThread(bCheckCVarOnGameThread);
+	return UCustomizableObjectSystemPrivate::IsUsingBenchmarkingSettings() ? false : CVarEnableMeshCache.GetValueOnAnyThread(bCheckCVarOnGameThread);
 }
 
 
 bool UCustomizableObjectSystem::ShouldClearWorkingMemoryOnUpdateEnd()
 {
-	return bIsBenchmarking ? true : CVarClearWorkingMemoryOnUpdateEnd.GetValueOnAnyThread();
+	return UCustomizableObjectSystemPrivate::IsUsingBenchmarkingSettings() ? true : CVarClearWorkingMemoryOnUpdateEnd.GetValueOnAnyThread();
 }
 
 
 bool UCustomizableObjectSystem::ShouldReuseTexturesBetweenInstances()
 {
-	return bIsBenchmarking ? false : CVarReuseImagesBetweenInstances.GetValueOnAnyThread();
+	return UCustomizableObjectSystemPrivate::IsUsingBenchmarkingSettings() ? false : CVarReuseImagesBetweenInstances.GetValueOnAnyThread();
 }
 
 
@@ -4155,7 +4170,7 @@ void UCustomizableObjectSystem::SetWorkingMemory(int32 Bytes)
 
 int32 UCustomizableObjectSystem::GetWorkingMemory() const
 {
-	return bIsBenchmarking ? 16384 : WorkingMemory;
+	return UCustomizableObjectSystemPrivate::IsUsingBenchmarkingSettings() ? 16384 : WorkingMemory;
 }
 
 
diff --git a/Engine/Plugins/Experimental/Mutable/Source/CustomizableObject/Public/MuCO/CustomizableObjectSystem.h b/Engine/Plugins/Experimental/Mutable/Source/CustomizableObject/Public/MuCO/CustomizableObjectSystem.h
--- a/Engine/Plugins/Experimental/Mutable/Source/CustomizableObject/Public/MuCO/CustomizableObjectSystem.h
+++ b/Engine/Plugins/Experimental/Mutable/Source/CustomizableObject/Public/MuCO/CustomizableObjectSystem.h
@@ -380,21 +380,15 @@ private:
 public:
 	
 	/**
-	 * Enable or disable the benchmarking mode for the MutableSystem.
-	 * @param bIsEnabled True to enable it and false to disable it
+	 * Enables the collection of internal Mutable performance data. It has a performance cost.
 	 */
-	static void SetBenchmarkState(bool bIsEnabled);
-	
+	void EnableBenchmark();
+
 	/**
-	 * Get if the system is in benchmarking mode or not.
-	 * @return True if it is in benchmarking mode and false otherwise
+	 * Disables the reporting of mutable instance benchmarking data.
 	 */
-	static bool IsBenchmarking();
+	void EndBenchmark();
 	
-private:
-
-	/** Flag that controls some of the settings used for the generation of instances. */
-	inline static bool bIsBenchmarking = false;
 };
 
 
diff --git a/Engine/Plugins/Experimental/Mutable/Source/CustomizableObjectEditor/Private/MuCOE/CustomizableObjectEditorPerformanceAnalyzer.cpp b/Engine/Plugins/Experimental/Mutable/Source/CustomizableObjectEditor/Private/MuCOE/CustomizableObjectEditorPerformanceAnalyzer.cpp
--- a/Engine/Plugins/Experimental/Mutable/Source/CustomizableObjectEditor/Private/MuCOE/CustomizableObjectEditorPerformanceAnalyzer.cpp
+++ b/Engine/Plugins/Experimental/Mutable/Source/CustomizableObjectEditor/Private/MuCOE/CustomizableObjectEditorPerformanceAnalyzer.cpp
@@ -620,7 +620,7 @@ void SCustomizableObjectEditorPerformanceAnalyzer::SetMutableBenchmarkingSystemS
 	FLogBenchmarkUtil::SetBenchmarkReportingStateOverride(bNewState);
 
 	// Setting this to false will allow for the system to retrieve the values set by CVars and other non code ways
-	UCustomizableObjectSystem::SetBenchmarkState(bNewState);
+	UCustomizableObjectSystemPrivate::SetUsageOfBenchmarkingSettings(bNewState);
 }
 
 
diff --git a/Engine/Plugins/Experimental/Mutable/Source/MutableValidation/Private/MuV/COIBakingTestCommandlet.cpp b/Engine/Plugins/Experimental/Mutable/Source/MutableValidation/Private/MuV/COIBakingTestCommandlet.cpp
--- a/Engine/Plugins/Experimental/Mutable/Source/MutableValidation/Private/MuV/COIBakingTestCommandlet.cpp
+++ b/Engine/Plugins/Experimental/Mutable/Source/MutableValidation/Private/MuV/COIBakingTestCommandlet.cpp
@@ -13,6 +13,7 @@
 #include "HAL/FileManager.h"
 #include "Interfaces/ITargetPlatformManagerModule.h"
 #include "MuCO/CustomizableObjectPrivate.h"
+#include "MuCO/CustomizableObjectSystemPrivate.h"
 #include "MuCOE/CustomizableObjectBenchmarkingUtils.h"
 
 
@@ -37,7 +38,9 @@ void OnInstanceUpdate(const FUpdateContext& Result)
 
 int32 UCOIBakingTestCommandlet::Main(const FString& Params)
 {
-	UCustomizableObjectSystem::SetBenchmarkState(true);
+	// Ensure we have set the mutable system to the benchmarking mode and that we are reporting benchmarking data
+	FLogBenchmarkUtil::SetBenchmarkReportingStateOverride(true);
+	UCustomizableObjectSystemPrivate::SetUsageOfBenchmarkingSettings(true);
 	
 	// Ensure we do not show any OK dialog since we are not an user that can interact with them
 	GIsRunningUnattendedScript = true;
diff --git a/Engine/Plugins/Experimental/Mutable/Source/MutableValidation/Private/MuV/COIBulkUpdateTestCommandlet.cpp b/Engine/Plugins/Experimental/Mutable/Source/MutableValidation/Private/MuV/COIBulkUpdateTestCommandlet.cpp
--- a/Engine/Plugins/Experimental/Mutable/Source/MutableValidation/Private/MuV/COIBulkUpdateTestCommandlet.cpp
+++ b/Engine/Plugins/Experimental/Mutable/Source/MutableValidation/Private/MuV/COIBulkUpdateTestCommandlet.cpp
@@ -12,13 +12,15 @@
 #include "Interfaces/ITargetPlatformManagerModule.h"
 #include "MuCO/CustomizableObjectPrivate.h"
 #include "MuCO/CustomizableObjectSystem.h"
+#include "MuCO/CustomizableObjectSystemPrivate.h"
 #include "MuCOE/CustomizableObjectBenchmarkingUtils.h"
 
 
 int32 UCOIBulkUpdateTestCommandlet::Main(const FString& Params)
 {
-	// Ensure we have the cvars used for our testing set
-	UCustomizableObjectSystem::SetBenchmarkState(true);
+	// Ensure we have set the mutable system to the benchmarking mode and that we are reporting benchmarking data
+	FLogBenchmarkUtil::SetBenchmarkReportingStateOverride(true);
+	UCustomizableObjectSystemPrivate::SetUsageOfBenchmarkingSettings(true);
 	
 	// Ensure we do not show any OK dialog since we are not an user that can interact with them
 	GIsRunningUnattendedScript = true;
diff --git a/Engine/Plugins/Experimental/Mutable/Source/MutableValidation/Private/MuV/CustomizableObjectValidationCommandlet.cpp b/Engine/Plugins/Experimental/Mutable/Source/MutableValidation/Private/MuV/CustomizableObjectValidationCommandlet.cpp
--- a/Engine/Plugins/Experimental/Mutable/Source/MutableValidation/Private/MuV/CustomizableObjectValidationCommandlet.cpp
+++ b/Engine/Plugins/Experimental/Mutable/Source/MutableValidation/Private/MuV/CustomizableObjectValidationCommandlet.cpp
@@ -16,6 +16,7 @@
 #include "MuR/Model.h"
 #include "Interfaces/ITargetPlatformManagerModule.h"
 #include "Interfaces/ITargetPlatform.h"
+#include "MuCO/CustomizableObjectSystemPrivate.h"
 #include "MuCOE/CustomizableObjectBenchmarkingUtils.h"
 #include "Serialization/MemoryWriter.h"
 
@@ -27,8 +28,9 @@ int32 UCustomizableObjectValidationCommandlet::Main(const FString& Params)
 	// Execution arguments for commandlet from IDE
 	// -run=CustomizableObjectValidation -CustomizableObject=(PathToCO)
 
-	// Ensure we have the cvars used for our testing set
-	UCustomizableObjectSystem::SetBenchmarkState(true);
+	// Ensure we have set the mutable system to the benchmarking mode and that we are reporting benchmarking data
+	FLogBenchmarkUtil::SetBenchmarkReportingStateOverride(true);
+	UCustomizableObjectSystemPrivate::SetUsageOfBenchmarkingSettings(true);
 	
 	// Ensure we do not show any OK dialog since we are not an user that can interact with them
 	GIsRunningUnattendedScript = true;
