[mutable] Improved Mutable Mesh constants slate so it displays useful data
- Created new slate to contain multipage lists of elements (as a way to avoid UE  crash due to too big of a slate)- It uses already tested code.
- Simplified some code from the Mutable Constants Widget.

diff --git a/Engine/Plugins/Experimental/Mutable/Source/CustomizableObjectEditor/Private/MuCOE/SMutableConstantsWidget.cpp b/Engine/Plugins/Experimental/Mutable/Source/CustomizableObjectEditor/Private/MuCOE/SMutableConstantsWidget.cpp
--- a/Engine/Plugins/Experimental/Mutable/Source/CustomizableObjectEditor/Private/MuCOE/SMutableConstantsWidget.cpp
+++ b/Engine/Plugins/Experimental/Mutable/Source/CustomizableObjectEditor/Private/MuCOE/SMutableConstantsWidget.cpp
@@ -11,6 +11,7 @@
 #include "Widgets/Layout/SScrollBox.h"
 #include "Widgets/Layout/SSeparator.h"
 #include "Widgets/Views/STileView.h"
+#include "MuCOE/Widgets/MutableMultiPageListView.h"
 
 class ITableRow;
 class SWidget;
@@ -57,33 +58,119 @@ namespace
 
 		return FText::FromString(OutputString);
 	}
+
+	
+	/**
+	* Get the amount of channels for all the buffers in the provided BuffSetet
+	* @param InBufferSet : The bufferset whose channels we want to count.
+	* @return : the amount of channels found 
+	*/
+	int32 GetChannelCountOfBufferSet(const mu::FMeshBufferSet& InBufferSet)
+	{
+		int32 MeshBufferChannelsCount = 0;
+		for (const mu::FMeshBuffer& MeshBuffer : InBufferSet.Buffers)
+		{
+			MeshBufferChannelsCount += MeshBuffer.Channels.Num();
+		}
+		return MeshBufferChannelsCount;
+	}
+
+
+	/**
+	 * Get the amount of channels found in all buffers found in the provided mutable mesh
+	 * @param MeshPtr Pointer to the mesh whose channels we want to count
+	 * @return : the amount of channels found.
+	 */
+	int32 GetMeshChannelCount (mu::Ptr<const mu::Mesh> MeshPtr)
+	{
+		check(MeshPtr)
+
+		int32 ChannelCount = 0;
+		ChannelCount += GetChannelCountOfBufferSet(MeshPtr->GetVertexBuffers());
+		ChannelCount += GetChannelCountOfBufferSet(MeshPtr->GetIndexBuffers());
+		return ChannelCount;
+	}
 }
 
 
 #pragma region SUPPORT_CLASSES
 
-class SMutableConstantMeshRow final : public STableRow<TSharedPtr<FMutableConstantMeshElement>>
+
+namespace MeshConstantTitles
+{
+	static const FName MeshID("Id");
+	static const FName MeshVertices("Vertices");
+	static const FName MeshIndices("Indices");
+	static const FName MeshChannels("BufferChannels");
+	static const FName MeshMemory("Memory");
+}
+
+class SMutableConstantMeshRow final : public SMultiColumnTableRow<TSharedPtr<FMutableConstantMeshElement>>
 {
 public:
 	void Construct(const FArguments& Args, const TSharedRef<STableViewBase>& InOwnerTableView, const TSharedPtr<FMutableConstantMeshElement>& InRowElement)
 	{
 		check(InRowElement);
+		RowElement = InRowElement;
+
+		SMultiColumnTableRow< TSharedPtr<FMutableConstantMeshElement> >::Construct(
+		STableRow::FArguments()
+			.ShowSelection(true)
+			, InOwnerTableView
+		);
+	}
+
+	virtual TSharedRef<SWidget> GenerateWidgetForColumn(const FName& ColumnName) override
+	{
+		check(RowElement);
+		check(RowElement->MeshPtr)
 		
-		const FText MeshProxyText =
-			FText::Format( LOCTEXT("MeshConstantProxyLabel", "{0}_MESH "),InRowElement->IndexOnSourceVector);
+		// Index
+		if (ColumnName == MeshConstantTitles::MeshID)
+		{
+			const FText IndexAsText = FText::AsNumber(RowElement->IndexOnSourceVector);
+			return SNew(STextBlock).Text(IndexAsText);
+		}
 
-		this->ChildSlot
-		[
-			SNew(STextBlock)
-			.Text(MeshProxyText)
-		];
+		// The amount of mesh vertex buffer and image buffer channels
+		if (ColumnName == MeshConstantTitles::MeshChannels)
+		{
+			mu::Ptr<const mu::Mesh> Mesh = RowElement->MeshPtr.get();
+			const int32 ChannelCount =  GetMeshChannelCount(Mesh);
+			const FText ChannelCountText = FText::AsNumber(ChannelCount);
+			return SNew(STextBlock).Text(ChannelCountText);
+		}
 		
-		STableRow< TSharedPtr<FMutableConstantMeshElement> >::ConstructInternal(
-			STableRow::FArguments()
-			.ShowSelection(true)
-			, InOwnerTableView);
+		// The amount of indices of the mesh
+		if (ColumnName == MeshConstantTitles::MeshIndices)
+		{
+			mu::Ptr<const mu::Mesh> Mesh = RowElement->MeshPtr.get();
+			const FText IndexCount = FText::AsNumber(Mesh->GetIndexCount());
+			return SNew(STextBlock).Text(IndexCount);
+		}
 
+		// The amount of vertices of the mesh
+		if (ColumnName == MeshConstantTitles::MeshVertices)
+		{
+			mu::Ptr<const mu::Mesh> Mesh = RowElement->MeshPtr.get();
+			const FText VertexCount = FText::AsNumber(Mesh->GetVertexCount());
+			return SNew(STextBlock).Text(VertexCount);
+		}
+
+		// Memory used by the mesh 
+		if (ColumnName == MeshConstantTitles::MeshMemory)
+		{
+			mu::Ptr<const mu::Mesh> Mesh = RowElement->MeshPtr.get();
+			const FText SizeAsText = GenerateTextForSize(Mesh->GetDataSize());
+			return SNew(STextBlock).Text(SizeAsText);
+		}
+
+		return SNullWidget::NullWidget;
 	}
+
+private:
+	TSharedPtr<FMutableConstantMeshElement> RowElement;
+	
 };
 
 class SMutableConstantStringRow final : public STableRow<TSharedPtr<FMutableConstantStringElement>>
@@ -198,38 +285,8 @@ public:
 		// Memory
 		if (ColumnName == ImageConstantTitles::ImageTotalMemory)
 		{
-			// Todo: Not sure if I do prefer this way of showing the sizes or the one used for the expandable area titles.
-			
-			uint32 SizeInBytes = Image->GetDataSize();
-			FString SizeUnit = " B";
-			
-			// If convertible to KB then display it as KB (Kilo Bytes)
-			if (SizeInBytes >= 1024 && SizeInBytes % 1024 == 0)
-			{
-				SizeUnit = " KB";
-				SizeInBytes /= 1024;
-			
-				// If convertible to MB then display it as MB (Mega Bytes)
-				if (SizeInBytes >= 1024 && SizeInBytes % 1024 == 0)
-				{
-					SizeUnit = " MB";
-					SizeInBytes /= 1024;
-				
-					// If convertible to GB then display it as GB (Giga Bytes)
-					if (SizeInBytes >= 1024 && SizeInBytes % 1024 == 0)
-					{
-						SizeUnit = " GB";
-						SizeInBytes /= 1024;
-					}
-				}
-			}
-			
-			FText SizeAsText = FText::GetEmpty();
-			if (SizeInBytes > 0)
-			{
-				// Return the text formatted as number
-				SizeAsText = FText::FromString( FString::FormatAsNumber(SizeInBytes) + SizeUnit);
-			}
+			const uint32 SizeInBytes = Image->GetDataSize();
+			const FText SizeAsText = GenerateTextForSize(SizeInBytes);
 
 			// Return the text with the size
 			return SNew(STextBlock).Text(SizeAsText);
@@ -395,7 +452,99 @@ void SMutableConstantsWidget::Construct(const FArguments& InArgs,const mu::FProg
 	
 	// Panel title
 	const FText ConstantsPanelTitle = LOCTEXT("ConstantsPannelName", "Constants : ");
+
+	// Hack to allow us to later set the array to be used. if no array is provided then the children slate to contain them later will not exist
+	TArray<TSharedPtr<FMutableConstantImageElement>> TempImageElementsEmptyArray;
+	
+	// Image Constants List View object and Handler
+	TSharedPtr<SListView<TSharedPtr<FMutableConstantImageElement>>> ConstantImagesListView =
+		SNew(SListView<TSharedPtr<FMutableConstantImageElement>>)
+	.OnGenerateRow(this,&SMutableConstantsWidget::OnGenerateImageRow)
+		// We do require to provide something here or the slate to contain the children will not get generated
+	.ListItemsSource(&TempImageElementsEmptyArray)		
+	.OnSelectionChanged(this,&SMutableConstantsWidget::OnSelectedImageChanged)
+	.SelectionMode(ESelectionMode::Single)
+	.HeaderRow
+	(
+		SNew(SHeaderRow)
+		+ SHeaderRow::Column(ImageConstantTitles::ImageID)
+		.DefaultLabel(FText(LOCTEXT("ImageId", "ID")))
+		.OnSort(this, &SMutableConstantsWidget::OnImageTableSortRequested)
+		.SortMode(this, &SMutableConstantsWidget::GetImageListColumnSortMode, ImageConstantTitles::ImageID )
+		.FillWidth(0.28f)
 	
+		+ SHeaderRow::Column(ImageConstantTitles::ImageSize)
+		.DefaultLabel(FText(LOCTEXT("ImageResolution","Resolution")))
+		.DefaultTooltip(FText(LOCTEXT("ImageResolutionolumnToolTip","Pixel resolution")))
+		.OnSort(this, &SMutableConstantsWidget::OnImageTableSortRequested)
+		.SortMode(this, &SMutableConstantsWidget::GetImageListColumnSortMode, ImageConstantTitles::ImageSize )
+		
+		+ SHeaderRow::Column(ImageConstantTitles::ImageMipMaps)
+		.DefaultLabel(FText(LOCTEXT("ImageMipMaps","Mip Maps")))
+		.DefaultTooltip(FText(LOCTEXT("ImageMipMapsColumnToolTip","Amount of Mip maps")))
+		.OnSort(this, &SMutableConstantsWidget::OnImageTableSortRequested)
+		.SortMode(this, &SMutableConstantsWidget::GetImageListColumnSortMode, ImageConstantTitles::ImageMipMaps )
+
+		+ SHeaderRow::Column(ImageConstantTitles::ImageFormat)
+		.DefaultLabel(FText(LOCTEXT("ImageFormat","Format")))
+		.DefaultTooltip(FText(LOCTEXT("ImageFormatColumnToolTip","Image Format")))
+		.OnSort(this, &SMutableConstantsWidget::OnImageTableSortRequested)
+		.SortMode(this, &SMutableConstantsWidget::GetImageListColumnSortMode, ImageConstantTitles::ImageFormat )
+
+		+ SHeaderRow::Column(ImageConstantTitles::ImageTotalMemory)
+		.DefaultLabel(FText(LOCTEXT("ImageMemorySize","Size")))
+		.DefaultTooltip(FText(LOCTEXT("ImageMemorySizeColumnToolTip","Memory size")))
+		.OnSort(this, &SMutableConstantsWidget::OnImageTableSortRequested)
+		.SortMode(this, &SMutableConstantsWidget::GetImageListColumnSortMode, ImageConstantTitles::ImageTotalMemory )
+
+	);
+
+
+	TArray<TSharedPtr<FMutableConstantMeshElement>> TempMeshElementsEmptyArray;
+
+	// Handled list view for mutable constant meshes
+	TSharedPtr<SListView<TSharedPtr<FMutableConstantMeshElement>>> ConstantMeshesListView =
+		SNew(SListView<TSharedPtr<FMutableConstantMeshElement>>)
+	.OnGenerateRow(this,&SMutableConstantsWidget::OnGenerateMeshRow)
+		// We do require to provide something here or the slate to contain the children will not get generated
+	.ListItemsSource(&TempMeshElementsEmptyArray)		
+	.OnSelectionChanged(this,&SMutableConstantsWidget::OnSelectedMeshChanged)
+	.SelectionMode(ESelectionMode::Single)
+	.HeaderRow
+	(
+		SNew(SHeaderRow)
+		+ SHeaderRow::Column(MeshConstantTitles::MeshID)
+		.DefaultLabel(FText(LOCTEXT("MeshId", "ID")))
+		.OnSort(this, &SMutableConstantsWidget::OnMeshTableSortRequested)
+		.SortMode(this, &SMutableConstantsWidget::GetMeshListColumnSortMode, MeshConstantTitles::MeshID)
+		.FillWidth(0.28f)
+	
+		+ SHeaderRow::Column(MeshConstantTitles::MeshVertices)
+		.DefaultLabel(FText(LOCTEXT("MeshVerticesCount","Vertices")))
+		.DefaultTooltip(FText(LOCTEXT("MeshVerticesCountColumnToolTip","Amount of vertice")))
+		.OnSort(this, &SMutableConstantsWidget::OnMeshTableSortRequested)
+		.SortMode(this, &SMutableConstantsWidget::GetMeshListColumnSortMode, MeshConstantTitles::MeshVertices)
+
+		+ SHeaderRow::Column(MeshConstantTitles::MeshIndices)
+		.DefaultLabel(FText(LOCTEXT("MeshIndicesCount","Indices")))
+		.DefaultTooltip(FText(LOCTEXT("MeshIndicesCountColumnToolTip","Amount of indices")))
+		.OnSort(this, &SMutableConstantsWidget::OnMeshTableSortRequested)
+		.SortMode(this, &SMutableConstantsWidget::GetMeshListColumnSortMode, MeshConstantTitles::MeshIndices)
+
+		+ SHeaderRow::Column(MeshConstantTitles::MeshChannels)
+		.DefaultLabel(FText(LOCTEXT("MeshVertexChannelsCount","Channels")))
+		.DefaultTooltip(FText(LOCTEXT("MeshVertexChannelsCountColumnToolTip","Amount of channels")))
+		.OnSort(this, &SMutableConstantsWidget::OnMeshTableSortRequested)
+		.SortMode(this, &SMutableConstantsWidget::GetMeshListColumnSortMode, MeshConstantTitles::MeshChannels)
+		
+		+ SHeaderRow::Column(MeshConstantTitles::MeshMemory)
+		.DefaultLabel(FText(LOCTEXT("MeshMemory","Size")))
+		.DefaultTooltip(FText(LOCTEXT("MeshMemoryColumnToolTip","Memory size")))
+		.OnSort(this, &SMutableConstantsWidget::OnMeshTableSortRequested)
+		.SortMode(this, &SMutableConstantsWidget::GetMeshListColumnSortMode, MeshConstantTitles::MeshMemory)
+	);
+	
+
 	// Child structure
 	this->ChildSlot
 	[
@@ -436,48 +585,10 @@ void SMutableConstantsWidget::Construct(const FArguments& InArgs,const mu::FProg
 			.AreaTitle(this,&SMutableConstantsWidget::OnDrawImagesAreaTitle)
 			.BodyContent()
 			[
-				SNew(SVerticalBox)
-
-				// Buttons being used to change the segment to be displayed
-				+SVerticalBox::Slot()
-				.AutoHeight()
-				[
-					GenerateImagesSegmentSelectionWidget()
-				]
-
-				// List view showing the selected segment of images
-				+ SVerticalBox::Slot()
-				[
-					SAssignNew(ConstantImagesListView,SListView<TSharedPtr<FMutableConstantImageElement>>)
-					.OnGenerateRow(this,&SMutableConstantsWidget::OnGenerateImageRow)
-					.ListItemsSource(&ConstantImageElementsSegment)
-					.OnSelectionChanged(this,&SMutableConstantsWidget::OnSelectedImageChanged)
-					.SelectionMode(ESelectionMode::Single)
-					.HeaderRow
-					(
-						SNew(SHeaderRow)
-						+ SHeaderRow::Column(ImageConstantTitles::ImageID)
-						.DefaultLabel(FText(LOCTEXT("ImageId", "ID")))
-						.OnSort(this, &SMutableConstantsWidget::OnImageTableSortRequested)
-						.FillWidth(0.28f)
-					
-						+ SHeaderRow::Column(ImageConstantTitles::ImageSize)
-						.DefaultLabel(FText(LOCTEXT("ImageResolution","Resolution")))
-						.OnSort(this, &SMutableConstantsWidget::OnImageTableSortRequested)
-						
-						+ SHeaderRow::Column(ImageConstantTitles::ImageMipMaps)
-						.DefaultLabel(FText(LOCTEXT("ImageMipMaps","Mip Maps")))
-						.OnSort(this, &SMutableConstantsWidget::OnImageTableSortRequested)
-						
-						+ SHeaderRow::Column(ImageConstantTitles::ImageFormat)
-						.DefaultLabel(FText(LOCTEXT("ImageFormat","Format")))
-						.OnSort(this, &SMutableConstantsWidget::OnImageTableSortRequested)
-						
-						+ SHeaderRow::Column(ImageConstantTitles::ImageTotalMemory)
-						.DefaultLabel(FText(LOCTEXT("ImageMemorySize","Size")))
-						.OnSort(this, &SMutableConstantsWidget::OnImageTableSortRequested)
-					)
-				]
+				// Custom slate that will handle the updating of the elements displayed by the ConstantImagesListView object
+					SAssignNew(ImageListViewHandler,SMutableMultiPageListView<TSharedPtr<FMutableConstantImageElement>>)
+						.HostedListView(ConstantImagesListView)
+						.ElementsToSeparateInPages(ConstantImageElements)
 			]
 		]
 
@@ -496,12 +607,10 @@ void SMutableConstantsWidget::Construct(const FArguments& InArgs,const mu::FProg
 			.AreaTitle(this,&SMutableConstantsWidget::OnDrawMeshesAreaTitle)
 			.BodyContent()
 			[
-				SAssignNew(ConstantMeshesSlate,STileView<TSharedPtr<FMutableConstantMeshElement>>)
-				.OnSelectionChanged(this,&SMutableConstantsWidget::OnSelectedMeshChanged)
-				.ListItemsSource(&ConstantMeshElements)
-				.ItemHeight(ProxyEntryHeight)
-				.OnGenerateTile(this,&SMutableConstantsWidget::OnGenerateMeshRow)
-				.SelectionMode(ESelectionMode::Single)
+				// Custom slate that will handle the updating of the elements displayed by the ConstantMeshesListView object
+				SAssignNew(MeshListViewHandler ,SMutableMultiPageListView<TSharedPtr<FMutableConstantMeshElement>>)
+				.HostedListView(ConstantMeshesListView)	
+				.ElementsToSeparateInPages(ConstantMeshElements)
 			]
 		]
 
@@ -754,123 +863,19 @@ TSharedRef<ITableRow> SMutableConstantsWidget::OnGenerateSkeletonRow(
 	return Row;
 }
 
-bool SMutableConstantsWidget::ShouldBackButtonBeEnabled() const
-{
-	return  CurrentArraySegment > 0;
-}
-
-bool SMutableConstantsWidget::ShouldNextButtonBeEnabled() const
-{
-	return TotalAmountOfSegments >= 1 && CurrentArraySegment < TotalAmountOfSegments - 1;
-}
-
-void SMutableConstantsWidget::RegenerateProxyImageArray()
-{
-	// Prepare the segment of elements to be filled up
-	ConstantImageElementsSegment.SetNum(0);
-	ConstantImageElementsSegment.Reserve(ElementsPerSegment);
-	
-	const uint32 StartingIndex = CurrentArraySegment * ElementsPerSegment;
-	const uint32 FinishIndex =
-		StartingIndex + FMath::Min(ElementsPerSegment, (static_cast<uint32>(ConstantImageElements.Num()) - StartingIndex));
-	// Iterate over the main array from the point where the segment starts to the point where it ends. Stop before if needed
-	for	(uint32 MainArrayIndex = StartingIndex ; MainArrayIndex < FinishIndex; ++ MainArrayIndex)
-	{
-		ConstantImageElementsSegment.Add(ConstantImageElements[MainArrayIndex]);
-	}
-	
-	if (ConstantImagesListView.IsValid())
-	{
-		ConstantImagesListView->RequestListRefresh();
-	}
-}
-
-TSharedRef<SWidget> SMutableConstantsWidget::GenerateImagesSegmentSelectionWidget()
-{
-	// If there are segments to navigate then generate the navigation area
-	if (TotalAmountOfSegments > 1)
-	{
-		TSharedRef<SHorizontalBox> Container =	SNew(SHorizontalBox)
-
-		+ SHorizontalBox::Slot()
-		.HAlign(EHorizontalAlignment::HAlign_Left)
-		.AutoWidth()
-		[
-			SNew(SButton)
-			.OnClicked(this,&SMutableConstantsWidget::OnImageFullBackButtonClicked)
-			.Text(FText(INVTEXT("|<")))
-			.IsEnabled(this,&SMutableConstantsWidget::ShouldBackButtonBeEnabled)
-		]
-		
-		+ SHorizontalBox::Slot()
-		.HAlign(EHorizontalAlignment::HAlign_Left)
-		.AutoWidth()
-		[
-			SNew(SButton)
-			.OnClicked(this,&SMutableConstantsWidget::OnImageBackButtonClicked)
-			.Text(FText(INVTEXT("<")))
-			.IsEnabled(this,&SMutableConstantsWidget::ShouldBackButtonBeEnabled)
-		]
-
-		+ SHorizontalBox::Slot()
-		.HAlign(EHorizontalAlignment::HAlign_Center)
-		[
-			SNew(STextBlock)
-			.Justification(ETextJustify::Center)
-			.Text(this,&SMutableConstantsWidget::OnDrawCurrentImageSegmentText)
-		]
-			
-		+ SHorizontalBox::Slot()
-		.HAlign(EHorizontalAlignment::HAlign_Right)
-		.AutoWidth()
-		[
-			SNew(SButton)
-			.OnClicked(this,&SMutableConstantsWidget::OnImageForwardButtonClicked)
-			.Text(FText(INVTEXT(">")))
-			.IsEnabled(this,&SMutableConstantsWidget::ShouldNextButtonBeEnabled)
-		]
-
-		+ SHorizontalBox::Slot()
-		.HAlign(EHorizontalAlignment::HAlign_Right)
-		.AutoWidth()
-		[
-			SNew(SButton)
-			.OnClicked(this,&SMutableConstantsWidget::OnImageFullForwardButtonClicked)
-			.Text(FText(INVTEXT(">|")))
-			.IsEnabled(this,&SMutableConstantsWidget::ShouldNextButtonBeEnabled)
-		];
-
-		return Container;
-	}
-
-	// If no more than 1 section is required
-	return SNullWidget::NullWidget;
-}
-
-FText SMutableConstantsWidget::OnDrawCurrentImageSegmentText() const
-{
-	FString BuiltText = FString::FromInt(  CurrentArraySegment + 1);
-	BuiltText += FString(" / ");
-	BuiltText += FString::FromInt(TotalAmountOfSegments);
-	
-	return FText::Format(LOCTEXT("PageNumber","Page : {0}"), FText::FromString(BuiltText));
-}
 
 void SMutableConstantsWidget::OnImageTableSortRequested(EColumnSortPriority::Type ColumnSortPriority, const FName& ColumnID,
                                                         EColumnSortMode::Type ColumnSortMode)
 {
-	// If the colum has been sorted on one way now do the inverse way of sorting
-	if (LastSortedColumnID == ColumnID)
-	{
-		bSortAscending = !bSortAscending;
-	}
+	ImageConstantsLastSortedColumnID = ColumnID;
+	ImageListSortMode = ColumnSortMode;
 
-	Algo::StableSort(ConstantImageElements, [&](const TSharedPtr<FMutableConstantImageElement>& A, const TSharedPtr<FMutableConstantImageElement>& B)
+	Algo::StableSort(*ConstantImageElements, [&](const TSharedPtr<FMutableConstantImageElement>& A, const TSharedPtr<FMutableConstantImageElement>& B)
 	{
 		// Sort by image id
 		if (ColumnID == ImageConstantTitles::ImageID)
 		{
-			if (bSortAscending)
+			if (ColumnSortMode == EColumnSortMode::Ascending)
 			{
 				return A->IndexOnSourceVector < B->IndexOnSourceVector;
 			}
@@ -880,7 +885,7 @@ void SMutableConstantsWidget::OnImageTableSortRequested(EColumnSortPriority::Typ
 		// Sort Image mip maps
 		if (ColumnID == ImageConstantTitles::ImageMipMaps)
 		{
-			if (bSortAscending)
+			if (ColumnSortMode == EColumnSortMode::Ascending)
 			{
 				return A->ImagePtr->GetLODCount() < B->ImagePtr->GetLODCount();
 			}
@@ -896,7 +901,7 @@ void SMutableConstantsWidget::OnImageTableSortRequested(EColumnSortPriority::Typ
 			const uint8 BImageFormatValue = static_cast<uint8>(B->ImagePtr->GetFormat());
 			const FString BImageFormatString = FString(mu::TypeInfo::s_imageFormatName[BImageFormatValue]);
 			
-			if (bSortAscending)
+			if (ColumnSortMode == EColumnSortMode::Ascending)
 			{
 				return AImageFormatString.Compare(BImageFormatString) < 0;
 			}
@@ -907,7 +912,7 @@ void SMutableConstantsWidget::OnImageTableSortRequested(EColumnSortPriority::Typ
 		// Sort by image size
 		if (ColumnID == ImageConstantTitles::ImageSize)
 		{
-			if (bSortAscending)
+			if (ColumnSortMode == EColumnSortMode::Ascending)
 			{
 				return A->ImagePtr->GetSizeX() *  A->ImagePtr->GetSizeY() <
 					B->ImagePtr->GetSizeX() *  B->ImagePtr->GetSizeY();
@@ -920,7 +925,7 @@ void SMutableConstantsWidget::OnImageTableSortRequested(EColumnSortPriority::Typ
 		// Sort by image used memory
 		if (ColumnID == ImageConstantTitles::ImageTotalMemory)
 		{
-			if (bSortAscending)
+			if (ColumnSortMode == EColumnSortMode::Ascending)
 			{
 				return A->ImagePtr->GetDataSize() < B->ImagePtr->GetDataSize();
 			}
@@ -929,55 +934,108 @@ void SMutableConstantsWidget::OnImageTableSortRequested(EColumnSortPriority::Typ
 		
 		return false;
 	});
-
-	// Store the last column to be able to do the inverse sorting order the next time it gets pressed
-	LastSortedColumnID = ColumnID;
-
-	RegenerateProxyImageArray();
-	ConstantImagesListView->RequestListRefresh();
+	
+	if (ImageListViewHandler)
+	{
+		ImageListViewHandler->RegeneratePage();
+	}
 }
 
 
-FReply SMutableConstantsWidget::OnImageBackButtonClicked()
+EColumnSortMode::Type SMutableConstantsWidget::GetImageListColumnSortMode(FName ColumnName) const
 {
-	CurrentArraySegment--;
-	
-	// Regenerate the proxy list with the elements on the currently selected list segment
-	RegenerateProxyImageArray();
+	if (ImageConstantsLastSortedColumnID != ColumnName)
+	{
+		return EColumnSortMode::None;
+	}
 
-	return FReply::Handled();
+	return ImageListSortMode;
 }
 
-FReply SMutableConstantsWidget::OnImageForwardButtonClicked()
+
+void SMutableConstantsWidget::OnMeshTableSortRequested(EColumnSortPriority::Type ColumnSortPriority, const FName& ColumnID, EColumnSortMode::Type ColumnSortMode)
 {
-	CurrentArraySegment++;
+	// If the colum has been sorted on one way now do the inverse way of sorting
+	MeshConstantsLastSortedColumnID = ColumnID;
+	MeshListSortMode = ColumnSortMode;
+
+	Algo::StableSort(*ConstantMeshElements, [&](const TSharedPtr<FMutableConstantMeshElement>& A, const TSharedPtr<FMutableConstantMeshElement>& B)
+	{
+		// Sort by mesh id
+		if (ColumnID == MeshConstantTitles::MeshID)
+		{
+			if (ColumnSortMode == EColumnSortMode::Ascending)
+			{
+				return A->IndexOnSourceVector < B->IndexOnSourceVector;
+			}
+			return A->IndexOnSourceVector > B->IndexOnSourceVector;
+		}
 	
-	// Regenerate the proxy list with the elements on the currently selected list segment
-	RegenerateProxyImageArray();
+		// Sort by vertex count
+		if (ColumnID == MeshConstantTitles::MeshVertices)
+		{
+			if (ColumnSortMode == EColumnSortMode::Ascending)
+			{
+				return A->MeshPtr->GetVertexCount() < B->MeshPtr->GetVertexCount();
+			}
+			return A->MeshPtr->GetVertexCount() > B->MeshPtr->GetVertexCount();
+		}
 
-	return FReply::Handled();
-}
+		// Sort by index count
+		if (ColumnID == MeshConstantTitles::MeshIndices)
+		{
+			if (ColumnSortMode == EColumnSortMode::Ascending)
+			{
+				return A->MeshPtr->GetIndexCount() < B->MeshPtr->GetIndexCount();
+			}
+			return A->MeshPtr->GetIndexCount() > B->MeshPtr->GetIndexCount();
+		}
 
-FReply SMutableConstantsWidget::OnImageFullBackButtonClicked()
-{ 
-	CurrentArraySegment = 0;
+		// Sort by the amount of channels in the vertex and index buffers
+		if (ColumnID == MeshConstantTitles::MeshChannels)
+		{
+			const int32 AChannelCount = GetMeshChannelCount(A->MeshPtr);
+			const int32 BChannelCount = GetMeshChannelCount(B->MeshPtr);
+			
+			if (ColumnSortMode == EColumnSortMode::Ascending)
+			{
+				return AChannelCount < BChannelCount;
+			}
+			return AChannelCount > BChannelCount;
+		}
+		
+		// Sort by the amount memory used by the mesh
+		if (ColumnID == MeshConstantTitles::MeshMemory)
+		{
+			if (ColumnSortMode == EColumnSortMode::Ascending)
+			{
+				return A->MeshPtr->GetDataSize() < B->MeshPtr->GetDataSize();
+			}
+			return A->MeshPtr->GetDataSize() > B->MeshPtr->GetDataSize();
+		}
 	
-	// Regenerate the proxy list with the elements on the currently selected list segment
-	RegenerateProxyImageArray();
-
-	return FReply::Handled();
+		return false;
+	});
+	
+	if (MeshListViewHandler)
+	{
+		MeshListViewHandler->RegeneratePage();
+	}
 }
 
-FReply SMutableConstantsWidget::OnImageFullForwardButtonClicked()
+
+
+EColumnSortMode::Type SMutableConstantsWidget::GetMeshListColumnSortMode(FName ColumnName) const
 {
-	CurrentArraySegment = TotalAmountOfSegments - 1;
-	
-	// Regenerate the proxy list with the elements on the currently selected list segment
-	RegenerateProxyImageArray();
+	if (MeshConstantsLastSortedColumnID != ColumnName)
+	{
+		return EColumnSortMode::None;
+	}
 
-	return FReply::Handled();
+	return MeshListSortMode;
 }
 
+
 #pragma endregion 
 
 #pragma  region  Expansions Handling
@@ -1121,7 +1179,15 @@ void SMutableConstantsWidget::LoadConstantImages()
 	check (MutableProgramPtr);
 
 	const int32 ConstantsCount = MutableProgramPtr->ConstantImages.Num();
-	ConstantImageElements.Empty(ConstantsCount);
+	if (ConstantImageElements)
+	{
+		ConstantImageElements->Empty(ConstantsCount);
+	}
+	else
+	{
+		ConstantImageElements = MakeShared<TArray<TSharedPtr<FMutableConstantImageElement>>>();
+		ConstantImageElements->Reserve(ConstantsCount);
+	}
 	
 	uint64 ConstantImagesAccumulatedSize = 0;
 
@@ -1136,24 +1202,11 @@ void SMutableConstantsWidget::LoadConstantImages()
 		
 		ConstantImagesAccumulatedSize += ConstantImageElement->ImagePtr->GetDataSize();
 		
-		ConstantImageElements.Add(ConstantImageElement);
+		ConstantImageElements->Add(ConstantImageElement);
 	}
 
 	// Cache the size in memory of the constants as a formatted text so it is able to be be later used by the UI
 	ConstantImagesFormattedSize = GenerateTextForSize(ConstantImagesAccumulatedSize);
-	
-	// Regenerate the array with only a part of the the full images array
-	RegenerateProxyImageArray();
-
-	// Compute the total amount of segments
-	{
-		// Compute the max amount of segments in case it needs to be updated
-		TotalAmountOfSegments = ConstantImageElements.Num() / ElementsPerSegment;
-		
-		// Add one more segment if there are elements not filling a segment
-		TotalAmountOfSegments = ConstantImageElements.Num() % ElementsPerSegment > 0 ? TotalAmountOfSegments + 1 : TotalAmountOfSegments;
-	}
-
 }
 
 void SMutableConstantsWidget::LoadConstantMeshes()
@@ -1161,7 +1214,15 @@ void SMutableConstantsWidget::LoadConstantMeshes()
 	check (MutableProgramPtr);
 	
 	const int32 ConstantsCount = MutableProgramPtr->ConstantMeshes.Num();
-	ConstantMeshElements.Empty(ConstantsCount);
+	if (ConstantMeshElements)
+	{
+		ConstantMeshElements->Empty(ConstantsCount);
+	}
+	else
+	{
+		ConstantMeshElements = MakeShared<TArray<TSharedPtr<FMutableConstantMeshElement>>>();
+		ConstantMeshElements->Reserve(ConstantsCount);
+	}
 	
 	uint64 ConstantMeshesAccumulatedSize = 0;
 	
@@ -1170,8 +1231,6 @@ void SMutableConstantsWidget::LoadConstantMeshes()
 		TSharedPtr< FMutableConstantMeshElement> ConstantMeshElement = MakeShared<FMutableConstantMeshElement>();
 		ConstantMeshElement->MeshPtr = MutableProgramPtr->ConstantMeshes[MeshIndex].Value;
 		ConstantMeshElement->IndexOnSourceVector = MeshIndex;
-		
-		uint64 ThisMeshSize = ConstantMeshElement->MeshPtr->GetDataSize();
 
 		// Actual core disk size would be:
 		//mu::OutputMemoryStream Stream;
@@ -1179,9 +1238,10 @@ void SMutableConstantsWidget::LoadConstantMeshes()
 		//ConstantMeshElement->MeshPtr->Serialise(Archive);
 		//uint64 OtherSize = Stream.GetBufferSize();
 
+		const uint64 ThisMeshSize = ConstantMeshElement->MeshPtr->GetDataSize();
 		ConstantMeshesAccumulatedSize += ThisMeshSize;
 		
-		ConstantMeshElements.Add(ConstantMeshElement);
+		ConstantMeshElements->Add(ConstantMeshElement);
 	}
 
 	// Cache the size in memory of the constants as a formatted text so it is able to be be later used by the UI
@@ -1343,7 +1403,7 @@ void SMutableConstantsWidget::ClearSelectedConstantItems(mu::DATATYPE ExceptionD
 {
 	if (ExceptionDataType != mu::DATATYPE::DT_MESH)
 	{
-		ConstantMeshesSlate->ClearSelection();
+		MeshListViewHandler->ClearSelection();
 	}
 	if (ExceptionDataType != mu::DATATYPE::DT_STRING)
 	{
@@ -1375,7 +1435,8 @@ void SMutableConstantsWidget::ClearSelectedConstantItems(mu::DATATYPE ExceptionD
 	}
 	if (ExceptionDataType != mu::DATATYPE::DT_IMAGE)
 	{
-		ConstantImagesListView->ClearSelection();
+		ImageListViewHandler->ClearSelection();
+		//ConstantImagesListView->ClearSelection();
 	}
 
 	// todo: if you add more slates for the types update this to make sure they behave as the others
@@ -1539,12 +1600,12 @@ FText SMutableConstantsWidget::OnDrawStringsAreaTitle() const
 
 FText SMutableConstantsWidget::OnDrawImagesAreaTitle() const
 {
-	return FText::Format( LOCTEXT("ImageConstantsTitle", "Image Constants ({0}) : {1} "), ConstantImageElements.Num(), ConstantImagesFormattedSize);
+	return FText::Format( LOCTEXT("ImageConstantsTitle", "Image Constants ({0}) : {1} "), ConstantImageElements->Num(), ConstantImagesFormattedSize);
 }
 
 FText SMutableConstantsWidget::OnDrawMeshesAreaTitle() const
 {
-	return FText::Format(LOCTEXT("MeshConstantsTitle", "Mesh Constants ({0}) : {1} "), ConstantMeshElements.Num(), ConstantMeshesFormattedSize) ;
+	return FText::Format(LOCTEXT("MeshConstantsTitle", "Mesh Constants ({0}) : {1} "), ConstantMeshElements->Num(), ConstantMeshesFormattedSize) ;
 }
 
 FText SMutableConstantsWidget::OnDrawLayoutsAreaTitle() const
diff --git a/Engine/Plugins/Experimental/Mutable/Source/CustomizableObjectEditor/Private/MuCOE/SMutableConstantsWidget.h b/Engine/Plugins/Experimental/Mutable/Source/CustomizableObjectEditor/Private/MuCOE/SMutableConstantsWidget.h
--- a/Engine/Plugins/Experimental/Mutable/Source/CustomizableObjectEditor/Private/MuCOE/SMutableConstantsWidget.h
+++ b/Engine/Plugins/Experimental/Mutable/Source/CustomizableObjectEditor/Private/MuCOE/SMutableConstantsWidget.h
@@ -9,11 +9,11 @@
 #include "Widgets/Views/SHeaderRow.h"
 #include "Widgets/Views/STileView.h"
 
-template <typename ItemType> class SListView;
 
+// Forward declarations
+template <typename ItemType> class SMutableMultiPageListView;
 class ITableRow;
 class SExpandableArea;
-// Forward declarations
 class SMutableCodeViewer;
 class STableViewBase;
 class SWidget;
@@ -158,8 +158,12 @@ private:
 	* Data backend for the lists of constants
 	*/
 
-	TArray<TSharedPtr<FMutableConstantImageElement>> ConstantImageElements;
-	TArray<TSharedPtr<FMutableConstantMeshElement>> ConstantMeshElements;
+	TSharedPtr<SMutableMultiPageListView<TSharedPtr<FMutableConstantImageElement>>> ImageListViewHandler;
+	TSharedPtr<TArray<TSharedPtr<FMutableConstantImageElement>>> ConstantImageElements;
+
+	TSharedPtr<SMutableMultiPageListView<TSharedPtr<FMutableConstantMeshElement>>> MeshListViewHandler;
+	TSharedPtr<TArray<TSharedPtr<FMutableConstantMeshElement>>> ConstantMeshElements;
+	
 	TArray<TSharedPtr<FMutableConstantStringElement>> ConstantStringElements;
 	TArray<TSharedPtr<FMutableConstantLayoutElement>> ConstantLayoutElements;
 	TArray<TSharedPtr<FMutableConstantProjectorElement>> ConstantProjectorElements;
@@ -211,68 +215,31 @@ private:
 	TSharedRef<ITableRow> OnGenerateCurveRow(TSharedPtr<FMutableConstantCurveElement> MutableConstantCurveElement, const TSharedRef<STableViewBase>& OwnerTable) const;
 	TSharedRef<ITableRow> OnGenerateSkeletonRow(TSharedPtr<FMutableConstantSkeletonElement> MutableConstantSkeletonElement, const TSharedRef<STableViewBase>& OwnerTable) const;
 	
-	/*
-	 * Image segment operation 
-	 */
-
-	// Part of the original array used do display only a part of the data
-	TArray<TSharedPtr<FMutableConstantImageElement>> ConstantImageElementsSegment;
-
-	/** Amount of elements an array segment can hold. The lowest the more stable the UI drawing gets
-	*(avoid crash due to oversize slate Y size) */
-	const uint32 ElementsPerSegment = 24;
-
-	// Amount of elements a segment can hold. The lowest the more stable the UI drawing gets (avoid crash due to oversize slate Y coordinate)
-	uint32 CurrentArraySegment = 0;
-	uint32 TotalAmountOfSegments = 0;
-
-	/** Provided the array on ConstantImageElements a new sub array gets generated for the currently selected segment */
-	void RegenerateProxyImageArray();
-	
-	/** Method designed to generate the UI controls to let the user change the image array segment being previewed */
-	TSharedRef<SWidget> GenerateImagesSegmentSelectionWidget();
-
-	/** Returns the text containing the page (segment) being drawn at the moment. Designed for the UI to show data to
-	 * the user
-	 */
-	FText OnDrawCurrentImageSegmentText() const;
-
-	/** Determines if the button that loads the previous image array segment should be interactable by the user
-	 * It only relates to UI, there are other internal checks to avoid the processing of an existent segment
-	 */
-	bool ShouldBackButtonBeEnabled() const;
-	/** Determines if the button that loads the next image array segment should be interactable by the user
-	 * It only relates to UI, there are other internal checks to avoid the processing of an existent segment
-	 */
-	bool ShouldNextButtonBeEnabled() const;
-
-	/** Action performed when the back button is pressed on the UI. Sets the current segment to the previous one*/
-	FReply OnImageBackButtonClicked();
-	/** Action performed when the forward button is pressed on the UI. Sets the current segment to the next one*/
-	FReply OnImageForwardButtonClicked();
-
-	/** Action performed when the full back button is pressed on the UI. Sets the current segment to the first one*/
-	FReply OnImageFullBackButtonClicked();
-	/** Action performed when the full forward button is pressed on the UI. Sets the current segment to the last one*/
-	FReply OnImageFullForwardButtonClicked();
-
 	
 	/*
-	 * Image Table sorting methods
+	 * Image List sorting methods
 	 */
-
-	/** List view showing the constant images data*/
-	TSharedPtr<SListView<TSharedPtr<FMutableConstantImageElement>>> ConstantImagesListView;
-
+	
 	/** Id of the last column the user decided to Sort. Usefully in order to interpolate ascending and descending sorting
 	 * order 
 	 */
-	FName LastSortedColumnID = "";
+	FName ImageConstantsLastSortedColumnID = "";
 	/** Variable holding what kind of sorting has been used on the last sorting operation*/
-	bool bSortAscending = false;
-
+	EColumnSortMode::Type ImageListSortMode = EColumnSortMode::Type::None;
+	
 	/** Callback method designed to sort the list of images. It sorts ConstantImageElements */
 	void OnImageTableSortRequested(EColumnSortPriority::Type ColumnSortPriority, const FName& ColumnID, EColumnSortMode::Type ColumnSortMode);
+	EColumnSortMode::Type GetImageListColumnSortMode(FName ColumnName) const;
+
+	/*
+	 * Mesh List sorting methods
+	 */
+	
+	FName MeshConstantsLastSortedColumnID = "";
+	EColumnSortMode::Type MeshListSortMode = EColumnSortMode::Type::None;
+	
+	void OnMeshTableSortRequested(EColumnSortPriority::Type ColumnSortPriority, const FName& ColumnID, EColumnSortMode::Type ColumnSortMode);
+	EColumnSortMode::Type GetMeshListColumnSortMode(FName ColumnName) const;
 
 	/*
 	 * Expandable areas objects and behaviours
diff --git a/Engine/Plugins/Experimental/Mutable/Source/CustomizableObjectEditor/Private/MuCOE/Widgets/MutableMultiPageListView.h b/Engine/Plugins/Experimental/Mutable/Source/CustomizableObjectEditor/Private/MuCOE/Widgets/MutableMultiPageListView.h
new file mode 100644
--- /dev/null
+++ b/Engine/Plugins/Experimental/Mutable/Source/CustomizableObjectEditor/Private/MuCOE/Widgets/MutableMultiPageListView.h
@@ -0,0 +1,332 @@
+﻿// Copyright Epic Games, Inc. All Rights Reserved.
+#pragma once
+
+#include "Widgets/SCompoundWidget.h"
+#include "Widgets/Views/SListView.h"
+#include "Widgets/SBoxPanel.h"
+#include "Widgets/Input/SButton.h"
+#include "Widgets/Text/STextBlock.h"
+
+#define LOCTEXT_NAMESPACE "CustomizableObjectEditor"
+
+/**
+ * Slate containing a list view designed to contain high amounts of elements. Will separate the elements in pages so the parent
+ * UI element drawing them can get the listview to fit without triggering any hard limit in the size of the containing slate object.
+ * This class was created out of the necessity of having lists with hundreds of elements (mutable constant images for example).
+ * 
+ * @tparam ItemType The type of item the list view is holding
+ */
+template <typename ItemType>
+class SMutableMultiPageListView final : public SCompoundWidget
+{
+public:
+	
+	SLATE_BEGIN_ARGS(SMutableMultiPageListView<ItemType>)
+			: _HostedListView(nullptr)
+			, _ElementsToSeparateInPages(nullptr) {}
+
+	    /** The list view we want to manipulate */
+    	SLATE_ARGUMENT(TSharedPtr<SListView<ItemType>>, HostedListView)
+
+	    /** The elements that got fed to the list view we are hosting */ 
+	    SLATE_ARGUMENT(TSharedPtr<const TArray<ItemType>>, ElementsToSeparateInPages)
+	    
+    SLATE_END_ARGS()
+	
+		
+	void Construct(const FArguments& InArgs);
+
+private:
+	
+	/**
+	 * Get the visibility state set for the navigation buttons box. If there is no need for the box existence it will then be collapsed
+	 * @return The visibility state to use in the navigation panel
+	 */
+	EVisibility GetPageNavigationBoxVisibility() const;
+
+	/**
+	 * Changes the page to the previous one
+	 * @return the reply status to know if the operation was handled or not.
+	 */
+	FReply OnBackButtonClicked();
+
+	/**
+	 * Changes the page to the next one
+	 * @return the reply status to know if the operation was handled or not.
+	 */
+	FReply OnForwardButtonClicked();
+
+	/**
+	 * Sets the page to the first one
+	 * @return the reply status to know if the operation was handled or not.
+	 */
+	FReply OnFullBackButtonClicked();
+
+	/**
+	 * Sets the page to the last one
+	 * @return the reply status to know if the operation was handled or not.
+	 */
+	FReply OnFullForwardButtonClicked();
+
+	/**
+	 * Controls if the back button should or should not be enabled
+	 * @return True if it should, false otherwise.
+	 */
+	bool ShouldBackButtonBeEnabled() const;
+
+	/**
+	 * Controls if the forward button should or should not be enabled
+	 * @return True if it should, false otherwise.
+	 */
+	bool ShouldNextButtonBeEnabled() const;
+
+	/**
+	 * Handles the drawing of the current / total pages UI text.
+	 * @return The text representing the currently displayed page
+	 */
+	FText OnDrawCurrentPageText() const;
+
+public:
+
+	/** Compute the elements to display for the current page */
+	void RegeneratePage();
+
+	/** Clear the selected item from the handled list view */
+	void ClearSelection();
+	
+private:
+	
+	/**
+	 * Amount of elements an array segment can hold. The lowest the more stable the UI drawing gets
+	 * (avoids crash due to oversize slate Y size)
+	 */
+	const uint32 ElementsPerPage = 24;
+
+	/** The index of the current page/segment. Will increase and decrease using the UI */
+	uint32 CurrentSegmentIndex = 0;
+
+	/**
+	 * The computed total amount of segments. It is based in the amount of elements provided in the AllElements array and the
+	 * value set in ElementsPerPage
+	 */
+	uint32 TotalAmountOfPages = 0;
+
+	/** Elements being displayed in the current page */
+	TArray<ItemType> CurrentSegmentElements;
+
+	/** List view handled by this slate. */
+	TSharedPtr<SListView<ItemType>> HostedListView;
+
+	/** Array with all the elements to manage from the provided ListView*/
+	TSharedPtr<const TArray<ItemType>> AllElements;
+};
+
+
+
+template <typename ItemType>
+void SMutableMultiPageListView<ItemType>::Construct(const FArguments& InArgs)
+{
+	HostedListView = InArgs._HostedListView;
+	AllElements = InArgs._ElementsToSeparateInPages;
+
+	check(HostedListView);
+	check(AllElements);
+
+	// Compute the max amount of pages in case it needs to be updated
+	TotalAmountOfPages = AllElements->Num() / ElementsPerPage;
+	// Add one more segment if there are elements not filling a full page
+	TotalAmountOfPages = AllElements->Num() % ElementsPerPage > 0 ? TotalAmountOfPages + 1 : TotalAmountOfPages;
+	
+	this->ChildSlot
+	[
+		SNew(SVerticalBox)
+
+		// Buttons being used to change the segment to be displayed
+		+SVerticalBox::Slot()
+		.AutoHeight()
+		.Padding(2,0)
+		[
+			SNew(SHorizontalBox)
+			.Visibility(this, &SMutableMultiPageListView::GetPageNavigationBoxVisibility)
+
+			+ SHorizontalBox::Slot()
+			.HAlign(EHorizontalAlignment::HAlign_Left)
+			.AutoWidth()
+			[
+				SNew(SButton)
+				.Text(FText(INVTEXT("|<")))
+				.OnClicked(this,&SMutableMultiPageListView::OnFullBackButtonClicked)
+				.IsEnabled(this,&SMutableMultiPageListView::ShouldBackButtonBeEnabled)
+			]
+
+			+ SHorizontalBox::Slot()
+			.HAlign(EHorizontalAlignment::HAlign_Left)
+			.AutoWidth()
+			[
+				SNew(SButton)
+				.Text(FText(INVTEXT("<")))
+				.OnClicked(this,&SMutableMultiPageListView::OnBackButtonClicked)
+				.IsEnabled(this,&SMutableMultiPageListView::ShouldBackButtonBeEnabled)
+			]
+			
+			+ SHorizontalBox::Slot()
+			.HAlign(EHorizontalAlignment::HAlign_Center)
+			[
+				SNew(STextBlock)
+				.Justification(ETextJustify::Center)
+				.Text(this,&SMutableMultiPageListView::OnDrawCurrentPageText)
+			]
+
+			+ SHorizontalBox::Slot()
+			.HAlign(EHorizontalAlignment::HAlign_Right)
+			.AutoWidth()
+			[
+				SNew(SButton)
+				.Text(FText(INVTEXT(">")))
+				.OnClicked(this,&SMutableMultiPageListView::OnForwardButtonClicked)
+				.IsEnabled(this,&SMutableMultiPageListView::ShouldNextButtonBeEnabled)
+			]
+
+			+ SHorizontalBox::Slot()
+			.HAlign(EHorizontalAlignment::HAlign_Right)
+			.AutoWidth()
+			[
+				SNew(SButton)
+				.Text(FText(INVTEXT(">|")))
+				.OnClicked(this,&SMutableMultiPageListView::OnFullForwardButtonClicked)
+				.IsEnabled(this,&SMutableMultiPageListView::ShouldNextButtonBeEnabled)
+			]
+		]
+		
+		// List view showing the selected segment of elements
+		+ SVerticalBox::Slot()
+		.Padding(2,5)
+		[
+			HostedListView.ToSharedRef()
+		]
+	];
+	
+	// Generate the array to be used as first displaying elements. 
+	RegeneratePage();
+	
+	// Tell the list to use our set of elements
+	//  note: If this code throws a runtime error may be because no children slot has been setup in the ListView. That could happen if
+	//  the list view provided did not get fully setup. For example, if no ItemsSource is defined no Child slot will be created and, therefore,
+	//  the SetItemsSource method will fail to run.
+	HostedListView->SetItemsSource(&CurrentSegmentElements);
+}
+
+
+template <typename ItemType>
+EVisibility SMutableMultiPageListView<ItemType>::GetPageNavigationBoxVisibility() const
+{
+	return TotalAmountOfPages > 0 ? EVisibility::Visible : EVisibility::Collapsed;
+}
+
+
+template <typename ItemType>
+FReply SMutableMultiPageListView<ItemType>::OnBackButtonClicked()
+{
+	CurrentSegmentIndex--;
+	
+	// Regenerate the proxy list with the elements on the currently selected list segment
+	RegeneratePage();
+
+	return FReply::Handled();
+}
+
+
+template <typename ItemType>
+FReply SMutableMultiPageListView<ItemType>::OnForwardButtonClicked()
+{
+	CurrentSegmentIndex++;
+	
+	// Regenerate the proxy list with the elements on the currently selected list segment
+	RegeneratePage();
+
+	return FReply::Handled();
+}
+
+
+template <typename ItemType>
+FReply SMutableMultiPageListView<ItemType>::OnFullBackButtonClicked()
+{ 
+	CurrentSegmentIndex = 0;
+	
+	// Regenerate the proxy list with the elements on the currently selected list segment
+	RegeneratePage();
+
+	return FReply::Handled();
+}
+
+
+template <typename ItemType>
+FReply SMutableMultiPageListView<ItemType>::OnFullForwardButtonClicked()
+{
+	CurrentSegmentIndex = TotalAmountOfPages - 1;
+	
+	// Regenerate the proxy list with the elements on the currently selected list segment
+	RegeneratePage();
+
+	return FReply::Handled();
+}
+
+
+template <typename ItemType>
+bool SMutableMultiPageListView<ItemType>::ShouldBackButtonBeEnabled() const
+{
+	return  CurrentSegmentIndex > 0;
+}
+
+
+template <typename ItemType>
+bool SMutableMultiPageListView<ItemType>::ShouldNextButtonBeEnabled() const
+{
+	return TotalAmountOfPages >= 1 && CurrentSegmentIndex < TotalAmountOfPages - 1;
+}
+
+
+template <typename ItemType>
+FText SMutableMultiPageListView<ItemType>::OnDrawCurrentPageText() const
+{
+	FString BuiltText = FString::FromInt(  CurrentSegmentIndex + 1);
+	BuiltText += FString(" / ");
+	BuiltText += FString::FromInt(TotalAmountOfPages);
+	
+	return FText::Format(LOCTEXT("PageNumber","Page : {0}"), FText::FromString(BuiltText));
+}
+
+
+template <typename ItemType>
+void SMutableMultiPageListView<ItemType>::RegeneratePage()
+{
+	// Prepare the segment of elements to be filled up
+	CurrentSegmentElements.SetNum(0);
+	CurrentSegmentElements.Reserve(ElementsPerPage);
+	
+	const uint32 TotalElements = AllElements->Num();
+	const uint32 StartingIndex = CurrentSegmentIndex * ElementsPerPage;
+	const uint32 FinishIndex = StartingIndex + FMath::Min(ElementsPerPage, (TotalElements - StartingIndex));
+	// Iterate over the main array from the point where the segment starts to the point where it ends. Stop before if needed
+	for	(uint32 MainArrayIndex = StartingIndex ; MainArrayIndex < FinishIndex; ++ MainArrayIndex)
+	{
+		CurrentSegmentElements.Add(AllElements.Get()->GetData()[MainArrayIndex]);
+	}
+	
+	if (HostedListView)
+	{
+		HostedListView->RequestListRefresh();
+	}
+}
+
+
+template <typename ItemType>
+void SMutableMultiPageListView<ItemType>::ClearSelection()
+{
+	if (HostedListView)
+	{
+		HostedListView->ClearSelection();
+	}
+}
+
+#undef LOCTEXT_NAMESPACE
