[Mutable] Reduce bone pose precision error when converting scaled bone transforms to mutable.

#rnx

diff --git a/Engine/Plugins/Experimental/Mutable/Source/CustomizableObjectEditor/Private/MuCOE/GenerateMutableSource/GenerateMutableSourceMesh.cpp b/Engine/Plugins/Experimental/Mutable/Source/CustomizableObjectEditor/Private/MuCOE/GenerateMutableSource/GenerateMutableSourceMesh.cpp
--- a/Engine/Plugins/Experimental/Mutable/Source/CustomizableObjectEditor/Private/MuCOE/GenerateMutableSource/GenerateMutableSourceMesh.cpp
+++ b/Engine/Plugins/Experimental/Mutable/Source/CustomizableObjectEditor/Private/MuCOE/GenerateMutableSource/GenerateMutableSourceMesh.cpp
@@ -886,6 +886,9 @@ mu::MeshPtr ConvertSkeletalMeshToMutable(const USkeletalMesh* InSkeletalMesh, co
 		TArray<mu::FBoneName> MutableBoneMap;
 		MutableBoneMap.SetNum(BoneMap.Num());
 		
+		TArray<FMatrix> ComposedRefPoseMatrices;
+		ComposedRefPoseMatrices.SetNum(NumRequiredBones);
+
 		const TArray<FMeshBoneInfo>& RefBoneInfo = InSkeletalMesh->GetRefSkeleton().GetRefBoneInfo();
 		for (int32 BoneIndex = 0; BoneIndex < NumRequiredBones; ++BoneIndex)
 		{
@@ -910,16 +913,24 @@ mu::MeshPtr ConvertSkeletalMeshToMutable(const USkeletalMesh* InSkeletalMesh, co
 				MutableBoneMap[BoneMapIndex] = BoneName;
 			}
 
+			if (ParentBoneIndex >= 0)
+			{
+				ComposedRefPoseMatrices[BoneIndex] = InSkeletalMesh->GetRefPoseMatrix(RefSkeletonBoneIndex) * ComposedRefPoseMatrices[ParentBoneIndex];
+			}
+			else
+			{
+				ComposedRefPoseMatrices[BoneIndex] = InSkeletalMesh->GetRefPoseMatrix(RefSkeletonBoneIndex);
+			}
+
 			// Set bone pose
-			FMatrix44f BaseInvMatrix = InSkeletalMesh->GetRefBasesInvMatrix()[RefSkeletonBoneIndex];
-			FTransform3f BaseInvTransform;
-			BaseInvTransform.SetFromMatrix(BaseInvMatrix);
+			FTransform3f BoneTransform;
+			BoneTransform.SetFromMatrix(FMatrix44f(ComposedRefPoseMatrices[BoneIndex]));
 
 			mu::EBoneUsageFlags BoneUsageFlags = mu::EBoneUsageFlags::None;
 			EnumAddFlags(BoneUsageFlags, BoneMapIndex != INDEX_NONE ? mu::EBoneUsageFlags::Skinning : mu::EBoneUsageFlags::None);
 			EnumAddFlags(BoneUsageFlags, ParentBoneIndex == INDEX_NONE ? mu::EBoneUsageFlags::Root : mu::EBoneUsageFlags::None);
 
-			MutableMesh->SetBonePose(BoneIndex, BoneName, BaseInvTransform.Inverse(), BoneUsageFlags);
+			MutableMesh->SetBonePose(BoneIndex, BoneName, BoneTransform, BoneUsageFlags);
 		}
 
 		MutableMesh->SetBoneMap(MutableBoneMap);
