[Mutable] Fix deadlock when compiling an object if there is an ongoing instance update.

#rnx

diff --git a/Engine/Plugins/Experimental/Mutable/Source/CustomizableObject/Internal/MuCO/ICustomizableObjectEditorModule.h b/Engine/Plugins/Experimental/Mutable/Source/CustomizableObject/Internal/MuCO/ICustomizableObjectEditorModule.h
--- a/Engine/Plugins/Experimental/Mutable/Source/CustomizableObject/Internal/MuCO/ICustomizableObjectEditorModule.h
+++ b/Engine/Plugins/Experimental/Mutable/Source/CustomizableObject/Internal/MuCO/ICustomizableObjectEditorModule.h
@@ -77,5 +77,8 @@ public:
 	/** Force finish current compile request and cancels all pending requests */
 	virtual void CancelCompileRequests() = 0;
 
+	/** Return the number of pending compilation requests. Ongoing requests included. */
+	virtual int32 GetNumCompileRequests() = 0;
+
 };
  
diff --git a/Engine/Plugins/Experimental/Mutable/Source/CustomizableObject/Private/MuCO/CustomizableObjectSystem.cpp b/Engine/Plugins/Experimental/Mutable/Source/CustomizableObject/Private/MuCO/CustomizableObjectSystem.cpp
--- a/Engine/Plugins/Experimental/Mutable/Source/CustomizableObject/Private/MuCO/CustomizableObjectSystem.cpp
+++ b/Engine/Plugins/Experimental/Mutable/Source/CustomizableObject/Private/MuCO/CustomizableObjectSystem.cpp
@@ -3607,8 +3607,14 @@ int32 UCustomizableObjectSystem::TickInternal(const bool bBlocking)
 
 	FMutableUpdateCandidate* LODUpdateCandidateFound = nullptr;
 	
+	bool bPendingCompilation = false;
+#if WITH_EDITOR
+	ICustomizableObjectEditorModule* EditorModule = ICustomizableObjectEditorModule::Get();
+	bPendingCompilation = EditorModule && EditorModule->GetNumCompileRequests() > 0;
+#endif
+
 	// Get a new operation if we aren't working on one
-	if (!Private->CurrentMutableOperation && bIsMutableEnabled)
+	if (!Private->CurrentMutableOperation && bIsMutableEnabled && !bPendingCompilation)
 	{
 		// Reset the instance relevancy
 		// The RequestedUpdates only refer to LOD changes. User Customization and discards are handled separately
@@ -3782,9 +3788,6 @@ int32 UCustomizableObjectSystem::TickInternal(const bool bBlocking)
 	{
 		AdvanceCurrentOperation();
 	}
-
-#if WITH_EDITOR
-#endif
 	
 	const int32 RemainingTasks = Private->MutableTaskGraph.Tick();
 
@@ -3812,7 +3815,6 @@ int32 UCustomizableObjectSystem::TickInternal(const bool bBlocking)
 	if (bBlocking)
 	{
 #if WITH_EDITOR
-		ICustomizableObjectEditorModule* EditorModule = ICustomizableObjectEditorModule::Get();
 		RemainingWork += EditorModule ? EditorModule->Tick(true) : 0;
 #endif
 		
diff --git a/Engine/Plugins/Experimental/Mutable/Source/CustomizableObjectEditor/Private/MuCOE/CustomizableObjectCompiler.cpp b/Engine/Plugins/Experimental/Mutable/Source/CustomizableObjectEditor/Private/MuCOE/CustomizableObjectCompiler.cpp
--- a/Engine/Plugins/Experimental/Mutable/Source/CustomizableObjectEditor/Private/MuCOE/CustomizableObjectCompiler.cpp
+++ b/Engine/Plugins/Experimental/Mutable/Source/CustomizableObjectEditor/Private/MuCOE/CustomizableObjectCompiler.cpp
@@ -187,6 +187,12 @@ void FCustomizableObjectCompiler::Compile(const TSharedRef<FCompilationRequest>&
 
 	UCustomizableObjectSystem* System = UCustomizableObjectSystem::GetInstanceChecked();
 
+	if (!CurrentRequest->IsAsyncCompilation())
+	{
+		// Sync compilation. Force finish all pending updates and async compilations
+		System->GetPrivate()->BlockTillAllRequestsFinished();
+	}
+
 	check(!CurrentObject->GetPrivate()->IsLocked());
 
 	// Lock object during asynchronous asset loading to avoid instance/mip updates and reentrant compilations
@@ -1321,6 +1327,12 @@ bool FCustomizableObjectCompiler::TryPopCompileRequest()
 		return false;
 	}
 
+	UCustomizableObjectSystemPrivate* SystemPrivate = UCustomizableObjectSystem::GetInstance()->GetPrivate();
+	if (SystemPrivate->CurrentMutableOperation)
+	{
+		return false;
+	}
+
 	Compile(CompileRequests.Pop());
 	return true;
 }
diff --git a/Engine/Plugins/Experimental/Mutable/Source/CustomizableObjectEditor/Private/MuCOE/CustomizableObjectEditorModule.cpp b/Engine/Plugins/Experimental/Mutable/Source/CustomizableObjectEditor/Private/MuCOE/CustomizableObjectEditorModule.cpp
--- a/Engine/Plugins/Experimental/Mutable/Source/CustomizableObjectEditor/Private/MuCOE/CustomizableObjectEditorModule.cpp
+++ b/Engine/Plugins/Experimental/Mutable/Source/CustomizableObjectEditor/Private/MuCOE/CustomizableObjectEditorModule.cpp
@@ -748,6 +748,12 @@ void FCustomizableObjectEditorModule::CancelCompileRequests()
 }
 
 
+int32 FCustomizableObjectEditorModule::GetNumCompileRequests()
+{
+	return Compiler.GetNumRemainingWork();
+}
+
+
 void FCustomizableObjectEditorModule::OnPreBeginPIE(const bool bIsSimulatingInEditor)
 {
 	if (IsRunningGame() || !UCustomizableObjectSystem::IsActive())
diff --git a/Engine/Plugins/Experimental/Mutable/Source/CustomizableObjectEditor/Private/MuCOE/CustomizableObjectEditorModule.h b/Engine/Plugins/Experimental/Mutable/Source/CustomizableObjectEditor/Private/MuCOE/CustomizableObjectEditorModule.h
--- a/Engine/Plugins/Experimental/Mutable/Source/CustomizableObjectEditor/Private/MuCOE/CustomizableObjectEditorModule.h
+++ b/Engine/Plugins/Experimental/Mutable/Source/CustomizableObjectEditor/Private/MuCOE/CustomizableObjectEditorModule.h
@@ -53,6 +53,8 @@ public:
 
 	virtual void CancelCompileRequests() override;
 
+	virtual int32 GetNumCompileRequests() override;
+
 	virtual TSharedPtr<FExtensibilityManager> GetCustomizableObjectEditorToolBarExtensibilityManager() override { return CustomizableObjectEditor_ToolBarExtensibilityManager; }
 	virtual TSharedPtr<FExtensibilityManager> GetCustomizableObjectEditorMenuExtensibilityManager() override { return CustomizableObjectEditor_MenuExtensibilityManager; }
 
