[Mutable] Allow the use of FByteBulkData to store and stream resources in cooked builds.
- Move the UCustomizableObjectBulk methods that generate and serialize streamable resources to separate functions from the MutablePrivate namespace.
- Add TSharedPtr to a struct containing all streamable bulk info to avoid unnecessary copies when streaming files from non-game thread threads.
- Add CVarMutableUseBulkData to use FByteBulkData in packaged builds.

#rnx

diff --git a/Engine/Plugins/Experimental/Mutable/Source/CustomizableObject/Internal/MuCO/CustomizableObjectPrivate.h b/Engine/Plugins/Experimental/Mutable/Source/CustomizableObject/Internal/MuCO/CustomizableObjectPrivate.h
--- a/Engine/Plugins/Experimental/Mutable/Source/CustomizableObject/Internal/MuCO/CustomizableObjectPrivate.h
+++ b/Engine/Plugins/Experimental/Mutable/Source/CustomizableObject/Internal/MuCO/CustomizableObjectPrivate.h
@@ -7,6 +7,7 @@
 #include "MuCO/StateMachine.h"
 #include "MuCO/CustomizableObjectUIData.h"
 #include "MuCO/CustomizableObjectIdentifier.h"
+#include "Serialization/BulkData.h"
 #include "Templates/SharedPointer.h"
 #include "UObject/WeakObjectPtr.h"
 #include "UObject/SoftObjectPtr.h"
@@ -34,7 +35,8 @@ class UAnimInstance;
 class UAssetUserData;
 class UCustomizableObject;
 class USkeletalMeshLODSettings;
-
+struct FModelResources;
+struct FModelStreamableBulkData;
 
 FGuid CUSTOMIZABLEOBJECT_API GenerateIdentifier(const UCustomizableObject& CustomizableObject);
 
@@ -691,6 +693,180 @@ struct CUSTOMIZABLEOBJECT_API FMutableStateData
 };
 
 
+#if WITH_EDITOR
+namespace MutablePrivate
+{
+	enum class EDataType : uint32 // uint32 for padding and DDC purposes
+	{
+		None = 0,
+		Model,
+		RealTimeMorph,
+		Clothing,
+	};
+
+
+	struct FBlock
+	{
+		/** Data Type*/
+		EDataType DataType;
+
+		/** Used on some data types as the index to the block stored in the CustomizableObject */
+		uint32 Id;
+
+		/** Size of the data block. */
+		uint32 Size;
+
+		/** Data flags, like "high-res". */
+		uint32 Flags;
+
+		/** Offset in the full source streamed data file that is created when compiling. */
+		uint64 Offset;
+
+		friend FArchive& operator<<(FArchive& Ar, FBlock& Data)
+		{
+			Ar << Data.DataType;
+			Ar << Data.Id;
+			Ar << Data.Size;
+			Ar << Data.Flags;
+			Ar << Data.Offset;
+			return Ar;
+		};
+	};
+	template<> struct TCanBulkSerialize<FBlock> { enum { Value = true }; };
+
+	struct CUSTOMIZABLEOBJECT_API FFile
+	{
+		EDataType DataType = EDataType::None;
+
+		uint32 Padding = 0;
+
+		/** Id generated from a hash of the file content + offset to avoid collisions. */
+		uint32 Id;
+
+		/** Common flags of the data stored in this file. See mu::ERomFlags. */
+		uint32 Flags = 0;
+
+		/** List of blocks that are contained in the file, in order. */
+		TArray<FBlock> Blocks;
+
+		/** Get the total size of blocks in this file. */
+		int64 GetSize() const;
+
+		/** Copy the requested block to the requested buffer and return its size. */
+		void GetFileData(struct FMutableCachedPlatformData*, TArray64<uint8>& DataDestination, bool bDropData);
+
+		friend FArchive& operator<<(FArchive& Ar, FFile& Data)
+		{
+			Ar << Data.DataType;
+			Ar << Data.Padding;
+			Ar << Data.Id;
+			Ar << Data.Flags;
+			Ar << Data.Blocks;
+			return Ar;
+		};
+	};
+
+	struct FModelStreamableData
+	{
+		void Get(uint32 Key, TArrayView64<uint8> Destination, bool bDropData)
+		{
+			TArray64<uint8>* Buffer = Data.Find(Key);
+			check(Buffer);
+			check(Destination.Num() == Buffer->Num());
+			FMemory::Memcpy(Destination.GetData(), Buffer->GetData(), Buffer->Num());
+
+			if (bDropData)
+			{
+				Buffer->Empty();
+			}
+		}
+
+		void Set(uint32 Key, const uint8* Source, int64 Size)
+		{
+			check(Source);
+			check(Size);
+			TArray64<uint8>& Buffer = Data.Add(Key);
+			check(Buffer.Num() == 0);
+			Buffer.SetNumUninitialized(Size);
+			FMemory::Memcpy(Buffer.GetData(), Source, Size);
+		}
+
+		// Temp, to be replaced with disk storage
+		TMap<uint32, TArray64<uint8> > Data;
+	};
+
+
+	struct CUSTOMIZABLEOBJECT_API FMutableCachedPlatformData
+	{
+		/** */
+		TArray64<uint8> ModelData;
+
+		/** */
+		FModelStreamableData ModelStreamableData;
+
+		/** */
+		TArray64<uint8> MorphData;
+
+		/** */
+		TArray64<uint8> ClothingData;
+
+		/** List of files to serialize. Each file has a list of binary blocks to be serialized. */
+		TArray<FFile> BulkDataFiles;
+	};
+
+
+	/** Compute the number of files and sizes the BulkData will be split into and update
+	 * the streamables's FileIds and Offsets.
+	 */
+	void CUSTOMIZABLEOBJECT_API GenerateBulkDataFilesList(
+		TSharedPtr<const mu::Model, ESPMode::ThreadSafe> Model,
+		TSharedPtr<FModelStreamableBulkData> StreamableBulkData,
+		const ITargetPlatform* TargetPlatform,
+		uint64 TargetBulkDataFileBytes,
+		TArray<FFile>& OutBulkDataFiles);
+
+	void CUSTOMIZABLEOBJECT_API SerializeBulkDataFiles(
+		FMutableCachedPlatformData& CachedPlatformData,
+		TArray<FFile>& BulkDataFiles,
+		TFunctionRef<void(FFile&, TArray64<uint8>&)> WriteFile,
+		bool bDropData
+	);
+}
+#endif
+
+
+struct CUSTOMIZABLEOBJECT_API FModelStreamableBulkData
+{
+	/** Map of Hash to Streaming blocks, used to stream a block of data representing a resource from the BulkData */
+	TMap<uint32, FMutableStreamableBlock> ModelStreamables;
+
+	TMap<uint32, FClothingStreamable> ClothingStreamables;
+
+	TMap<uint32, FRealTimeMorphStreamable> RealTimeMorphStreamables;
+
+	TMap<uint32, FByteBulkData> HashToBulkData;
+
+	void Serialize(FArchive& Ar, UObject* Owner, bool bCooked);
+
+#if WITH_EDITORONLY_DATA
+	friend FArchive& operator<<(FArchive& Ar, FModelStreamableBulkData& Struct);
+#endif
+};
+
+/** Interface class to allow custom serialization of FModelStreamableBulkData and its FBulkData. */
+UCLASS()
+class UModelStreamableData : public UObject
+{
+	GENERATED_BODY()
+
+	UModelStreamableData();
+
+public:
+	virtual void Serialize(FArchive& Ar) override;
+
+	TSharedPtr<FModelStreamableBulkData> StreamingData;
+};
+
 // Referenced materials, skeletons, passthrough textures...
 USTRUCT()
 struct FModelResources
@@ -773,17 +949,12 @@ struct FModelResources
 	UPROPERTY()
 	TMap<FString, FMutableStateData> StateUIDataMap;
 
-	UPROPERTY()
-	TMap<uint32, FRealTimeMorphStreamable> RealTimeMorphStreamables;
-
 	UPROPERTY()
 	TArray<FCustomizableObjectClothConfigData> ClothSharedConfigsData;	
 
 	UPROPERTY()
 	TArray<FCustomizableObjectClothingAssetData> ClothingAssetsData;
 
-	UPROPERTY()
-	TMap<uint32, FClothingStreamable> ClothingStreamables;
 
 	/** Currently not used, this option should be selectable from editor maybe as a compilation flag */
 	UPROPERTY()
@@ -799,10 +970,6 @@ struct FModelResources
 	// Stores what param names use a certain table as a table can be used from multiple table nodes, useful for partial compilations to restrict params
 	TMap<TObjectPtr<const UDataTable>, TSet<FString>> TableToParamNames;
 #endif
-	
-	/** Map of Hash to Streaming blocks, used to stream a block of data representing a resource from the BulkData */
-	UPROPERTY()
-	TMap<uint32, FMutableStreamableBlock> HashToStreamableBlock;
 
 	/** Max number of components in the compiled Model. */
 	UPROPERTY()
@@ -822,47 +989,6 @@ struct FModelResources
 };
 
 
-struct FModelStreamableData
-{
-	void Get(uint32 Key, uint8* Destination) const
-	{
-		check(Destination);
-		const TArray64<uint8>* Buffer = Data.Find(Key);
-		check(Buffer);
-		FMemory::Memcpy(Destination, Buffer->GetData(), Buffer->Num());
-	}
-
-	void Set(uint32 Key, const uint8* Source, int64 Size)
-	{
-		check(Source);
-		check(Size);
-		TArray64<uint8>& Buffer = Data.Add(Key);
-		check(Buffer.Num()==0);
-		Buffer.SetNumUninitialized(Size);
-		FMemory::Memcpy(Buffer.GetData(), Source, Size);
-	}
-
-	// Temp, to be replaced with disk storage
-	TMap<uint32, TArray64<uint8> > Data;
-};
-
-
-struct CUSTOMIZABLEOBJECT_API FMutableCachedPlatformData
-{
-	/** */
-	TArray64<uint8> ModelData;
-
-	/** */
-	FModelStreamableData ModelStreamableData;
-
-	/** */
-	TArray64<uint8> MorphData;
-	
-	/** */
-	TArray64<uint8> ClothingData;
-};
-
-
 USTRUCT()
 struct CUSTOMIZABLEOBJECT_API FMutableMeshComponentData
 {
@@ -892,6 +1018,10 @@ class CUSTOMIZABLEOBJECT_API UCustomizableObjectPrivate : public UObject
 
 	TSharedPtr<mu::Model, ESPMode::ThreadSafe> MutableModel;
 
+	/** Stores streamable data info to be used by MutableModel In-Game. Cooked resources. */
+	UPROPERTY()
+	TObjectPtr<UModelStreamableData> ModelStreamableData;
+
 	/** Stores resources to be used by MutableModel In-Game. Cooked resources. */
 	UPROPERTY()
 	FModelResources ModelResources;
@@ -903,6 +1033,11 @@ class CUSTOMIZABLEOBJECT_API UCustomizableObjectPrivate : public UObject
 	 */
 	UPROPERTY(Transient)
 	FModelResources ModelResourcesEditor;
+
+	/**
+	 * Stores streamable data info to be used by MutableModel in the Editor. Editor streaming.
+	 */
+	TSharedPtr<FModelStreamableBulkData> ModelStreamableDataEditor;
 #endif
 
 public:
@@ -917,8 +1052,12 @@ public:
 
 #if WITH_EDITORONLY_DATA
 	FModelResources& GetModelResources(bool bIsCooking);
+
+	void SetModelStreamableBulkData(const TSharedPtr<FModelStreamableBulkData>& StreamableData, bool bIsCooking);
 #endif
 
+	TSharedPtr<FModelStreamableBulkData> GetModelStreamableBulkData(bool bIsCooking = false);
+
 	// See UCustomizableObjectSystem::LockObject()
 	bool IsLocked() const;
 
@@ -1080,7 +1219,7 @@ public:
 
 #if WITH_EDITOR
 	/** Map of PlatformName to CachedPlatformData. Only valid while cooking. */
-	TMap<FString, FMutableCachedPlatformData> CachedPlatformsData;
+	TMap<FString, MutablePrivate::FMutableCachedPlatformData> CachedPlatformsData;
 #endif
 #endif
 
@@ -1172,6 +1311,8 @@ public:
 
 		CompilationOptimizationsMeshFormat,
 
+		ModelStreamableBulkData,
+
 		// -----<new versions can be added above this line>--------
 		LastCustomizableObjectVersion
 	};
diff --git a/Engine/Plugins/Experimental/Mutable/Source/CustomizableObject/Internal/MuCO/UnrealMutableModelDiskStreamer.h b/Engine/Plugins/Experimental/Mutable/Source/CustomizableObject/Internal/MuCO/UnrealMutableModelDiskStreamer.h
--- a/Engine/Plugins/Experimental/Mutable/Source/CustomizableObject/Internal/MuCO/UnrealMutableModelDiskStreamer.h
+++ b/Engine/Plugins/Experimental/Mutable/Source/CustomizableObject/Internal/MuCO/UnrealMutableModelDiskStreamer.h
@@ -2,13 +2,13 @@
 
 #pragma once
 
-#include "MuCO/CustomizableObject.h"
 #include "MuCO/CustomizableObjectPrivate.h"
 #include "MuR/Serialisation.h"
 
 class FArchive;
 class IAsyncReadFileHandle;
 class IAsyncReadRequest;
+class IBulkDataIORequest;
 class UCustomizableObject;
 namespace mu { class Model; }
 struct FMutableStreamableBlock;
@@ -64,7 +64,8 @@ protected:
 
 	struct FReadRequest
 	{
-		TSharedPtr<IAsyncReadRequest> ReadRequest = nullptr;
+		TSharedPtr<IBulkDataIORequest> BulkReadRequest = nullptr;
+		TSharedPtr<IAsyncReadRequest> FileReadRequest = nullptr;
 		TSharedPtr<FAsyncFileCallBack> FileCallback;
 	};
 	
@@ -73,10 +74,11 @@ protected:
 	{
 		TWeakPtr<const mu::Model> Model;
 		TMap<OPERATION_ID, FReadRequest> CurrentReadRequests;
-		TMap<uint32, FMutableStreamableBlock> StreamableBlocks;
 		FString BulkFilePrefix;
 
 		TMap<uint32, TSharedPtr<IAsyncReadFileHandle>> ReadFileHandles;
+
+		TSharedPtr<FModelStreamableBulkData> ModelStreamableBulkData;
 	};
 
 	TArray<FObjectData> Objects;
@@ -134,7 +136,7 @@ class CUSTOMIZABLEOBJECT_API FUnrealMutableModelBulkWriterCook : public mu::Mode
 {
 public:
 	// 
-	FUnrealMutableModelBulkWriterCook(FArchive* InMainDataArchive = nullptr, FModelStreamableData* InStreamedData = nullptr);
+	FUnrealMutableModelBulkWriterCook(FArchive* InMainDataArchive = nullptr, MutablePrivate::FModelStreamableData* InStreamedData = nullptr);
 
 	// mu::ModelWriter interface
 	void OpenWriteFile(uint32 BlockKey) override;
@@ -147,7 +149,7 @@ protected:
 	FArchive* MainDataArchive = nullptr;
 
 	// Non-owned pointer to an archive where we'll store the resouces (streamable)
-	FModelStreamableData* StreamedData = nullptr;
+	MutablePrivate::FModelStreamableData* StreamedData = nullptr;
 
 	uint32 CurrentKey = 0;
 
diff --git a/Engine/Plugins/Experimental/Mutable/Source/CustomizableObject/Private/MuCO/CustomizableObject.cpp b/Engine/Plugins/Experimental/Mutable/Source/CustomizableObject/Private/MuCO/CustomizableObject.cpp
--- a/Engine/Plugins/Experimental/Mutable/Source/CustomizableObject/Private/MuCO/CustomizableObject.cpp
+++ b/Engine/Plugins/Experimental/Mutable/Source/CustomizableObject/Private/MuCO/CustomizableObject.cpp
@@ -61,6 +61,14 @@ TAutoConsoleVariable<int32> CVarPackagedDataBytesLimitOverride(
 	TEXT(" <0 : Use value defined in the CO\n")
 	TEXT(" >=0  : Use this value instead\n"));
 
+
+TAutoConsoleVariable<bool> CVarMutableUseBulkData(
+	TEXT("Mutable.UseBulkData"),
+	false,
+	TEXT("Switch between .utoc/.ucas (FBulkData) and .mut files (CookAdditionalFiles).\n")
+	TEXT("True - Use FBulkData to store streamable data.\n")
+	TEXT("False - Use Mut files to store streamable data\n"));
+
 #endif
 
 #if WITH_EDITORONLY_DATA
@@ -166,15 +174,52 @@ void UCustomizableObject::PreSave(FObjectPreSaveContext ObjectSaveContext)
 		// Load cached data before saving
 		if (GetPrivate()->TryLoadCompiledCookDataForPlatform(TargetPlatform))
 		{
-			// Create an export object to manage the streamable data
-			if (!BulkData)
+			MutablePrivate::FMutableCachedPlatformData& CachedPlatformData = *GetPrivate()->CachedPlatformsData.Find(TargetPlatform->PlatformName());
+
+			// Generate list of files and update streamable blocks ids and offsets
+			FCompilationOptions Options = GetPrivate()->GetCompileOptions();
+			TSharedPtr<FModelStreamableBulkData> ModelStreamableBulkData = GetPrivate()->GetModelStreamableBulkData(true);
+
+			MutablePrivate::GenerateBulkDataFilesList(GetPrivate()->GetModel(), ModelStreamableBulkData, TargetPlatform, Options.PackagedDataBytesLimit, CachedPlatformData.BulkDataFiles);
+
+			const bool bUseBulkData = CVarMutableUseBulkData.GetValueOnAnyThread();
+			if (bUseBulkData)
 			{
-				BulkData = UE::Mutable::Private::MoveOldObjectAndCreateNew<UCustomizableObjectBulk>(UCustomizableObjectBulk::StaticClass(), this);
-			}
-			BulkData->Mark(OBJECTMARK_TagExp);
+				const int32 NumBulkDataFiles = CachedPlatformData.BulkDataFiles.Num();
+				ModelStreamableBulkData->HashToBulkData.Empty(NumBulkDataFiles);
+				ModelStreamableBulkData->HashToBulkData.Reserve(NumBulkDataFiles);
 
-			// Split streamable data into smaller chunks and fix up the CO HashToStreamableBlock's FileIndex and Offset
-			BulkData->PrepareBulkData(this, ObjectSaveContext.GetTargetPlatform());
+				const auto WriteBulkData = [&ModelStreamableBulkData](MutablePrivate::FFile& File, TArray64<uint8>& FileBulkData)
+					{
+						MUTABLE_CPUPROFILER_SCOPE(WriteBulkData);
+
+						FByteBulkData& BulkData = ModelStreamableBulkData->HashToBulkData.FindOrAdd(File.Id);
+
+						BulkData.Lock(LOCK_READ_WRITE);
+						uint8* Ptr = (uint8*)BulkData.Realloc(FileBulkData.Num());
+						FMemory::Memcpy(Ptr, FileBulkData.GetData(), FileBulkData.Num());
+						BulkData.Unlock();
+
+						uint32 BulkDataFlags = BULKDATA_PayloadInSeperateFile | BULKDATA_Force_NOT_InlinePayload;
+						if (File.Flags == uint32(mu::ERomFlags::HighRes))
+						{
+							BulkDataFlags |= BULKDATA_OptionalPayload;
+						}
+						BulkData.SetBulkDataFlags(BulkDataFlags);
+					};
+
+				bool bDropData = true;
+				MutablePrivate::SerializeBulkDataFiles(CachedPlatformData, CachedPlatformData.BulkDataFiles, WriteBulkData, bDropData);
+			}
+			else 
+			{
+				// Create an export object to manage the streamable data
+				if (!BulkData)
+				{
+					BulkData = UE::Mutable::Private::MoveOldObjectAndCreateNew<UCustomizableObjectBulk>(UCustomizableObjectBulk::StaticClass(), this);
+				}
+				BulkData->Mark(OBJECTMARK_TagExp);
+			}
 		}
 		else
 		{
@@ -230,7 +275,7 @@ void UCustomizableObjectPrivate::SetIsChildObject(bool bIsChildObject)
 
 bool UCustomizableObjectPrivate::TryLoadCompiledCookDataForPlatform(const ITargetPlatform* TargetPlatform)
 {
-	const FMutableCachedPlatformData* PlatformData = CachedPlatformsData.Find(TargetPlatform->PlatformName());
+	const MutablePrivate::FMutableCachedPlatformData* PlatformData = CachedPlatformsData.Find(TargetPlatform->PlatformName());
 	if (!PlatformData)
 	{
 		return false;
@@ -632,13 +677,8 @@ void UCustomizableObjectPrivate::SaveCompiledData(FArchive& MemoryWriter, bool b
 	MemoryWriter << LocalModelResources.ParameterUIDataMap;
 	MemoryWriter << LocalModelResources.StateUIDataMap;
 
-	MemoryWriter << LocalModelResources.RealTimeMorphStreamables;
-	
-	MemoryWriter << LocalModelResources.ClothingStreamables;
-	MemoryWriter << LocalModelResources.ClothingAssetsData; 	
-	MemoryWriter << LocalModelResources.ClothSharedConfigsData; 
-
-	MemoryWriter << LocalModelResources.HashToStreamableBlock;
+	MemoryWriter << LocalModelResources.ClothingAssetsData;
+	MemoryWriter << LocalModelResources.ClothSharedConfigsData;
 
 	MemoryWriter << LocalModelResources.NumComponents;
 	MemoryWriter << LocalModelResources.NumLODs;
@@ -657,6 +697,11 @@ void UCustomizableObjectPrivate::SaveCompiledData(FArchive& MemoryWriter, bool b
 		MemoryWriter << LocalModelResources.EditorOnlyMorphTargetReconstructionData;
 		MemoryWriter << LocalModelResources.EditorOnlyClothingMeshToMeshVertData;
 	}
+
+	{
+		FModelStreamableBulkData& ModelStreamables = *GetModelStreamableBulkData(false).Get();
+		MemoryWriter << ModelStreamables;
+	}
 }
 
 
@@ -778,13 +823,8 @@ void UCustomizableObjectPrivate::LoadCompiledData(FArchive& MemoryReader, const
 		MemoryReader << LocalModelResource.ParameterUIDataMap;
 		MemoryReader << LocalModelResource.StateUIDataMap;
 
-		MemoryReader << LocalModelResource.RealTimeMorphStreamables;
-		
-		MemoryReader << LocalModelResource.ClothingStreamables;
 		MemoryReader << LocalModelResource.ClothingAssetsData; 
 		MemoryReader << LocalModelResource.ClothSharedConfigsData; 
-		
-		MemoryReader << LocalModelResource.HashToStreamableBlock;
 
 		MemoryReader << LocalModelResource.NumComponents;
 		MemoryReader << LocalModelResource.NumLODs;
@@ -818,6 +858,14 @@ void UCustomizableObjectPrivate::LoadCompiledData(FArchive& MemoryReader, const
 			}
 		}
 
+		{
+			TSharedPtr<FModelStreamableBulkData> LocalModelStreamablesPtr = MakeShared<FModelStreamableBulkData>();
+			FModelStreamableBulkData& LocalModelStreamables = *LocalModelStreamablesPtr.Get();
+			MemoryReader << LocalModelStreamables;
+
+			SetModelStreamableBulkData(LocalModelStreamablesPtr, bIsCooking);
+		}
+
 		bool bModelSerialized = false;
 		MemoryReader << bModelSerialized;
 
@@ -912,6 +960,7 @@ void UCustomizableObjectPrivate::CompileForTargetPlatform(const ITargetPlatform*
 	ICustomizableObjectEditorModule* EditorModule = ICustomizableObjectEditorModule::Get();
 	if (!EditorModule || !EditorModule->IsRootObject(*CustomizableObject))
 	{
+		SetIsChildObject(true);
 		return;
 	}
 
@@ -1984,6 +2033,77 @@ FArchive& operator<<(FArchive& Ar, FMutableStateData& Struct)
 	return Ar;
 }
 
+#if WITH_EDITORONLY_DATA
+FArchive& operator<<(FArchive& Ar, FModelStreamableBulkData& Data)
+{
+	Ar << Data.ModelStreamables;
+	Ar << Data.ClothingStreamables;
+	Ar << Data.RealTimeMorphStreamables;
+
+	// Don't serialize FByteBulkData manually, the data will be skipped.
+
+	return Ar;
+}
+#endif
+
+
+void FModelStreamableBulkData::Serialize(FArchive& Ar, UObject* Owner, bool bCooked)
+{
+	MUTABLE_CPUPROFILER_SCOPE(FModelStreamableBulkData::Serialize);
+
+	Ar << ModelStreamables;
+	Ar << ClothingStreamables;
+	Ar << RealTimeMorphStreamables;
+
+	if (bCooked)
+	{
+		int32 NumBulkDatas = HashToBulkData.Num();
+		Ar << NumBulkDatas;
+
+		if (Ar.IsSaving())
+		{
+			for (TPair<uint32, FByteBulkData>& BulkData : HashToBulkData)
+			{
+				Ar << BulkData.Key;
+				BulkData.Value.Serialize(Ar, Owner);
+			}
+		}
+		else
+		{
+			HashToBulkData.Reserve(NumBulkDatas);
+			for (int32 Index = 0; Index < NumBulkDatas; ++Index)
+			{
+				uint32 ResourceId = 0;
+				Ar << ResourceId;
+
+				FByteBulkData& BulkData = HashToBulkData.Add(ResourceId);
+				BulkData.Serialize(Ar, Owner);
+			}
+		}
+	}
+}
+
+
+UModelStreamableData::UModelStreamableData()
+{
+	StreamingData = MakeShared<FModelStreamableBulkData>();
+}
+
+
+void UModelStreamableData::Serialize(FArchive& Ar)
+{
+	Super::Serialize(Ar);
+
+	bool bCooked = Ar.IsCooking();
+	Ar << bCooked;
+
+	if (bCooked && !IsTemplate() && !Ar.IsCountingMemory())
+	{
+		check(StreamingData);
+		StreamingData->Serialize(Ar, GetOutermostObject(), bCooked);
+	}
+}
+
 
 UCustomizableObjectPrivate::UCustomizableObjectPrivate()
 {
@@ -2031,6 +2151,39 @@ TSharedPtr<const mu::Model, ESPMode::ThreadSafe> UCustomizableObjectPrivate::Get
 	return MutableModel;
 }
 
+#if WITH_EDITOR
+void UCustomizableObjectPrivate::SetModelStreamableBulkData(const TSharedPtr<FModelStreamableBulkData>& StreamableData, bool bIsCooking)
+{
+	if (bIsCooking)
+	{
+		if (!ModelStreamableData)
+		{
+			ModelStreamableData = NewObject<UModelStreamableData>(GetOuter());
+		}
+
+		ModelStreamableData->StreamingData = StreamableData;
+	}
+	else
+	{
+		ModelStreamableDataEditor = StreamableData;
+	}
+}
+#endif
+
+TSharedPtr<FModelStreamableBulkData> UCustomizableObjectPrivate::GetModelStreamableBulkData(bool bIsCooking)
+{
+#if WITH_EDITOR
+	if (bIsCooking)
+	{
+		return ModelStreamableData ? ModelStreamableData->StreamingData : nullptr;
+	}
+
+	return ModelStreamableDataEditor;
+#else
+	return ModelStreamableData ? ModelStreamableData->StreamingData : nullptr;
+#endif
+}
+
 
 const FModelResources& UCustomizableObjectPrivate::GetModelResources() const
 {
@@ -2155,149 +2308,73 @@ FCompilationOptions UCustomizableObjectPrivate::GetCompileOptions() const
 }
 #endif
 
-
-//-------------------------------------------------------------------------------------------------
-//-------------------------------------------------------------------------------------------------
-//-------------------------------------------------------------------------------------------------
-void UCustomizableObjectBulk::PostLoad()
-{
-	UObject::PostLoad();
-
-	const FString OutermostName = GetOutermost()->GetName();
-	FString PackageFilename = FPackageName::LongPackageNameToFilename(OutermostName);
-	FPaths::MakeStandardFilename(PackageFilename);
-	BulkFilePrefix = PackageFilename;
-}
-
-TUniquePtr<IAsyncReadFileHandle> UCustomizableObjectBulk::OpenFileAsyncRead(uint32 FileId, uint32 Flags) const
-{
-	check(IsInGameThread());
-
-	FString FilePath = FString::Printf(TEXT("%s-%08x.mut"), *BulkFilePrefix, FileId);
-	if (Flags == uint32(mu::ERomFlags::HighRes))
-	{
-		FilePath += TEXT(".high");
-	}
-
-	IAsyncReadFileHandle* Result = FPlatformFileManager::Get().GetPlatformFile().OpenAsyncRead(*FilePath);
-	
-	// Result being null does not mean the file does not exist. A request has to be made. Let the callee deal with it.
-	//UE_CLOG(!Result, LogMutable, Warning, TEXT("CustomizableObjectBulkData: Failed to open file [%s]."), *FilePath);
-
-	return TUniquePtr<IAsyncReadFileHandle>(Result);
-}
-
 #if WITH_EDITOR
-
-
-int64 UCustomizableObjectBulk::FFile::GetSize() const
-{
-	int64 FileSize = 0;
-	for (const FBlock& Block : Blocks)
-	{
-		FileSize += Block.Size;
-	}
-	return FileSize;
-}
-
-
-void UCustomizableObjectBulk::FFile::GetFileData(struct FMutableCachedPlatformData* PlatformData, uint8* DataDestination) const
+namespace MutablePrivate
 {
-	check(PlatformData);
-	check(DataDestination);
-
-	const uint8* SourceData = nullptr;
-	if (DataType == EDataType::Model)
+	int64 FFile::GetSize() const
 	{
-		for (int32 BlockIndex = 0; BlockIndex < Blocks.Num(); ++BlockIndex)
+		int64 FileSize = 0;
+		for (const FBlock& Block : Blocks)
 		{
-			const FBlock& Block = Blocks[BlockIndex];
-			PlatformData->ModelStreamableData.Get(Block.Id, DataDestination);
-			DataDestination += Block.Size;
+			FileSize += Block.Size;
 		}
-		return;
+		return FileSize;
 	}
-	else if (DataType == EDataType::RealTimeMorph)
-	{
-		SourceData = PlatformData->MorphData.GetData();
-	}
-	else if (DataType == EDataType::Clothing)
-	{
-		SourceData = PlatformData->ClothingData.GetData();
-	}
-	else
-	{
-		checkf(false, TEXT("Unknown file DataType found."));
-	}
-
-	for (int32 BlockIndex = 0; BlockIndex<Blocks.Num(); ++BlockIndex)
-	{
-		const FBlock& Block = Blocks[BlockIndex];
-		FMemory::Memcpy(DataDestination, SourceData + Block.Offset, Block.Size);
-		DataDestination += Block.Size;
-	}
-}
 
 
-void UCustomizableObjectBulk::CookAdditionalFilesOverride(const TCHAR* PackageFilename,
-	const ITargetPlatform* TargetPlatform,
-	TFunctionRef<void(const TCHAR* Filename, void* Data, int64 Size)> WriteAdditionalFile)
-{
-	// Don't save streamed data on server builds since it won't be used anyway.
-	if (TargetPlatform->IsServerOnly())
-	{
-		return;
-	}
-	
-	check(CustomizableObject);
-	
-	FMutableCachedPlatformData* PlatformData = CustomizableObject->GetPrivate()->CachedPlatformsData.Find(TargetPlatform->PlatformName());
-	check(PlatformData);
-
-	const int32 NumBulkDataFiles = BulkDataFiles.Num();
-	for(int32 FileIndex = 0; FileIndex < NumBulkDataFiles; ++FileIndex)
+	void FFile::GetFileData(FMutableCachedPlatformData* PlatformData, TArray64<uint8>& DestData, bool bDropData)
 	{
-		const FFile& CurrentFile = BulkDataFiles[FileIndex];
-		int64 FileSize = CurrentFile.GetSize();
-
-		// Get the file data in memory
-		TArray64<uint8> FileBulkData;
-		FileBulkData.SetNumUninitialized(FileSize);
-		uint8* FileData = FileBulkData.GetData();
-		CurrentFile.GetFileData(PlatformData, FileData);
+		check(PlatformData);
 
-		// Path to the asset
-		const FString CookedFilePath = FPaths::GetPath(PackageFilename);
+		const uint64 DestSize = DestData.Num();
+		const int32 NumBlocks = Blocks.Num();
 
-		FString CookedBulkFileName = FString::Printf(TEXT("%s/%s-%08x.mut"), *CookedFilePath, *CustomizableObject->GetName(), CurrentFile.Id);
-
-		if (CurrentFile.Flags == uint32(mu::ERomFlags::HighRes))
+		const uint8* SourceData = nullptr;
+		if (DataType == EDataType::Model)
+		{
+			for (int32 BlockIndex = 0; BlockIndex < NumBlocks; ++BlockIndex)
+			{
+				const FBlock& Block = Blocks[BlockIndex];
+				check(Block.Offset + Block.Size <= DestSize);
+				PlatformData->ModelStreamableData.Get(Block.Id, TArrayView64<uint8>(DestData.GetData() + Block.Offset, Block.Size), bDropData);
+			}
+			return;
+		}
+		else if (DataType == EDataType::RealTimeMorph)
 		{
-			// We can do something different here for high-res data.
-			// For example: change the file name. We also need to detect it when generating the file name for loading.
-			CookedBulkFileName += TEXT(".high");
+			SourceData = PlatformData->MorphData.GetData();
+		}
+		else if (DataType == EDataType::Clothing)
+		{
+			SourceData = PlatformData->ClothingData.GetData();
+		}
+		else
+		{
+			checkf(false, TEXT("Unknown file DataType found."));
 		}
 
-		WriteAdditionalFile(*CookedBulkFileName, FileBulkData.GetData(), FileBulkData.Num());
+		for (int32 BlockIndex = 0; BlockIndex < NumBlocks; ++BlockIndex)
+		{
+			const FBlock& Block = Blocks[BlockIndex];
+			check(Block.Offset + Block.Size <= DestSize);
+			FMemory::Memcpy(DestData.GetData() + Block.Offset, SourceData + Block.Offset, Block.Size);
+		}
 	}
-}
 
 
-namespace MutablePrivate
-{
 	// TODO:
 	// To avoid influence of the order of the streamed data (their index), classify it recursively based on hash values
 	// until the tree leaves have either a single block, or a sum of blocks below the desired file size.
 	struct FClassifyNode
 	{
-		TArray<UCustomizableObjectBulk::FBlock> Blocks;
+		TArray<FBlock> Blocks;
 
 		//TSharedPtr<FClassifyNode> Child0;
 		//TSharedPtr<FClassifyNode> Child1;
 		//uint32 Depth=0;
 	};
 
-	void AddNode(TMap<uint32, FClassifyNode>& Nodes, int32 Slack, const UCustomizableObjectBulk::FBlock& Block)
+	void AddNode(TMap<uint32, FClassifyNode>& Nodes, int32 Slack, const FBlock& Block)
 	{
 		FClassifyNode& Root = Nodes.FindOrAdd(Block.Flags);
 		if (Root.Blocks.IsEmpty())
@@ -2307,221 +2384,315 @@ namespace MutablePrivate
 
 		Root.Blocks.Add(Block);
 	}
-}
 
+	void GenerateBulkDataFilesList(
+		TSharedPtr<const mu::Model, ESPMode::ThreadSafe> Model,
+		TSharedPtr<FModelStreamableBulkData> ModelStreamableBulkData,
+		const ITargetPlatform* TargetPlatform,
+		uint64 TargetBulkDataFileBytes,
+		TArray<FFile>& OutBulkDataFiles)
+	{
+		MUTABLE_CPUPROFILER_SCOPE(GenerateBulkDataFilesList);
 
-void UCustomizableObjectBulk::PrepareBulkData(UCustomizableObject* InOuter, const ITargetPlatform* TargetPlatform)
-{
-	CustomizableObject = InOuter;
-	check(CustomizableObject);
+		OutBulkDataFiles.Empty();
 
-	BulkDataFiles.Empty();
+		if (!Model || !ModelStreamableBulkData)
+		{
+			return;
+		}
 
-	TSharedPtr<const mu::Model, ESPMode::ThreadSafe> Model = CustomizableObject->GetPrivate()->GetModel();
-	if (!Model)
-	{
-		return;
-	}
+		const uint64 MaxChunkSize = UCustomizableObjectSystem::GetInstance()->GetMaxChunkSizeForPlatform(TargetPlatform);
+		TargetBulkDataFileBytes = FMath::Min(TargetBulkDataFileBytes, MaxChunkSize);
 
-	FModelResources& ModelResources = CustomizableObject->GetPrivate()->GetModelResources(true);
+		// Root nodes by flags.
+		const int32 NumRoms = Model->GetRomCount();
+		TMap<uint32, FClassifyNode> RootNode;
 
-	uint64 TargetBulkDataFileBytes = CustomizableObject->GetPrivate()->GetCompileOptions().PackagedDataBytesLimit;
-	const uint64 MaxChunkSize = UCustomizableObjectSystem::GetInstance()->GetMaxChunkSizeForPlatform(TargetPlatform);
-	TargetBulkDataFileBytes = FMath::Min(TargetBulkDataFileBytes, MaxChunkSize);
+		// Create blocks data.
+		{
+			for (int32 RomIndex = 0; RomIndex < NumRoms; ++RomIndex)
+			{
+				uint32 BlockId = Model->GetRomId(RomIndex);
+				const uint32 BlockSize = Model->GetRomSize(RomIndex);
+				const mu::ERomFlags BlockFlags = Model->GetRomFlags(RomIndex);
 
-	// Root nodes by flags.
-	const int32 NumRoms = Model->GetRomCount();
-	TMap<uint32, MutablePrivate::FClassifyNode> RootNode;
+				FBlock CurrentBlock = { EDataType::Model, BlockId, BlockSize, uint32(BlockFlags), 0 };
+				AddNode(RootNode, NumRoms, CurrentBlock);
+			}
+		}
 
-	// Create blocks data.
-	{		
-		for (int32 RomIndex = 0; RomIndex < NumRoms; ++RomIndex)
+		// TODO: This should create a new classification branch when the tree is implemented.
+		// For now append after Model roms. 
 		{
-			uint32 BlockId = Model->GetRomId(RomIndex);
-			const uint32 BlockSize = Model->GetRomSize(RomIndex);
-			const mu::ERomFlags BlockFlags = Model->GetRomFlags(RomIndex);
+			uint64 SourceOffset = 0;
 
-			FBlock CurrentBlock = { EDataType::Model, BlockId, BlockSize, uint32(BlockFlags), 0 };
-			MutablePrivate::AddNode(RootNode, NumRoms, CurrentBlock);
-		}
-	}
+			TArray<FMutableStreamableBlock> RealTimeMorphTargetsBlocks;
 
-	// TODO: This should create a new classification branch when the tree is implemented.
-	// For now append after Model roms. 
-	{
-		uint64 SourceOffset = 0;
-		
-		TArray<FMutableStreamableBlock> RealTimeMorphTargetsBlocks;
-		
-		const TMap<uint32, FRealTimeMorphStreamable>& RealTimeMorphStreamables = ModelResources.RealTimeMorphStreamables;
+			const TMap<uint32, FRealTimeMorphStreamable>& RealTimeMorphStreamables = ModelStreamableBulkData->RealTimeMorphStreamables;
 
-		const int32 NumBlocks = RealTimeMorphTargetsBlocks.Num();
-		for (const TPair<uint32, FRealTimeMorphStreamable>& MorphStreamable : RealTimeMorphStreamables)
-		{
-			const uint32 BlockSize = MorphStreamable.Value.Size;
-			
-			check(SourceOffset == MorphStreamable.Value.Block.Offset);
-			uint32 Flags = 0;
-			FBlock CurrentBlock = { EDataType::RealTimeMorph, MorphStreamable.Key, BlockSize, Flags, SourceOffset };
-			MutablePrivate::AddNode(RootNode, NumRoms, CurrentBlock);
+			const int32 NumBlocks = RealTimeMorphTargetsBlocks.Num();
+			for (const TPair<uint32, FRealTimeMorphStreamable>& MorphStreamable : RealTimeMorphStreamables)
+			{
+				const uint32 BlockSize = MorphStreamable.Value.Size;
 
-			SourceOffset += BlockSize;
-		}
-	}
+				uint32 Flags = 0;
+				FBlock CurrentBlock = { EDataType::RealTimeMorph, MorphStreamable.Key, BlockSize, Flags, SourceOffset };
+				AddNode(RootNode, NumRoms, CurrentBlock);
 
-	// TODO: This should create a new classification branch when the tree is implemented.
-	// For now append after Model roms. 
-	{
-		uint64 SourceOffset = 0;
-		
-		TArray<FMutableStreamableBlock> ClothingBlocks;
-		
-		const TMap<uint32, FClothingStreamable>& ClothingStreamables = ModelResources.ClothingStreamables;
+				SourceOffset += BlockSize;
+			}
+		}
 
-		const int32 NumBlocks = ClothingBlocks.Num();
-		for (const TPair<uint32, FClothingStreamable>& ClothStreamable : ClothingStreamables)
+		// TODO: This should create a new classification branch when the tree is implemented.
+		// For now append after Model roms. 
 		{
-			const uint32 BlockSize = ClothStreamable.Value.Size;
-			
-			check(SourceOffset == ClothStreamable.Value.Block.Offset);
-			uint32 Flags = 0;
-			FBlock CurrentBlock = { EDataType::Clothing, ClothStreamable.Key, BlockSize, Flags, SourceOffset };
-			MutablePrivate::AddNode(RootNode, NumRoms, CurrentBlock);
+			uint64 SourceOffset = 0;
 
-			SourceOffset += BlockSize;
-		}
-	}
+			TArray<FMutableStreamableBlock> ClothingBlocks;
 
+			const TMap<uint32, FClothingStreamable>& ClothingStreamables = ModelStreamableBulkData->ClothingStreamables;
 
-	for (int32 FlagClassIndex = 0; FlagClassIndex < RootNode.Num(); ++FlagClassIndex)
-	{
-		MutablePrivate::FClassifyNode& Root = RootNode[FlagClassIndex];
+			const int32 NumBlocks = ClothingBlocks.Num();
+			for (const TPair<uint32, FClothingStreamable>& ClothStreamable : ClothingStreamables)
+			{
+				const uint32 BlockSize = ClothStreamable.Value.Size;
 
-		// Temp: Group by order in the array
-		for (int32 BlockIndex = 0; BlockIndex < Root.Blocks.Num(); )
-		{
-			int32 CurrentFileSize = 0;
+				//check(SourceOffset == ClothStreamable.Value.Block.Offset);
+				uint32 Flags = 0;
+				FBlock CurrentBlock = { EDataType::Clothing, ClothStreamable.Key, BlockSize, Flags, SourceOffset };
+				AddNode(RootNode, NumRoms, CurrentBlock);
 
-			FFile CurrentFile;
-			CurrentFile.DataType = Root.Blocks[BlockIndex].DataType;
-			CurrentFile.Flags = Root.Blocks[BlockIndex].Flags;
+				SourceOffset += BlockSize;
+			}
+		}
 
-			while (BlockIndex < Root.Blocks.Num())
+		for (int32 FlagClassIndex = 0; FlagClassIndex < RootNode.Num(); ++FlagClassIndex)
+		{
+			FClassifyNode& Root = RootNode[FlagClassIndex];
+
+			// Temp: Group by order in the array
+			for (int32 BlockIndex = 0; BlockIndex < Root.Blocks.Num(); )
 			{
-				FBlock CurrentBlock = Root.Blocks[BlockIndex];
+				int32 CurrentFileSize = 0;
+
+				FFile CurrentFile;
+				CurrentFile.DataType = Root.Blocks[BlockIndex].DataType;
+				CurrentFile.Flags = Root.Blocks[BlockIndex].Flags;
 
-				// Next file?
-				// Store different data types in different files. Blocks should be sorted by DataType so data is properly packeted
-				if (CurrentFile.DataType != CurrentBlock.DataType)
+				while (BlockIndex < Root.Blocks.Num())
 				{
-					break;
-				}
+					FBlock CurrentBlock = Root.Blocks[BlockIndex];
 
-				// Different flags go to different files
-				check (CurrentFile.Flags == CurrentBlock.Flags)
+					// Next file?
+					// Store different data types in different files. Blocks should be sorted by DataType so data is properly packeted
+					if (CurrentFile.DataType != CurrentBlock.DataType)
+					{
+						break;
+					}
 
-				if (CurrentFileSize > 0 && CurrentFileSize + CurrentBlock.Size > TargetBulkDataFileBytes)
-				{
-					break;
-				}
+					// Different flags go to different files
+					check(CurrentFile.Flags == CurrentBlock.Flags)
 
-				// Add the block to the current file
-				CurrentFile.Blocks.Add(CurrentBlock);
-				CurrentFileSize += CurrentBlock.Size;
+						if (CurrentFileSize > 0 && CurrentFileSize + CurrentBlock.Size > TargetBulkDataFileBytes)
+						{
+							break;
+						}
 
-				// Next block
-				++BlockIndex;
-			}
+					// Add the block to the current file
+					CurrentFile.Blocks.Add(CurrentBlock);
+					CurrentFileSize += CurrentBlock.Size;
 
-			BulkDataFiles.Add(MoveTemp(CurrentFile));
-		}
-	}
+					// Next block
+					++BlockIndex;
+				}
 
-	// Create the file list
-	for (int32 FileIndex = 0; FileIndex < BulkDataFiles.Num(); ++FileIndex)
-	{
-		// Generate the id for this file
-		FFile& CurrentFile = BulkDataFiles[FileIndex];
-		uint32 FileId = uint32(CurrentFile.DataType);
-		for (int32 FileBlockIndex = 0; FileBlockIndex < CurrentFile.Blocks.Num(); ++FileBlockIndex)
-		{
-			FBlock ThisBlock = CurrentFile.Blocks[FileBlockIndex];
-			FileId = HashCombine(FileId, ThisBlock.Id);
+				OutBulkDataFiles.Add(MoveTemp(CurrentFile));
+			}
 		}
 
-		// Ensure the FileId is unique
-		bool bUnique = false;
-		while (!bUnique)
+		// Create the file list
+		for (int32 FileIndex = 0; FileIndex < OutBulkDataFiles.Num(); ++FileIndex)
 		{
-			bUnique = true;
-			for (int32 PreviousFileIndex = 0; PreviousFileIndex < FileIndex; ++PreviousFileIndex)
+			// Generate the id for this file
+			FFile& CurrentFile = OutBulkDataFiles[FileIndex];
+			uint32 FileId = uint32(CurrentFile.DataType);
+			for (int32 FileBlockIndex = 0; FileBlockIndex < CurrentFile.Blocks.Num(); ++FileBlockIndex)
 			{
-				if (BulkDataFiles[PreviousFileIndex].Id == FileId)
+				FBlock ThisBlock = CurrentFile.Blocks[FileBlockIndex];
+				FileId = HashCombine(FileId, ThisBlock.Id);
+			}
+
+			// Ensure the FileId is unique
+			bool bUnique = false;
+			while (!bUnique)
+			{
+				bUnique = true;
+				for (int32 PreviousFileIndex = 0; PreviousFileIndex < FileIndex; ++PreviousFileIndex)
 				{
-					bUnique = false;
-					++FileId;
-					break;
+					if (OutBulkDataFiles[PreviousFileIndex].Id == FileId)
+					{
+						bUnique = false;
+						++FileId;
+						break;
+					}
 				}
 			}
-		}
 
-		// Set it to the editor-only file descriptor
-		CurrentFile.Id = FileId;
+			// Set it to the editor-only file descriptor
+			CurrentFile.Id = FileId;
 
-		if (CurrentFile.DataType == EDataType::Model)
-		{
-			// Set it to all streamable blocks
-			uint32 OffsetInFile = 0;
-			for (int32 FileBlockIndex = 0; FileBlockIndex < CurrentFile.Blocks.Num(); ++FileBlockIndex)
+			if (CurrentFile.DataType == EDataType::Model)
 			{
-				FBlock ThisBlock = CurrentFile.Blocks[FileBlockIndex];
-
-				FMutableStreamableBlock* StreamableBlock = ModelResources.HashToStreamableBlock.Find(ThisBlock.Id);
-				StreamableBlock->FileId = FileId;
-				StreamableBlock->Offset = OffsetInFile;
-				check(StreamableBlock->Flags==CurrentFile.Flags);
-				OffsetInFile += ThisBlock.Size;
+				// Set it to all streamable blocks
+				uint32 OffsetInFile = 0;
+				for (int32 FileBlockIndex = 0; FileBlockIndex < CurrentFile.Blocks.Num(); ++FileBlockIndex)
+				{
+					FBlock& ThisBlock = CurrentFile.Blocks[FileBlockIndex];
+					ThisBlock.Offset = OffsetInFile;
+
+					FMutableStreamableBlock* StreamableBlock = ModelStreamableBulkData->ModelStreamables.Find(ThisBlock.Id);
+					StreamableBlock->FileId = FileId;
+					StreamableBlock->Offset = OffsetInFile;
+					check(StreamableBlock->Flags == CurrentFile.Flags);
+					OffsetInFile += ThisBlock.Size;
+				}
 			}
-		}
-		else if (CurrentFile.DataType == EDataType::RealTimeMorph)
-		{
-			TMap<uint32, FRealTimeMorphStreamable>& MorphBlocks = ModelResources.RealTimeMorphStreamables;
-			// Set it to all streamable blocks
-			uint32 OffsetInFile = 0;
-			for (int32 FileBlockIndex = 0; FileBlockIndex < CurrentFile.Blocks.Num(); ++FileBlockIndex)
+			else if (CurrentFile.DataType == EDataType::RealTimeMorph)
 			{
-				FBlock ThisBlock = CurrentFile.Blocks[FileBlockIndex];
+				TMap<uint32, FRealTimeMorphStreamable>& MorphBlocks = ModelStreamableBulkData->RealTimeMorphStreamables;
+				// Set it to all streamable blocks
+				uint32 OffsetInFile = 0;
+				for (int32 FileBlockIndex = 0; FileBlockIndex < CurrentFile.Blocks.Num(); ++FileBlockIndex)
+				{
+					FBlock& ThisBlock = CurrentFile.Blocks[FileBlockIndex];
+					ThisBlock.Offset = OffsetInFile;
 
-				FMutableStreamableBlock& StreamableBlock = MorphBlocks[ThisBlock.Id].Block;
-				StreamableBlock.FileId = FileId;
-				StreamableBlock.Offset = OffsetInFile;
-				OffsetInFile += ThisBlock.Size;
+					FMutableStreamableBlock& StreamableBlock = MorphBlocks[ThisBlock.Id].Block;
+					StreamableBlock.FileId = FileId;
+					StreamableBlock.Offset = OffsetInFile;
+					OffsetInFile += ThisBlock.Size;
+				}
 			}
-		}
-		else if (CurrentFile.DataType == EDataType::Clothing)
-		{
-			TMap<uint32, FClothingStreamable>& ClothBlocks = ModelResources.ClothingStreamables;
-			// Set it to all streamable blocks
-			uint32 OffsetInFile = 0;
-			for (int32 FileBlockIndex = 0; FileBlockIndex < CurrentFile.Blocks.Num(); ++FileBlockIndex)
+			else if (CurrentFile.DataType == EDataType::Clothing)
 			{
-				FBlock ThisBlock = CurrentFile.Blocks[FileBlockIndex];
+				TMap<uint32, FClothingStreamable>& ClothBlocks = ModelStreamableBulkData->ClothingStreamables;
+				// Set it to all streamable blocks
+				uint32 OffsetInFile = 0;
+				for (int32 FileBlockIndex = 0; FileBlockIndex < CurrentFile.Blocks.Num(); ++FileBlockIndex)
+				{
+					FBlock& ThisBlock = CurrentFile.Blocks[FileBlockIndex];
+					ThisBlock.Offset = OffsetInFile;
 
-				FMutableStreamableBlock& StreamableBlock = ClothBlocks[ThisBlock.Id].Block;
-				StreamableBlock.FileId = FileId;
-				StreamableBlock.Offset = OffsetInFile;
-				OffsetInFile += ThisBlock.Size;
+					FMutableStreamableBlock& StreamableBlock = ClothBlocks[ThisBlock.Id].Block;
+					StreamableBlock.FileId = FileId;
+					StreamableBlock.Offset = OffsetInFile;
+					OffsetInFile += ThisBlock.Size;
+				}
+			}
+			else
+			{
+				UE_LOG(LogMutable, Error, TEXT("Unknown DataType found while fixing streaming block files ids."));
+				check(false);
 			}
 		}
-		else
+	}
+
+
+	void CUSTOMIZABLEOBJECT_API SerializeBulkDataFiles(
+		FMutableCachedPlatformData& CachedPlatformData,
+		TArray<FFile>& BulkDataFiles,
+		TFunctionRef<void(FFile&, TArray64<uint8>&)> WriteFile,
+		bool bDropData)
+	{
+		TArray64<uint8> FileBulkData;
+
+		const int32 NumBulkDataFiles = BulkDataFiles.Num();
+		for (int32 FileIndex = 0; FileIndex < NumBulkDataFiles; ++FileIndex)
 		{
-			UE_LOG(LogMutable, Error, TEXT("Unknown DataType found while fixing streaming block files ids."));
-			check(false);
+			MutablePrivate::FFile& CurrentFile = BulkDataFiles[FileIndex];
+
+			const int64 FileSize = CurrentFile.GetSize();
+			FileBulkData.SetNumUninitialized(FileSize, EAllowShrinking::No);
+
+			// Get the file data in memory
+			CurrentFile.GetFileData(&CachedPlatformData, FileBulkData, bDropData);
+
+			WriteFile(CurrentFile, FileBulkData);
 		}
 	}
 }
 #endif // WITH_EDITOR
 
+//-------------------------------------------------------------------------------------------------
+//-------------------------------------------------------------------------------------------------
+//-------------------------------------------------------------------------------------------------
+void UCustomizableObjectBulk::PostLoad()
+{
+	UObject::PostLoad();
+
+	const FString OutermostName = GetOutermost()->GetName();
+	FString PackageFilename = FPackageName::LongPackageNameToFilename(OutermostName);
+	FPaths::MakeStandardFilename(PackageFilename);
+	BulkFilePrefix = PackageFilename;
+}
+
+TUniquePtr<IAsyncReadFileHandle> UCustomizableObjectBulk::OpenFileAsyncRead(uint32 FileId, uint32 Flags) const
+{
+	check(IsInGameThread());
+
+	FString FilePath = FString::Printf(TEXT("%s-%08x.mut"), *BulkFilePrefix, FileId);
+	if (Flags == uint32(mu::ERomFlags::HighRes))
+	{
+		FilePath += TEXT(".high");
+	}
+
+	IAsyncReadFileHandle* Result = FPlatformFileManager::Get().GetPlatformFile().OpenAsyncRead(*FilePath);
+	
+	// Result being null does not mean the file does not exist. A request has to be made. Let the callee deal with it.
+	//UE_CLOG(!Result, LogMutable, Warning, TEXT("CustomizableObjectBulkData: Failed to open file [%s]."), *FilePath);
+
+	return TUniquePtr<IAsyncReadFileHandle>(Result);
+}
+
+#if WITH_EDITOR
+
+void UCustomizableObjectBulk::CookAdditionalFilesOverride(const TCHAR* PackageFilename,
+	const ITargetPlatform* TargetPlatform,
+	TFunctionRef<void(const TCHAR* Filename, void* Data, int64 Size)> WriteAdditionalFile)
+{
+	// Don't save streamed data on server builds since it won't be used anyway.
+	if (TargetPlatform->IsServerOnly())
+	{
+		return;
+	}
+
+	UCustomizableObject* CustomizableObject = CastChecked<UCustomizableObject>(GetOutermostObject());
+
+	MutablePrivate::FMutableCachedPlatformData* PlatformData = CustomizableObject->GetPrivate()->CachedPlatformsData.Find(TargetPlatform->PlatformName());
+	check(PlatformData);
+
+	const FString CookedBulkFileName = FString::Printf(TEXT("%s/%s"), *FPaths::GetPath(PackageFilename), *CustomizableObject->GetName());
+
+	const auto WriteFile = [WriteAdditionalFile, CookedBulkFileName](MutablePrivate::FFile& File, TArray64<uint8>& FileBulkData)
+		{
+			FString FileName = CookedBulkFileName + FString::Printf(TEXT("-%08x.mut"), File.Id);
+
+			if (File.Flags == uint32(mu::ERomFlags::HighRes))
+			{
+				// We can do something different here for high-res data.
+				// For example: change the file name. We also need to detect it when generating the file name for loading.
+				FileName += TEXT(".high");
+			}
+
+			WriteAdditionalFile(*FileName, FileBulkData.GetData(), FileBulkData.Num());
+		};
+
+	bool bDropData = true;
+	MutablePrivate::SerializeBulkDataFiles(*PlatformData, PlatformData->BulkDataFiles, WriteFile, bDropData);
+}
+#endif // WITH_EDITOR
+
+
 //-------------------------------------------------------------------------------------------------
 //-------------------------------------------------------------------------------------------------
 //-------------------------------------------------------------------------------------------------
diff --git a/Engine/Plugins/Experimental/Mutable/Source/CustomizableObject/Private/MuCO/CustomizableObjectInstance.cpp b/Engine/Plugins/Experimental/Mutable/Source/CustomizableObject/Private/MuCO/CustomizableObjectInstance.cpp
--- a/Engine/Plugins/Experimental/Mutable/Source/CustomizableObject/Private/MuCO/CustomizableObjectInstance.cpp
+++ b/Engine/Plugins/Experimental/Mutable/Source/CustomizableObject/Private/MuCO/CustomizableObjectInstance.cpp
@@ -52,6 +52,7 @@
 #include "PhysicsEngine/SkeletalBodySetup.h"
 #include "Hash/CityHash.h"
 #include "MuCO/CustomizableObjectCustomVersion.h"
+#include "Serialization/BulkData.h"
 
 #include UE_INLINE_GENERATED_CPP_BY_NAME(CustomizableObjectInstance)
 
@@ -5209,6 +5210,7 @@ UE::Tasks::FTask UCustomizableInstancePrivate::LoadAdditionalAssetsAndData(
 	UCustomizableObject* CustomizableObject = GetPublic()->GetCustomizableObject();
 
 	const FModelResources& ModelResources = CustomizableObject->GetPrivate()->GetModelResources();
+	const TSharedPtr<FModelStreamableBulkData>& ModelStreamableBulkData = CustomizableObject->GetPrivate()->GetModelStreamableBulkData();
 
 	AssetsToStream.Empty();
 	TArray<uint32> RealTimeMorphStreamableBlocksToStream;
@@ -5342,7 +5344,7 @@ UE::Tasks::FTask UCustomizableInstancePrivate::LoadAdditionalAssetsAndData(
 					{
 						check(TypedResourceId.Id != 0 && TypedResourceId.Id <= TNumericLimits<uint32>::Max());
 
-						const TMap<uint32, FRealTimeMorphStreamable>& MorphsStreamables = ModelResources.RealTimeMorphStreamables;
+						const TMap<uint32, FRealTimeMorphStreamable>& MorphsStreamables = ModelStreamableBulkData->RealTimeMorphStreamables;
 						if (MorphsStreamables.Contains((uint32)TypedResourceId.Id))
 						{
 							RealTimeMorphStreamableBlocksToStream.AddUnique(TypedResourceId.Id);
@@ -5356,7 +5358,7 @@ UE::Tasks::FTask UCustomizableInstancePrivate::LoadAdditionalAssetsAndData(
 					{
 						check(TypedResourceId.Id != 0 && TypedResourceId.Id <= TNumericLimits<uint32>::Max());
 						
-						const TMap<uint32, FClothingStreamable>& ClothingStreamables = ModelResources.ClothingStreamables;
+						const TMap<uint32, FClothingStreamable>& ClothingStreamables = ModelStreamableBulkData->ClothingStreamables;
 						if (const FClothingStreamable* ClothingStreamable = ClothingStreamables.Find(TypedResourceId.Id))
 						{
 							// TODO: Add async loading of ClothingAsset Data. This could be loaded as an streamead resource similar to and the asset user data.
@@ -5531,6 +5533,7 @@ UE::Tasks::FTask UCustomizableInstancePrivate::LoadAdditionalAssetsAndData(
 	// File handles will end up owned by the gather task.
 	TArray<TUniquePtr<IAsyncReadFileHandle>> OpenFileHandles;
 	TArray<UE::Tasks::TTask<TUniquePtr<IAsyncReadRequest>>> ReadRequestTasks;
+	TArray<UE::Tasks::TTask<TUniquePtr<IBulkDataIORequest>>> BulkReadRequestTasks;
 	
 	bool bHasInvalidMesh = false;
 	bool bUpdateMeshes = DoComponentsNeedUpdate(GetPublic(), OperationData, bHasInvalidMesh);
@@ -5545,7 +5548,7 @@ UE::Tasks::FTask UCustomizableInstancePrivate::LoadAdditionalAssetsAndData(
 			MUTABLE_CPUPROFILER_SCOPE(RealTimeMorphStreamingEditor);
 			for (uint32 BlockId : RealTimeMorphStreamableBlocksToStream)
 			{	
-				const FRealTimeMorphStreamable& RealTimeMorphStreamable = ModelResources.RealTimeMorphStreamables[BlockId]; 
+				const FRealTimeMorphStreamable& RealTimeMorphStreamable = ModelStreamableBulkData->RealTimeMorphStreamables[BlockId];
 			
 				const FMutableStreamableBlock& Block = RealTimeMorphStreamable.Block;
 
@@ -5567,7 +5570,7 @@ UE::Tasks::FTask UCustomizableInstancePrivate::LoadAdditionalAssetsAndData(
 			MUTABLE_CPUPROFILER_SCOPE(ClothingStreamingEditor);
 			for (uint32 BlockId : ClothingStreamableBlocksToStream)
 			{	
-				const FClothingStreamable& ClothingStreamable = ModelResources.ClothingStreamables[BlockId]; 
+				const FClothingStreamable& ClothingStreamable = ModelStreamableBulkData->ClothingStreamables[BlockId]; 
 			
 				const FMutableStreamableBlock& Block = ClothingStreamable.Block;
 
@@ -5597,13 +5600,18 @@ UE::Tasks::FTask UCustomizableInstancePrivate::LoadAdditionalAssetsAndData(
 			uint32 FileId;
 		};
 
+		const bool bUseFBulkData = !ModelStreamableBulkData->HashToBulkData.IsEmpty();
+
 		TArray<FBlockReadInfo> BlockReadInfos;
 		BlockReadInfos.Reserve(16);
 		
 		const UCustomizableObjectBulk* BulkData = CustomizableObject->GetPrivate()->GetStreamableBulkData();
-		UE_CLOG(!BulkData, LogMutable, Error, TEXT("BulkData object for CustomizableObject [%s] not found."), 
+		if (!bUseFBulkData)
+		{
+			UE_CLOG(!BulkData, LogMutable, Error, TEXT("BulkData object for CustomizableObject [%s] not found."),
 				*CustomizableObject->GetFName().ToString());
-		check(BulkData);
+			check(BulkData);
+		}
 
 		TArray<uint32> OpenFilesIds;
 		const int32 NumMorphBlocks = RealTimeMorphStreamableBlocksToStream.Num();
@@ -5612,7 +5620,7 @@ UE::Tasks::FTask UCustomizableInstancePrivate::LoadAdditionalAssetsAndData(
 			MUTABLE_CPUPROFILER_SCOPE(RealTimeMorphStreamingRequest_Alloc);
 
 			const int32 BlockId = RealTimeMorphStreamableBlocksToStream[I];
-			const FRealTimeMorphStreamable& Streamable = ModelResources.RealTimeMorphStreamables[BlockId];
+			const FRealTimeMorphStreamable& Streamable = ModelStreamableBulkData->RealTimeMorphStreamables[BlockId];
 			const FMutableStreamableBlock& Block = Streamable.Block; 
 		
 			FInstanceUpdateData::FMorphTargetMeshData& ReadDestData = 
@@ -5624,26 +5632,33 @@ UE::Tasks::FTask UCustomizableInstancePrivate::LoadAdditionalAssetsAndData(
 				continue;
 			}
 
-			ReadDestData.NameResolutionMap = ModelResources.RealTimeMorphStreamables[BlockId].NameResolutionMap;
+			ReadDestData.NameResolutionMap = ModelStreamableBulkData->RealTimeMorphStreamables[BlockId].NameResolutionMap;
 
 			check(Streamable.Size % sizeof(FMorphTargetVertexData) == 0);
 			uint32 NumElems = Streamable.Size / sizeof(FMorphTargetVertexData);
 
 			ReadDestData.Data.SetNumUninitialized(NumElems);
 
-			int32 FileHandleIndex = OpenFilesIds.Find(Block.FileId);
-			if (FileHandleIndex == INDEX_NONE && BulkData)
+			IAsyncReadFileHandle* FileHandle = nullptr;
+
+			if (!bUseFBulkData)
 			{
-				TUniquePtr<IAsyncReadFileHandle> ReadFileHandle = BulkData->OpenFileAsyncRead(Block.FileId,Block.Flags);
+				int32 FileHandleIndex = OpenFilesIds.Find(Block.FileId);
+				if (FileHandleIndex == INDEX_NONE && BulkData)
+				{
+					TUniquePtr<IAsyncReadFileHandle> ReadFileHandle = BulkData->OpenFileAsyncRead(Block.FileId, Block.Flags);
 
-				OpenFileHandles.Emplace(MoveTemp(ReadFileHandle));
-				FileHandleIndex = OpenFilesIds.Add(Block.FileId);
-			}	
+					FileHandle = ReadFileHandle.Get();
+
+					OpenFileHandles.Emplace(MoveTemp(ReadFileHandle));
+					FileHandleIndex = OpenFilesIds.Add(Block.FileId);
+				}
+			}
 
 			BlockReadInfos.Emplace(FBlockReadInfo
 			{ 
 				Block.Offset,
-				OpenFileHandles[FileHandleIndex].Get(), 
+				FileHandle,
 				MakeArrayView(reinterpret_cast<uint8*>(ReadDestData.Data.GetData()), ReadDestData.Data.Num()*sizeof(FMorphTargetVertexData)),
 				Block.FileId
 			});
@@ -5655,7 +5670,7 @@ UE::Tasks::FTask UCustomizableInstancePrivate::LoadAdditionalAssetsAndData(
 			MUTABLE_CPUPROFILER_SCOPE(ClothingStreamingRequest_Alloc);
 
 			const int32 BlockId = ClothingStreamableBlocksToStream[I];
-			const FClothingStreamable& ClothingStreamable = ModelResources.ClothingStreamables[BlockId];
+			const FClothingStreamable& ClothingStreamable = ModelStreamableBulkData->ClothingStreamables[BlockId];
 			const FMutableStreamableBlock& Block = ClothingStreamable.Block;
 		
 			FInstanceUpdateData::FClothingMeshData& ReadDestData = 
@@ -5676,19 +5691,26 @@ UE::Tasks::FTask UCustomizableInstancePrivate::LoadAdditionalAssetsAndData(
 
 			ReadDestData.Data.SetNumUninitialized(NumElems);
 
-			int32 FileHandleIndex = OpenFilesIds.Find(Block.FileId);
-			if (FileHandleIndex == INDEX_NONE && BulkData)
+			IAsyncReadFileHandle* FileHandle = nullptr;
+
+			if (!bUseFBulkData)
 			{
-				TUniquePtr<IAsyncReadFileHandle> ReadFileHandle = BulkData->OpenFileAsyncRead(Block.FileId, Block.Flags);
+				int32 FileHandleIndex = OpenFilesIds.Find(Block.FileId);
+				if (FileHandleIndex == INDEX_NONE && BulkData)
+				{
+					TUniquePtr<IAsyncReadFileHandle> ReadFileHandle = BulkData->OpenFileAsyncRead(Block.FileId, Block.Flags);
 
-				OpenFileHandles.Emplace(MoveTemp(ReadFileHandle));
-				FileHandleIndex = OpenFilesIds.Add(Block.FileId);
-			}	
+					FileHandle = ReadFileHandle.Get();
+
+					OpenFileHandles.Emplace(MoveTemp(ReadFileHandle));
+					FileHandleIndex = OpenFilesIds.Add(Block.FileId);
+				}
+			}
 
 			BlockReadInfos.Emplace(FBlockReadInfo
 			{ 
 				Block.Offset,
-				OpenFileHandles[FileHandleIndex].Get(), 
+				FileHandle,
 				MakeArrayView(reinterpret_cast<uint8*>(ReadDestData.Data.GetData()), ReadDestData.Data.Num()*sizeof(FCustomizableObjectMeshToMeshVertData)),
 				Block.FileId
 			});
@@ -5697,50 +5719,83 @@ UE::Tasks::FTask UCustomizableInstancePrivate::LoadAdditionalAssetsAndData(
 
 		for (const FBlockReadInfo& BlockReadInfo : BlockReadInfos)
 		{
-			if (!BlockReadInfo.FileHandle)
+			if (bUseFBulkData)
 			{
-				continue;
-			}	
-	
-			ReadRequestTasks.Emplace(UE::Tasks::Launch(TEXT("CustomizableObjectInstanceReadRequestTask"),
-			[
-				OwnedOperationData = OperationData.ToSharedPtr(), // Keep a reference to make sure allocated memory is always alive.
-				ReadDataReadyEvent = StreamingCompletionEvents.Emplace_GetRef(TEXT("AsyncReadDataReadyEvent")),
-				Block              = BlockReadInfo,
-				Priority 		   = bEnableHighPriorityLoading ? AIOP_High : AIOP_Normal
-			]() -> TUniquePtr<IAsyncReadRequest>
-			{
-				MUTABLE_CPUPROFILER_SCOPE(CustomizableInstanceLoadBlocksAsyncRead_Request);
-				FAsyncFileCallBack ReadRequestCallBack = 
-						[OwnedOperationData, ReadDataReadyEvent, FileId = Block.FileId](bool bWasCancelled, IAsyncReadRequest*) mutable
-						{
-							if (bWasCancelled)
+				BulkReadRequestTasks.Emplace(UE::Tasks::Launch(TEXT("CustomizableObjectInstanceBulkReadRequestTask"),
+					[
+						OwnedOperationData = OperationData.ToSharedPtr(), // Keep a reference to make sure allocated memory is always alive.
+							ModelStreamableBulkData,
+							ReadDataReadyEvent = StreamingCompletionEvents.Emplace_GetRef(TEXT("AsyncReadDataReadyEvent")),
+							Block = BlockReadInfo,
+							Priority = bEnableHighPriorityLoading ? AIOP_High : AIOP_Normal
+					]() -> TUniquePtr<IBulkDataIORequest>
+					{
+						MUTABLE_CPUPROFILER_SCOPE(CustomizableInstanceLoadBlocksAsyncRead_Request);
+						FBulkDataIORequestCallBack IOCallback = 
+							[OwnedOperationData, ReadDataReadyEvent, FileId = Block.FileId](bool bWasCancelled, IBulkDataIORequest*) mutable
 							{
-								UE_LOG(LogMutable, Warning, TEXT("An AsyncReadRequest to file %08x was cancelled. The file may not exist."), FileId);
-							}
-							ReadDataReadyEvent.Trigger();
-						}; 
-
-				return TUniquePtr<IAsyncReadRequest>(Block.FileHandle->ReadRequest(
-						Block.Offset, 
-						(int64)Block.AllocatedMemoryView.Num(),
-						Priority,
-						&ReadRequestCallBack,
-						Block.AllocatedMemoryView.GetData()));
-			}, 
-			bEnableHighPriorityLoading ? UE::Tasks::ETaskPriority::High : UE::Tasks::ETaskPriority::Normal));
+								if (bWasCancelled)
+								{
+									UE_LOG(LogMutable, Warning, TEXT("An AsyncReadRequest to file %08x was cancelled. The file may not exist."), FileId);
+								}
+								ReadDataReadyEvent.Trigger();
+							};
+
+						FByteBulkData* ByteBulkData = ModelStreamableBulkData->HashToBulkData.Find(Block.FileId);
+						check(ByteBulkData);
+
+						return TUniquePtr<IBulkDataIORequest>(ByteBulkData->CreateStreamingRequest(
+							Block.Offset,
+							(int64)Block.AllocatedMemoryView.Num(),
+							Priority,
+							&IOCallback,
+							Block.AllocatedMemoryView.GetData()));
+					},
+					bEnableHighPriorityLoading ? UE::Tasks::ETaskPriority::High : UE::Tasks::ETaskPriority::Normal));
+			}
+			else if (BlockReadInfo.FileHandle)
+			{
+				ReadRequestTasks.Emplace(UE::Tasks::Launch(TEXT("CustomizableObjectInstanceReadRequestTask"),
+					[
+						OwnedOperationData = OperationData.ToSharedPtr(), // Keep a reference to make sure allocated memory is always alive.
+							ReadDataReadyEvent = StreamingCompletionEvents.Emplace_GetRef(TEXT("AsyncReadDataReadyEvent")),
+							Block = BlockReadInfo,
+							Priority = bEnableHighPriorityLoading ? AIOP_High : AIOP_Normal
+					]() -> TUniquePtr<IAsyncReadRequest>
+					{
+						MUTABLE_CPUPROFILER_SCOPE(CustomizableInstanceLoadBlocksAsyncRead_Request);
+						FAsyncFileCallBack ReadRequestCallBack =
+							[OwnedOperationData, ReadDataReadyEvent, FileId = Block.FileId](bool bWasCancelled, IAsyncReadRequest*) mutable
+							{
+								if (bWasCancelled)
+								{
+									UE_LOG(LogMutable, Warning, TEXT("An AsyncReadRequest to file %08x was cancelled. The file may not exist."), FileId);
+								}
+								ReadDataReadyEvent.Trigger();
+							};
+
+						return TUniquePtr<IAsyncReadRequest>(Block.FileHandle->ReadRequest(
+							Block.Offset,
+							(int64)Block.AllocatedMemoryView.Num(),
+							Priority,
+							&ReadRequestCallBack,
+							Block.AllocatedMemoryView.GetData()));
+					},
+					bEnableHighPriorityLoading ? UE::Tasks::ETaskPriority::High : UE::Tasks::ETaskPriority::Normal));
+			}
 		}	
 
 #endif
 
 	}
-
-	if (AssetsToStream.Num() > 0 || OpenFileHandles.Num() > 0)
+	
+	if (AssetsToStream.Num() > 0 || OpenFileHandles.Num() > 0 || BulkReadRequestTasks.Num() > 0)
 	{
 		return UE::Tasks::Launch(TEXT("GatherStreamingRequestsCompletionTask"),
 				[
 					ReadRequestTasks = MoveTemp(ReadRequestTasks),
-					OpenFileHandles  = MoveTemp(OpenFileHandles)
+					OpenFileHandles  = MoveTemp(OpenFileHandles),
+					BulkReadRequestTasks = MoveTemp(BulkReadRequestTasks)
 				]() mutable 
 				{
 					for (UE::Tasks::TTask<TUniquePtr<IAsyncReadRequest>>& ReadRequestTask : ReadRequestTasks)
@@ -5754,8 +5809,20 @@ UE::Tasks::FTask UCustomizableInstancePrivate::LoadAdditionalAssetsAndData(
 						}
 					}
 					
+					for (UE::Tasks::TTask<TUniquePtr<IBulkDataIORequest>>& BulkReadRequestTask : BulkReadRequestTasks)
+					{
+						// GetResult() may wait for the task to complete, this should not be a problem as this task
+						// prerequisites guarantee ReadRequestTask has at least started execution.
+						TUniquePtr<IBulkDataIORequest>& ReadRequest = BulkReadRequestTask.GetResult();
+						if (ReadRequest)
+						{
+							ReadRequest->WaitCompletion();
+						}
+					}
+
 					ReadRequestTasks.Empty();
 					OpenFileHandles.Empty();
+					BulkReadRequestTasks.Empty();
 				},
 				StreamingCompletionEvents,
 				bEnableHighPriorityLoading ? UE::Tasks::ETaskPriority::High : UE::Tasks::ETaskPriority::Normal);
@@ -5763,6 +5830,7 @@ UE::Tasks::FTask UCustomizableInstancePrivate::LoadAdditionalAssetsAndData(
 	else
 	{
 		check(ReadRequestTasks.Num() == 0);
+		check(BulkReadRequestTasks.Num() == 0);
 		return UE::Tasks::MakeCompletedTask<void>();
 	}
 }
diff --git a/Engine/Plugins/Experimental/Mutable/Source/CustomizableObject/Private/MuCO/UnrealMutableModelDiskStreamer.cpp b/Engine/Plugins/Experimental/Mutable/Source/CustomizableObject/Private/MuCO/UnrealMutableModelDiskStreamer.cpp
--- a/Engine/Plugins/Experimental/Mutable/Source/CustomizableObject/Private/MuCO/UnrealMutableModelDiskStreamer.cpp
+++ b/Engine/Plugins/Experimental/Mutable/Source/CustomizableObject/Private/MuCO/UnrealMutableModelDiskStreamer.cpp
@@ -4,12 +4,13 @@
 
 #include "Async/AsyncFileHandle.h"
 #include "HAL/PlatformFile.h"
-#include "MuCO/CustomizableObjectPrivate.h"
+#include "MuCO/CustomizableObject.h"
 #include "MuCO/CustomizableObjectSystem.h"
 #include "MuR/Model.h"
 #include "MuR/MutableTrace.h"
 #include "MuCO/LogBenchmarkUtil.h"
 #include "HAL/PlatformFileManager.h"
+#include "Serialization/BulkData.h"
 
 
 DECLARE_DWORD_ACCUMULATOR_STAT(TEXT("Streaming Ops"), STAT_MutableStreamingOps, STATGROUP_Mutable);
@@ -86,26 +87,9 @@ bool FUnrealMutableModelBulkReader::PrepareStreamingForObject(UCustomizableObjec
 	{
 		FObjectData NewData;
 		NewData.Model = TWeakPtr<const mu::Model>(CustomizableObject->GetPrivate()->GetModel());
+		NewData.ModelStreamableBulkData = CustomizableObject->GetPrivate()->GetModelStreamableBulkData();
 
-#if WITH_EDITOR
-		FString FolderPath = CustomizableObject->GetPrivate()->GetCompiledDataFolderPath();
-		FString FullFileName = FolderPath + CustomizableObject->GetPrivate()->GetCompiledDataFileName(false, nullptr, true);
-		NewData.BulkFilePrefix = *FullFileName;
-#else
-		const UCustomizableObjectBulk* BulkData = CustomizableObject->GetPrivate()->GetStreamableBulkData();
-		if (!BulkData)
-		{
-			UE_LOG(LogMutable, Warning, TEXT("Streaming: Customizable Object %s is missing the BulkData export."), *CustomizableObject->GetName());
-			return false;
-		}
-
-		NewData.BulkFilePrefix = BulkData->GetBulkFilePrefix();
-#endif
-
-		const FModelResources& ModelResources = CustomizableObject->GetPrivate()->GetModelResources();
-		// \TODO: unnecessary runtime map copy. This should be stored in a TSharedPtr and only copy the pointer.
-		NewData.StreamableBlocks = ModelResources.HashToStreamableBlock;
-		if (NewData.StreamableBlocks.IsEmpty())
+		if (!NewData.ModelStreamableBulkData || NewData.ModelStreamableBulkData->ModelStreamables.IsEmpty())
 		{
 			UE_LOG(LogMutable, Warning, TEXT("Streaming: Customizable Object %s has no data to stream."), *CustomizableObject->GetName());
 
@@ -115,6 +99,24 @@ bool FUnrealMutableModelBulkReader::PrepareStreamingForObject(UCustomizableObjec
 			return false;
 		}
 
+#if WITH_EDITOR
+		FString FolderPath = CustomizableObject->GetPrivate()->GetCompiledDataFolderPath();
+		FString FullFileName = FolderPath + CustomizableObject->GetPrivate()->GetCompiledDataFileName(false, nullptr, true);
+		NewData.BulkFilePrefix = *FullFileName;
+#else
+		if (NewData.ModelStreamableBulkData->HashToBulkData.IsEmpty())
+		{
+			const UCustomizableObjectBulk* BulkData = CustomizableObject->GetPrivate()->GetStreamableBulkData();
+			if (!BulkData)
+			{
+				UE_LOG(LogMutable, Warning, TEXT("Streaming: Customizable Object %s is missing the BulkData export."), *CustomizableObject->GetName());
+				return false;
+			}
+
+			NewData.BulkFilePrefix = BulkData->GetBulkFilePrefix();
+		}
+#endif
+
 		Objects.Add(MoveTemp(NewData));
 	}
 
@@ -178,7 +180,14 @@ void FUnrealMutableModelBulkReader::EndStreaming()
 	{
 		for (TPair<OPERATION_ID, FReadRequest>& it : o.CurrentReadRequests)
 		{
-			it.Value.ReadRequest->WaitCompletion();
+			if (it.Value.FileReadRequest)
+			{
+				it.Value.FileReadRequest->WaitCompletion();
+			}
+			else if (it.Value.BulkReadRequest)
+			{
+				it.Value.BulkReadRequest->WaitCompletion();
+			}
 		}
 	}
 	Objects.Empty();
@@ -214,7 +223,7 @@ mu::ModelReader::OPERATION_ID FUnrealMutableModelBulkReader::BeginReadBlock(cons
 		return -1;
 	}
 
-	const FMutableStreamableBlock* Block = ObjectData->StreamableBlocks.Find(Key);
+	const FMutableStreamableBlock* Block = ObjectData->ModelStreamableBulkData->ModelStreamables.Find(Key);
 	if (!Block)
 	{
 		// File Handle not found! This shouldn't really happen.
@@ -230,59 +239,83 @@ mu::ModelReader::OPERATION_ID FUnrealMutableModelBulkReader::BeginReadBlock(cons
 
 	OPERATION_ID Result = ++LastOperationID;
 
-	int32 BulkDataOffsetInFile = 0;
+	if (FByteBulkData* BulkData = ObjectData->ModelStreamableBulkData->HashToBulkData.Find(Block->FileId))
+	{
+		FBulkDataIORequestCallBack IOCallback;
+
+		if (CompletionCallback)
+		{
+			IOCallback = [CompletionCallbackCapture = *CompletionCallback](bool bWasCancelled, IBulkDataIORequest*)
+				{
+					CompletionCallbackCapture(!bWasCancelled);
+				};
+		}
+
+		FReadRequest& Request = ObjectData->CurrentReadRequests.Add(Result);
+
+		Request.BulkReadRequest = MakeShareable(BulkData->CreateStreamingRequest(
+			Block->Offset,
+			size,
+			(EAsyncIOPriorityAndFlags)StreamPriority,
+			&IOCallback,
+			reinterpret_cast<uint8*>(pBuffer)));
+	}
+	else
+	{
+		int32 BulkDataOffsetInFile = 0;
 #if WITH_EDITOR
-	BulkDataOffsetInFile = sizeof(MutableCompiledDataStreamHeader);
+		BulkDataOffsetInFile = sizeof(MutableCompiledDataStreamHeader);
 #endif
 
-	FReadRequest ReadRequest;
+		FReadRequest ReadRequest;
 
-	if (CompletionCallback)
-	{
-		ReadRequest.FileCallback = MakeShared<TFunction<void(bool, IAsyncReadRequest*)>>([CompletionCallbackCapture = *CompletionCallback](bool bWasCancelled, IAsyncReadRequest*) -> void
+		if (CompletionCallback)
 		{
-			CompletionCallbackCapture(!bWasCancelled);			
-		});		
-	}
-
-	TSharedPtr<IAsyncReadFileHandle> FileHandle;
-	{
-		FScopeLock Lock(&FileHandlesCritical);
+			ReadRequest.FileCallback = MakeShared<TFunction<void(bool, IAsyncReadRequest*)>>([CompletionCallbackCapture = *CompletionCallback](bool bWasCancelled, IAsyncReadRequest*) -> void
+				{
+					CompletionCallbackCapture(!bWasCancelled);
+				});
+		}
 
-		TSharedPtr<IAsyncReadFileHandle>& Found = ObjectData->ReadFileHandles.FindOrAdd( Block->FileId );
-		if (!Found)
+		TSharedPtr<IAsyncReadFileHandle> FileHandle;
 		{
+			FScopeLock Lock(&FileHandlesCritical);
+
+			TSharedPtr<IAsyncReadFileHandle>& Found = ObjectData->ReadFileHandles.FindOrAdd(Block->FileId);
+			if (!Found)
+			{
 #if WITH_EDITOR
-			FString FilePath = ObjectData->BulkFilePrefix;
+				FString FilePath = ObjectData->BulkFilePrefix;
 #else
-			FString FilePath = FString::Printf(TEXT("%s-%08x.mut"), *ObjectData->BulkFilePrefix, Block->FileId);
-			if (Block->Flags == uint32(mu::ERomFlags::HighRes))
-			{
-				FilePath += TEXT(".high");
-			}
+				FString FilePath = FString::Printf(TEXT("%s-%08x.mut"), *ObjectData->BulkFilePrefix, Block->FileId);
+				if (Block->Flags == uint32(mu::ERomFlags::HighRes))
+				{
+					FilePath += TEXT(".high");
+				}
 #endif
 
-			Found = MakeShareable(FPlatformFileManager::Get().GetPlatformFile().OpenAsyncRead(*FilePath));
+				Found = MakeShareable(FPlatformFileManager::Get().GetPlatformFile().OpenAsyncRead(*FilePath));
 
-			if (!Found)
-			{
-				UE_LOG(LogMutable, Error, TEXT("Failed to create AsyncReadFileHandle. File Path [%s]."), *FilePath);
-				check(false);
-				return -1;
+				if (!Found)
+				{
+					UE_LOG(LogMutable, Error, TEXT("Failed to create AsyncReadFileHandle. File Path [%s]."), *FilePath);
+					check(false);
+					return -1;
+				}
 			}
+
+			FileHandle = Found;
 		}
 
-		FileHandle = Found;
-	}
-	
-	ReadRequest.ReadRequest = MakeShareable(FileHandle->ReadRequest(
-		BulkDataOffsetInFile + Block->Offset,
-		size,
-		(EAsyncIOPriorityAndFlags)StreamPriority,
-		ReadRequest.FileCallback.Get(),
-		reinterpret_cast<uint8*>(pBuffer)));
+		ReadRequest.FileReadRequest = MakeShareable(FileHandle->ReadRequest(
+			BulkDataOffsetInFile + Block->Offset,
+			size,
+			(EAsyncIOPriorityAndFlags)StreamPriority,
+			ReadRequest.FileCallback.Get(),
+			reinterpret_cast<uint8*>(pBuffer)));
 
-	ObjectData->CurrentReadRequests.Add(Result, ReadRequest);
+		ObjectData->CurrentReadRequests.Add(Result, ReadRequest);
+	}
 
 	INC_DWORD_STAT(STAT_MutableStreamingOps);
 
@@ -298,7 +331,14 @@ bool FUnrealMutableModelBulkReader::IsReadCompleted(mu::ModelReader::OPERATION_I
 	{
 		if (FReadRequest* ReadRequest = o.CurrentReadRequests.Find(OperationId))
 		{
-			return ReadRequest->ReadRequest->PollCompletion();
+			if (ReadRequest->FileReadRequest)
+			{
+				return ReadRequest->FileReadRequest->PollCompletion();
+			}
+			else if(ReadRequest->BulkReadRequest)
+			{
+				return ReadRequest->BulkReadRequest->PollCompletion();
+			}
 		}
 	}
 
@@ -319,16 +359,31 @@ bool FUnrealMutableModelBulkReader::EndRead(mu::ModelReader::OPERATION_ID Operat
 		FReadRequest* ReadRequest = o.CurrentReadRequests.Find(OperationId);
 		if (ReadRequest)
 		{
-			if (ReadRequest->ReadRequest)
+			if (ReadRequest->FileReadRequest)
+			{
+				bool bCompleted = ReadRequest->FileReadRequest->WaitCompletion();
+				if (!bCompleted)
+				{
+					UE_LOG(LogMutable, Error, TEXT("Operation failed to complete in EndRead."));
+					check(false);
+					bSuccess = false;
+				}
+				else if (ReadRequest->FileReadRequest->GetReadResults()==nullptr)
+				{
+					// This means we have failed: file not found?
+					bSuccess = false;
+				}
+			}
+			else if (ReadRequest->BulkReadRequest)
 			{
-				bool bCompleted = ReadRequest->ReadRequest->WaitCompletion();
+				bool bCompleted = ReadRequest->BulkReadRequest->WaitCompletion();
 				if (!bCompleted)
 				{
 					UE_LOG(LogMutable, Error, TEXT("Operation failed to complete in EndRead."));
 					check(false);
 					bSuccess = false;
 				}
-				else if (ReadRequest->ReadRequest->GetReadResults()==nullptr)
+				else if (ReadRequest->BulkReadRequest->GetReadResults() == nullptr)
 				{
 					// This means we have failed: file not found?
 					bSuccess = false;
@@ -391,7 +446,7 @@ void FUnrealMutableModelBulkWriterEditor::CloseWriteFile()
 
 
 //-------------------------------------------------------------------------------------------------
-FUnrealMutableModelBulkWriterCook::FUnrealMutableModelBulkWriterCook(FArchive* InMainDataArchive, FModelStreamableData* InStreamedData)
+FUnrealMutableModelBulkWriterCook::FUnrealMutableModelBulkWriterCook(FArchive* InMainDataArchive, MutablePrivate::FModelStreamableData* InStreamedData)
 {
 	MainDataArchive = InMainDataArchive;
 	StreamedData = InStreamedData;
diff --git a/Engine/Plugins/Experimental/Mutable/Source/CustomizableObject/Public/MuCO/CustomizableObject.h b/Engine/Plugins/Experimental/Mutable/Source/CustomizableObject/Public/MuCO/CustomizableObject.h
--- a/Engine/Plugins/Experimental/Mutable/Source/CustomizableObject/Public/MuCO/CustomizableObject.h
+++ b/Engine/Plugins/Experimental/Mutable/Source/CustomizableObject/Public/MuCO/CustomizableObject.h
@@ -279,69 +279,10 @@ public:
 	//~ Begin UObject Interface
 	virtual void CookAdditionalFilesOverride(const TCHAR*, const ITargetPlatform*, TFunctionRef<void(const TCHAR*, void*, int64)> ) override;
 	//~ End UObject Interface
-
-	/** Compute the number of files and sizes the BulkData will be split into and fix up 
-	 * the HashToStreamableBlock's FileIds and Offsets. 
-	 */	
-	void PrepareBulkData(UCustomizableObject* InCustomizableObject, const ITargetPlatform* TargetPlatform);
-
 #endif
 
 #if WITH_EDITOR
-
-	enum class EDataType : uint8
-	{
-		None = 0,
-		Model,
-		RealTimeMorph,
-		Clothing,
-	};
-
-	struct FBlock
-	{
-		/** Data Type*/
-		EDataType DataType;
-	
-		/** Used on some data types as the index to the block stored in the CustomizableObject */
-		uint32 Id;
-
-		/** Size of the data block. */
-		uint32 Size;
-
-		/** Data flags, like "high-res". */
-		uint32 Flags;
-
-		/** Offset in the full source streamed data file that is created when compiling. */
-		uint64 Offset;
-	};
-
 private:
-
-	struct FFile
-	{
-		EDataType DataType;
-
-		/** Id generated from a hash of the file content + offset to avoid collisions. */
-		uint32 Id;
-
-		/** Common flags of the data stored in this file. See mu::ERomFlags. */
-		uint32 Flags=0;
-
-		/** List of blocks that are contained in the file, in order. */
-		TArray<FBlock> Blocks;
-
-		/** Get the total size of blocks in this file. */
-		int64 GetSize() const;
-
-		/** Copy the requested block to the requested buffer and return its size. */
-		void GetFileData(struct FMutableCachedPlatformData*, uint8* DataDestination) const;
-	};
-
-	/** Helper to store the size of each BulkData partition. Only valid while cooking */
-	TArray<FFile> BulkDataFiles;
-
-	/** Helper to retrieve the BulkData from within the CookAdditionalFilesOverride */
-	TObjectPtr<UCustomizableObject> CustomizableObject;
 #endif
 
 	/** Prefix to locate bulkfiles for loading, using the file ids in each FMutableStreamableBlock. */
diff --git a/Engine/Plugins/Experimental/Mutable/Source/CustomizableObjectEditor/Private/MuCOE/CustomizableObjectCompileRunnable.h b/Engine/Plugins/Experimental/Mutable/Source/CustomizableObjectEditor/Private/MuCOE/CustomizableObjectCompileRunnable.h
--- a/Engine/Plugins/Experimental/Mutable/Source/CustomizableObjectEditor/Private/MuCOE/CustomizableObjectCompileRunnable.h
+++ b/Engine/Plugins/Experimental/Mutable/Source/CustomizableObjectEditor/Private/MuCOE/CustomizableObjectCompileRunnable.h
@@ -123,7 +123,7 @@ public:
 	TArray64<uint8> ModelBytes;
 
 	// Model streamed data
-	FModelStreamableData ModelStreamableData;
+	MutablePrivate::FModelStreamableData ModelStreamableData;
 
 	// Bytes store streameable files coming form the CO itself.
 	TArray64<uint8> MorphDataBytes;
diff --git a/Engine/Plugins/Experimental/Mutable/Source/CustomizableObjectEditor/Private/MuCOE/CustomizableObjectCompiler.cpp b/Engine/Plugins/Experimental/Mutable/Source/CustomizableObjectEditor/Private/MuCOE/CustomizableObjectCompiler.cpp
--- a/Engine/Plugins/Experimental/Mutable/Source/CustomizableObjectEditor/Private/MuCOE/CustomizableObjectCompiler.cpp
+++ b/Engine/Plugins/Experimental/Mutable/Source/CustomizableObjectEditor/Private/MuCOE/CustomizableObjectCompiler.cpp
@@ -930,6 +930,8 @@ void FCustomizableObjectCompiler::CompileInternal(bool bAsync)
 	}
 	else
 	{
+		ModelStreamableBulkData = MakeShared<FModelStreamableBulkData>();
+
 		FModelResources ModelResources;
 		
 		ModelResources.ReferenceSkeletalMeshesData = MoveTemp(GenerationContext.ReferenceSkeletalMeshesData);
@@ -1019,14 +1021,14 @@ void FCustomizableObjectCompiler::CompileInternal(bool bAsync)
 			RealTimeMorphDataSize += MeshData.Value.Data.Num();
 		}
 		
-		ModelResources.RealTimeMorphStreamables.Empty(32);
+		ModelStreamableBulkData->RealTimeMorphStreamables.Empty(32);
 		ModelResources.EditorOnlyMorphTargetReconstructionData.Empty(RealTimeMorphDataSize);
 
 		uint64 RealTimeMorphDataOffsetInBytes = 0;
 		for (const TPair<uint32, FRealTimeMorphMeshData>& MeshData : GenerationContext.RealTimeMorphTargetPerMeshData)
 		{
 			const uint32 DataSizeInBytes = (uint32)MeshData.Value.Data.Num()*sizeof(FMorphTargetVertexData); 
-			FRealTimeMorphStreamable& ResourceMeshData = ModelResources.RealTimeMorphStreamables.FindOrAdd(MeshData.Key);
+			FRealTimeMorphStreamable& ResourceMeshData = ModelStreamableBulkData->RealTimeMorphStreamables.FindOrAdd(MeshData.Key);
 			
 			check(ResourceMeshData.NameResolutionMap.IsEmpty());
 			check(ResourceMeshData.Size == 0);
@@ -1047,14 +1049,14 @@ void FCustomizableObjectCompiler::CompileInternal(bool bAsync)
 			ClothingDataNum += MeshData.Value.Data.Num();
 		}
 		
-		ModelResources.ClothingStreamables.Empty(32);
+		ModelStreamableBulkData->ClothingStreamables.Empty(32);
 		ModelResources.EditorOnlyClothingMeshToMeshVertData.Empty(ClothingDataNum);
 
 		uint64 ClothingDataOffsetInBytes = 0;
 		for (const TPair<uint32, FClothingMeshData>& MeshData : GenerationContext.ClothingPerMeshData)
 		{
 			const uint32 DataSizeInBytes = (uint32)MeshData.Value.Data.Num()*sizeof(FCustomizableObjectMeshToMeshVertData); 
-			FClothingStreamable& ResourceMeshData = ModelResources.ClothingStreamables.FindOrAdd(MeshData.Key);
+			FClothingStreamable& ResourceMeshData = ModelStreamableBulkData->ClothingStreamables.FindOrAdd(MeshData.Key);
 			
 			check(ResourceMeshData.ClothingAssetIndex == INDEX_NONE);
 			check(ResourceMeshData.ClothingAssetLOD == INDEX_NONE);
@@ -1415,13 +1417,10 @@ void FCustomizableObjectCompiler::FinishCompilationTask()
 	// At this point it is assumed that all data goes into a single file.
 	if (CurrentModel)
 	{
-		// Always work with the ModelResources (Editor) when compiling. They'll be copied to the cooked version during PreSave.
-		FModelResources& ModelResources = CurrentObject->GetPrivate()->GetModelResources(false);
-
 		const int32 NumStreamingFiles = CurrentModel->GetRomCount();
 
-		TSharedPtr<TMap<uint32, FMutableStreamableBlock>> HashToStreamableBlock = MakeShared<TMap<uint32, FMutableStreamableBlock>>();
-		HashToStreamableBlock->Empty(NumStreamingFiles);
+		TMap<uint32, FMutableStreamableBlock>& ModelStreamables = ModelStreamableBulkData->ModelStreamables;
+		ModelStreamables.Empty(NumStreamingFiles);
 
 		uint64 Offset = 0;
 		for (int32 FileIndex = 0; FileIndex < NumStreamingFiles; ++FileIndex)
@@ -1429,11 +1428,12 @@ void FCustomizableObjectCompiler::FinishCompilationTask()
 			const uint32 ResourceId = CurrentModel->GetRomId(FileIndex);
 			const uint32 ResourceSize = CurrentModel->GetRomSize(FileIndex);
 			mu::ERomFlags Flags = CurrentModel->GetRomFlags(FileIndex);
-			HashToStreamableBlock->Add(ResourceId, FMutableStreamableBlock{ 0, uint32(Flags), Offset });
+			ModelStreamables.Add(ResourceId, FMutableStreamableBlock{ 0, uint32(Flags), Offset });
 			Offset += ResourceSize;
 		}
 
-		ModelResources.HashToStreamableBlock = *(HashToStreamableBlock.Get());
+		// Always work with the ModelStreamableData (Editor) when compiling. They'll be copied to the cooked version during PreSave.
+		CurrentObject->GetPrivate()->SetModelStreamableBulkData(ModelStreamableBulkData, false);
 	}
 
 	// Generate ParameterProperties and IntParameterLookUpTable
@@ -1467,7 +1467,7 @@ void FCustomizableObjectCompiler::FinishSavingDerivedDataTask()
 
 		check(!CurrentObject->GetPrivate()->CachedPlatformsData.Find(PlatformName));
 
-		FMutableCachedPlatformData& Data = CurrentObject->GetPrivate()->CachedPlatformsData.Add(PlatformName);
+		MutablePrivate::FMutableCachedPlatformData& Data = CurrentObject->GetPrivate()->CachedPlatformsData.Add(PlatformName);
 
 		// Cache CO data and mu::Model
 		FMemoryWriter64 MemoryWriter(Data.ModelData);
diff --git a/Engine/Plugins/Experimental/Mutable/Source/CustomizableObjectEditor/Private/MuCOE/CustomizableObjectCompiler.h b/Engine/Plugins/Experimental/Mutable/Source/CustomizableObjectEditor/Private/MuCOE/CustomizableObjectCompiler.h
--- a/Engine/Plugins/Experimental/Mutable/Source/CustomizableObjectEditor/Private/MuCOE/CustomizableObjectCompiler.h
+++ b/Engine/Plugins/Experimental/Mutable/Source/CustomizableObjectEditor/Private/MuCOE/CustomizableObjectCompiler.h
@@ -168,6 +168,7 @@ private:
 	TArray<TObjectPtr<UCustomizableObject>> ArrayGCProtect;
 
 	TSharedPtr<mu::Model, ESPMode::ThreadSafe> CurrentModel;
+	TSharedPtr<FModelStreamableBulkData> ModelStreamableBulkData;
 
 	// Protected from GC with FCustomizableObjectCompiler::AddReferencedObjects
 	TObjectPtr<UCustomizableObject> CurrentObject = nullptr;
diff --git a/Engine/Plugins/Experimental/Mutable/Source/CustomizableObjectEditor/Private/MuCOE/CustomizableObjectCookPackageSplitter.cpp b/Engine/Plugins/Experimental/Mutable/Source/CustomizableObjectEditor/Private/MuCOE/CustomizableObjectCookPackageSplitter.cpp
--- a/Engine/Plugins/Experimental/Mutable/Source/CustomizableObjectEditor/Private/MuCOE/CustomizableObjectCookPackageSplitter.cpp
+++ b/Engine/Plugins/Experimental/Mutable/Source/CustomizableObjectEditor/Private/MuCOE/CustomizableObjectCookPackageSplitter.cpp
@@ -153,7 +153,7 @@ bool FCustomizableObjectCookPackageSplitter::ShouldSplit(UObject* SplitData)
 {
 	const UCustomizableObject* Object = CastChecked<UCustomizableObject>(SplitData);
 
-	return Object->GetPrivate()->GetStreamedResourceData().Num() > 0 || Object->GetPrivate()->GetStreamedExtensionData().Num() > 0;
+	return !Object->IsChildObject() && (Object->GetPrivate()->GetStreamedResourceData().Num() > 0 || Object->GetPrivate()->GetStreamedExtensionData().Num() > 0);
 }
 
 TArray<ICookPackageSplitter::FGeneratedPackage> FCustomizableObjectCookPackageSplitter::GetGenerateList(
