[mutable] Layout improvements.
- Support for blocks defined outside the unit UV space (editor and runtime), but in the positive 2D quarter-space.
- Add specific options for autogeneration of UV blocks.
- Move autogeneration of blocks to the CO compilation stage.
- Show preview of autogenerated blocks and rename "Generate" to "Consolidate".
- Add UV-island based block auto-generation.
#rnx

diff --git a/Engine/Plugins/Experimental/Mutable/Source/CustomizableObject/Internal/MuCO/CustomizableObjectPrivate.h b/Engine/Plugins/Experimental/Mutable/Source/CustomizableObject/Internal/MuCO/CustomizableObjectPrivate.h
--- a/Engine/Plugins/Experimental/Mutable/Source/CustomizableObject/Internal/MuCO/CustomizableObjectPrivate.h
+++ b/Engine/Plugins/Experimental/Mutable/Source/CustomizableObject/Internal/MuCO/CustomizableObjectPrivate.h
@@ -1313,6 +1313,8 @@ public:
 
 		ModelStreamableBulkData,
 
+		LayoutBlocksAsInt32,
+
 		// -----<new versions can be added above this line>--------
 		LastCustomizableObjectVersion
 	};
diff --git a/Engine/Plugins/Experimental/Mutable/Source/CustomizableObjectEditor/Private/MuCOE/CustomizableObjectLayout.cpp b/Engine/Plugins/Experimental/Mutable/Source/CustomizableObjectEditor/Private/MuCOE/CustomizableObjectLayout.cpp
--- a/Engine/Plugins/Experimental/Mutable/Source/CustomizableObjectEditor/Private/MuCOE/CustomizableObjectLayout.cpp
+++ b/Engine/Plugins/Experimental/Mutable/Source/CustomizableObjectEditor/Private/MuCOE/CustomizableObjectLayout.cpp
@@ -4,8 +4,9 @@
 
 #include "Engine/StaticMesh.h"
 #include "MuCO/CustomizableObjectCompilerTypes.h"
+#include "MuCO/CustomizableObjectInstancePrivate.h"
 #include "MuCOE/CustomizableObjectCompiler.h"
-#include "MuCOE/GenerateMutableSource/GenerateMutableSourceMesh.h"
+#include "MuCOE/GenerateMutableSource/GenerateMutableSourceLayout.h"
 #include "MuCOE/GraphTraversal.h"
 #include "MuCOE/ICustomizableObjectEditor.h"
 #include "MuCOE/Nodes/CustomizableObjectNodeLayoutBlocks.h"
@@ -36,12 +37,6 @@ void UCustomizableObjectLayout::SetLayout(UObject* InMesh, int32 LODIndex, int32
 }
 
 
-void UCustomizableObjectLayout::SetPackingStrategy(ECustomizableObjectTextureLayoutPackingStrategy Strategy)
-{
-	PackingStrategy = Strategy;
-}
-
-
 mu::EPackStrategy ConvertLayoutStrategy(const ECustomizableObjectTextureLayoutPackingStrategy LayoutPackStrategy)
 {
 	mu::EPackStrategy PackStrategy = mu::EPackStrategy::Fixed;
@@ -86,64 +81,92 @@ void UCustomizableObjectLayout::SetLayoutName(FString Name)
 }
 
 
-void UCustomizableObjectLayout::GenerateBlocksFromUVs()
+void UCustomizableObjectLayout::GenerateAutomaticBlocksFromUVs()
 {
-	mu::Ptr<mu::NodeLayout> LayoutNode;
-
 	UCustomizableObjectNode* Node = Cast<UCustomizableObjectNode>(GetOuter());
 
-	if (Node && Mesh)
+	if (!Node || !Mesh)
 	{
-		//Creating a GenerationContext
-		FCustomizableObjectCompiler Compiler;
-		UCustomizableObject* Object = Node->GetGraphEditor()->GetCustomizableObject();
-		FCompilationOptions Options = Object->GetPrivate()->GetCompileOptions();
-	
-		FMutableGraphGenerationContext GenerationContext(Object, &Compiler, Options);
-	
-		//Transforming skeletalmesh to mutable mesh
-		mu::MeshPtr	MutableMesh = nullptr;
-		
-		if (USkeletalMesh* SkeletalMesh = Cast<USkeletalMesh>(Mesh))
-		{
-			// We don't need all the data to generate the blocks
-			const EMutableMeshConversionFlags ShapeFlags = 
-					EMutableMeshConversionFlags::IgnoreSkinning | 
-					EMutableMeshConversionFlags::IgnorePhysics;
+		return;
+	}
 
-			GenerationContext.MeshGenerationFlags.Push(ShapeFlags);
+	if (AutomaticBlocksStrategy == ECustomizableObjectLayoutAutomaticBlocksStrategy::Ignore)
+	{
+		return;
+	}
 
-			GenerationContext.ComponentInfos.Add(FMutableComponentInfo(FName(), SkeletalMesh));
-			MutableMesh = ConvertSkeletalMeshToMutable(SkeletalMesh, TSoftClassPtr<UAnimInstance>(), LOD, Material, LOD, Material, GenerationContext, Node, nullptr);
-		}
-		else if (UStaticMesh* StaticMesh = Cast<UStaticMesh>(Mesh))
-		{
-			MutableMesh = ConvertStaticMeshToMutable(StaticMesh, LOD, Material, GenerationContext, Node);
-		}
-	
-		if (MutableMesh)
-		{
-			// Generating blocks with the mutable mesh
-			LayoutNode = mu::NodeLayout::GenerateLayoutBlocks(MutableMesh, UVChannel, GridSize.X, GridSize.Y);
-		}
-	
-		if (LayoutNode)
+	// Create a GenerationContext
+	FCustomizableObjectCompiler Compiler;
+	UCustomizableObject* Object = Node->GetGraphEditor()->GetCustomizableObject();
+	FCompilationOptions Options = Object->GetPrivate()->GetCompileOptions();
+
+	FMutableGraphGenerationContext GenerationContext(Object, &Compiler, Options);
+
+	bool bOutWasEmpty = false;
+	mu::Ptr<mu::NodeLayout> LayoutNode = CreateMutableLayoutNode(GenerationContext, this, false, bOutWasEmpty );
+
+	if (LayoutNode)
+	{
+		AutomaticBlocks.Empty();
+
+		// Generate the editor layout blocks from the mutable layout
+		for (int32 i = 0; i < LayoutNode->Blocks.Num(); ++i)
 		{
-			Blocks.Empty();
-		
-			// Generating the layout blocks with the mutable layout
-			for (int i = 0; i < LayoutNode->Blocks.Num(); ++i)
+			FIntPoint Min = FIntPoint(LayoutNode->Blocks[i].Min.X, LayoutNode->Blocks[i].Min.Y);
+			FIntPoint Size = FIntPoint(LayoutNode->Blocks[i].Size.X, LayoutNode->Blocks[i].Size.Y);
+
+			// Ignore blocks contained inside any block in the initial block set.
+			bool bContainedInInitialSet = false;
+			for (const FCustomizableObjectLayoutBlock& Block : Blocks)
 			{
-				mu::FImageSize Min = LayoutNode->Blocks[i].Min;
-				mu::FImageSize Size = LayoutNode->Blocks[i].Size;
-		
-				FCustomizableObjectLayoutBlock block(FIntPoint(Min.X, Min.Y), FIntPoint(Min.X + Size.X, Min.Y + Size.Y));
-				Blocks.Add(block);
+				FInt32Rect ExistingRect(Block.Min, Block.Max + FIntPoint(1));
+				if (ExistingRect.Contains(Min) && ExistingRect.Contains(Min + Size))
+				{
+					bContainedInInitialSet = true;
+					break;
+				}
+			}
+			if (bContainedInInitialSet)
+			{
+				continue;
+			}
+
+			FCustomizableObjectLayoutBlock Block(FIntPoint(Min.X, Min.Y), FIntPoint(Min.X + Size.X, Min.Y + Size.Y));
+			Block.bIsAutomatic = true;
+			Block.Id = FGuid::NewGuid();
+
+			mu::Ptr<mu::Image> Mask = LayoutNode->Blocks[i].Mask;
+			if (Mask)
+			{
+				UTexture2D* UnrealImage = NewObject<UTexture2D>(UTexture2D::StaticClass());
+
+				FMutableModelImageProperties Props;
+				Props.Filter = TF_Nearest;
+				Props.SRGB = true;
+				Props.LODBias = 0;
+				ConvertImage(UnrealImage, Mask, Props );
+				UnrealImage->NeverStream = true;
+				UnrealImage->UpdateResource();
+
+				Block.Mask = UnrealImage;
 			}
-		
-			Node->PostEditChange();
-			Node->GetGraph()->MarkPackageDirty();
+			AutomaticBlocks.Add(Block);
 		}
+
+		Node->PostEditChange();
+		Node->GetGraph()->MarkPackageDirty();
+	}
+}
+
+
+void UCustomizableObjectLayout::ConsolidateAutomaticBlocks()
+{
+	Blocks.Append(AutomaticBlocks);
+	AutomaticBlocks.Empty();
+	for (FCustomizableObjectLayoutBlock& Block : Blocks)
+	{
+		Block.Id = FGuid::NewGuid();
+		Block.bIsAutomatic = false;
 	}
 }
 
@@ -172,7 +195,7 @@ void UCustomizableObjectLayout::GetUVChannel(TArray<FVector2f>& UVs, int32 UVCha
 
 int32 UCustomizableObjectLayout::FindBlock(const FGuid& InId) const
 {
-	for (int Index = 0; Index < Blocks.Num(); ++Index)
+	for (int32 Index = 0; Index < Blocks.Num(); ++Index)
 	{
 		if (Blocks[Index].Id == InId)
 		{
@@ -196,8 +219,4 @@ void UCustomizableObjectLayout::SetIgnoreWarningsLOD(int32 LODValue)
 }
 
 
-void UCustomizableObjectLayout::SetBlockReductionMethod(ECustomizableObjectLayoutBlockReductionMethod Method)
-{
-	BlockReductionMethod = Method;
-}
 #undef LOCTEXT_NAMESPACE
diff --git a/Engine/Plugins/Experimental/Mutable/Source/CustomizableObjectEditor/Private/MuCOE/GenerateMutableSource/GenerateMutableSourceLayout.cpp b/Engine/Plugins/Experimental/Mutable/Source/CustomizableObjectEditor/Private/MuCOE/GenerateMutableSource/GenerateMutableSourceLayout.cpp
--- a/Engine/Plugins/Experimental/Mutable/Source/CustomizableObjectEditor/Private/MuCOE/GenerateMutableSource/GenerateMutableSourceLayout.cpp
+++ b/Engine/Plugins/Experimental/Mutable/Source/CustomizableObjectEditor/Private/MuCOE/GenerateMutableSource/GenerateMutableSourceLayout.cpp
@@ -2,11 +2,13 @@
 
 #include "MuCOE/GenerateMutableSource/GenerateMutableSourceLayout.h"
 
+#include "MuCOE/GenerateMutableSource/GenerateMutableSourceMesh.h"
 #include "MuCOE/CustomizableObjectCompiler.h"
 #include "MuCOE/CustomizableObjectLayout.h"
 #include "MuCOE/GraphTraversal.h"
 #include "MuCOE/Nodes/CustomizableObjectNodeLayoutBlocks.h"
 #include "MuCOE/Nodes/CustomizableObjectNodeSkeletalMesh.h"
+#include "Engine/StaticMesh.h"
 
 #define LOCTEXT_NAMESPACE "CustomizableObjectEditor"
 
@@ -84,14 +86,13 @@ mu::Ptr<mu::NodeLayout> CreateMutableLayoutNode(FMutableGraphGenerationContext&
 	bWasEmpty = false;
 	mu::Ptr<mu::NodeLayout> LayoutNode = new mu::NodeLayout;
 
-	LayoutNode->Size = UE::Math::TIntVector2<uint16>(UnrealLayout->GetGridSize().X, UnrealLayout->GetGridSize().Y);
-	LayoutNode->MaxSize = UE::Math::TIntVector2<uint16>(UnrealLayout->GetMaxGridSize().X, UnrealLayout->GetMaxGridSize().Y);
-	LayoutNode->Blocks.SetNum(UnrealLayout->Blocks.Num() ? UnrealLayout->Blocks.Num() : 1);
+	LayoutNode->Size = FIntVector2(UnrealLayout->GetGridSize().X, UnrealLayout->GetGridSize().Y);
+	LayoutNode->MaxSize = FIntVector2(UnrealLayout->GetMaxGridSize().X, UnrealLayout->GetMaxGridSize().Y);
 
-	mu::EPackStrategy PackStrategy = ConvertLayoutStrategy(UnrealLayout->GetPackingStrategy());
+	mu::EPackStrategy PackStrategy = ConvertLayoutStrategy(UnrealLayout->PackingStrategy);
 	LayoutNode->Strategy = PackStrategy;
 
-	LayoutNode->ReductionMethod = (UnrealLayout->GetBlockReductionMethod() == ECustomizableObjectLayoutBlockReductionMethod::Halve ? mu::EReductionMethod::Halve : mu::EReductionMethod::Unitary);
+	LayoutNode->ReductionMethod = (UnrealLayout->BlockReductionMethod == ECustomizableObjectLayoutBlockReductionMethod::Halve ? mu::EReductionMethod::Halve : mu::EReductionMethod::Unitary);
 
 	if (bIgnoreLayoutWarnings)
 	{
@@ -103,21 +104,71 @@ mu::Ptr<mu::NodeLayout> CreateMutableLayoutNode(FMutableGraphGenerationContext&
 		LayoutNode->FirstLODToIgnoreWarnings = UnrealLayout->GetIgnoreVertexLayoutWarnings() ? UnrealLayout->GetFirstLODToIgnoreWarnings() : -1;
 	}
 
-	if (UnrealLayout->Blocks.Num())
+	LayoutNode->Blocks.SetNum(UnrealLayout->Blocks.Num());
+	for (int32 BlockIndex = 0; BlockIndex < UnrealLayout->Blocks.Num(); ++BlockIndex)
 	{
-		for (int BlockIndex = 0; BlockIndex < UnrealLayout->Blocks.Num(); ++BlockIndex)
+		LayoutNode->Blocks[BlockIndex] = ToMutable(GenerationContext, UnrealLayout->Blocks[BlockIndex]);
+	}
+
+
+	ECustomizableObjectLayoutAutomaticBlocksStrategy AutomaticBlockStrategy = UnrealLayout->AutomaticBlocksStrategy;
+
+	if (AutomaticBlockStrategy == ECustomizableObjectLayoutAutomaticBlocksStrategy::Ignore)
+	{
+		// Legacy behavior
+		if (!UnrealLayout->Blocks.Num())
 		{
-			LayoutNode->Blocks[BlockIndex] = ToMutable(GenerationContext, UnrealLayout->Blocks[BlockIndex]);
+			bWasEmpty = true;
+			LayoutNode->Blocks.SetNum(1);
+			LayoutNode->Blocks[0].Min = { 0,0 };
+			LayoutNode->Blocks[0].Size = LayoutNode->Size;
+			LayoutNode->Blocks[0].Priority = 0;
+			LayoutNode->Blocks[0].bReduceBothAxes = false;
+			LayoutNode->Blocks[0].bReduceByTwo = false;
 		}
 	}
 	else
 	{
-		bWasEmpty = true;
-		LayoutNode->Blocks[0].Min = { 0,0 };
-		LayoutNode->Blocks[0].Size = { uint16(UnrealLayout->GetGridSize().X), uint16(UnrealLayout->GetGridSize().Y) };
-		LayoutNode->Blocks[0].Priority = 0;
-		LayoutNode->Blocks[0].bReduceBothAxes = false;
-		LayoutNode->Blocks[0].bReduceByTwo = false;
+		// Convert the UE mesh in the layout into a Mutable mesh.
+		mu::Ptr<mu::Mesh> MutableMesh = nullptr;
+
+		if (USkeletalMesh* SkeletalMesh = Cast<USkeletalMesh>(UnrealLayout->GetMesh()))
+		{
+			// We don't need all the data to generate the blocks
+			const EMutableMeshConversionFlags ShapeFlags =
+				EMutableMeshConversionFlags::IgnoreSkinning |
+				EMutableMeshConversionFlags::IgnorePhysics;
+
+			GenerationContext.MeshGenerationFlags.Push(ShapeFlags);
+
+			GenerationContext.ComponentInfos.Add(FMutableComponentInfo(FName(), SkeletalMesh));
+			MutableMesh = ConvertSkeletalMeshToMutable(SkeletalMesh, TSoftClassPtr<UAnimInstance>(), UnrealLayout->GetLOD(), UnrealLayout->GetMaterial(), UnrealLayout->GetLOD(), UnrealLayout->GetMaterial(), GenerationContext, nullptr, nullptr);
+
+			GenerationContext.MeshGenerationFlags.Pop();
+		}
+		else if (UStaticMesh* StaticMesh = Cast<UStaticMesh>(UnrealLayout->GetMesh()))
+		{
+			MutableMesh = ConvertStaticMeshToMutable(StaticMesh, UnrealLayout->GetLOD(), UnrealLayout->GetMaterial(), GenerationContext, nullptr);
+		}
+
+		if (MutableMesh)
+		{
+			// Generating blocks with the mutable mesh
+			if (AutomaticBlockStrategy == ECustomizableObjectLayoutAutomaticBlocksStrategy::Rectangles)
+			{
+				LayoutNode->GenerateLayoutBlocks(MutableMesh, UnrealLayout->GetUVChannel());
+			}
+			else if (AutomaticBlockStrategy == ECustomizableObjectLayoutAutomaticBlocksStrategy::UVIslands)
+			{
+				bool bMergeChildBlocks = UnrealLayout->AutomaticBlocksMergeStrategy == ECustomizableObjectLayoutAutomaticBlocksMergeStrategy::MergeChildBlocks;
+				LayoutNode->GenerateLayoutBlocksFromUVIslands(MutableMesh, UnrealLayout->GetUVChannel(), bMergeChildBlocks);
+			}
+			else
+			{
+				// Unimplemented
+				check(false);
+			}
+		}
 	}
 
 	return LayoutNode;
diff --git a/Engine/Plugins/Experimental/Mutable/Source/CustomizableObjectEditor/Private/MuCOE/GenerateMutableSource/GenerateMutableSourceMesh.cpp b/Engine/Plugins/Experimental/Mutable/Source/CustomizableObjectEditor/Private/MuCOE/GenerateMutableSource/GenerateMutableSourceMesh.cpp
--- a/Engine/Plugins/Experimental/Mutable/Source/CustomizableObjectEditor/Private/MuCOE/GenerateMutableSource/GenerateMutableSourceMesh.cpp
+++ b/Engine/Plugins/Experimental/Mutable/Source/CustomizableObjectEditor/Private/MuCOE/GenerateMutableSource/GenerateMutableSourceMesh.cpp
@@ -1162,7 +1162,7 @@ mu::MeshPtr ConvertSkeletalMeshToMutable(const USkeletalMesh* InSkeletalMesh, co
 
 	// TODO: Add Mesh generation flags to not include RT Morph and clothing if not needed.
 	int32 NextBufferIndex = MutableMesh->VertexBuffers.Buffers.Num();
-	if (GenerationContext.Options.bRealTimeMorphTargetsEnabled)
+	if (GenerationContext.Options.bRealTimeMorphTargetsEnabled && CurrentNode)
 	{
 
 		// This call involves resolving every TObjectPtr<UMorphTarget> to a UMorphTarget*, so
diff --git a/Engine/Plugins/Experimental/Mutable/Source/CustomizableObjectEditor/Private/MuCOE/Nodes/CustomizableObjectNodeEditMaterialDetails.cpp b/Engine/Plugins/Experimental/Mutable/Source/CustomizableObjectEditor/Private/MuCOE/Nodes/CustomizableObjectNodeEditMaterialDetails.cpp
--- a/Engine/Plugins/Experimental/Mutable/Source/CustomizableObjectEditor/Private/MuCOE/Nodes/CustomizableObjectNodeEditMaterialDetails.cpp
+++ b/Engine/Plugins/Experimental/Mutable/Source/CustomizableObjectEditor/Private/MuCOE/Nodes/CustomizableObjectNodeEditMaterialDetails.cpp
@@ -4,6 +4,7 @@
 
 #include "DetailLayoutBuilder.h"
 #include "DetailWidgetRow.h"
+#include "IDetailGroup.h"
 #include "IDetailsView.h"
 #include "MuCOE/Nodes/CustomizableObjectNodeEditMaterial.h"
 #include "MuCOE/Nodes/CustomizableObjectNodeMaterialBase.h"
@@ -11,6 +12,7 @@
 #include "MuCOE/PinViewer/SPinViewer.h"
 #include "Widgets/Layout/SBox.h"
 #include "Widgets/Text/STextBlock.h"
+#include "Widgets/Input/STextComboBox.h"
 #include "UObject/Package.h"
 
 class FString;
@@ -35,49 +37,109 @@ void FCustomizableObjectNodeEditMaterialDetails::CustomizeDetails( IDetailLayout
 		Node = Cast<UCustomizableObjectNodeEditMaterial>(DetailsView->GetSelectedObjects()[0].Get());
 	}
 
+	IDetailCategoryBuilder& LayoutCategory = DetailBuilder.EditCategory("LayoutOptions");
 	IDetailCategoryBuilder& BlocksCategory = DetailBuilder.EditCategory( "Blocks" );
 
-	if (Node)
+	if (!Node)
 	{
-		// Add blocks selector
-		LayoutBlocksEditor = SNew(SCustomizableObjectNodeLayoutBlocksEditor);
-
-		BlocksCategory.AddCustomRow(LOCTEXT("BlocksDetails_BlockInstructions", "BlockInstructions"))
-		[
-			SNew(SBox)
-			.HeightOverride(700.0f)
-			.WidthOverride(700.0f)
+		BlocksCategory.AddCustomRow(LOCTEXT("BlocksDetails_NodeNotFound", "NodeNotFound"))
 			[
-				LayoutBlocksEditor.ToSharedRef()
-			]
-		];
+				SNew(STextBlock)
+					.Text(LOCTEXT("Node not found", "Node not found"))
+			];
+		return;
+	}
+
+	// Grid size combo
+	{
+		int32 MaxGridSize = 128;
+		LayoutGridSizes.Empty();
 
-		// Try to find the parent layout, because we want to show its UVs in the widget
-		UCustomizableObjectLayout* ParentLayout = nullptr;
-		if (UCustomizableObjectNodeMaterialBase* ParentMaterialNode = Node->GetParentMaterialNode())
+		TSharedPtr<FString> CurrentSize;
+		for (int32 Size = 1; Size <= MaxGridSize; Size *= 2)
 		{
-			TArray<UCustomizableObjectLayout*> Layouts = ParentMaterialNode->GetLayouts();
+			LayoutGridSizes.Add(MakeShareable(new FString(FString::Printf(TEXT("%d x %d"), Size, Size))));
 
-			if (!Layouts.IsValidIndex(Node->ParentLayoutIndex))
-			{
-				UE_LOG(LogMutable, Warning, TEXT("[%s] UCustomizableObjectNodeEditMaterial refers to an invalid texture layout index %d. Parent node has %d layouts."),
-					*Node->GetOutermost()->GetName(), Node->ParentLayoutIndex, Layouts.Num());
-			}
-			else
+			if (Node->Layout->GetGridSize() == FIntPoint(Size))
 			{
-				ParentLayout = Layouts[Node->ParentLayoutIndex];
+				CurrentSize = LayoutGridSizes.Last();
 			}
 		}
 
-		LayoutBlocksEditor->SetCurrentLayout(Node->Layout, ParentLayout);
+		IDetailGroup& LayoutOptionsGroup = LayoutCategory.AddGroup(TEXT("LayoutOptionsGroup"), LOCTEXT("LayoutGroup", "Layout Group"), false, true);
+		LayoutOptionsGroup.HeaderRow()
+			.NameContent()
+			[
+				SNew(STextBlock)
+					.Text(LOCTEXT("LayoutGridSizeText", "Grid Size"))
+					.Font(DetailBuilder.GetDetailFont())
+			]
+			.ValueContent()
+			[
+				SNew(STextComboBox)
+					.InitiallySelectedItem(CurrentSize)
+					.OptionsSource(&LayoutGridSizes)
+					.OnSelectionChanged(this, &FCustomizableObjectNodeEditMaterialDetails::OnGridSizeChanged)
+					.Font(DetailBuilder.GetDetailFont())
+			];
 	}
-	else
-	{
-		BlocksCategory.AddCustomRow( LOCTEXT("BlocksDetails_NodeNotFound", "NodeNotFound") )
+
+	// Block editor
+	LayoutBlocksEditor = SNew(SCustomizableObjectNodeLayoutBlocksEditor);
+
+	BlocksCategory.AddCustomRow(LOCTEXT("BlocksDetails_BlockInstructions", "BlockInstructions"))
+	[
+		SNew(SBox)
+		.HeightOverride(700.0f)
+		.WidthOverride(700.0f)
 		[
-			SNew( STextBlock )
-			.Text( LOCTEXT( "Node not found", "Node not found" ) )
-		];
+			LayoutBlocksEditor.ToSharedRef()
+		]
+	];
+
+	UpdateLayout();
+}
+
+
+void FCustomizableObjectNodeEditMaterialDetails::UpdateLayout()
+{
+	// Try to find the parent layout, because we want to show its UVs in the widget
+	UCustomizableObjectLayout* ParentLayout = nullptr;
+	if (UCustomizableObjectNodeMaterialBase* ParentMaterialNode = Node->GetParentMaterialNode())
+	{
+		TArray<UCustomizableObjectLayout*> Layouts = ParentMaterialNode->GetLayouts();
+
+		if (!Layouts.IsValidIndex(Node->ParentLayoutIndex))
+		{
+			UE_LOG(LogMutable, Warning, TEXT("[%s] UCustomizableObjectNodeEditMaterial refers to an invalid texture layout index %d. Parent node has %d layouts."),
+				*Node->GetOutermost()->GetName(), Node->ParentLayoutIndex, Layouts.Num());
+		}
+		else
+		{
+			ParentLayout = Layouts[Node->ParentLayoutIndex];
+		}
+	}
+
+	LayoutBlocksEditor->SetCurrentLayout(Node->Layout, ParentLayout);
+}
+
+
+void FCustomizableObjectNodeEditMaterialDetails::OnGridSizeChanged(TSharedPtr<FString> NewSelection, ESelectInfo::Type SelectInfo)
+{
+	if (!Node->Layout)
+	{
+		return;
+	}
+
+	int32 Size = 1 << LayoutGridSizes.Find(NewSelection);
+
+	if (Node->Layout->GetGridSize().X != Size || Node->Layout->GetGridSize().Y != Size)
+	{
+		Node->Layout->SetGridSize(FIntPoint(Size));
+
+		Node->MarkPackageDirty();
+
+		UpdateLayout();
 	}
 }
 
diff --git a/Engine/Plugins/Experimental/Mutable/Source/CustomizableObjectEditor/Private/MuCOE/Nodes/CustomizableObjectNodeEditMaterialDetails.h b/Engine/Plugins/Experimental/Mutable/Source/CustomizableObjectEditor/Private/MuCOE/Nodes/CustomizableObjectNodeEditMaterialDetails.h
--- a/Engine/Plugins/Experimental/Mutable/Source/CustomizableObjectEditor/Private/MuCOE/Nodes/CustomizableObjectNodeEditMaterialDetails.h
+++ b/Engine/Plugins/Experimental/Mutable/Source/CustomizableObjectEditor/Private/MuCOE/Nodes/CustomizableObjectNodeEditMaterialDetails.h
@@ -26,4 +26,13 @@ private:
 
 	// Layout block editor widget
 	TSharedPtr<SCustomizableObjectNodeLayoutBlocksEditor> LayoutBlocksEditor;
+
+	/** List of available layout grid sizes. */
+	TArray< TSharedPtr< FString > > LayoutGridSizes;
+
+	/** Layout Options Callbacks */
+	void OnGridSizeChanged(TSharedPtr<FString> NewSelection, ESelectInfo::Type SelectInfo);
+
+	/** Reset the layout in the widget to force a refresh. */
+	void UpdateLayout();
 };
diff --git a/Engine/Plugins/Experimental/Mutable/Source/CustomizableObjectEditor/Private/MuCOE/Nodes/CustomizableObjectNodeLayoutBlocks.cpp b/Engine/Plugins/Experimental/Mutable/Source/CustomizableObjectEditor/Private/MuCOE/Nodes/CustomizableObjectNodeLayoutBlocks.cpp
--- a/Engine/Plugins/Experimental/Mutable/Source/CustomizableObjectEditor/Private/MuCOE/Nodes/CustomizableObjectNodeLayoutBlocks.cpp
+++ b/Engine/Plugins/Experimental/Mutable/Source/CustomizableObjectEditor/Private/MuCOE/Nodes/CustomizableObjectNodeLayoutBlocks.cpp
@@ -186,7 +186,7 @@ void UCustomizableObjectNodeLayoutBlocks::Serialize(FArchive& Ar)
 		Layout->SetGridSize(GridSize_DEPRECATED);
 		Layout->SetMaxGridSize(MaxGridSize_DEPRECATED);
 		Layout->Blocks = Blocks_DEPRECATED;
-		Layout->SetPackingStrategy(PackingStrategy_DEPRECATED);
+		Layout->PackingStrategy = PackingStrategy_DEPRECATED;
 
 		if (Layout->GetGridSize() == FIntPoint::ZeroValue)
 		{
diff --git a/Engine/Plugins/Experimental/Mutable/Source/CustomizableObjectEditor/Private/MuCOE/Nodes/CustomizableObjectNodeLayoutBlocksDetails.cpp b/Engine/Plugins/Experimental/Mutable/Source/CustomizableObjectEditor/Private/MuCOE/Nodes/CustomizableObjectNodeLayoutBlocksDetails.cpp
--- a/Engine/Plugins/Experimental/Mutable/Source/CustomizableObjectEditor/Private/MuCOE/Nodes/CustomizableObjectNodeLayoutBlocksDetails.cpp
+++ b/Engine/Plugins/Experimental/Mutable/Source/CustomizableObjectEditor/Private/MuCOE/Nodes/CustomizableObjectNodeLayoutBlocksDetails.cpp
@@ -50,44 +50,8 @@ void FCustomizableObjectNodeLayoutBlocksDetails::CustomizeDetails( IDetailLayout
 	{
 		LayoutBlocksEditor = SNew(SCustomizableObjectNodeLayoutBlocksEditor);
 
-		TSharedPtr<FString> CurrentSize, CurrentStrategy, CurrentMaxSize, CurrentReductionMethod;
-		FillComboBoxOptionsArrays(CurrentSize, CurrentStrategy, CurrentMaxSize, CurrentReductionMethod);
-
-		// Layout size selector widget
-		CustomizableObjectCategory.AddCustomRow(LOCTEXT("BlocksDetails_SizeSelector", "SizeSelector"))
-		.NameContent()
-		[
-			SNew(STextBlock)
-			.Text(LOCTEXT("LayoutGridSizeText", "Grid Size"))
-			.Font(DetailBuilder.GetDetailFont())
-		]
-		.ValueContent()
-		[
-			SNew(STextComboBox)
-			.InitiallySelectedItem(CurrentSize)
-			.OptionsSource(&LayoutGridSizes)
-			.OnSelectionChanged(this, &FCustomizableObjectNodeLayoutBlocksDetails::OnGridSizeChanged)
-			.Font(DetailBuilder.GetDetailFont())
-		];
-
-
-		// Layout Strategy options. Hardcoded: we should get names and tooltips from the enum property
-		{
-			LayoutPackingStrategies.Empty();
-			LayoutPackingStrategiesTooltips.Empty();
-
-			LayoutPackingStrategies.Add(MakeShareable(new FString("Resizable")));
-			LayoutPackingStrategiesTooltips.Add(LOCTEXT("LayoutDetails_ResizableStrategyTooltip", "In a layout merge, Layout size will increase if blocks don't fit inside."));
-
-			LayoutPackingStrategies.Add(MakeShareable(new FString("Fixed")));
-			LayoutPackingStrategiesTooltips.Add(LOCTEXT("LayoutDetails_FixedStrategyTooltip", "In a layout merge, the layout will increase its size until the maximum layout grid size"
-				"\nBlock sizes will be reduced if they don't fit inside the layout."
-				"\nSet the reduction priority of each block to control which blocks are reduced first and how they are reduced."));
-
-			LayoutPackingStrategies.Add(MakeShareable(new FString("Overlay")));
-			LayoutPackingStrategiesTooltips.Add(LOCTEXT("LayoutDetails_OverlayStrategyTooltip", "In a layout merge, the layout will not be modified and blocks will be ignored."
-				"\nExtend material nodes just add their layouts on top of the base one"));
-		}
+		TSharedPtr<FString> CurrentSize, CurrentStrategy, CurrentAutoBlocks, CurrentAutoBlocksMerge, CurrentMaxSize, CurrentReductionMethod;
+		FillComboBoxOptionsArrays(CurrentSize, CurrentStrategy, CurrentAutoBlocks, CurrentAutoBlocksMerge, CurrentMaxSize, CurrentReductionMethod);
 
 		// Layout strategy selector group widget
 		IDetailGroup* LayoutStrategyOptionsGroup = &CustomizableObjectCategory.AddGroup(TEXT("LayoutStrategyOptionsGroup"), LOCTEXT("LayoutStrategyGroup", "Layout Strategy Group"), false, true);
@@ -95,8 +59,8 @@ void FCustomizableObjectNodeLayoutBlocksDetails::CustomizeDetails( IDetailLayout
 		.NameContent()
 		[
 			SNew(STextBlock)
-			.Text(LOCTEXT("LayoutStrategy_Text", "Layout Strategy:"))
-			.ToolTipText(LOCTEXT("LayoutStrategyTooltup", "Selects the packing strategy in case of a layout merge."))
+			.Text(LOCTEXT("LayoutStrategy_Text", "Layout Strategy"))
+			.ToolTipText(LOCTEXT("LayoutStrategyTooltip", "Selects the UV packing strategy in case of a mesh merge or clip."))
 			.Font(DetailBuilder.GetDetailFont())
 		]
 		.ValueContent()
@@ -114,6 +78,93 @@ void FCustomizableObjectNodeLayoutBlocksDetails::CustomizeDetails( IDetailLayout
 			]
 		];
 
+		// Automatic block generation strategy
+		LayoutStrategyOptionsGroup->AddWidgetRow()
+			.Visibility(TAttribute<EVisibility>::CreateLambda([this]()
+				{
+					return Node->Layout->PackingStrategy != ECustomizableObjectTextureLayoutPackingStrategy::Overlay ? EVisibility::Visible : EVisibility::Collapsed;
+				}))
+			.NameContent()
+			[
+				SNew(STextBlock)
+					.Text(LOCTEXT("AutoBlockStrategy_Text", "Automatic Blocks Strategy"))
+					.ToolTipText(LOCTEXT("AutoBlockStrategyTooltip", "Selects the strategy to create layout blocks from unassigned UVs."))
+					.Font(DetailBuilder.GetDetailFont())
+			]
+			.ValueContent()
+			[
+				SNew(SSearchableComboBox)
+					.InitiallySelectedItem(CurrentAutoBlocks)
+					.OptionsSource(&AutoBlocksStrategies)
+					.OnSelectionChanged(this, &FCustomizableObjectNodeLayoutBlocksDetails::OnAutoBlocksChanged)
+					.OnGenerateWidget(this, &FCustomizableObjectNodeLayoutBlocksDetails::OnGenerateAutoBlocksComboBox)
+					.ToolTipText(this, &FCustomizableObjectNodeLayoutBlocksDetails::GetSelectedAutoBlocksTooltip)
+					[
+						SNew(STextBlock)
+							.Text(this, &FCustomizableObjectNodeLayoutBlocksDetails::GetSelectedAutoBlocksName)
+							.Font(IDetailLayoutBuilder::GetDetailFont())
+					]
+			];
+
+		// Option to merge child automatic blocks
+		LayoutStrategyOptionsGroup->AddWidgetRow()
+			.Visibility(TAttribute<EVisibility>::CreateLambda([this]()
+				{
+					if (Node->Layout->PackingStrategy == ECustomizableObjectTextureLayoutPackingStrategy::Overlay
+						||
+						Node->Layout->AutomaticBlocksStrategy != ECustomizableObjectLayoutAutomaticBlocksStrategy::UVIslands)
+					{
+						return EVisibility::Collapsed;
+					}
+					return EVisibility::Visible;
+				}))
+			.NameContent()
+			[
+				SNew(STextBlock)
+					.Text(LOCTEXT("AutoBlockMergeStrategy_Text", "Automatic Blocks Merge Strategy"))
+					.ToolTipText(LOCTEXT("AutoBlockMergeStrategyTooltip", "Selects the strategy to merge blocks during automatic generation."))
+					.Font(DetailBuilder.GetDetailFont())
+			]
+			.ValueContent()
+			[
+				SNew(SSearchableComboBox)
+					.InitiallySelectedItem(CurrentAutoBlocksMerge)
+					.OptionsSource(&AutoBlocksMergeStrategies)
+					.OnSelectionChanged(this, &FCustomizableObjectNodeLayoutBlocksDetails::OnAutoBlocksMergeChanged)
+					.OnGenerateWidget(this, &FCustomizableObjectNodeLayoutBlocksDetails::OnGenerateAutoBlocksMergeComboBox)
+					.ToolTipText(this, &FCustomizableObjectNodeLayoutBlocksDetails::GetSelectedAutoBlocksMergeTooltip)
+					[
+						SNew(STextBlock)
+							.Text(this, &FCustomizableObjectNodeLayoutBlocksDetails::GetSelectedAutoBlocksMergeName)
+							.Font(IDetailLayoutBuilder::GetDetailFont())
+					]
+			];
+
+		// Layout size selector widget
+		LayoutStrategyOptionsGroup->AddWidgetRow()
+			.Visibility(TAttribute<EVisibility>::CreateLambda([this]()
+				{
+					if (Node->Layout->PackingStrategy == ECustomizableObjectTextureLayoutPackingStrategy::Overlay)
+					{
+						return EVisibility::Collapsed;
+					}
+					return EVisibility::Visible;
+				}))
+			.NameContent()
+			[
+				SNew(STextBlock)
+					.Text(LOCTEXT("LayoutGridSizeText", "Grid Size"))
+					.Font(DetailBuilder.GetDetailFont())
+			]
+			.ValueContent()
+			[
+				SNew(STextComboBox)
+					.InitiallySelectedItem(CurrentSize)
+					.OptionsSource(&LayoutGridSizes)
+					.OnSelectionChanged(this, &FCustomizableObjectNodeLayoutBlocksDetails::OnGridSizeChanged)
+					.Font(DetailBuilder.GetDetailFont())
+			];
+
 
 		// Max layout size selector widget
 		LayoutStrategyOptionsGroup->AddWidgetRow()
@@ -121,7 +172,7 @@ void FCustomizableObjectNodeLayoutBlocksDetails::CustomizeDetails( IDetailLayout
 		.NameContent()
 		[
 			SNew(STextBlock)
-			.Text(LOCTEXT("MaxLayoutSize_Text", "Max Layout Size:"))
+			.Text(LOCTEXT("MaxLayoutSize_Text", "Max Layout Size"))
 			.Font(DetailBuilder.GetDetailFont())
 		]
 		.ValueContent()
@@ -150,7 +201,7 @@ void FCustomizableObjectNodeLayoutBlocksDetails::CustomizeDetails( IDetailLayout
 		.NameContent()
 		[
 			SNew(STextBlock)
-			.Text(LOCTEXT("ReductionMethod_Text", "Reduction Method:"))
+			.Text(LOCTEXT("ReductionMethod_Text", "Reduction Method"))
 			.ToolTipText(LOCTEXT("Reduction_Method_Tooltip", "Select how blocks will be reduced in case that they do not fit in the layout."))
 			.Font(DetailBuilder.GetDetailFont())
 		]
@@ -175,7 +226,7 @@ void FCustomizableObjectNodeLayoutBlocksDetails::CustomizeDetails( IDetailLayout
 		.NameContent()
 		[
 			SNew(STextBlock)
-			.Text(LOCTEXT("LayoutOptions_IgnoreLodsCheckBox_Text", "Ignore Unassigned Vertices Warning:"))
+			.Text(LOCTEXT("LayoutOptions_IgnoreLodsCheckBox_Text", "Ignore Unassigned Vertices Warning"))
 			.ToolTipText(LOCTEXT("LayoutOptions_IgnoreLodsCheckBox_Tooltip",
 				"If true, warning message \"Source mesh has vertices not assigned to any layout block\" will be ignored."
 				"\n Note:"
@@ -196,7 +247,7 @@ void FCustomizableObjectNodeLayoutBlocksDetails::CustomizeDetails( IDetailLayout
 		.NameContent()
 		[
 			SAssignNew(LODSelectorTextWidget, STextBlock)
-			.Text(LOCTEXT("LayoutOptions_IgnoreLod_Text", "First LOD to ignore:"))
+			.Text(LOCTEXT("LayoutOptions_IgnoreLod_Text", "First LOD to ignore"))
 			.ToolTipText(LOCTEXT("LayoutOptions_IgnoreLod_Tooltip", "LOD from which vertex warning messages will be ignored."))
 			.Font(IDetailLayoutBuilder::GetDetailFont())
 			.IsEnabled(Node->Layout ? Node->Layout->GetIgnoreVertexLayoutWarnings() : false)
@@ -267,12 +318,18 @@ void FCustomizableObjectNodeLayoutBlocksDetails::OnLODBoxValueChanged(int32 Valu
 }
 
 
-void FCustomizableObjectNodeLayoutBlocksDetails::FillComboBoxOptionsArrays(TSharedPtr<FString>& CurrGridSize, TSharedPtr<FString>& CurrStrategy, TSharedPtr<FString>& CurrMaxSize, TSharedPtr<FString>& CurrRedMethod)
+void FCustomizableObjectNodeLayoutBlocksDetails::FillComboBoxOptionsArrays(
+	TSharedPtr<FString>& CurrGridSize, 
+	TSharedPtr<FString>& CurrStrategy, 
+	TSharedPtr<FString>& CurrAutoBlocks,
+	TSharedPtr<FString>& CurrAutoBlocksMerge,
+	TSharedPtr<FString>& CurrMaxSize,
+	TSharedPtr<FString>& CurrRedMethod)
 {
 	if (Node->Layout)
 	{
 		// Static const variable?
-		int32 MaxGridSize = 32;
+		int32 MaxGridSize = 128;
 		LayoutGridSizes.Empty();
 
 		for (int32 Size = 1; Size <= MaxGridSize; Size *= 2)
@@ -290,22 +347,95 @@ void FCustomizableObjectNodeLayoutBlocksDetails::FillComboBoxOptionsArrays(TShar
 			}
 		}
 
-		LayoutPackingStrategies.Empty();
-		LayoutPackingStrategies.Add(MakeShareable(new FString("Resizable")));
-		LayoutPackingStrategies.Add(MakeShareable(new FString("Fixed")));
-		LayoutPackingStrategies.Add(MakeShareable(new FString("Overlay")));
-		CurrStrategy = LayoutPackingStrategies[(uint32)Node->Layout->GetPackingStrategy()];
+		// Layout Strategy options. Hardcoded: we should get names and tooltips from the enum property
+		{
+			LayoutPackingStrategies.Empty();
+			LayoutPackingStrategiesOptions.Empty();
+
+			LayoutPackingStrategies.Add(MakeShareable(new FString("Package and Grow Textures")));
+			LayoutPackingStrategiesOptions.Add(
+				{
+					ECustomizableObjectTextureLayoutPackingStrategy::Resizable ,
+					LOCTEXT("LayoutDetails_ResizableStrategyTooltip", "In a layout merge, Layout size will increase if blocks don't fit inside.")
+				});
+
+			LayoutPackingStrategies.Add(MakeShareable(new FString("Package and Shrink Textures")));
+			LayoutPackingStrategiesOptions.Add(
+				{
+					ECustomizableObjectTextureLayoutPackingStrategy::Fixed,
+					LOCTEXT("LayoutDetails_FixedStrategyTooltip", "In a layout merge, the layout will increase its size until the maximum layout grid size"
+					"\nBlock sizes will be reduced if they don't fit inside the layout."
+					"\nSet the reduction priority of each block to control which blocks are reduced first and how they are reduced.")
+				});
+
+			LayoutPackingStrategies.Add(MakeShareable(new FString("Overlay")));
+			LayoutPackingStrategiesOptions.Add(
+				{
+					ECustomizableObjectTextureLayoutPackingStrategy::Overlay,
+					LOCTEXT("LayoutDetails_OverlayStrategyTooltip", "In a layout merge, the layout will not be modified and blocks will be ignored."
+					"\nExtend material nodes just add their layouts on top of the base one")
+				});
+		}
+		CurrStrategy = LayoutPackingStrategies[(uint32)Node->Layout->PackingStrategy];
 
 		BlockReductionMethods.Empty();
 		BlockReductionMethods.Add(MakeShareable(new FString("Halve")));
 		BlockReductionMethods.Add(MakeShareable(new FString("Unitary")));
-		CurrRedMethod = BlockReductionMethods[(uint32)Node->Layout->GetBlockReductionMethod()];
+		CurrRedMethod = BlockReductionMethods[(uint32)Node->Layout->BlockReductionMethod];
+
+		{
+			AutoBlocksStrategies.Empty();
+			AutoBlocksStrategiesOptions.Empty();
+
+			AutoBlocksStrategies.Add(MakeShareable(new FString("Rectangles")));
+			AutoBlocksStrategiesOptions.Add(
+				{
+					ECustomizableObjectLayoutAutomaticBlocksStrategy::Rectangles ,
+					LOCTEXT("AutoBlockDetails_RectanglesStrategyTooltip", "Try to build rectangles splitting the UVs.")
+				});
+
+			AutoBlocksStrategies.Add(MakeShareable(new FString("UV islands")));
+			AutoBlocksStrategiesOptions.Add(
+				{
+					ECustomizableObjectLayoutAutomaticBlocksStrategy::UVIslands,
+					LOCTEXT("AutoBlockDetails_UVIslandsStrategyTooltip", "Try to build rectangles around each UV island, with a mask.")
+				});
+
+			AutoBlocksStrategies.Add(MakeShareable(new FString("Ignore (legacy)")));
+			AutoBlocksStrategiesOptions.Add(
+				{
+					ECustomizableObjectLayoutAutomaticBlocksStrategy::Ignore,
+					LOCTEXT("AutoBlockDetails_IgnoreStrategyTooltip", "Legacy behavior: assign to first block, or ignore if none.")
+				});
+		}
+		CurrAutoBlocks = AutoBlocksStrategies[(uint32)Node->Layout->AutomaticBlocksStrategy];
+
+		{
+			AutoBlocksMergeStrategies.Empty();
+			AutoBlocksMergeStrategiesOptions.Empty();
+
+			AutoBlocksMergeStrategies.Add(MakeShareable(new FString("Don't merge")));
+			AutoBlocksMergeStrategiesOptions.Add(
+				{
+					ECustomizableObjectLayoutAutomaticBlocksMergeStrategy::DontMerge ,
+					LOCTEXT("AutoBlockMerge_DontMergeTooltip", "Don't merge and make each UV island a unique block.")
+				});
+
+			AutoBlocksMergeStrategies.Add(MakeShareable(new FString("Merge child blocks")));
+			AutoBlocksMergeStrategiesOptions.Add(
+				{
+					ECustomizableObjectLayoutAutomaticBlocksMergeStrategy::MergeChildBlocks,
+					LOCTEXT("AutoBlockMerge_ChildBlocksTooltip", "Merge the blocks that are already fully included in another block.")
+				});
+		}
+		CurrAutoBlocksMerge = AutoBlocksMergeStrategies[(uint32)Node->Layout->AutomaticBlocksMergeStrategy];
+
 	}
 }
 
 EVisibility FCustomizableObjectNodeLayoutBlocksDetails::FixedStrategyOptionsVisibility() const
 {
-	return Node->Layout->GetPackingStrategy() == ECustomizableObjectTextureLayoutPackingStrategy::Fixed ? EVisibility::Visible : EVisibility::Collapsed;
+	return Node->Layout->PackingStrategy == ECustomizableObjectTextureLayoutPackingStrategy::Fixed ? EVisibility::Visible : EVisibility::Collapsed;
 }
 
 
@@ -319,16 +449,10 @@ void FCustomizableObjectNodeLayoutBlocksDetails::OnGridSizeChanged(TSharedPtr<FS
 		{
 			Node->Layout->SetGridSize(FIntPoint(Size));
 
-			// Adjust all the blocks sizes
-			for (int b = 0; b < Node->Layout->Blocks.Num(); ++b)
-			{
-				Node->Layout->Blocks[b].Min.X = FMath::Min(Node->Layout->Blocks[b].Min.X, Size - 1);
-				Node->Layout->Blocks[b].Min.Y = FMath::Min(Node->Layout->Blocks[b].Min.Y, Size - 1);
-				Node->Layout->Blocks[b].Max.X = FMath::Min(Node->Layout->Blocks[b].Max.X, Size);
-				Node->Layout->Blocks[b].Max.Y = FMath::Min(Node->Layout->Blocks[b].Max.Y, Size);
-			}
-
 			Node->MarkPackageDirty();
+
+			// Reset to update the UI.
+			LayoutBlocksEditor->SetCurrentLayout(Node->Layout);
 		}
 	}
 }
@@ -338,12 +462,51 @@ void FCustomizableObjectNodeLayoutBlocksDetails::OnLayoutPackingStrategyChanged(
 {
 	if (Node->Layout)
 	{
-		uint32 selection = LayoutPackingStrategies.IndexOfByKey(NewSelection);
+		uint32 Selection = LayoutPackingStrategies.IndexOfByKey(NewSelection);
+
+		if (Node->Layout->PackingStrategy != LayoutPackingStrategiesOptions[Selection].Value)
+		{
+			Node->Layout->PackingStrategy = LayoutPackingStrategiesOptions[Selection].Value;
+			Node->MarkPackageDirty();
+
+			// Reset to update the UI.
+			LayoutBlocksEditor->SetCurrentLayout(Node->Layout);
+		}
+	}
+}
+
+
+void FCustomizableObjectNodeLayoutBlocksDetails::OnAutoBlocksChanged(TSharedPtr<FString> NewSelection, ESelectInfo::Type SelectInfo)
+{
+	if (Node->Layout)
+	{
+		uint32 Selection = AutoBlocksStrategies.IndexOfByKey(NewSelection);
+
+		if (Node->Layout->AutomaticBlocksStrategy != AutoBlocksStrategiesOptions[Selection].Value)
+		{
+			Node->Layout->AutomaticBlocksStrategy = AutoBlocksStrategiesOptions[Selection].Value;
+			Node->MarkPackageDirty();
+
+			// Reset to update the UI.
+			LayoutBlocksEditor->SetCurrentLayout(Node->Layout);
+		}
+	}
+}
+
+
+void FCustomizableObjectNodeLayoutBlocksDetails::OnAutoBlocksMergeChanged(TSharedPtr<FString> NewSelection, ESelectInfo::Type SelectInfo)
+{
+	if (Node->Layout)
+	{
+		uint32 Selection = AutoBlocksMergeStrategies.IndexOfByKey(NewSelection);
 
-		if (Node->Layout->GetPackingStrategy() != (ECustomizableObjectTextureLayoutPackingStrategy)selection)
+		if (Node->Layout->AutomaticBlocksMergeStrategy != AutoBlocksMergeStrategiesOptions[Selection].Value)
 		{
-			Node->Layout->SetPackingStrategy((ECustomizableObjectTextureLayoutPackingStrategy)selection);
+			Node->Layout->AutomaticBlocksMergeStrategy = AutoBlocksMergeStrategiesOptions[Selection].Value;
 			Node->MarkPackageDirty();
+
+			// Reset to update the UI.
+			LayoutBlocksEditor->SetCurrentLayout(Node->Layout);
 		}
 	}
 }
@@ -353,7 +516,7 @@ void FCustomizableObjectNodeLayoutBlocksDetails::OnMaxGridSizeChanged(TSharedPtr
 {
 	if (Node->Layout)
 	{
-		int Size = 1 << LayoutGridSizes.Find(NewSelection);
+		int32 Size = 1 << LayoutGridSizes.Find(NewSelection);
 
 		if (Node->Layout->GetMaxGridSize().X != Size || Node->Layout->GetMaxGridSize().Y != Size)
 		{
@@ -368,11 +531,11 @@ void FCustomizableObjectNodeLayoutBlocksDetails::OnReductionMethodChanged(TShare
 {
 	if (Node->Layout)
 	{
-		uint32 selection = BlockReductionMethods.IndexOfByKey(NewSelection);
+		uint32 Selection = BlockReductionMethods.IndexOfByKey(NewSelection);
 
-		if (Node->Layout->GetBlockReductionMethod() != (ECustomizableObjectLayoutBlockReductionMethod)selection)
+		if (Node->Layout->BlockReductionMethod != (ECustomizableObjectLayoutBlockReductionMethod)Selection)
 		{
-			Node->Layout->SetBlockReductionMethod((ECustomizableObjectLayoutBlockReductionMethod)selection);
+			Node->Layout->BlockReductionMethod = (ECustomizableObjectLayoutBlockReductionMethod)Selection;
 			Node->MarkPackageDirty();
 		}
 	}
@@ -387,11 +550,57 @@ TSharedRef<SWidget> FCustomizableObjectNodeLayoutBlocksDetails::OnGenerateStrate
 	{
 		int32 TooltipIndex = LayoutPackingStrategies.IndexOfByKey(InItem);
 
-		if (LayoutPackingStrategiesTooltips.IsValidIndex(TooltipIndex))
+		if (LayoutPackingStrategies.IsValidIndex(TooltipIndex))
 		{
 			//A list of tool tips should have been populated in a 1 to 1 correspondance
-			check(LayoutPackingStrategies.Num() == LayoutPackingStrategiesTooltips.Num());
-			Tooltip = LayoutPackingStrategiesTooltips[TooltipIndex];
+			check(LayoutPackingStrategies.Num() == LayoutPackingStrategies.Num());
+			Tooltip = LayoutPackingStrategiesOptions[TooltipIndex].Tooltip;
+		}
+	}
+
+	return SNew(STextBlock)
+		.Text(FText::FromString(*InItem.Get()))
+		.Font(IDetailLayoutBuilder::GetDetailFont())
+		.ToolTipText(Tooltip);
+}
+
+
+TSharedRef<SWidget> FCustomizableObjectNodeLayoutBlocksDetails::OnGenerateAutoBlocksComboBox(TSharedPtr<FString> InItem) const
+{
+	FText Tooltip;
+
+	if (InItem.IsValid())
+	{
+		int32 TooltipIndex = AutoBlocksStrategies.IndexOfByKey(InItem);
+
+		if (AutoBlocksStrategies.IsValidIndex(TooltipIndex))
+		{
+			//A list of tool tips should have been populated in a 1 to 1 correspondance
+			check(AutoBlocksStrategies.Num() == AutoBlocksStrategies.Num());
+			Tooltip = AutoBlocksStrategiesOptions[TooltipIndex].Tooltip;
+		}
+	}
+
+	return SNew(STextBlock)
+		.Text(FText::FromString(*InItem.Get()))
+		.Font(IDetailLayoutBuilder::GetDetailFont())
+		.ToolTipText(Tooltip);
+}
+
+
+TSharedRef<SWidget> FCustomizableObjectNodeLayoutBlocksDetails::OnGenerateAutoBlocksMergeComboBox(TSharedPtr<FString> InItem) const
+{
+	FText Tooltip;
+
+	if (InItem.IsValid())
+	{
+		int32 TooltipIndex = AutoBlocksMergeStrategies.IndexOfByKey(InItem);
+
+		if (AutoBlocksMergeStrategies.IsValidIndex(TooltipIndex))
+		{
+			//A list of tool tips should have been populated in a 1 to 1 correspondance
+			check(AutoBlocksMergeStrategies.Num() == AutoBlocksMergeStrategies.Num());
+			Tooltip = AutoBlocksMergeStrategiesOptions[TooltipIndex].Tooltip;
 		}
 	}
 
@@ -429,7 +638,29 @@ FText FCustomizableObjectNodeLayoutBlocksDetails::GetSelectedLayoutStrategyName(
 {
 	if (Node->Layout)
 	{
-		return FText::FromString(*LayoutPackingStrategies[(uint32)Node->Layout->GetPackingStrategy()]);
+		return FText::FromString(*LayoutPackingStrategies[(uint32)Node->Layout->PackingStrategy]);
+	}
+
+	return FText();
+}
+
+
+FText FCustomizableObjectNodeLayoutBlocksDetails::GetSelectedAutoBlocksName() const
+{
+	if (Node->Layout)
+	{
+		return FText::FromString(*AutoBlocksStrategies[(uint32)Node->Layout->AutomaticBlocksStrategy]);
+	}
+
+	return FText();
+}
+
+
+FText FCustomizableObjectNodeLayoutBlocksDetails::GetSelectedAutoBlocksMergeName() const
+{
+	if (Node->Layout)
+	{
+		return FText::FromString(*AutoBlocksMergeStrategies[(uint32)Node->Layout->AutomaticBlocksMergeStrategy]);
 	}
 
 	return FText();
@@ -440,7 +671,7 @@ FText FCustomizableObjectNodeLayoutBlocksDetails::GetSelectedLayoutReductionMeth
 {
 	if (Node->Layout)
 	{
-		return FText::FromString(*BlockReductionMethods[(uint32)Node->Layout->GetBlockReductionMethod()]);
+		return FText::FromString(*BlockReductionMethods[(uint32)Node->Layout->BlockReductionMethod]);
 	}
 
 	return FText();
@@ -452,9 +683,9 @@ FText FCustomizableObjectNodeLayoutBlocksDetails::GetSelectedLayoutStrategyToolt
 	if (Node->Layout)
 	{
 		//A list of tool tips should have been populated in a 1 to 1 correspondance
-		check(LayoutPackingStrategies.Num() == LayoutPackingStrategiesTooltips.Num());
+		check(LayoutPackingStrategies.Num() == LayoutPackingStrategiesOptions.Num());
 
-		return LayoutPackingStrategiesTooltips[(uint32)Node->Layout->GetPackingStrategy()];
+		return LayoutPackingStrategiesOptions[(uint32)Node->Layout->PackingStrategy].Tooltip;
 	}
 
 	return FText();
@@ -468,7 +699,35 @@ FText FCustomizableObjectNodeLayoutBlocksDetails::GetSelectedLayoutReductionMeth
 		//A list of tool tips should have been populated in a 1 to 1 correspondance
 		check(BlockReductionMethods.Num() == BlockReductionMethodsTooltips.Num());
 
-		return BlockReductionMethodsTooltips[(uint32)Node->Layout->GetBlockReductionMethod()];
+		return BlockReductionMethodsTooltips[(uint32)Node->Layout->BlockReductionMethod];
+	}
+
+	return FText();
+}
+
+
+FText FCustomizableObjectNodeLayoutBlocksDetails::GetSelectedAutoBlocksTooltip() const
+{
+	if (Node->Layout)
+	{
+		//A list of tool tips should have been populated in a 1 to 1 correspondance
+		check(AutoBlocksStrategies.Num() == AutoBlocksStrategiesOptions.Num());
+
+		return AutoBlocksStrategiesOptions[(uint32)Node->Layout->AutomaticBlocksStrategy].Tooltip;
+	}
+
+	return FText();
+}
+
+
+FText FCustomizableObjectNodeLayoutBlocksDetails::GetSelectedAutoBlocksMergeTooltip() const
+{
+	if (Node->Layout)
+	{
+		//A list of tool tips should have been populated in a 1 to 1 correspondance
+		check(AutoBlocksMergeStrategies.Num() == AutoBlocksMergeStrategiesOptions.Num());
+
+		return AutoBlocksMergeStrategiesOptions[(uint32)Node->Layout->AutomaticBlocksMergeStrategy].Tooltip;
 	}
 
 	return FText();
diff --git a/Engine/Plugins/Experimental/Mutable/Source/CustomizableObjectEditor/Private/MuCOE/Nodes/CustomizableObjectNodeLayoutBlocksDetails.h b/Engine/Plugins/Experimental/Mutable/Source/CustomizableObjectEditor/Private/MuCOE/Nodes/CustomizableObjectNodeLayoutBlocksDetails.h
--- a/Engine/Plugins/Experimental/Mutable/Source/CustomizableObjectEditor/Private/MuCOE/Nodes/CustomizableObjectNodeLayoutBlocksDetails.h
+++ b/Engine/Plugins/Experimental/Mutable/Source/CustomizableObjectEditor/Private/MuCOE/Nodes/CustomizableObjectNodeLayoutBlocksDetails.h
@@ -4,6 +4,7 @@
 
 
 #include "CustomizableObjectNodeDetails.h"
+#include "MuCOE/CustomizableObjectLayout.h"
 #include "IDetailCustomization.h"
 #include "UObject/WeakObjectPtr.h"
 
@@ -29,25 +30,38 @@ public:
 
 private:
 
-	/** Returns the visibility of the Fixed layout widgets */
+	/** Returns the visibility of some widgets based on the layout strategy. */
 	EVisibility FixedStrategyOptionsVisibility() const;
 
 	/** Fills the combo box arrays sources */
-	void FillComboBoxOptionsArrays(TSharedPtr<FString>& CurrGridSize, TSharedPtr<FString>& CurrStrategy, TSharedPtr<FString>& CurrMaxSize, TSharedPtr<FString>& CurrRedMethod);
+	void FillComboBoxOptionsArrays(TSharedPtr<FString>& CurrGridSize, 
+		TSharedPtr<FString>& CurrStrategy, 
+		TSharedPtr<FString>& CurrAutoBlocks,
+		TSharedPtr<FString>& CurrAutoBlocksMerge,
+		TSharedPtr<FString>& CurrMaxSize,
+		TSharedPtr<FString>& CurrRedMethod);
 
 	/** Layout Options Callbacks */
 	void OnGridSizeChanged(TSharedPtr<FString> NewSelection, ESelectInfo::Type SelectInfo);
 	void OnMaxGridSizeChanged(TSharedPtr<FString> NewSelection, ESelectInfo::Type SelectInfo);
 	void OnReductionMethodChanged(TSharedPtr<FString> NewSelection, ESelectInfo::Type SelectInfo);
 	void OnLayoutPackingStrategyChanged(TSharedPtr<FString> NewSelection, ESelectInfo::Type SelectInfo);
+	void OnAutoBlocksChanged(TSharedPtr<FString> NewSelection, ESelectInfo::Type SelectInfo);
+	void OnAutoBlocksMergeChanged(TSharedPtr<FString> NewSelection, ESelectInfo::Type SelectInfo);
 	void OnIgnoreErrorsCheckStateChanged(ECheckBoxState State);
 	void OnLODBoxValueChanged(int32 Value);
 	TSharedRef<SWidget> OnGenerateStrategyComboBox(TSharedPtr<FString> InItem) const;
+	TSharedRef<SWidget> OnGenerateAutoBlocksComboBox(TSharedPtr<FString> InItem) const;
+	TSharedRef<SWidget> OnGenerateAutoBlocksMergeComboBox(TSharedPtr<FString> InItem) const;
 	TSharedRef<SWidget> OnGenerateReductionMethodComboBox(TSharedPtr<FString> InItem) const;
 	FText GetSelectedLayoutStrategyName() const;
+	FText GetSelectedAutoBlocksName() const;
+	FText GetSelectedAutoBlocksMergeName() const;
 	FText GetSelectedLayoutReductionMethodName() const;
 	FText GetSelectedLayoutStrategyTooltip() const;
 	FText GetSelectedLayoutReductionMethodTooltip() const;
+	FText GetSelectedAutoBlocksTooltip() const;
+	FText GetSelectedAutoBlocksMergeTooltip() const;
 
 private:
 
@@ -65,11 +79,32 @@ private:
 	TArray< TSharedPtr< FString > > LayoutGridSizes;
 
 	/** List of available layout packing strategies. */
-	TArray< TSharedPtr< FString > > LayoutPackingStrategies;
-	TArray<FText> LayoutPackingStrategiesTooltips;
+	struct FPackingStrategyOption
+	{
+		ECustomizableObjectTextureLayoutPackingStrategy Value;
+		FText Tooltip;
+	};
+	TArray<TSharedPtr<FString>> LayoutPackingStrategies;
+	TArray<FPackingStrategyOption> LayoutPackingStrategiesOptions;
 
 	/** List of available block reduction methods. */
 	TArray< TSharedPtr< FString > > BlockReductionMethods;
 	TArray<FText> BlockReductionMethodsTooltips;
 
+	struct FAutoBlocksStrategyOption
+	{
+		ECustomizableObjectLayoutAutomaticBlocksStrategy Value;
+		FText Tooltip;
+	};
+	TArray<TSharedPtr<FString>> AutoBlocksStrategies;
+	TArray<FAutoBlocksStrategyOption> AutoBlocksStrategiesOptions;
+
+	struct FAutoBlocksMergeStrategyOption
+	{
+		ECustomizableObjectLayoutAutomaticBlocksMergeStrategy Value;
+		FText Tooltip;
+	};
+	TArray<TSharedPtr<FString>> AutoBlocksMergeStrategies;
+	TArray<FAutoBlocksMergeStrategyOption> AutoBlocksMergeStrategiesOptions;
+
 };
diff --git a/Engine/Plugins/Experimental/Mutable/Source/CustomizableObjectEditor/Private/MuCOE/Nodes/CustomizableObjectNodeRemoveMeshBlocksDetails.cpp b/Engine/Plugins/Experimental/Mutable/Source/CustomizableObjectEditor/Private/MuCOE/Nodes/CustomizableObjectNodeRemoveMeshBlocksDetails.cpp
--- a/Engine/Plugins/Experimental/Mutable/Source/CustomizableObjectEditor/Private/MuCOE/Nodes/CustomizableObjectNodeRemoveMeshBlocksDetails.cpp
+++ b/Engine/Plugins/Experimental/Mutable/Source/CustomizableObjectEditor/Private/MuCOE/Nodes/CustomizableObjectNodeRemoveMeshBlocksDetails.cpp
@@ -4,12 +4,14 @@
 
 #include "DetailLayoutBuilder.h"
 #include "DetailWidgetRow.h"
+#include "IDetailGroup.h"
 #include "IDetailsView.h"
 #include "MuCOE/Nodes/CustomizableObjectNodeRemoveMeshBlocks.h"
 #include "MuCOE/Nodes/CustomizableObjectNodeMaterialBase.h"
 #include "MuCOE/SCustomizableObjectNodeLayoutBlocksEditor.h"
 #include "Widgets/Layout/SBox.h"
 #include "Widgets/Text/STextBlock.h"
+#include "Widgets/Input/STextComboBox.h"
 #include "UObject/Package.h"
 
 class FString;
@@ -24,7 +26,7 @@ TSharedRef<IDetailCustomization> FCustomizableObjectNodeRemoveMeshBlocksDetails:
 }
 
 
-void FCustomizableObjectNodeRemoveMeshBlocksDetails::CustomizeDetails( IDetailLayoutBuilder& DetailBuilder )
+void FCustomizableObjectNodeRemoveMeshBlocksDetails::CustomizeDetails(IDetailLayoutBuilder& DetailBuilder)
 {
 	FCustomizableObjectNodeEditMaterialBaseDetails::CustomizeDetails(DetailBuilder);
 
@@ -35,51 +37,112 @@ void FCustomizableObjectNodeRemoveMeshBlocksDetails::CustomizeDetails( IDetailLa
 		Node = Cast<UCustomizableObjectNodeRemoveMeshBlocks>(DetailsView->GetSelectedObjects()[0].Get());
 	}
 
-	IDetailCategoryBuilder& BlocksCategory = DetailBuilder.EditCategory( "Blocks" );
+	IDetailCategoryBuilder& LayoutCategory = DetailBuilder.EditCategory("LayoutOptions");
+	IDetailCategoryBuilder& BlocksCategory = DetailBuilder.EditCategory("Blocks");
 
-	if (Node)
+	if (!Node)
 	{
-		// Add blocks selector
-		LayoutBlocksEditor = SNew(SCustomizableObjectNodeLayoutBlocksEditor);
-
-		BlocksCategory.AddCustomRow(LOCTEXT("BlocksDetails_BlockInstructions", "BlockInstructions"))
-		[
-			SNew(SBox)
-			.HeightOverride(700.0f)
-			.WidthOverride(700.0f)
+		BlocksCategory.AddCustomRow(LOCTEXT("BlocksDetails_NodeNotFound", "NodeNotFound"))
 			[
-				LayoutBlocksEditor.ToSharedRef()
-			]
-		];
+				SNew(STextBlock)
+					.Text(LOCTEXT("Node not found", "Node not found"))
+			];
+		return;
+	}
 
-		// Try to find the parent layout, because we want to show its UVs in the widget
-		UCustomizableObjectLayout* ParentLayout = nullptr;
-		if (UCustomizableObjectNodeMaterialBase* ParentMaterialNode = Node->GetParentMaterialNode())
+	// Grid size combo
+	{
+		int32 MaxGridSize = 128;
+		LayoutGridSizes.Empty();
+
+		TSharedPtr<FString> CurrentSize;
+		for (int32 Size = 1; Size <= MaxGridSize; Size *= 2)
 		{
-			TArray<UCustomizableObjectLayout*> Layouts = ParentMaterialNode->GetLayouts();
+			LayoutGridSizes.Add(MakeShareable(new FString(FString::Printf(TEXT("%d x %d"), Size, Size))));
 
-			if (!Layouts.IsValidIndex(Node->ParentLayoutIndex))
+			if (Node->Layout->GetGridSize() == FIntPoint(Size))
 			{
-				UE_LOG(LogMutable, Warning, TEXT("[%s] UCustomizableObjectNodeRemoveMeshBlocks refers to an invalid texture layout index %d. Parent node has %d layouts."),
-					*Node->GetOutermost()->GetName(), Node->ParentLayoutIndex, Layouts.Num());
-			}
-			else
-			{
-				ParentLayout = Layouts[Node->ParentLayoutIndex];
+				CurrentSize = LayoutGridSizes.Last();
 			}
 		}
 
-		LayoutBlocksEditor->SetCurrentLayout(Node->Layout, ParentLayout);
+		IDetailGroup& LayoutOptionsGroup = LayoutCategory.AddGroup(TEXT("LayoutOptionsGroup"), LOCTEXT("LayoutGroup", "Layout Group"), false, true);
+		LayoutOptionsGroup.HeaderRow()
+			.NameContent()
+			[
+				SNew(STextBlock)
+					.Text(LOCTEXT("LayoutGridSizeText", "Grid Size"))
+					.Font(DetailBuilder.GetDetailFont())
+			]
+			.ValueContent()
+			[
+				SNew(STextComboBox)
+					.InitiallySelectedItem(CurrentSize)
+					.OptionsSource(&LayoutGridSizes)
+					.OnSelectionChanged(this, &FCustomizableObjectNodeRemoveMeshBlocksDetails::OnGridSizeChanged)
+					.Font(DetailBuilder.GetDetailFont())
+			];
 	}
-	else
-	{
-		BlocksCategory.AddCustomRow( LOCTEXT("BlocksDetails_NodeNotFound", "NodeNotFound") )
+
+	// Block editor
+	LayoutBlocksEditor = SNew(SCustomizableObjectNodeLayoutBlocksEditor);
+
+	BlocksCategory.AddCustomRow(LOCTEXT("BlocksDetails_BlockInstructions", "BlockInstructions"))
 		[
-			SNew( STextBlock )
-			.Text( LOCTEXT( "Node not found", "Node not found" ) )
+			SNew(SBox)
+				.HeightOverride(700.0f)
+				.WidthOverride(700.0f)
+				[
+					LayoutBlocksEditor.ToSharedRef()
+				]
 		];
+
+	UpdateLayout();
+}
+
+
+void FCustomizableObjectNodeRemoveMeshBlocksDetails::UpdateLayout()
+{
+	// Try to find the parent layout, because we want to show its UVs in the widget
+	UCustomizableObjectLayout* ParentLayout = nullptr;
+	if (UCustomizableObjectNodeMaterialBase* ParentMaterialNode = Node->GetParentMaterialNode())
+	{
+		TArray<UCustomizableObjectLayout*> Layouts = ParentMaterialNode->GetLayouts();
+
+		if (!Layouts.IsValidIndex(Node->ParentLayoutIndex))
+		{
+			UE_LOG(LogMutable, Warning, TEXT("[%s] UCustomizableObjectNodeRemoveMeshBlocks refers to an invalid texture layout index %d. Parent node has %d layouts."),
+				*Node->GetOutermost()->GetName(), Node->ParentLayoutIndex, Layouts.Num());
+		}
+		else
+		{
+			ParentLayout = Layouts[Node->ParentLayoutIndex];
+		}
 	}
 
+	LayoutBlocksEditor->SetCurrentLayout(Node->Layout, ParentLayout);
+
 }
 
+
+void FCustomizableObjectNodeRemoveMeshBlocksDetails::OnGridSizeChanged(TSharedPtr<FString> NewSelection, ESelectInfo::Type SelectInfo)
+{
+	if (!Node->Layout)
+	{
+		return;
+	}
+
+	int32 Size = 1 << LayoutGridSizes.Find(NewSelection);
+
+	if (Node->Layout->GetGridSize().X != Size || Node->Layout->GetGridSize().Y != Size)
+	{
+		Node->Layout->SetGridSize(FIntPoint(Size));
+
+		Node->MarkPackageDirty();
+
+		UpdateLayout();
+	}
+}
+
+
 #undef LOCTEXT_NAMESPACE
diff --git a/Engine/Plugins/Experimental/Mutable/Source/CustomizableObjectEditor/Private/MuCOE/Nodes/CustomizableObjectNodeRemoveMeshBlocksDetails.h b/Engine/Plugins/Experimental/Mutable/Source/CustomizableObjectEditor/Private/MuCOE/Nodes/CustomizableObjectNodeRemoveMeshBlocksDetails.h
--- a/Engine/Plugins/Experimental/Mutable/Source/CustomizableObjectEditor/Private/MuCOE/Nodes/CustomizableObjectNodeRemoveMeshBlocksDetails.h
+++ b/Engine/Plugins/Experimental/Mutable/Source/CustomizableObjectEditor/Private/MuCOE/Nodes/CustomizableObjectNodeRemoveMeshBlocksDetails.h
@@ -27,4 +27,14 @@ private:
 
 	// Layout block editor widget
 	TSharedPtr<SCustomizableObjectNodeLayoutBlocksEditor> LayoutBlocksEditor;
+
+	/** List of available layout grid sizes. */
+	TArray< TSharedPtr< FString > > LayoutGridSizes;
+
+	/** Layout Options Callbacks */
+	void OnGridSizeChanged(TSharedPtr<FString> NewSelection, ESelectInfo::Type SelectInfo);
+
+	/** Reset the layout in the widget to force a refresh. */
+	void UpdateLayout();
+
 };
diff --git a/Engine/Plugins/Experimental/Mutable/Source/CustomizableObjectEditor/Private/MuCOE/Nodes/CustomizableObjectNodeTable.cpp b/Engine/Plugins/Experimental/Mutable/Source/CustomizableObjectEditor/Private/MuCOE/Nodes/CustomizableObjectNodeTable.cpp
--- a/Engine/Plugins/Experimental/Mutable/Source/CustomizableObjectEditor/Private/MuCOE/Nodes/CustomizableObjectNodeTable.cpp
+++ b/Engine/Plugins/Experimental/Mutable/Source/CustomizableObjectEditor/Private/MuCOE/Nodes/CustomizableObjectNodeTable.cpp
@@ -852,8 +852,10 @@ void UCustomizableObjectNodeTable::RemapPinsData(const TMap<UEdGraphPin*, UEdGra
 							NewLayout->Blocks = OldLayout->Blocks;
 							NewLayout->SetGridSize(OldLayout->GetGridSize());
 							NewLayout->SetMaxGridSize(OldLayout->GetMaxGridSize());
-							NewLayout->SetPackingStrategy(OldLayout->GetPackingStrategy());
-							NewLayout->SetBlockReductionMethod(OldLayout->GetBlockReductionMethod());
+							NewLayout->PackingStrategy = OldLayout->PackingStrategy;
+							NewLayout->AutomaticBlocksStrategy = OldLayout->AutomaticBlocksStrategy;
+							NewLayout->AutomaticBlocksMergeStrategy = OldLayout->AutomaticBlocksMergeStrategy;
+							NewLayout->BlockReductionMethod = OldLayout->BlockReductionMethod;
 
 							break;
 						}
diff --git a/Engine/Plugins/Experimental/Mutable/Source/CustomizableObjectEditor/Private/MuCOE/Nodes/CustomizableObjectNodeTableDetails.cpp b/Engine/Plugins/Experimental/Mutable/Source/CustomizableObjectEditor/Private/MuCOE/Nodes/CustomizableObjectNodeTableDetails.cpp
--- a/Engine/Plugins/Experimental/Mutable/Source/CustomizableObjectEditor/Private/MuCOE/Nodes/CustomizableObjectNodeTableDetails.cpp
+++ b/Engine/Plugins/Experimental/Mutable/Source/CustomizableObjectEditor/Private/MuCOE/Nodes/CustomizableObjectNodeTableDetails.cpp
@@ -773,7 +773,7 @@ EVisibility FCustomizableObjectNodeTableDetails::LayoutOptionsVisibility() const
 
 EVisibility FCustomizableObjectNodeTableDetails::FixedStrategyOptionsVisibility() const
 {
-	return (SelectedLayout.IsValid() && SelectedLayout->GetPackingStrategy() == ECustomizableObjectTextureLayoutPackingStrategy::Fixed) ? EVisibility::Visible : EVisibility::Collapsed;
+	return (SelectedLayout.IsValid() && SelectedLayout->PackingStrategy == ECustomizableObjectTextureLayoutPackingStrategy::Fixed) ? EVisibility::Visible : EVisibility::Collapsed;
 }
 
 
@@ -797,8 +797,8 @@ void FCustomizableObjectNodeTableDetails::FillLayoutComboBoxOptions()
 			}
 		}
 
-		StrategyComboBox->SetSelectedItem(LayoutPackingStrategies[(uint32)SelectedLayout->GetPackingStrategy()]);
-		ReductionMethodComboBox->SetSelectedItem(BlockReductionMethods[(uint32)SelectedLayout->GetBlockReductionMethod()]);
+		StrategyComboBox->SetSelectedItem(LayoutPackingStrategies[(uint32)SelectedLayout->PackingStrategy]);
+		ReductionMethodComboBox->SetSelectedItem(BlockReductionMethods[(uint32)SelectedLayout->BlockReductionMethod]);
 	}
 }
 
@@ -853,7 +853,7 @@ FText FCustomizableObjectNodeTableDetails::GetSelectedLayoutStrategyName() const
 {
 	if (SelectedLayout.IsValid())
 	{
-		return FText::FromString(*LayoutPackingStrategies[(uint32)SelectedLayout->GetPackingStrategy()]);
+		return FText::FromString(*LayoutPackingStrategies[(uint32)SelectedLayout->PackingStrategy]);
 	}
 
 	return FText();
@@ -864,7 +864,7 @@ FText FCustomizableObjectNodeTableDetails::GetSelectedLayoutReductionMethodName(
 {
 	if (SelectedLayout.IsValid())
 	{
-		return FText::FromString(*BlockReductionMethods[(uint32)SelectedLayout->GetBlockReductionMethod()]);
+		return FText::FromString(*BlockReductionMethods[(uint32)SelectedLayout->BlockReductionMethod]);
 	}
 
 	return FText();
@@ -875,20 +875,21 @@ void FCustomizableObjectNodeTableDetails::OnGridSizeChanged(TSharedPtr<FString>
 {
 	if (SelectedLayout.IsValid())
 	{
-		int Size = 1 << LayoutGridSizes.Find(NewSelection);
+		int32 Size = 1 << LayoutGridSizes.Find(NewSelection);
 
 		if (SelectedLayout->GetGridSize().X != Size || SelectedLayout->GetGridSize().Y != Size)
 		{
 			SelectedLayout->SetGridSize(FIntPoint(Size));
 
 			// Adjust all the blocks sizes
-			for (int b = 0; b < SelectedLayout->Blocks.Num(); ++b)
-			{
-				SelectedLayout->Blocks[b].Min.X = FMath::Min(SelectedLayout->Blocks[b].Min.X, Size - 1);
-				SelectedLayout->Blocks[b].Min.Y = FMath::Min(SelectedLayout->Blocks[b].Min.Y, Size - 1);
-				SelectedLayout->Blocks[b].Max.X = FMath::Min(SelectedLayout->Blocks[b].Max.X, Size);
-				SelectedLayout->Blocks[b].Max.Y = FMath::Min(SelectedLayout->Blocks[b].Max.Y, Size);
-			}
+			//for (int32 BlockIndex = 0; BlockIndex < SelectedLayout->Blocks.Num(); ++BlockIndex)
+			//{
+			//	FCustomizableObjectLayoutBlock& Block = SelectedLayout->Blocks[BlockIndex];
+			//	Block.Min.X = FMath::Min(Block.Min.X, Size - 1);
+			//	Block.Min.Y = FMath::Min(Block.Min.Y, Size - 1);
+			//	Block.Max.X = FMath::Min(Block.Max.X, Size);
+			//	Block.Max.Y = FMath::Min(Block.Max.Y, Size);
+			//}
 
 			Node->MarkPackageDirty();
 		}
@@ -900,11 +901,11 @@ void FCustomizableObjectNodeTableDetails::OnLayoutPackingStrategyChanged(TShared
 {
 	if (SelectedLayout.IsValid())
 	{
-		uint32 selection = LayoutPackingStrategies.IndexOfByKey(NewSelection);
+		uint32 Selection = LayoutPackingStrategies.IndexOfByKey(NewSelection);
 
-		if (SelectedLayout->GetPackingStrategy() != (ECustomizableObjectTextureLayoutPackingStrategy)selection)
+		if (SelectedLayout->PackingStrategy != (ECustomizableObjectTextureLayoutPackingStrategy)Selection)
 		{
-			SelectedLayout->SetPackingStrategy((ECustomizableObjectTextureLayoutPackingStrategy)selection);
+			SelectedLayout->PackingStrategy=(ECustomizableObjectTextureLayoutPackingStrategy)Selection;
 			Node->MarkPackageDirty();
 		}
 	}
@@ -915,7 +916,7 @@ void FCustomizableObjectNodeTableDetails::OnMaxGridSizeChanged(TSharedPtr<FStrin
 {
 	if (SelectedLayout.IsValid())
 	{
-		int Size = 1 << LayoutGridSizes.Find(NewSelection);
+		int32 Size = 1 << LayoutGridSizes.Find(NewSelection);
 
 		if (SelectedLayout->GetMaxGridSize().X != Size || SelectedLayout->GetMaxGridSize().Y != Size)
 		{
@@ -930,11 +931,11 @@ void FCustomizableObjectNodeTableDetails::OnReductionMethodChanged(TSharedPtr<FS
 {
 	if (SelectedLayout.IsValid())
 	{
-		uint32 selection = BlockReductionMethods.IndexOfByKey(NewSelection);
+		uint32 Selection = BlockReductionMethods.IndexOfByKey(NewSelection);
 
-		if (SelectedLayout->GetBlockReductionMethod() != (ECustomizableObjectLayoutBlockReductionMethod)selection)
+		if (SelectedLayout->BlockReductionMethod != (ECustomizableObjectLayoutBlockReductionMethod)Selection)
 		{
-			SelectedLayout->SetBlockReductionMethod((ECustomizableObjectLayoutBlockReductionMethod)selection);
+			SelectedLayout->BlockReductionMethod=(ECustomizableObjectLayoutBlockReductionMethod)Selection;
 			Node->MarkPackageDirty();
 		}
 	}
@@ -948,7 +949,7 @@ FText FCustomizableObjectNodeTableDetails::GetSelectedLayoutStrategyTooltip() co
 		//A list of tool tips should have been populated in a 1 to 1 correspondance
 		check(LayoutPackingStrategies.Num() == LayoutPackingStrategiesTooltips.Num());
 
-		return LayoutPackingStrategiesTooltips[(uint32)SelectedLayout->GetPackingStrategy()];
+		return LayoutPackingStrategiesTooltips[(uint32)SelectedLayout->PackingStrategy];
 	}
 
 	return FText();
@@ -962,7 +963,7 @@ FText FCustomizableObjectNodeTableDetails::GetSelectedLayoutReductionMethodToolt
 		//A list of tool tips should have been populated in a 1 to 1 correspondance
 		check(BlockReductionMethods.Num() == BlockReductionMethodsTooltips.Num());
 
-		return BlockReductionMethodsTooltips[(uint32)SelectedLayout->GetBlockReductionMethod()];
+		return BlockReductionMethodsTooltips[(uint32)SelectedLayout->BlockReductionMethod];
 	}
 
 	return FText();
diff --git a/Engine/Plugins/Experimental/Mutable/Source/CustomizableObjectEditor/Private/MuCOE/SCustomizableObjectLayoutGrid.cpp b/Engine/Plugins/Experimental/Mutable/Source/CustomizableObjectEditor/Private/MuCOE/SCustomizableObjectLayoutGrid.cpp
--- a/Engine/Plugins/Experimental/Mutable/Source/CustomizableObjectEditor/Private/MuCOE/SCustomizableObjectLayoutGrid.cpp
+++ b/Engine/Plugins/Experimental/Mutable/Source/CustomizableObjectEditor/Private/MuCOE/SCustomizableObjectLayoutGrid.cpp
@@ -137,6 +137,7 @@ private:
 	/** Default colors. */
 	FColor SelectedBlockColor = FColor(75, 106, 230, 155);
 	FColor UnselectedBlockColor = FColor(230, 199, 75, 155);
+	FColor AutomaticBlockColor = FColor(125, 125, 125, 125);
 };
 
 
@@ -160,7 +161,6 @@ void SCustomizableObjectLayoutGrid::Construct( const FArguments& InArgs )
 	for (int32 BufferIndex=0;BufferIndex< UE_MUTABLE_UI_DRAWBUFFERS; ++BufferIndex)
 	{
 		UVCanvasDrawers[BufferIndex] = TSharedPtr<FUVCanvasDrawer>(new FUVCanvasDrawer());
-		UVCanvasDrawers[BufferIndex]->SetLayoutMode(Mode);
 	}
 }
 
@@ -211,7 +211,7 @@ int32 SCustomizableObjectLayoutGrid::OnPaint(const FPaintArgs& Args, const FGeom
 
 	FVector2D OldSize = Size;
 
-	double ZoomFactor = FMath::Pow(2.0,double(Zoom-1));
+	double ZoomFactor = PointOfView.GetZoomFactor();
 	Size *= ZoomFactor;
 
 	float AuxCellSize = Size.X / GridSizePoint.X;
@@ -220,7 +220,7 @@ int32 SCustomizableObjectLayoutGrid::OnPaint(const FPaintArgs& Args, const FGeom
 	FVector2D Offset = FVector2D((AreaSize - Size).X / 2.0f, 0.0f);
 	
 	// Drawing Origin
-	FVector2D Origin = BorderPadding + Offset + PaddingAmount;
+	FVector2D Origin = BorderPadding + Offset + PointOfView.PaddingAmount;
 
 	// Setting Canvas Drawing Rectangles
 	FSlateRect SlateCanvasRect = AllottedGeometry.GetLayoutBoundingRect();
@@ -239,8 +239,12 @@ int32 SCustomizableObjectLayoutGrid::OnPaint(const FPaintArgs& Args, const FGeom
 		FMath::TruncToInt(FMath::Max(0.0f, ClippedCanvasRect.Bottom)));
 
 	TSharedPtr<class FUVCanvasDrawer> UVCanvasDrawer = UVCanvasDrawers[CurrentDrawBuffer];
+	
+	ELayoutGridMode GridMode = Mode.Get();
+	UVCanvasDrawer->SetLayoutMode(GridMode);
 	UVCanvasDrawer->InitializeDrawingData(UVLayout, UnassignedUVLayoutVertices, Blocks.Get(), SelectedBlocks);
 	UVCanvasDrawer->Initialize(CanvasRect, ClippingRect, Origin * AllottedGeometry.Scale, Size * AllottedGeometry.Scale, GridSizePoint, AuxCellSize * AllottedGeometry.Scale);
+
 	FSlateDrawElement::MakeCustom(OutDrawElements, RetLayerId, UVCanvasDrawer);
 
 	const auto MakeYellowSquareLine = [&](const TArray<FVector2D>& Points) -> void
@@ -250,7 +254,7 @@ int32 SCustomizableObjectLayoutGrid::OnPaint(const FPaintArgs& Args, const FGeom
 	};
 
 	// Drawing Multi-Selection rect
-	if (Mode == ELGM_Edit && bIsSelecting)
+	if (GridMode == ELGM_Edit && bIsSelecting)
 	{
 		TArray<FVector2D> SelectionSquarePoints;
 		SelectionSquarePoints.SetNum(2);
@@ -285,6 +289,7 @@ int32 SCustomizableObjectLayoutGrid::OnPaint(const FPaintArgs& Args, const FGeom
 
 void SCustomizableObjectLayoutGrid::Tick( const FGeometry& AllottedGeometry, const double InCurrentTime, const float InDeltaTime )
 {
+	// Swap the rendering buffer.
 	CurrentDrawBuffer = FMath::Modulo(CurrentDrawBuffer+1, UE_MUTABLE_UI_DRAWBUFFERS);
 
 	const FVector2D BorderPadding = FVector2D(2,2);
@@ -303,12 +308,12 @@ void SCustomizableObjectLayoutGrid::Tick( const FGeometry& AllottedGeometry, con
 	}
 
 	FVector2D OldSize = Size;
-	double ZoomFactor = FMath::Pow(2.0, double(Zoom - 1));
+	double ZoomFactor = PointOfView.GetZoomFactor();
 	Size *= ZoomFactor;
 
 	CellSize = Size.X/GridSize.Get().X;
 	FVector2D Offset = FVector2D((AreaSize - Size).X / 2.0f, 0.0f);
-	FVector2D Origin = BorderPadding + Offset + PaddingAmount;
+	FVector2D Origin = BorderPadding + Offset + PointOfView.PaddingAmount;
 	DrawOrigin = Origin;
 
 	BlockRects.Empty();
@@ -331,7 +336,7 @@ void SCustomizableObjectLayoutGrid::Tick( const FGeometry& AllottedGeometry, con
 	}
 
 	// Update selection list
-	for (int32 SelectedBlockIndex=0; SelectedBlockIndex <SelectedBlocks.Num();)
+	for (int32 SelectedBlockIndex=0; SelectedBlockIndex<SelectedBlocks.Num();)
 	{
 		bool bFound = false;
 		for (const FCustomizableObjectLayoutBlock& Block : CurrentBlocks)
@@ -365,9 +370,9 @@ FReply SCustomizableObjectLayoutGrid::OnMouseButtonDown( const FGeometry& MyGeom
 {
 	FReply Reply = FReply::Unhandled();
 
-	if (Mode == ELGM_Edit)
+	ELayoutGridMode GridMode = Mode.Get();
 	{
-		if (MouseEvent.GetEffectingButton() == EKeys::LeftMouseButton)
+		if (MouseEvent.GetEffectingButton() == EKeys::LeftMouseButton && GridMode == ELGM_Edit)
 		{
 			bHasDragged = false;
 			bIsDragging = false;
@@ -416,7 +421,7 @@ FReply SCustomizableObjectLayoutGrid::OnMouseButtonDown( const FGeometry& MyGeom
 		}
 		else if (MouseEvent.GetEffectingButton() == EKeys::RightMouseButton)
 		{
-			//Mouse position
+			// Mouse position
 			FVector2D Pos = MyGeometry.AbsoluteToLocal(MouseEvent.GetScreenSpacePosition());
 			FVector2D CellDelta = (Pos - DrawOrigin) / CellSize;
 
@@ -431,97 +436,101 @@ FReply SCustomizableObjectLayoutGrid::OnMouseButtonDown( const FGeometry& MyGeom
 			}
 			MenuBuilder.EndSection();
 
-			MenuBuilder.BeginSection("Block Management", LOCTEXT("BlockActionsTitle", "Block Actions"));
+			if (GridMode == ELGM_Edit)
 			{
-				if (SelectedBlocks.Num())
+				MenuBuilder.BeginSection("Block Management", LOCTEXT("BlockActionsTitle", "Block Actions"));
 				{
-					FUIAction DeleteAction(FExecuteAction::CreateSP(this, &SCustomizableObjectLayoutGrid::DeleteSelectedBlocks));
-					MenuBuilder.AddMenuEntry(LOCTEXT("DeleteBlocksLabel", "Delete"), LOCTEXT("DeleteBlocksTooltip", "Delete Selected Blocks"), FSlateIcon(), DeleteAction);
+					if (SelectedBlocks.Num())
+					{
+						FUIAction DeleteAction(FExecuteAction::CreateSP(this, &SCustomizableObjectLayoutGrid::DeleteSelectedBlocks));
+						MenuBuilder.AddMenuEntry(LOCTEXT("DeleteBlocksLabel", "Delete"), LOCTEXT("DeleteBlocksTooltip", "Delete Selected Blocks"), FSlateIcon(), DeleteAction);
 
-					FUIAction DuplicateAction(FExecuteAction::CreateSP(this, &SCustomizableObjectLayoutGrid::DuplicateBlocks));
-					MenuBuilder.AddMenuEntry(LOCTEXT("DuplicateBlocksLabel", "Duplicate"), LOCTEXT("DuplicateBlocksTooltip", "Duplicate Selected Blocks"), FSlateIcon(), DuplicateAction);
-				}
-				else
-				{
-					FUIAction AddNewBlockAction(FExecuteAction::CreateSP(this, &SCustomizableObjectLayoutGrid::GenerateNewBlock, CellDelta));
-					MenuBuilder.AddMenuEntry(LOCTEXT("AddNewBlockLabel", "Add Block"), LOCTEXT("AddNewBlockTooltip", "Add New Block"), FSlateIcon(), AddNewBlockAction);
+						FUIAction DuplicateAction(FExecuteAction::CreateSP(this, &SCustomizableObjectLayoutGrid::DuplicateBlocks));
+						MenuBuilder.AddMenuEntry(LOCTEXT("DuplicateBlocksLabel", "Duplicate"), LOCTEXT("DuplicateBlocksTooltip", "Duplicate Selected Blocks"), FSlateIcon(), DuplicateAction);
+					}
+					else
+					{
+						FUIAction AddNewBlockAction(FExecuteAction::CreateSP(this, &SCustomizableObjectLayoutGrid::GenerateNewBlock, CellDelta));
+						MenuBuilder.AddMenuEntry(LOCTEXT("AddNewBlockLabel", "Add Block"), LOCTEXT("AddNewBlockTooltip", "Add New Block"), FSlateIcon(), AddNewBlockAction);
+					}
 				}
-			}
-			MenuBuilder.EndSection();
+				MenuBuilder.EndSection();
 
-			MenuBuilder.BeginSection("Block Properties for Fixed Layout", LOCTEXT("BlockPropertiesFixedTitle", "Block Properties for Fixed Layout"));
-			{
-				if (SelectedBlocks.Num())
+				MenuBuilder.BeginSection("Block Properties for Fixed Layout", LOCTEXT("BlockPropertiesFixedTitle", "Block Properties for Fixed Layout"));
 				{
-					MenuBuilder.AddWidget(
-						SNew(SBox)
-						.WidthOverride(125.0f)
-						.ToolTipText(LOCTEXT("SetBlockPriority_Tooltip", "Sets the block priority for a Fixed Layout Strategy."))
-						[
-							SNew(SNumericEntryBox<int32>)
-							.MinValue(0)
-							.MaxValue(INT_MAX)
-							.MaxSliderValue(100)
-							.AllowSpin(SelectedBlocks.Num() == 1)
-							.Value(this, &SCustomizableObjectLayoutGrid::GetBlockPriorityValue)
-							.UndeterminedString(LOCTEXT("MultipleValues", "Multiples Values"))
-							.OnValueChanged(this, &SCustomizableObjectLayoutGrid::OnBlockPriorityChanged)
-							.EditableTextBoxStyle(&UE_MUTABLE_GET_WIDGETSTYLE<FEditableTextBoxStyle>("NormalEditableTextBox"))
-						]
-					, FText::FromString("Block Priority"), true);
-
-					MenuBuilder.AddWidget(
-						SNew(SBox)
-						.WidthOverride(125.0f)
-						.ToolTipText(LOCTEXT("SetBlockSymmetry_Tooltip", "If true, this block will be reduced in both axes at the same time in a Fixed Layout Strategy."))
-						[
-							SNew(SCheckBox)
-							.IsChecked(this, &SCustomizableObjectLayoutGrid::GetReductionMethodBoolValue, EFRO_Symmetry)
-							.OnCheckStateChanged(this, &SCustomizableObjectLayoutGrid::OnReduceBlockSymmetricallyChanged)
-						]
-					, FText::FromString("Reduce Symmetrically"), true);
-
-					MenuBuilder.AddWidget(
-						SNew(SBox)
-						.WidthOverride(125.0f)
-						.ToolTipText(LOCTEXT("SetBlockReduceByTwo_Tooltip", "Only for Unitary reduction. If true, this option reduces each time the block by two block units."))
-						[
-							SNew(SCheckBox)
-							.IsChecked(this, &SCustomizableObjectLayoutGrid::GetReductionMethodBoolValue, EFRO_RedyceByTwo)
-							.OnCheckStateChanged(this, &SCustomizableObjectLayoutGrid::OnReduceBlockByTwoChanged)
-						]
-					, FText::FromString("Reduce by Two"), true);
+					if (SelectedBlocks.Num())
+					{
+						MenuBuilder.AddWidget(
+							SNew(SBox)
+							.WidthOverride(125.0f)
+							.ToolTipText(LOCTEXT("SetBlockPriority_Tooltip", "Sets the block priority for a Fixed Layout Strategy."))
+							[
+								SNew(SNumericEntryBox<int32>)
+									.MinValue(0)
+									.MaxValue(INT_MAX)
+									.MaxSliderValue(100)
+									.AllowSpin(SelectedBlocks.Num() == 1)
+									.Value(this, &SCustomizableObjectLayoutGrid::GetBlockPriorityValue)
+									.UndeterminedString(LOCTEXT("MultipleValues", "Multiples Values"))
+									.OnValueChanged(this, &SCustomizableObjectLayoutGrid::OnBlockPriorityChanged)
+									.EditableTextBoxStyle(&UE_MUTABLE_GET_WIDGETSTYLE<FEditableTextBoxStyle>("NormalEditableTextBox"))
+							]
+							, FText::FromString("Block Priority"), true);
+
+						MenuBuilder.AddWidget(
+							SNew(SBox)
+							.WidthOverride(125.0f)
+							.ToolTipText(LOCTEXT("SetBlockSymmetry_Tooltip", "If true, this block will be reduced in both axes at the same time in a Fixed Layout Strategy."))
+							[
+								SNew(SCheckBox)
+									.IsChecked(this, &SCustomizableObjectLayoutGrid::GetReductionMethodBoolValue, EFRO_Symmetry)
+									.OnCheckStateChanged(this, &SCustomizableObjectLayoutGrid::OnReduceBlockSymmetricallyChanged)
+							]
+							, FText::FromString("Reduce Symmetrically"), true);
+
+						MenuBuilder.AddWidget(
+							SNew(SBox)
+							.WidthOverride(125.0f)
+							.ToolTipText(LOCTEXT("SetBlockReduceByTwo_Tooltip", "Only for Unitary reduction. If true, this option reduces each time the block by two block units."))
+							[
+								SNew(SCheckBox)
+									.IsChecked(this, &SCustomizableObjectLayoutGrid::GetReductionMethodBoolValue, EFRO_RedyceByTwo)
+									.OnCheckStateChanged(this, &SCustomizableObjectLayoutGrid::OnReduceBlockByTwoChanged)
+							]
+							, FText::FromString("Reduce by Two"), true);
+					}
 				}
-			}
-			MenuBuilder.EndSection();
+				MenuBuilder.EndSection();
 
 
-			MenuBuilder.BeginSection("Block Properties for Masks", LOCTEXT("BlockPropertiesMaskTitle", "Block Mask"));
-			{
-				if (SelectedBlocks.Num())
+				MenuBuilder.BeginSection("Block Properties for Masks", LOCTEXT("BlockPropertiesMaskTitle", "Block Mask"));
 				{
-					MenuBuilder.AddWidget(
-						SNew(SBox)
-						.WidthOverride(125.0f)
-						.ToolTipText(LOCTEXT("SetBlockMask_Tooltip", "Sets the UV mask texture for the block."))
-						[
-							SNew(SStandAloneAssetPicker)
-								.OnAssetSelected(this, &SCustomizableObjectLayoutGrid::OnMaskAssetSelected)
-								.OnGetAllowedClasses(FOnGetAllowedClasses::CreateLambda([](TArray<const UClass*>& OutClasses) {OutClasses.Add(UTexture2D::StaticClass()); }))
-								.InitialAsset( GetBlockMaskValue() )
-						]
-						, FText::FromString("Block Mask"), true);
-
-					// TODO: Additional properties: color used for preview?
+					if (SelectedBlocks.Num())
+					{
+						MenuBuilder.AddWidget(
+							SNew(SBox)
+							.WidthOverride(125.0f)
+							.ToolTipText(LOCTEXT("SetBlockMask_Tooltip", "Sets the UV mask texture for the block."))
+							[
+								SNew(SStandAloneAssetPicker)
+									.OnAssetSelected(this, &SCustomizableObjectLayoutGrid::OnMaskAssetSelected)
+									.OnGetAllowedClasses(FOnGetAllowedClasses::CreateLambda([](TArray<const UClass*>& OutClasses) {OutClasses.Add(UTexture2D::StaticClass()); }))
+									.InitialAsset(GetBlockMaskValue())
+							]
+							, FText::FromString("Block Mask"), true);
+
+						// TODO: Additional properties: color used for preview?
+					}
 				}
+				MenuBuilder.EndSection();
 			}
-			MenuBuilder.EndSection();
 
 			FWidgetPath WidgetPath = MouseEvent.GetEventPath() != nullptr ? *MouseEvent.GetEventPath() : FWidgetPath();
 			FSlateApplication::Get().PushMenu(AsShared(), WidgetPath, MenuBuilder.MakeWidget(), FSlateApplication::Get().GetCursorPos(), FPopupTransitionEffect(FPopupTransitionEffect::ContextMenu));
 
 			Reply = FReply::Handled();
 		}
+
 		else if (MouseEvent.GetEffectingButton() == EKeys::MiddleMouseButton)
 		{
 			bIsPadding = true;
@@ -595,12 +604,9 @@ FReply SCustomizableObjectLayoutGrid::OnMouseButtonUp( const FGeometry& MyGeomet
 {
 	FReply Reply = FReply::Unhandled();
 
-	if (Mode == ELGM_Show)
-	{
-		return SCompoundWidget::OnMouseButtonUp(MyGeometry, MouseEvent);
-	}
+	ELayoutGridMode GridMode = Mode.Get();
 
-	if ( MouseEvent.GetEffectingButton()==EKeys::LeftMouseButton)
+	if ( MouseEvent.GetEffectingButton()==EKeys::LeftMouseButton && GridMode == ELGM_Edit)
 	{
 		bIsDragging = false;
 		bIsResizing = false;
@@ -613,12 +619,12 @@ FReply SCustomizableObjectLayoutGrid::OnMouseButtonUp( const FGeometry& MyGeomet
 		
 		// Selection before reset
 		TArray<FGuid> OldSelection = SelectedBlocks;
-		TArray<FGuid> OldPossibleSelection = PossibleSelectedBlocks;
 
+		TArray<FGuid> OldPossibleSelection = PossibleSelectedBlocks;
 		PossibleSelectedBlocks.Reset();
 
 		// Reset selection if multi selection is not enabled
-		if (Mode == ELGM_Edit && !bLeftShift && !bHasDragged)
+		if (GridMode == ELGM_Edit && !bLeftShift && !bHasDragged)
 		{
 			// Only one selected block allowed in edit mode.
 			SelectedBlocks.Reset();
@@ -628,11 +634,11 @@ FReply SCustomizableObjectLayoutGrid::OnMouseButtonUp( const FGeometry& MyGeomet
 		{
 			if (!bHasDragged)
 			{
-				//Backward iteration to select the block rendered in front of the rest
+				// Backward iteration to select the block rendered in front of the rest
 				const TArray<FCustomizableObjectLayoutBlock>& CurrentBlocks = Blocks.Get();
 				for (int32 i = CurrentBlocks.Num() - 1; i > -1; --i)
 				{
-					if (MouseOnBlock(CurrentBlocks[i].Id, Pos))
+					if ( (!CurrentBlocks[i].bIsAutomatic) && MouseOnBlock(CurrentBlocks[i].Id, Pos))
 					{
 						PossibleSelectedBlocks.Add(CurrentBlocks[i].Id);
 					}
@@ -642,7 +648,7 @@ FReply SCustomizableObjectLayoutGrid::OnMouseButtonUp( const FGeometry& MyGeomet
 
 				for (int32 i = 0; i < PossibleSelectedBlocks.Num(); ++i)
 				{
-					if (bLeftShift || Mode == ELGM_Select)
+					if (bLeftShift)
 					{
 						if (PossibleSelectedBlocks.Num() == 1)
 						{
@@ -712,6 +718,11 @@ FReply SCustomizableObjectLayoutGrid::OnMouseButtonUp( const FGeometry& MyGeomet
 			const TArray<FCustomizableObjectLayoutBlock>& CurrentBlocks = Blocks.Get();
 			for (int32 i = 0; i < CurrentBlocks.Num(); ++i)
 			{
+				if (CurrentBlocks[i].bIsAutomatic)
+				{
+					continue;
+				}
+
 				FBox2D CurrentBlock(FVector2D(BlockRects[CurrentBlocks[i].Id].Rect.Min), FVector2D(BlockRects[CurrentBlocks[i].Id].Rect.Min + BlockRects[CurrentBlocks[i].Id].Rect.Size));
 				
 				if (SelectedBlocks.Contains(CurrentBlocks[i].Id))
@@ -762,12 +773,9 @@ FReply SCustomizableObjectLayoutGrid::OnMouseMove( const FGeometry& MyGeometry,
 {
 	CurrentMousePosition = MyGeometry.AbsoluteToLocal(MouseEvent.GetScreenSpacePosition());
 
-	if (Mode != ELGM_Edit)
-	{
-		return SCompoundWidget::OnMouseMove(MyGeometry, MouseEvent);
-	}
+	ELayoutGridMode GridMode = Mode.Get();
 
-	if(MouseEvent.IsMouseButtonDown(EKeys::LeftMouseButton))
+	if(MouseEvent.IsMouseButtonDown(EKeys::LeftMouseButton) && GridMode == ELGM_Edit)
 	{
 		FVector2D Pos = MyGeometry.AbsoluteToLocal(MouseEvent.GetScreenSpacePosition());
 
@@ -788,6 +796,7 @@ FReply SCustomizableObjectLayoutGrid::OnMouseMove( const FGeometry& MyGeometry,
 
 				if (!bIsResizing)
 				{
+					// Bounding box of all selected blocks in grid units.
 					FIntRect TotalBlock;
 					bool bFirstBlock = true;
 
@@ -812,8 +821,14 @@ FReply SCustomizableObjectLayoutGrid::OnMouseMove( const FGeometry& MyGeometry,
 
 					FIntPoint Grid = GridSize.Get();
 					FIntRect BlockMovement = TotalBlock;
-					BlockMovement.Min.X = FMath::Max(0, FMath::Min(TotalBlock.Min.X + CellDeltaX, Grid.X - TotalBlock.Size().X));
-					BlockMovement.Min.Y = FMath::Max(0, FMath::Min(TotalBlock.Min.Y + CellDeltaY, Grid.Y - TotalBlock.Size().Y));
+
+					// Block movement in layouts is restricted to the positive quadrant.
+					//BlockMovement.Min.X = FMath::Max(0, FMath::Min(TotalBlock.Min.X + CellDeltaX, Grid.X - TotalBlock.Size().X));
+					//BlockMovement.Min.Y = FMath::Max(0, FMath::Min(TotalBlock.Min.Y + CellDeltaY, Grid.Y - TotalBlock.Size().Y));
+					BlockMovement.Min.X = FMath::Max(0, TotalBlock.Min.X + CellDeltaX);
+					BlockMovement.Min.Y = FMath::Max(0, TotalBlock.Min.Y + CellDeltaY);
+					//BlockMovement.Min.X = TotalBlock.Min.X + CellDeltaX;
+					//BlockMovement.Min.Y = TotalBlock.Min.Y + CellDeltaY;
 
 					BlockMovement.Max = BlockMovement.Min + TotalBlock.Size();
 
@@ -851,8 +866,11 @@ FReply SCustomizableObjectLayoutGrid::OnMouseMove( const FGeometry& MyGeometry,
 							FIntPoint Grid = GridSize.Get();
 
 							FIntPoint BlockSize = Block.Size();
-							Block.Max.X = FMath::Max(Block.Min.X + 1, FMath::Min(Block.Max.X + CellDeltaX, Grid.X));
-							Block.Max.Y = FMath::Max(Block.Min.Y + 1, FMath::Min(Block.Max.Y + CellDeltaY, Grid.Y));
+							// Block movement in layouts is restricted to the positive quadrant.
+							//Block.Max.X = FMath::Max(Block.Min.X + 1, FMath::Min(Block.Max.X + CellDeltaX, Grid.X));
+							//Block.Max.Y = FMath::Max(Block.Min.Y + 1, FMath::Min(Block.Max.Y + CellDeltaY, Grid.Y));
+							Block.Max.X = Block.Max.X + CellDeltaX;
+							Block.Max.Y = Block.Max.Y + CellDeltaY;
 
 							if (Block != InitialBlock)
 							{
@@ -893,7 +911,7 @@ FReply SCustomizableObjectLayoutGrid::OnMouseMove( const FGeometry& MyGeometry,
 		}
 	}
 	
-	if (!bIsDragging && !bIsResizing && SelectedBlocks.Num()==1)
+	if (!bIsDragging && !bIsResizing && SelectedBlocks.Num()==1 && GridMode == ELGM_Edit)
 	{
 		const TArray<FCustomizableObjectLayoutBlock>& CurrentBlocks = Blocks.Get();
 		for (int32 i = CurrentBlocks.Num() - 1; i > -1; --i)
@@ -919,7 +937,7 @@ FReply SCustomizableObjectLayoutGrid::OnMouseMove( const FGeometry& MyGeometry,
 		if (MouseEvent.IsMouseButtonDown(EKeys::MiddleMouseButton))
 		{
 			FVector2D Pos = MyGeometry.AbsoluteToLocal(MouseEvent.GetScreenSpacePosition());
-			PaddingAmount += Pos - PaddingStart;
+			PointOfView.PaddingAmount += Pos - PaddingStart;
 			PaddingStart = Pos;
 		}
 		else
@@ -949,37 +967,36 @@ FReply SCustomizableObjectLayoutGrid::OnMouseWheel(const FGeometry& MyGeometry,
 	constexpr int32 MinZoomLevel = -2;
 	constexpr int32 MaxZoomLevel = 3;
 
-	if (Mode == ELGM_Edit)
 	{
-		double OldZoomFactor = FMath::Pow(2.0, double(Zoom - 1));
-		FVector2D UnzoomedPadding = PaddingAmount * (1.0f / OldZoomFactor);
+		double OldZoomFactor = PointOfView.GetZoomFactor();
+		FVector2D UnzoomedPadding = PointOfView.PaddingAmount * (1.0f / OldZoomFactor);
 
 		if (MouseEvent.GetWheelDelta() > 0)
 		{
-			int32 NewZoomLevel = FMath::Min(Zoom + 1, MaxZoomLevel);
-			if (Zoom != NewZoomLevel)
+			int32 NewZoomLevel = FMath::Min(PointOfView.Zoom + 1, MaxZoomLevel);
+			if (PointOfView.Zoom != NewZoomLevel)
 			{
 				//FVector2D GridCenter = DrawOrigin + (FVector2D((float)GridSize.Get().X, (float)GridSize.Get().Y) / 2.0f) * CellSize;
 				//DistanceFromOrigin = CurrentMousePosition - GridCenter;
 
-				Zoom = NewZoomLevel;
+				PointOfView.Zoom = NewZoomLevel;
 			}
 		}
 		else
 		{
-			int32 NewZoomLevel = FMath::Max(Zoom - 1, MinZoomLevel);
-			if (Zoom != NewZoomLevel)
+			int32 NewZoomLevel = FMath::Max(PointOfView.Zoom - 1, MinZoomLevel);
+			if (PointOfView.Zoom != NewZoomLevel)
 			{
 				//DistanceFromOrigin = FVector2D::Zero();
 				//PaddingAmount = FVector2D::Zero();
 
-				Zoom = NewZoomLevel;
+				PointOfView.Zoom = NewZoomLevel;
 			}
 		}
 
-		double NewZoomFactor = FMath::Pow(2.0, double(Zoom - 1));
+		double NewZoomFactor = PointOfView.GetZoomFactor();
 		FVector2D RezoomedPadding = UnzoomedPadding * NewZoomFactor;
-		PaddingAmount = RezoomedPadding;
+		PointOfView.PaddingAmount = RezoomedPadding;
 
 		return FReply::Handled().SetUserFocus(SharedThis(this), EFocusCause::Mouse, true);
 	}
@@ -990,7 +1007,8 @@ FReply SCustomizableObjectLayoutGrid::OnMouseWheel(const FGeometry& MyGeometry,
 
 FReply SCustomizableObjectLayoutGrid::OnKeyDown(const FGeometry& MyGeometry, const FKeyEvent& InKeyEvent)
 {
-	if (Mode != ELGM_Edit)
+	ELayoutGridMode GridMode = Mode.Get();
+	if (GridMode != ELGM_Edit)
 	{
 		return SCompoundWidget::OnKeyDown(MyGeometry, InKeyEvent);
 	}
@@ -1047,12 +1065,6 @@ void SCustomizableObjectLayoutGrid::SetSelectedBlock(FGuid block )
 }
 
 
-void SCustomizableObjectLayoutGrid::SetSelectedBlocks( const TArray<FGuid>& blocks )
-{
-	SelectedBlocks = blocks;
-}
-
-
 const TArray<FGuid>& SCustomizableObjectLayoutGrid::GetSelectedBlocks() const
 {
 	return SelectedBlocks;
@@ -1067,8 +1079,8 @@ void SCustomizableObjectLayoutGrid::DeleteSelectedBlocks()
 
 void SCustomizableObjectLayoutGrid::ResetView()
 {
-	Zoom = 1;
-	PaddingAmount = { 0,0 };
+	PointOfView.Zoom = 1;
+	PointOfView.PaddingAmount = { 0,0 };
 }
 
 
@@ -1372,7 +1384,7 @@ void FUVCanvasDrawer::Draw_RenderThread(FRDGBuilder& GraphBuilder, const FDrawPa
 	Canvas.SetRenderTargetScissorRect(RenderTarget->GetClippingRect());
 
 	// Number of tiles to render in each axis including the unit tile.
-	constexpr int32 NumTiles = 8;
+	constexpr int32 NumTiles = 4;
 
 	// Num Lines
 	const uint32 NumAxisLines = 2;
@@ -1406,7 +1418,8 @@ void FUVCanvasDrawer::Draw_RenderThread(FRDGBuilder& GraphBuilder, const FDrawPa
 	const FHitProxyId HitProxyId = Canvas.GetHitProxyId();
 
 	int32 FullTilesSize = NumTiles * Size.X;
-	FVector2D TilesOrigin = Origin - FVector2D(FullTilesSize / 2, FullTilesSize / 2);
+	//FVector2D TilesOrigin = Origin - FVector2D(FullTilesSize / 2, FullTilesSize / 2);
+	FVector2D TilesOrigin = Origin;
 
 	// Create lines as pairs of points
 	FVector LinePoints[2];
@@ -1422,6 +1435,7 @@ void FUVCanvasDrawer::Draw_RenderThread(FRDGBuilder& GraphBuilder, const FDrawPa
 
 
 	// Drawing Extended Grid
+	if (LayoutMode!=ELGM_ShowUVsOnly)
 	{
 		// Vertical Lines
 		for (int32 LineIndex = 0; LineIndex < NumTiles*GridSize.X + 1; LineIndex++)
@@ -1443,6 +1457,7 @@ void FUVCanvasDrawer::Draw_RenderThread(FRDGBuilder& GraphBuilder, const FDrawPa
 	}
 
 	// Drawing Unit Grid
+	if (LayoutMode != ELGM_ShowUVsOnly)
 	{
 		// Vertical Lines
 		for (int32 LineIndex = 0; LineIndex < GridSize.X + 1; LineIndex++)
@@ -1524,29 +1539,37 @@ void FUVCanvasDrawer::Draw_RenderThread(FRDGBuilder& GraphBuilder, const FDrawPa
 	}
 
 	// Drawing Blocks
-	for (const FCustomizableObjectLayoutBlock& Block : Blocks)
+	if (LayoutMode != ELGM_ShowUVsOnly)
 	{
-		const FColor SelectionBlockColor = SelectedBlocks.Contains(Block.Id) ? SelectedBlockColor : UnselectedBlockColor;
+		for (const FCustomizableObjectLayoutBlock& Block : Blocks)
+		{
+			FColor BlockColor = AutomaticBlockColor;
+			if (!Block.bIsAutomatic)
+			{
+				BlockColor = SelectedBlocks.Contains(Block.Id) ? SelectedBlockColor : UnselectedBlockColor;
+			}
 
-		const FVector2f BlockMin(Block.Min);
-		const FVector2f BlockMax(Block.Max);
 
-		// Selection Block
-		FRect2D SelectionBlock;
-		SelectionBlock.Min = FVector2f(Origin) + BlockMin * CellSize + CellSize * 0.1f;
-		SelectionBlock.Size = (BlockMax - BlockMin) * CellSize - CellSize * 0.2f;
+			const FVector2f BlockMin(Block.Min);
+			const FVector2f BlockMax(Block.Max);
 
-		DrawBlock(BatchedElements, HitProxyId, SelectionBlock, SelectionBlockColor, Block.Mask.Get());
+			// Selection Block
+			FRect2D BlockRect;
+			BlockRect.Min = FVector2f(Origin) + BlockMin * CellSize + CellSize * 0.1f;
+			BlockRect.Size = (BlockMax - BlockMin) * CellSize - CellSize * 0.2f;
 
-		if (LayoutMode == ELayoutGridMode::ELGM_Edit)
-		{
-			// Resize Block
-			FRect2D ResizeBlock;;
-			float HandleRectSize = FMath::Log2(float(GridSize.X)) / 10.0f;
-			ResizeBlock.Size = FVector2f(CellSize) * HandleRectSize;
-			ResizeBlock.Min = SelectionBlock.Min + SelectionBlock.Size - ResizeBlock.Size;
+			DrawBlock(BatchedElements, HitProxyId, BlockRect, BlockColor, Block.Mask.Get());
+
+			if (LayoutMode == ELayoutGridMode::ELGM_Edit && !Block.bIsAutomatic)
+			{
+				// Resize Block
+				FRect2D ResizeBlock;;
+				float HandleRectSize = FMath::Log2(float(GridSize.X)) / 10.0f;
+				ResizeBlock.Size = FVector2f(CellSize) * HandleRectSize;
+				ResizeBlock.Min = BlockRect.Min + BlockRect.Size - ResizeBlock.Size;
 
-			DrawBlock(BatchedElements, HitProxyId, ResizeBlock, ResizeBlockColor);
+				DrawBlock(BatchedElements, HitProxyId, ResizeBlock, ResizeBlockColor);
+			}
 		}
 	}
 
@@ -1566,7 +1589,11 @@ void FUVCanvasDrawer::DrawBlock(FBatchedElements* BatchedElements, const FHitPro
 
 	auto VertexToMaskUVs = [this](const FVector4& V)
 		{
-			return ( FVector2D(V.X,V.Y) - Origin ) / (FVector2D(CellSize) * FVector2D(GridSize.X, GridSize.X));
+			FVector2D Result = ( FVector2D(V.X,V.Y) - Origin ) / (FVector2D(CellSize) * FVector2D(GridSize.X, GridSize.X));
+			// TODO Modulo doesn't work with cross-tile blocks: use tiling.
+			Result.X = FMath::Fmod(Result.X, 1.0);
+			Result.Y = FMath::Fmod(Result.Y, 1.0);
+			return Result;
 		};
 
 	// Brush Paint triangle
@@ -1576,13 +1603,17 @@ void FUVCanvasDrawer::DrawBlock(FBatchedElements* BatchedElements, const FHitPro
 		int32 V2 = BatchedElements->AddVertex(Vert2, VertexToMaskUVs(Vert2), Color, HitProxyId);
 		int32 V3 = BatchedElements->AddVertex(Vert3, VertexToMaskUVs(Vert3), Color, HitProxyId);
 
-		FTexture* Texture = GWhiteTexture;
+		EBlendMode Mode = EBlendMode::BLEND_Translucent;
+		BatchedElements->AddTriangle(V0, V1, V2, GWhiteTexture, Mode);
+		BatchedElements->AddTriangle(V1, V3, V2, GWhiteTexture, Mode);
+
 		if (Mask)
 		{
-			Texture = Mask->GetResource();
+			Mode = EBlendMode::BLEND_Additive;
+			FTexture* Texture = Mask->GetResource();
+			BatchedElements->AddTriangle(V0, V1, V2, Texture, Mode);
+			BatchedElements->AddTriangle(V1, V3, V2, Texture, Mode);
 		}
-		BatchedElements->AddTriangle(V0, V1, V2, Texture, EBlendMode::BLEND_Translucent);
-		BatchedElements->AddTriangle(V1, V3, V2, Texture, EBlendMode::BLEND_Translucent);
 	}
 
 	// Drawing Outline to selected Blocks
diff --git a/Engine/Plugins/Experimental/Mutable/Source/CustomizableObjectEditor/Private/MuCOE/SCustomizableObjectLayoutGrid.h b/Engine/Plugins/Experimental/Mutable/Source/CustomizableObjectEditor/Private/MuCOE/SCustomizableObjectLayoutGrid.h
--- a/Engine/Plugins/Experimental/Mutable/Source/CustomizableObjectEditor/Private/MuCOE/SCustomizableObjectLayoutGrid.h
+++ b/Engine/Plugins/Experimental/Mutable/Source/CustomizableObjectEditor/Private/MuCOE/SCustomizableObjectLayoutGrid.h
@@ -29,7 +29,7 @@ typedef enum
 {
 	ELGM_Show,
 	ELGM_Edit,
-	ELGM_Select
+	ELGM_ShowUVsOnly,
 } ELayoutGridMode;
 
 struct FBlockWidgetData
@@ -61,11 +61,11 @@ public:
 
 	SLATE_BEGIN_ARGS( SCustomizableObjectLayoutGrid ){}
 
+		SLATE_ATTRIBUTE(ELayoutGridMode, Mode)
 		SLATE_ATTRIBUTE( FIntPoint, GridSize )
-		SLATE_ATTRIBUTE( TArray<FCustomizableObjectLayoutBlock>, Blocks )
+		SLATE_ATTRIBUTE(TArray<FCustomizableObjectLayoutBlock>, Blocks)
 		SLATE_ARGUMENT( TArray<FVector2f>, UVLayout )
 		SLATE_ARGUMENT( TArray<FVector2f>, UnassignedUVLayoutVertices )
-		SLATE_ARGUMENT( ELayoutGridMode, Mode )
 		SLATE_ARGUMENT( FColor, SelectionColor )
 		SLATE_EVENT( FBlockChangedDelegate, OnBlockChanged )
 		SLATE_EVENT( FBlockSelectionChangedDelegate, OnSelectionChanged )
@@ -100,7 +100,6 @@ public:
 
 	/** Set the currently selected block */
 	void SetSelectedBlock( FGuid block );
-	void SetSelectedBlocks( const TArray<FGuid>& blocks );
 
 	/** */
 	const TArray<FGuid>& GetSelectedBlocks() const;
@@ -127,6 +126,24 @@ public:
 	/** Set the grid and blocks to show in the widget. */
 	void SetBlocks( const FIntPoint& GridSize, const TArray<FCustomizableObjectLayoutBlock>& Blocks);
 
+	struct FPointOfView
+	{
+		/** Amount of padding since start dragging */
+		FVector2D PaddingAmount = FVector2D::Zero();
+
+		/** Level of zoom */
+		int32 Zoom = 1;
+
+		/** */
+		double GetZoomFactor() const
+		{
+			return FMath::Pow(2.0, double(Zoom - 1));
+		}
+
+	};
+
+	/** Current point of view. */
+	FPointOfView PointOfView;
 
 private:
 
@@ -170,7 +187,7 @@ private:
 	/** Size of the grid in blocks */
 	TAttribute<FIntPoint> GridSize;
 
-	/** Array with all the blocks of the layout */
+	/** Array with all the editable blocks of the layout */
 	TAttribute< TArray<FCustomizableObjectLayoutBlock> > Blocks;
 
 	/** Array with all the UVs to draw in the layout */
@@ -180,7 +197,7 @@ private:
 	TArray<FVector2f> UnassignedUVLayoutVertices;
 
 	/** Layout mode */
-	ELayoutGridMode Mode = ELGM_Show;
+	TAttribute<ELayoutGridMode> Mode = ELGM_Show;
 
 	float CellSize = 0.0f;
 
@@ -216,15 +233,9 @@ private:
 	/** Position where the layout grid starts to be drawn */
 	FVector2D DrawOrigin;
 
-	/** Amount of padding since start dragging */
-	FVector2D PaddingAmount = FVector2D::Zero();
-
 	/** Position where the padding started */
 	FVector2D PaddingStart;
 
-	/** Level of zoom */
-	int32 Zoom = 1;
-
 	/** Selection Rectangle */
 	FRect2D SelectionRect;
 
diff --git a/Engine/Plugins/Experimental/Mutable/Source/CustomizableObjectEditor/Private/MuCOE/SCustomizableObjectNodeLayoutBlocksEditor.cpp b/Engine/Plugins/Experimental/Mutable/Source/CustomizableObjectEditor/Private/MuCOE/SCustomizableObjectNodeLayoutBlocksEditor.cpp
--- a/Engine/Plugins/Experimental/Mutable/Source/CustomizableObjectEditor/Private/MuCOE/SCustomizableObjectNodeLayoutBlocksEditor.cpp
+++ b/Engine/Plugins/Experimental/Mutable/Source/CustomizableObjectEditor/Private/MuCOE/SCustomizableObjectNodeLayoutBlocksEditor.cpp
@@ -39,7 +39,7 @@ public:
 	/**  */
 	TSharedPtr< FUICommandInfo > AddBlock;
 	TSharedPtr< FUICommandInfo > RemoveBlock;
-	TSharedPtr< FUICommandInfo > GenerateBlocks;
+	TSharedPtr< FUICommandInfo > ConsolidateBlocks;
 
 	/**
 	 * Initialize commands
@@ -48,7 +48,7 @@ public:
 	{
 		UI_COMMAND( AddBlock, "Add Block", "Add a new block to the layout.", EUserInterfaceActionType::Button, FInputChord() );
 		UI_COMMAND( RemoveBlock, "Remove Block", "Remove a block from the layout.", EUserInterfaceActionType::Button, FInputChord() );
-		UI_COMMAND( GenerateBlocks, "Generate Blocks", "Generate Blocks automatically from UVs", EUserInterfaceActionType::Button, FInputChord() );
+		UI_COMMAND(ConsolidateBlocks, "Consolidate Blocks", "Convert automatic blocks into user blocks.", EUserInterfaceActionType::Button, FInputChord() );
 	}
 };
 
@@ -66,16 +66,15 @@ void SCustomizableObjectNodeLayoutBlocksEditor::Construct(const FArguments& InAr
 }
 
 
-void SCustomizableObjectNodeLayoutBlocksEditor::SetCurrentLayout( UCustomizableObjectLayout* Layout, UCustomizableObjectLayout* UVOverrideLayout )
+void SCustomizableObjectNodeLayoutBlocksEditor::SetCurrentLayout( UCustomizableObjectLayout* InLayout, UCustomizableObjectLayout* InUVOverrideLayout )
 {
-	CurrentLayout = Layout;
+	CurrentLayout = InLayout;
+	UVOverrideLayout = InUVOverrideLayout;
 
-	UCustomizableObjectLayout* LayoutForUVs = UVOverrideLayout ? UVOverrideLayout : Layout;
+	UCustomizableObjectLayout* LayoutForUVs = UVOverrideLayout ? UVOverrideLayout : CurrentLayout;
 
-	// Try to locate the source mesh
 	TArray<FVector2f> UVs;
 	TArray<FVector2f> UnassignedUVs;
-
 	if (CurrentLayout)
 	{
 		LayoutForUVs->GetUVChannel(UVs, LayoutForUVs->GetUVChannel());
@@ -88,6 +87,15 @@ void SCustomizableObjectNodeLayoutBlocksEditor::SetCurrentLayout( UCustomizableO
 		}
 	}
 
+	// Save some layout widget state to persist between updates
+	bool bHadWidget = false;
+	SCustomizableObjectLayoutGrid::FPointOfView OldView;
+	if (LayoutGridWidget)
+	{
+		bHadWidget = true;
+		OldView = LayoutGridWidget->PointOfView;
+	}
+
 	this->ChildSlot
 	[
 		SNew(SVerticalBox)
@@ -103,7 +111,7 @@ void SCustomizableObjectNodeLayoutBlocksEditor::SetCurrentLayout( UCustomizableO
 		.FillHeight(1)
 		[
 			SAssignNew(LayoutGridWidget, SCustomizableObjectLayoutGrid)
-			.Mode(ELGM_Edit)
+			.Mode(this, &SCustomizableObjectNodeLayoutBlocksEditor::GetGridMode)
 			.GridSize(this, &SCustomizableObjectNodeLayoutBlocksEditor::GetGridSize)
 			.Blocks(this, &SCustomizableObjectNodeLayoutBlocksEditor::GetBlocks)
 			.UVLayout(UVs)
@@ -118,11 +126,16 @@ void SCustomizableObjectNodeLayoutBlocksEditor::SetCurrentLayout( UCustomizableO
 			.OnSetBlockMask(this, &SCustomizableObjectNodeLayoutBlocksEditor::OnSetBlockMask)
 		]
 	];	
-}
 
+	if (CurrentLayout)
+	{
+		CurrentLayout->GenerateAutomaticBlocksFromUVs();
+	}
 
-SCustomizableObjectNodeLayoutBlocksEditor::~SCustomizableObjectNodeLayoutBlocksEditor()
-{
+	if (bHadWidget)
+	{
+		LayoutGridWidget->PointOfView = OldView;
+	}
 }
 
 
@@ -141,13 +154,29 @@ TSharedRef<SWidget> SCustomizableObjectNodeLayoutBlocksEditor::BuildLayoutToolBa
 	const ISlateStyle* const StyleSet = &FCoreStyle::Get();
 	const FName& StyleName = "ToolBar";
 
-	LayoutToolbarBuilder.BeginSection("Blocks");
+	if (CurrentLayout->PackingStrategy!=ECustomizableObjectTextureLayoutPackingStrategy::Overlay)
 	{
-		LayoutToolbarBuilder.AddToolBarButton(FLayoutEditorCommands::Get().AddBlock);
-		LayoutToolbarBuilder.AddToolBarButton(FLayoutEditorCommands::Get().RemoveBlock);
-		LayoutToolbarBuilder.AddToolBarButton(FLayoutEditorCommands::Get().GenerateBlocks);
+		LayoutToolbarBuilder.BeginSection("Blocks");
+		{
+			LayoutToolbarBuilder.AddToolBarButton(FLayoutEditorCommands::Get().AddBlock);
+			LayoutToolbarBuilder.AddToolBarButton(FLayoutEditorCommands::Get().RemoveBlock);
+			
+			// Disable block consolidation if we are defininf blocks on top of another layout
+			bool bCanConsolidate = !UVOverrideLayout;
+
+			// Disable block consolidation if the cuyrrent automatic strategy doesn't generate blocks.
+			if (CurrentLayout->AutomaticBlocksStrategy != ECustomizableObjectLayoutAutomaticBlocksStrategy::Rectangles)
+			{
+				bCanConsolidate = false;
+			}
+
+			if (bCanConsolidate)
+			{
+				LayoutToolbarBuilder.AddToolBarButton(FLayoutEditorCommands::Get().ConsolidateBlocks);
+			}
+		}
+		LayoutToolbarBuilder.EndSection();
 	}
-	LayoutToolbarBuilder.EndSection();
 
 	LayoutToolbarBuilder.BeginSection("Info");
 	{
@@ -191,6 +220,8 @@ void SCustomizableObjectNodeLayoutBlocksEditor::OnAddBlock()
 		FCustomizableObjectLayoutBlock Block;
 		CurrentLayout->Blocks.Add(Block);
 
+		CurrentLayout->GenerateAutomaticBlocksFromUVs();
+
 		if (LayoutGridWidget.IsValid())
 		{
 			LayoutGridWidget->SetSelectedBlock(Block.Id);
@@ -208,6 +239,8 @@ void SCustomizableObjectNodeLayoutBlocksEditor::OnAddBlockAt(const FIntPoint Min
 		
 		FCustomizableObjectLayoutBlock block(Min,Max);
 		CurrentLayout->Blocks.Add(block);
+
+		CurrentLayout->GenerateAutomaticBlocksFromUVs();
 	}
 }
 
@@ -229,26 +262,41 @@ void SCustomizableObjectNodeLayoutBlocksEditor::OnRemoveBlock()
 					It.RemoveCurrent();
 				}
 			}
+
+			CurrentLayout->GenerateAutomaticBlocksFromUVs();
 		}
 	}
 }
 
 
-void SCustomizableObjectNodeLayoutBlocksEditor::OnGenerateBlocks()
+void SCustomizableObjectNodeLayoutBlocksEditor::OnConsolidateBlocks()
 {
 	if (CurrentLayout)
 	{
-		const FScopedTransaction Transaction(LOCTEXT("OnGenerateBlocks", "Generate Blocks"));
+		const FScopedTransaction Transaction(LOCTEXT("OnConsolidateBlocks", "Consolidate Blocks"));
 		CurrentLayout->Modify();
-		
-		CurrentLayout->GenerateBlocksFromUVs();
+				
+		CurrentLayout->GenerateAutomaticBlocksFromUVs();
+
+		CurrentLayout->ConsolidateAutomaticBlocks();
 	}
 }
 
 
+ELayoutGridMode SCustomizableObjectNodeLayoutBlocksEditor::GetGridMode() const
+{
+	if (CurrentLayout && CurrentLayout->PackingStrategy != ECustomizableObjectTextureLayoutPackingStrategy::Overlay)
+	{
+		return ELayoutGridMode::ELGM_Edit;
+	}
+
+	return ELayoutGridMode::ELGM_ShowUVsOnly;
+}
+
+
 FIntPoint SCustomizableObjectNodeLayoutBlocksEditor::GetGridSize() const
 {
-	if ( CurrentLayout )
+	if (CurrentLayout)
 	{
 		return CurrentLayout->GetGridSize();
 	}
@@ -273,6 +321,8 @@ void SCustomizableObjectNodeLayoutBlocksEditor::OnBlockChanged( FGuid BlockId, F
 				break;
 			}
 		}
+
+		CurrentLayout->GenerateAutomaticBlocksFromUVs();
 	}
 }
 
@@ -365,6 +415,8 @@ void SCustomizableObjectNodeLayoutBlocksEditor::OnSetBlockMask(UTexture2D* InVal
 				}
 			}
 		}
+
+		CurrentLayout->GenerateAutomaticBlocksFromUVs();
 	}
 }
 
@@ -378,6 +430,8 @@ TArray<FCustomizableObjectLayoutBlock> SCustomizableObjectNodeLayoutBlocksEditor
 		Blocks = CurrentLayout->Blocks;
 	}
 
+	Blocks.Append(CurrentLayout->AutomaticBlocks );
+
 	return Blocks;
 }
 
@@ -402,8 +456,8 @@ void SCustomizableObjectNodeLayoutBlocksEditor::BindCommands()
 		FIsActionChecked() );
 
 	UICommandList->MapAction(
-		Commands.GenerateBlocks,
-		FExecuteAction::CreateSP(this, &SCustomizableObjectNodeLayoutBlocksEditor::OnGenerateBlocks),
+		Commands.ConsolidateBlocks,
+		FExecuteAction::CreateSP(this, &SCustomizableObjectNodeLayoutBlocksEditor::OnConsolidateBlocks),
 		FCanExecuteAction(),
 		FIsActionChecked());
 }
@@ -435,12 +489,17 @@ TSharedPtr<IToolTip> SCustomizableObjectNodeLayoutBlocksEditor::GenerateInfoTool
 	};
 
 	// Duplicate command
-	BuildShortcutAndTooltip(LOCTEXT("ShortCut_DuplicateBlocks", "CTRL + D"), LOCTEXT("Tooltip_DuplicateBlocks", "Duplicate selected block/s"));
-	BuildShortcutAndTooltip(LOCTEXT("ShortCut_CreateNewBlock", "CTRL + N"), LOCTEXT("Tooltip_CreateNewBlock", "Create new block"));
-	BuildShortcutAndTooltip(LOCTEXT("ShortCut_FillGridSize", "CTRL + F"), LOCTEXT("Tooltip_FillGridSize", "Resize selected block/s to grid size"));
-	BuildShortcutAndTooltip(LOCTEXT("ShortCut_DeleteSelectedBlock","DEL"), LOCTEXT("Tooltip_DeleteSelectedBlock","Delete selected block/s"));
-	BuildShortcutAndTooltip(LOCTEXT("ShortCut_SelectMultipleBlocksOneByOne","SHIFT + L Click"), LOCTEXT("Tooltip_SelectMultipleBlocksOneByOne","Select multiple blocks one by one"));
-	BuildShortcutAndTooltip(LOCTEXT("ShortCut_SelectMultipleBlocks","L Click + Drag"), LOCTEXT("Tooltip_SelectMultipleBlocks","Select blocks that intersect with the yellow rectangle"));
+	if (CurrentLayout->PackingStrategy != ECustomizableObjectTextureLayoutPackingStrategy::Overlay)
+	{
+		BuildShortcutAndTooltip(LOCTEXT("ShortCut_DuplicateBlocks", "CTRL + D"), LOCTEXT("Tooltip_DuplicateBlocks", "Duplicate selected block/s"));
+		BuildShortcutAndTooltip(LOCTEXT("ShortCut_CreateNewBlock", "CTRL + N"), LOCTEXT("Tooltip_CreateNewBlock", "Create new block"));
+		BuildShortcutAndTooltip(LOCTEXT("ShortCut_FillGridSize", "CTRL + F"), LOCTEXT("Tooltip_FillGridSize", "Resize selected block/s to grid size"));
+		BuildShortcutAndTooltip(LOCTEXT("ShortCut_DeleteSelectedBlock", "DEL"), LOCTEXT("Tooltip_DeleteSelectedBlock", "Delete selected block/s"));
+		BuildShortcutAndTooltip(LOCTEXT("ShortCut_SelectMultipleBlocksOneByOne", "SHIFT + L Click"), LOCTEXT("Tooltip_SelectMultipleBlocksOneByOne", "Select multiple blocks one by one"));
+		BuildShortcutAndTooltip(LOCTEXT("ShortCut_SelectMultipleBlocks", "L Click + Drag"), LOCTEXT("Tooltip_SelectMultipleBlocks", "Select blocks that intersect with the yellow rectangle"));
+	}
+	BuildShortcutAndTooltip(LOCTEXT("ShortCut_Pan", "M Click + Drag"), LOCTEXT("Tooltip_Pan", "Pan the UV view."));
+	BuildShortcutAndTooltip(LOCTEXT("ShortCut_Zoom", "M Wheel"), LOCTEXT("Tooltip_Zoom", "Zoom in and out the UV view."));
 
 	return SNew(SToolTip)
 	[
diff --git a/Engine/Plugins/Experimental/Mutable/Source/CustomizableObjectEditor/Private/MuCOE/SCustomizableObjectNodeLayoutBlocksEditor.h b/Engine/Plugins/Experimental/Mutable/Source/CustomizableObjectEditor/Private/MuCOE/SCustomizableObjectNodeLayoutBlocksEditor.h
--- a/Engine/Plugins/Experimental/Mutable/Source/CustomizableObjectEditor/Private/MuCOE/SCustomizableObjectNodeLayoutBlocksEditor.h
+++ b/Engine/Plugins/Experimental/Mutable/Source/CustomizableObjectEditor/Private/MuCOE/SCustomizableObjectNodeLayoutBlocksEditor.h
@@ -4,6 +4,7 @@
 
 #include "UObject/GCObject.h"
 #include "Widgets/SCompoundWidget.h"
+#include "MuCOE/SCustomizableObjectLayoutGrid.h"
 
 namespace ESelectInfo { enum Type : int; }
 
@@ -31,7 +32,6 @@ public:
 	SLATE_END_ARGS()
 
 	SCustomizableObjectNodeLayoutBlocksEditor();
-	virtual ~SCustomizableObjectNodeLayoutBlocksEditor();
 
 	void Construct(const FArguments& InArgs);
 	
@@ -54,9 +54,12 @@ public:
 
 private:
 
-	/** */
+	/** Layout whose blocksa re being edited. */
 	TObjectPtr<class UCustomizableObjectLayout> CurrentLayout;
 
+	/** If valid, layout use to show the UVs instead of CurrentLayout. */
+	TObjectPtr<class UCustomizableObjectLayout> UVOverrideLayout;
+
 	/** */
 	TSharedPtr<SCustomizableObjectLayoutGrid> LayoutGridWidget;
 
@@ -66,6 +69,7 @@ private:
 private:
 
 	/** */
+	ELayoutGridMode GetGridMode() const;
 	FIntPoint GetGridSize() const;
 	void OnBlockChanged(FGuid BlockId, FIntRect Block );
 	TArray<FCustomizableObjectLayoutBlock> GetBlocks() const;
@@ -77,8 +81,8 @@ private:
 	void OnAddBlockAt(const FIntPoint Min, const FIntPoint Max);
 	void OnRemoveBlock();
 
-	/** Generate layout blocks using UVs*/
-	void OnGenerateBlocks();
+	/** Turn the automatic layout blocks into user-created blocks. */
+	void OnConsolidateBlocks();
 
 	/** Sets the block priority from the input text. */
 	void OnSetBlockPriority(int32 InValue);
diff --git a/Engine/Plugins/Experimental/Mutable/Source/CustomizableObjectEditor/Private/MuCOE/SMutableLayoutViewer.cpp b/Engine/Plugins/Experimental/Mutable/Source/CustomizableObjectEditor/Private/MuCOE/SMutableLayoutViewer.cpp
--- a/Engine/Plugins/Experimental/Mutable/Source/CustomizableObjectEditor/Private/MuCOE/SMutableLayoutViewer.cpp
+++ b/Engine/Plugins/Experimental/Mutable/Source/CustomizableObjectEditor/Private/MuCOE/SMutableLayoutViewer.cpp
@@ -104,8 +104,8 @@ void SMutableLayoutViewer::Tick(const FGeometry& AllottedGeometry, const double
 			Blocks.SetNum( BlockCount );
 			for (int32 BlockIndex = 0; BlockIndex < BlockCount; ++BlockIndex)
 			{
-				UE::Math::TIntVector2<uint16> Min = MutableLayout->Blocks[BlockIndex].Min;
-				UE::Math::TIntVector2<uint16> Size = MutableLayout->Blocks[BlockIndex].Size;
+				FIntVector2 Min = MutableLayout->Blocks[BlockIndex].Min;
+				FIntVector2 Size = MutableLayout->Blocks[BlockIndex].Size;
 
 				Blocks[BlockIndex].Min[0] = Min[0];
 				Blocks[BlockIndex].Min[1] = Min[1];
diff --git a/Engine/Plugins/Experimental/Mutable/Source/CustomizableObjectEditor/Public/MuCOE/CustomizableObjectLayout.h b/Engine/Plugins/Experimental/Mutable/Source/CustomizableObjectEditor/Public/MuCOE/CustomizableObjectLayout.h
--- a/Engine/Plugins/Experimental/Mutable/Source/CustomizableObjectEditor/Public/MuCOE/CustomizableObjectLayout.h
+++ b/Engine/Plugins/Experimental/Mutable/Source/CustomizableObjectEditor/Public/MuCOE/CustomizableObjectLayout.h
@@ -4,6 +4,7 @@
 
 
 #include "MuT/NodeLayout.h"
+#include "MuR/Image.h"
 #include "CustomizableObjectLayout.generated.h"
 
 
@@ -20,6 +21,29 @@ enum class ECustomizableObjectTextureLayoutPackingStrategy : uint8
 
 mu::EPackStrategy ConvertLayoutStrategy(const ECustomizableObjectTextureLayoutPackingStrategy LayoutPackStrategy);
 
+UENUM()
+enum class ECustomizableObjectLayoutAutomaticBlocksStrategy : uint8
+{
+	// Create rectangles on a grid splitting the UV space if possible
+	Rectangles = 0 UMETA(DisplayName = "Rectangles"),
+	// Detect UV islands and create blocks for them with masks
+	UVIslands = 1 UMETA(DisplayName = "UV Islands"),
+	// Don't create automatic blocks, and ignore UVs that don't have a manual block already. They get assigned to the first avilable block or ignored if none. This is the legacy behavior.
+	Ignore = 2 UMETA(DisplayName = "Ignore"),
+};
+
+UENUM()
+enum class ECustomizableObjectLayoutAutomaticBlocksMergeStrategy : uint8
+{
+	// Don't merge the blocks
+	DontMerge = 0 UMETA(DisplayName = "Don't merge"),
+	// Merge the block if a block is entirely included in another block
+	MergeChildBlocks = 1 UMETA(DisplayName = "Child blocks"),
+	// TODO: this option would merge only when the child falls entirely inside the mask
+	// MergeChildInsideMask = 2 UMETA(DisplayName = "Child in mask"),
+};
+
+
 // Fixed Layout reduction methods
 UENUM()
 enum class ECustomizableObjectLayoutBlockReductionMethod : uint8
@@ -72,6 +96,9 @@ struct CUSTOMIZABLEOBJECTEDITOR_API FCustomizableObjectLayoutBlock
 	/** Block mask to use to filter the UVs when assigning them to the block. */
 	UPROPERTY(EditAnywhere, Category = CustomizableObject)
 	TObjectPtr<class UTexture2D> Mask = nullptr;
+
+	/** Transient flag used in the UI to differentiate between manual and automatic blocks. */
+	bool bIsAutomatic = false;
 };
 
 UCLASS()
@@ -85,14 +112,11 @@ public:
 
 	// Sets the layout parameters
 	void SetLayout(UObject* InMesh, int32 LODIndex, int32 MatIndex, int32 UVIndex);
-	void SetPackingStrategy(ECustomizableObjectTextureLayoutPackingStrategy Strategy);
 	void SetGridSize(FIntPoint Size);
 	void SetMaxGridSize(FIntPoint Size);
 	void SetLayoutName(FString Name);
 	void SetIgnoreVertexLayoutWarnings(bool bValue);
 	void SetIgnoreWarningsLOD(int32 LODValue);
-	void SetBlockReductionMethod(ECustomizableObjectLayoutBlockReductionMethod Method);
-
 
 	int32 GetLOD() const { return LOD; }
 	int32 GetMaterial() const { return Material; }
@@ -101,23 +125,44 @@ public:
 	UObject* GetMesh() const { return Mesh; }
 	FIntPoint GetGridSize() const { return GridSize; }
 	FIntPoint GetMaxGridSize() const { return MaxGridSize; }
-	ECustomizableObjectTextureLayoutPackingStrategy GetPackingStrategy() const { return PackingStrategy; }
 	bool GetIgnoreVertexLayoutWarnings() const { return bIgnoreUnassignedVertexWarning; };
 	int32 GetFirstLODToIgnoreWarnings() const { return FirstLODToIgnore; };
-	ECustomizableObjectLayoutBlockReductionMethod GetBlockReductionMethod()const { return BlockReductionMethod; }
 
+	// TODO: Remove the default channel index. Anywhere using the default is probably wrong.
 	void GetUVChannel(TArray<FVector2f>& UVs, int32 UVChannelIndex = 0) const;
 
-	// Get a block index in the array from its id. Return -1 if not found.
+	/** Get a block index in the array from its id.Return - 1 if not found. */
 	int32 FindBlock(const FGuid& InId) const;
 
-	void GenerateBlocksFromUVs();
+	/** Generate all the transient UV layout automatic blocks with unassigned UVs. */
+	void GenerateAutomaticBlocksFromUVs();
+
+	/** Convert the transient automatic blocks into real blocks. */
+	void ConsolidateAutomaticBlocks();
 
+	/** List of blocks manually defined in the layout. */
 	UPROPERTY()
 	TArray<FCustomizableObjectLayoutBlock> Blocks;
 
+	/** List of blocks automatically defined (in the layout) for preview. */
+	UPROPERTY(Transient)
+	TArray<FCustomizableObjectLayoutBlock> AutomaticBlocks;
+
+	/** List of UVs to highlight in the layout because they have issues. */
 	TArray< TArray<FVector2f> > UnassignedUVs;
 
+	UPROPERTY()
+	ECustomizableObjectTextureLayoutPackingStrategy PackingStrategy = ECustomizableObjectTextureLayoutPackingStrategy::Resizable;
+
+	UPROPERTY()
+	ECustomizableObjectLayoutAutomaticBlocksStrategy AutomaticBlocksStrategy = ECustomizableObjectLayoutAutomaticBlocksStrategy::Rectangles;
+
+	UPROPERTY()
+	ECustomizableObjectLayoutAutomaticBlocksMergeStrategy AutomaticBlocksMergeStrategy = ECustomizableObjectLayoutAutomaticBlocksMergeStrategy::MergeChildBlocks;
+
+	UPROPERTY()
+	ECustomizableObjectLayoutBlockReductionMethod BlockReductionMethod = ECustomizableObjectLayoutBlockReductionMethod::Halve;
+
 private:
 
 	UPROPERTY()
@@ -135,13 +180,10 @@ private:
 	UPROPERTY()
 	FIntPoint GridSize;
 
-	/** Used with the fixed layout strategy. */
+	/** Maximum grid size the layout can grow to. Used with the fixed layout strategy. */
 	UPROPERTY()
 	FIntPoint MaxGridSize;
 
-	UPROPERTY()
-	ECustomizableObjectTextureLayoutPackingStrategy PackingStrategy = ECustomizableObjectTextureLayoutPackingStrategy::Resizable;
-
 	UPROPERTY()
 	FString LayoutName;
 
@@ -153,7 +195,4 @@ private:
 	UPROPERTY()
 	int32 FirstLODToIgnore = 0;
 
-	UPROPERTY()
-	ECustomizableObjectLayoutBlockReductionMethod BlockReductionMethod = ECustomizableObjectLayoutBlockReductionMethod::Halve;
-
 };
diff --git a/Engine/Plugins/Experimental/Mutable/Source/MutableRuntime/Internal/MuR/OpImageGrow.h b/Engine/Plugins/Experimental/Mutable/Source/MutableRuntime/Internal/MuR/OpImageGrow.h
new file mode 100644
--- /dev/null
+++ b/Engine/Plugins/Experimental/Mutable/Source/MutableRuntime/Internal/MuR/OpImageGrow.h
@@ -0,0 +1,124 @@
+// Copyright Epic Games, Inc. All Rights Reserved.
+
+#pragma once
+
+#include "MuR/ImagePrivate.h"
+#include "MuR/Platform.h"
+#include "MuR/ConvertData.h"
+
+
+namespace mu
+{
+
+	//---------------------------------------------------------------------------------------------
+	//! Grow the image as if it was a bitmap
+	//---------------------------------------------------------------------------------------------
+	inline void ImageGrow( Image* pImage )
+	{
+		check(pImage && pImage->GetFormat()== EImageFormat::IF_L_UBYTE );
+
+		MUTABLE_CPUPROFILER_SCOPE(ImageGrow1)
+
+		int sizeX = pImage->GetSizeX();
+		int sizeY = pImage->GetSizeY();
+		check( sizeX>0 && sizeY>0 );
+
+		ImagePtrConst pSource = pImage->Clone();
+
+		int rowSize = sizeX*GetImageFormatData( pImage->GetFormat() ).BytesPerBlock;
+
+		for ( int y=0; y<sizeY; ++y )
+		{
+			int y0 = FMath::Max( 0, FMath::Min( sizeY-1, y-1 ) );
+			int y1 = FMath::Max( 0, FMath::Min( sizeY-1, y   ) );
+			int y2 = FMath::Max( 0, FMath::Min( sizeY-1, y+1 ) );
+
+            const uint8_t* pS0 = pSource->GetLODData(0) + y0 * rowSize;
+            const uint8_t* pS1 = pSource->GetLODData(0) + y1 * rowSize;
+            const uint8_t* pS2 = pSource->GetLODData(0) + y2 * rowSize;
+
+            uint8_t* pR = pImage->GetLODData(0) + y1 * rowSize;
+
+			if ( pS0[0] || pS2[0] || pS1[1] )
+			{
+				pR[0] = 255;
+			}
+
+			if ( pS0[sizeX-1] || pS2[sizeX-1] || pS1[sizeX-2] )
+			{
+				pR[sizeX-1] = 255;
+			}
+
+			for ( int x=1; x<sizeX-1; ++x )
+			{
+				if ( pS0[x] || pS2[x] || pS1[x-1] || pS1[x+1] )
+				{
+					pR[x] = 255;
+				}
+			}
+
+		}
+
+	}
+
+
+	//---------------------------------------------------------------------------------------------
+	//! Grow the image copying border colours, by using an external mask as reference
+	//---------------------------------------------------------------------------------------------
+	inline void ImageGrow( Image* pImage, const Image* pMask )
+	{
+		MUTABLE_CPUPROFILER_SCOPE(ImageGrow2)
+
+		check( pMask->GetFormat()== EImageFormat::IF_L_UBYTE );
+
+		int sizeX = pImage->GetSizeX();
+		int sizeY = pImage->GetSizeY();
+		check( sizeX>0 && sizeY>0 );
+
+        if (sizeX<2 || sizeY<2)
+        {
+            return;
+        }
+
+		int pixelSize = GetImageFormatData( pImage->GetFormat() ).BytesPerBlock;
+		int rowSize = sizeX*pixelSize;
+
+		for ( int y=0; y<sizeY; ++y )
+		{
+            const uint8_t* pS0 = y>0 ? pMask->GetLODData(0) + (y-1) * sizeX : nullptr;
+            const uint8_t* pS1 = pMask->GetLODData(0) + y * sizeX;
+            const uint8_t* pS2 = y<sizeY-1 ? pMask->GetLODData(0) + (y+1) * sizeX : nullptr;
+
+            const uint8_t* pR0 = y>0 ? pImage->GetLODData(0) + (y-1) * rowSize : nullptr;
+                  uint8_t* pR1 = pImage->GetLODData(0) + y * rowSize;
+            const uint8_t* pR2 = y<sizeY-1 ? pImage->GetLODData(0) + (y+1) * rowSize : nullptr;
+
+			for ( int x=0; x<sizeX; ++x )
+			{
+				if (pS1[x])
+				{
+					// Nothing to do.
+				}
+				else if ( y>0 && pS0[x] )
+				{
+					memcpy( &pR1[x*pixelSize], &pR0[x*pixelSize], pixelSize );
+				}
+				else if ( y<sizeY-1 && pS2[x] )
+				{
+					memcpy( &pR1[x*pixelSize], &pR2[x*pixelSize], pixelSize );
+				}
+				else if ( x<sizeX-1 && pS1[x+1] )
+				{
+					memcpy( &pR1[x*pixelSize], &pR1[(x+1)*pixelSize], pixelSize );
+				}
+				else if ( x>0 && pS1[x-1] )
+				{
+					memcpy( &pR1[x*pixelSize], &pR1[(x-1)*pixelSize], pixelSize );
+				}
+			}
+
+		}
+
+	}
+
+}
diff --git a/Engine/Plugins/Experimental/Mutable/Source/MutableRuntime/Internal/MuR/OpMeshMerge.h b/Engine/Plugins/Experimental/Mutable/Source/MutableRuntime/Internal/MuR/OpMeshMerge.h
--- a/Engine/Plugins/Experimental/Mutable/Source/MutableRuntime/Internal/MuR/OpMeshMerge.h
+++ b/Engine/Plugins/Experimental/Mutable/Source/MutableRuntime/Internal/MuR/OpMeshMerge.h
@@ -540,7 +540,10 @@ namespace mu
 		{
 			MUTABLE_CPUPROFILER_SCOPE(Pose);
 
-			Result->BonePoses.Reserve(Result->GetSkeleton()->GetBoneCount());
+			if (Result->GetSkeleton())
+			{
+				Result->BonePoses.Reserve(Result->GetSkeleton()->GetBoneCount());
+			}
 
 			// Copy poses from the first mesh
 			Result->BonePoses = pFirst->BonePoses;
diff --git a/Engine/Plugins/Experimental/Mutable/Source/MutableRuntime/Internal/MuR/Raster.h b/Engine/Plugins/Experimental/Mutable/Source/MutableRuntime/Internal/MuR/Raster.h
new file mode 100644
--- /dev/null
+++ b/Engine/Plugins/Experimental/Mutable/Source/MutableRuntime/Internal/MuR/Raster.h
@@ -0,0 +1,261 @@
+// Copyright Epic Games, Inc. All Rights Reserved.
+
+#pragma once
+
+#include "HAL/Platform.h"
+#include "Math/UnrealMathUtility.h"
+#include "Templates/Function.h"
+
+namespace mu
+{
+
+	template<int NUM_INTERPOLATORS>
+	class RasterVertex
+	{
+	public:
+		RasterVertex() : x(0), y(0) {}
+		RasterVertex( float aX, float aY ) : x(aX), y(aY) {}
+
+		//! Raster coordinates of the vertex in image space.
+		float x,y;
+
+		//! Extra channels to interpolate accross de triangles.
+		float interpolators[NUM_INTERPOLATORS];
+	};
+
+
+	//! Example rasterizer pixel processor.
+	//! It only takes the color interpolated and assigns it to the pixel.
+	class ColorPixelProcessor
+	{
+	public:
+		inline void ProcessPixel(uint8* pBufferPos, float interpolatedValues[3]) const
+		{
+			pBufferPos[0] = (uint8)(interpolatedValues[0]*0xFF);
+			pBufferPos[1] = (uint8)(interpolatedValues[1]*0xFF);
+			pBufferPos[2] = (uint8)(interpolatedValues[2]*0xFF);
+		}
+
+		inline void operator()(uint8* BufferPos, float InterpolatedValues[3]) const
+		{
+			ProcessPixel(BufferPos, InterpolatedValues);
+		}
+	};
+
+	template<int32 NumInterpolatedValues>
+	using TTriangleRasterPixelProcRefType = TFunctionRef<void(uint8*, float[NumInterpolatedValues])>;
+
+	//! Templatized rasterizer. It is not very efficient on the interpolating side, but the cost
+    //! per pixel is supposed to be huge compared to the raster process (for lightmaps or
+    //! normalmaps).
+	//! TODO Optimize anyway.
+	//! TODO Option for two sided? right now it culls back facing.
+    template<int32 NUM_INTERPOLATORS, typename INT=int64>
+	void Triangle( uint8* Buffer, int32 BufferSize, int Width, int Height, int PixelSize
+				 , const RasterVertex<NUM_INTERPOLATORS> &av1
+				 , const RasterVertex<NUM_INTERPOLATORS> &av2
+				 , const RasterVertex<NUM_INTERPOLATORS> &av3
+				 , TTriangleRasterPixelProcRefType<NUM_INTERPOLATORS> PixelProc
+				 , bool cullBackface = true )
+	{
+		check(Width * Height * PixelSize <= BufferSize);
+
+		float interpolatedValues[NUM_INTERPOLATORS];
+
+		int Stride = Width*PixelSize;
+
+		// Check backface culling
+		// TODO better option than resorting vertices?
+		float lCrossZ = (av2.x-av1.x) * (av3.y-av1.y) - (av2.y-av1.y) * (av3.x-av1.x);
+		const RasterVertex<NUM_INTERPOLATORS>& v1 = av1;
+		const RasterVertex<NUM_INTERPOLATORS>& v2 = cullBackface ? av2 : ( (lCrossZ>0) ? av3 : av2);
+		const RasterVertex<NUM_INTERPOLATORS>& v3 = cullBackface ? av3 : ( (lCrossZ>0) ? av2 : av3);
+
+	    // 28.4 fixed-point coordinates
+        const INT Y1 = (INT)(16.0f * v1.y + 0.5f);
+        const INT Y2 = (INT)(16.0f * v2.y + 0.5f);
+        const INT Y3 = (INT)(16.0f * v3.y + 0.5f);
+
+        const INT X1 = (INT)(16.0f * v1.x + 0.5f);
+        const INT X2 = (INT)(16.0f * v2.x + 0.5f);
+        const INT X3 = (INT)(16.0f * v3.x + 0.5f);
+
+	    // Deltas
+        const INT DX12 = X1 - X2;
+        const INT DX23 = X2 - X3;
+        const INT DX31 = X3 - X1;
+
+        const INT DY12 = Y1 - Y2;
+        const INT DY23 = Y2 - Y3;
+        const INT DY31 = Y3 - Y1;
+
+	    // Fixed-point deltas
+        const INT FDX12 = DX12 * 16;
+        const INT FDX23 = DX23 * 16;
+        const INT FDX31 = DX31 * 16;
+
+        const INT FDY12 = DY12 * 16;
+        const INT FDY23 = DY23 * 16;
+        const INT FDY31 = DY31 * 16;
+
+		// Block size, standard 8x8 (must be power of two)
+		constexpr INT q = 8;
+
+		// Skip partial blocks
+		INT ClampWidth = Width & ~(q - 1);
+		INT ClampHeight = Height & ~(q - 1);
+
+	    // Bounding rectangle
+        INT minx = FMath::Max(INT(0),				(FMath::Min3(X1, X2, X3) + 0xF) / 16 );
+        INT maxx = FMath::Min(INT(ClampWidth -1),	(FMath::Max3(X1, X2, X3) + 0xF) / 16 );
+        INT miny = FMath::Max(INT(0),				(FMath::Min3(Y1, Y2, Y3) + 0xF) / 16 );
+        INT maxy = FMath::Min(INT(ClampHeight -1),	(FMath::Max3(Y1, Y2, Y3) + 0xF) / 16 );
+
+		// Start in corner of 8x8 block
+		minx &= ~(q - 1);
+		miny &= ~(q - 1);
+
+		Buffer += miny * Stride;
+
+	    // Half-edge constants
+        INT C1 = DY12 * X1 - DX12 * Y1;
+        INT C2 = DY23 * X2 - DX23 * Y2;
+        INT C3 = DY31 * X3 - DX31 * Y3;
+
+	    // Correct for fill convention
+	    if(DY12 < 0 || (DY12 == 0 && DX12 > 0)) C1++;
+	    if(DY23 < 0 || (DY23 == 0 && DX23 > 0)) C2++;
+	    if(DY31 < 0 || (DY31 == 0 && DX31 > 0)) C3++;
+
+		// Interpolation process constant stuff
+        float ux = v2.x-v1.x;
+        float uy = v2.y-v1.y;
+        float vx = v3.x-v1.x;
+        float vy = v3.y-v1.y;
+        float det = ux*vy-vx*uy;
+        ux /= det;
+        uy /= det;
+        vx /= det;
+        vy /= det;
+
+	    // Loop through blocks
+        for(INT y = miny; y < maxy; y += q)
+	    {
+            for(INT x = minx; x < maxx; x += q)
+	        {
+	            // Corners of block
+                INT x0 = x * 16;
+                INT x1 = (x + q - 1) * 16;
+                INT y0 = y * 16;
+                INT y1 = (y + q - 1) * 16;
+
+	            // Evaluate half-space functions
+	            bool a00 = C1 + DX12 * y0 - DY12 * x0 > 0;
+	            bool a10 = C1 + DX12 * y0 - DY12 * x1 > 0;
+	            bool a01 = C1 + DX12 * y1 - DY12 * x0 > 0;
+	            bool a11 = C1 + DX12 * y1 - DY12 * x1 > 0;
+                INT a = (INT(a00) * 1) | (INT(a10) * 2) | (INT(a01) * 4) | (INT(a11) * 8);
+
+	            bool b00 = C2 + DX23 * y0 - DY23 * x0 > 0;
+	            bool b10 = C2 + DX23 * y0 - DY23 * x1 > 0;
+	            bool b01 = C2 + DX23 * y1 - DY23 * x0 > 0;
+	            bool b11 = C2 + DX23 * y1 - DY23 * x1 > 0;
+                INT b = (INT(b00) * 1) | (INT(b10) * 2) | (INT(b01) * 4) | (INT(b11) * 8);
+
+	            bool c00 = C3 + DX31 * y0 - DY31 * x0 > 0;
+	            bool c10 = C3 + DX31 * y0 - DY31 * x1 > 0;
+	            bool c01 = C3 + DX31 * y1 - DY31 * x0 > 0;
+	            bool c11 = C3 + DX31 * y1 - DY31 * x1 > 0;
+                INT c = (INT(c00) * 1) | (INT(c10) * 2) | (INT(c01) * 4) | (INT(c11) * 8);
+
+	            // Skip block when outside an edge
+				if (a == 0x0 || b == 0x0 || c == 0x0)
+				{
+					continue;
+				}
+
+				uint8 *CurrentBuffer = Buffer;
+
+	            // Accept whole block when totally covered
+	            if(a == 0xF && b == 0xF && c == 0xF)
+	            {
+                    for(INT iy = 0; iy < q; iy++)
+	                {
+                        for(INT ix = x; ix < x + q; ix++)
+	                    {
+	                        // Interpolate values
+	                        float px = ix-v1.x;
+	                        float py = y+iy-v1.y;
+							float alpha = px*vy+py*-vx;
+							float beta = px*-uy+py*ux;
+							alpha = FMath::Max(0.0, FMath::Min(1.0f, alpha ));
+							beta = FMath::Max(0.0, FMath::Min(1.0f, beta ));
+							float gamma = FMath::Max(0.0, 1.0f-alpha-beta );
+
+	                        for (int ii=0; ii<NUM_INTERPOLATORS; ii++)
+							{
+                                interpolatedValues[ii] = gamma * v1.interpolators[ii]
+													   + alpha * v2.interpolators[ii]
+													   + beta * v3.interpolators[ii];
+							}
+
+	                        // Process pixel
+							PixelProc(&CurrentBuffer[ix*PixelSize], interpolatedValues);
+	                    }
+
+						CurrentBuffer += Stride;
+	                }
+	            }
+	            else // Partially covered block
+	            {
+                    INT CY1 = C1 + DX12 * y0 - DY12 * x0;
+                    INT CY2 = C2 + DX23 * y0 - DY23 * x0;
+                    INT CY3 = C3 + DX31 * y0 - DY31 * x0;
+
+                    for(INT iy = y; iy < y + q; iy++)
+	                {
+                        INT CX1 = CY1;
+                        INT CX2 = CY2;
+                        INT CX3 = CY3;
+
+                        for(INT ix = x; ix < x + q; ix++)
+	                    {
+	                        if(CX1 > 0 && CX2 > 0 && CX3 > 0)
+							{
+		                        // Interpolate values
+		                        float px = ix-v1.x;
+		                        float py = iy-v1.y;
+								float alpha = px*vy+py*-vx;
+								float beta = px*-uy+py*ux;
+								alpha = FMath::Max(0.0, FMath::Min(1.0f, alpha ));
+								beta = FMath::Max(0.0, FMath::Min(1.0f, beta ));
+								float gamma = FMath::Max(0.0, 1.0f-alpha-beta );
+
+		                        for (int ii=0; ii<NUM_INTERPOLATORS; ii++)
+								{
+									interpolatedValues[ii] = gamma * v1.interpolators[ii]
+															+ alpha * v2.interpolators[ii]
+															+ beta * v3.interpolators[ii];
+								}
+
+								PixelProc(&CurrentBuffer[ix*PixelSize], interpolatedValues);
+	                        }
+
+	                        CX1 -= FDY12;
+	                        CX2 -= FDY23;
+	                        CX3 -= FDY31;
+	                    }
+
+	                    CY1 += FDX12;
+	                    CY2 += FDX23;
+	                    CY3 += FDX31;
+
+						CurrentBuffer += Stride;
+	                }
+	            }
+	        }
+
+			Buffer += q * Stride;
+	    }
+	}
+}
diff --git a/Engine/Plugins/Experimental/Mutable/Source/MutableRuntime/Private/MuR/CodeRunner.cpp b/Engine/Plugins/Experimental/Mutable/Source/MutableRuntime/Private/MuR/CodeRunner.cpp
--- a/Engine/Plugins/Experimental/Mutable/Source/MutableRuntime/Private/MuR/CodeRunner.cpp
+++ b/Engine/Plugins/Experimental/Mutable/Source/MutableRuntime/Private/MuR/CodeRunner.cpp
@@ -4366,7 +4366,7 @@ namespace mu
 
 				// Apply the mipmap reduction to the crop rectangle.
 				int32 MipsToSkip = item.ExecutionOptions;
-				box<UE::Math::TIntVector2<uint16>> rect;
+				box<FIntVector2> rect;
 				rect.min[0] = args.minX / (1 << MipsToSkip);
 				rect.min[1] = args.minY / (1 << MipsToSkip);
 				rect.size[0] = pB->GetSizeX();
diff --git a/Engine/Plugins/Experimental/Mutable/Source/MutableRuntime/Private/MuR/CodeRunnerTasked.cpp b/Engine/Plugins/Experimental/Mutable/Source/MutableRuntime/Private/MuR/CodeRunnerTasked.cpp
--- a/Engine/Plugins/Experimental/Mutable/Source/MutableRuntime/Private/MuR/CodeRunnerTasked.cpp
+++ b/Engine/Plugins/Experimental/Mutable/Source/MutableRuntime/Private/MuR/CodeRunnerTasked.cpp
@@ -1919,7 +1919,7 @@ namespace mu
 		Ptr<const Image> Block;
 		Ptr<const Image> Mask;
 		Ptr<Image> Result;
-		box< UE::Math::TIntVector2<uint16> > Rect;
+		box<FIntVector2> Rect;
 		FImageOperator::FImagePixelFormatFunc ImagePixelFormatFunc;
 	};
 
@@ -1950,7 +1950,7 @@ namespace mu
 			Mask = Runner->LoadImage({ Args.mask, Op.ExecutionIndex, Op.ExecutionOptions });
 		}
 
-		box< UE::Math::TIntVector2<uint16> > RectInblocks;
+		box<FIntVector2> RectInblocks;
 		RectInblocks.min = Layout->Blocks[RelBlockIndex].Min;
 		RectInblocks.size = Layout->Blocks[RelBlockIndex].Size;
 
@@ -1986,7 +1986,7 @@ namespace mu
 			EImageFormat Format = GetMostGenericFormat(Result->GetFormat(), Block->GetFormat());
 
 			// Resize image if it doesn't fit in the new block size
-			if (Block->GetSize() != Rect.size)
+			if (FIntVector2(Block->GetSize()) != Rect.size)
 			{
 				MUTABLE_CPUPROFILER_SCOPE(ImageComposeWithoutMask_BlockResize);
 
diff --git a/Engine/Plugins/Experimental/Mutable/Source/MutableRuntime/Private/MuR/Layout.cpp b/Engine/Plugins/Experimental/Mutable/Source/MutableRuntime/Private/MuR/Layout.cpp
--- a/Engine/Plugins/Experimental/Mutable/Source/MutableRuntime/Private/MuR/Layout.cpp
+++ b/Engine/Plugins/Experimental/Mutable/Source/MutableRuntime/Private/MuR/Layout.cpp
@@ -202,7 +202,7 @@ namespace mu {
 	bool Layout::IsSingleBlockAndFull() const
 	{
 		if (Blocks.Num() == 1
-			&& Blocks[0].Min == UE::Math::TIntVector2<uint16>(0, 0)
+			&& Blocks[0].Min == FIntVector2(0, 0)
 			&& Blocks[0].Size == Size)
 		{
 			return true;
diff --git a/Engine/Plugins/Experimental/Mutable/Source/MutableRuntime/Private/MuR/OpImageCompose.cpp b/Engine/Plugins/Experimental/Mutable/Source/MutableRuntime/Private/MuR/OpImageCompose.cpp
--- a/Engine/Plugins/Experimental/Mutable/Source/MutableRuntime/Private/MuR/OpImageCompose.cpp
+++ b/Engine/Plugins/Experimental/Mutable/Source/MutableRuntime/Private/MuR/OpImageCompose.cpp
@@ -7,7 +7,7 @@
 namespace mu
 {
 
-	void FImageOperator::ImageCompose(Image* Base, const Image* Block, const box< UE::Math::TIntVector2<uint16> >& Rect)
+	void FImageOperator::ImageCompose(Image* Base, const Image* Block, const box< FIntVector2 >& Rect)
 	{
 		check(Base && Block);
 		check(Base != Block);
@@ -63,8 +63,8 @@ namespace mu
 				Rect.min[1] / FormatInfo.PixelsPerBlockY);
 
 			UE::Math::TIntVector2<uint16> BlockMipSize(
-				FMath::DivideAndRoundUp(Rect.size[0], uint16(FormatInfo.PixelsPerBlockX)),
-				FMath::DivideAndRoundUp(Rect.size[1], uint16(FormatInfo.PixelsPerBlockY)));
+				FMath::DivideAndRoundUp(Rect.size[0], int32(FormatInfo.PixelsPerBlockX)),
+				FMath::DivideAndRoundUp(Rect.size[1], int32(FormatInfo.PixelsPerBlockY)));
 
 			int32 DoneMips = 0;
 			for (; DoneMips < Base->GetLODCount() && DoneMips < Block->GetLODCount(); ++DoneMips)
diff --git a/Engine/Plugins/Experimental/Mutable/Source/MutableRuntime/Private/MuR/OpImageGrow.h b/Engine/Plugins/Experimental/Mutable/Source/MutableRuntime/Private/MuR/OpImageGrow.h
deleted file mode 100644
diff --git a/Engine/Plugins/Experimental/Mutable/Source/MutableRuntime/Private/MuR/OpImageLayer.h b/Engine/Plugins/Experimental/Mutable/Source/MutableRuntime/Private/MuR/OpImageLayer.h
--- a/Engine/Plugins/Experimental/Mutable/Source/MutableRuntime/Private/MuR/OpImageLayer.h
+++ b/Engine/Plugins/Experimental/Mutable/Source/MutableRuntime/Private/MuR/OpImageLayer.h
@@ -1976,7 +1976,7 @@ namespace mu
 			Image* BaseImage,
 			const Image* MaskImage,
 			const Image* BlendedImage,
-			const box<UE::Math::TIntVector2<uint16>>& Rect)
+			const box<FIntVector2>& Rect)
 	{
 		check(BaseImage->GetSizeX() >= Rect.min[0] + Rect.size[0]);
 		check(BaseImage->GetSizeY() >= Rect.min[1] + Rect.size[1]);
@@ -2553,7 +2553,7 @@ namespace mu
 	inline void ImageBlendOnBaseNoAlpha(Image* BaseImage,
 		const Image* MaskImage,
 		const Image* BlendedImage,
-		const box<UE::Math::TIntVector2<uint16>>& Rect)
+		const box<FIntVector2>& Rect)
 	{
 		ImageLayerOnBaseNoAlpha<BlendChannel, false>(BaseImage, MaskImage, BlendedImage, Rect);
 	}
diff --git a/Engine/Plugins/Experimental/Mutable/Source/MutableRuntime/Private/MuR/OpLayoutPack.cpp b/Engine/Plugins/Experimental/Mutable/Source/MutableRuntime/Private/MuR/OpLayoutPack.cpp
--- a/Engine/Plugins/Experimental/Mutable/Source/MutableRuntime/Private/MuR/OpLayoutPack.cpp
+++ b/Engine/Plugins/Experimental/Mutable/Source/MutableRuntime/Private/MuR/OpLayoutPack.cpp
@@ -14,7 +14,7 @@ namespace mu
 	{
 		FPackLayoutBlock() {}
 
-		FPackLayoutBlock(int32 InIndex, UE::Math::TIntVector2<uint16> InSize, int32 InPriority = 0, bool bInReduceBothAxes = false, bool bInReduceByTwo = false)
+		FPackLayoutBlock(int32 InIndex, FIntVector2 InSize, int32 InPriority = 0, bool bInReduceBothAxes = false, bool bInReduceByTwo = false)
 		{
 			Index = InIndex;
 			size = InSize;
@@ -24,7 +24,7 @@ namespace mu
 		}
 
 		int32 Index = -1;
-		UE::Math::TIntVector2<uint16> size = UE::Math::TIntVector2<uint16>(0, 0);
+		FIntVector2 size = { 0, 0 };
 		int32 priority = 0;
 		bool bReduceBothAxes = false;
 		bool bReduceByTwo = false;
@@ -74,7 +74,7 @@ namespace mu
 
 	struct FScratchLayoutPack
 	{
-		TArray< UE::Math::TIntVector2<uint16> > blocks;
+		TArray< FIntVector2 > blocks;
 		TArray< FPackLayoutBlock > sorted;
 		TArray< FIntVector2 > positions;
 		TArray< int32 > priorities;
@@ -107,7 +107,7 @@ namespace mu
     }
 
 
-	inline void ReductionOperation(uint16& BlockSize, EReductionMethod ReductionMethod, int32 bReduceByTwo)
+	inline void ReductionOperation(int32& BlockSize, EReductionMethod ReductionMethod, int32 bReduceByTwo)
 	{
 		if (ReductionMethod == EReductionMethod::Unitary)
 		{
@@ -129,7 +129,9 @@ namespace mu
 
 		int32 oldBlockArea = scratch.sorted[r_it].size[0] * scratch.sorted[r_it].size[1];
 
-		if (scratch.sorted[r_it].size[0] != 1 || scratch.sorted[r_it].size[1] != 1)
+		if (oldBlockArea>0
+			&&
+			(scratch.sorted[r_it].size[0] != 1 || scratch.sorted[r_it].size[1] != 1))
 		{
 			if (scratch.sorted[r_it].bReduceBothAxes)
 			{
@@ -212,7 +214,7 @@ namespace mu
 	}
 
 
-	inline bool SetPositions(int32 bestY,int32 layoutSizeY, uint16* maxX, uint16* maxY, FScratchLayoutPack& scratch, EPackStrategy packStrategy)
+	inline bool SetPositions(int32 bestY,int32 layoutSizeY, int32* maxX, int32* maxY, FScratchLayoutPack& scratch, EPackStrategy packStrategy)
 	{
 		bool fits = true;
 
@@ -269,8 +271,8 @@ namespace mu
 					}
 
 					// Does it make an unfillable hole with the top or side?
-					uint16 minX = TNumericLimits<uint16>::Max();
-					uint16 minY = TNumericLimits<uint16>::Max();
+					int32 minX = TNumericLimits<int32>::Max();
+					int32 minY = TNumericLimits<int32>::Max();
 					for (size_t b = 0; b < scratch.sorted.Num(); ++b)
 					{
 						if (!packedFlag[b] && b != candidate)
@@ -349,7 +351,7 @@ namespace mu
 
 			// Store
 			scratch.positions[scratch.sorted[best].Index] = FIntVector2(bestX, bestLevel);
-			*maxY = FMath::Max(*maxY, uint16(bestLevel + scratch.sorted[best].size[1]) );
+			*maxY = FMath::Max(*maxY, bestLevel + scratch.sorted[best].size[1] );
 
 			if (packStrategy == EPackStrategy::Fixed && *maxY > layoutSizeY)
 			{
@@ -413,13 +415,13 @@ namespace mu
 		EReductionMethod ReductionMethod = pSourceLayout->ReductionMethod;
 
         // Look for the maximum block sizes on the layout and the total area
-		uint16 maxX = 0;
-		uint16 maxY = 0;
+		int32 maxX = 0;
+		int32 maxY = 0;
         int32 area = 0;
 
         for ( int32 Index=0; Index<BlockCount; ++Index )
         {
-            box< UE::Math::TIntVector2<uint16> > b;
+            box< FIntVector2 > b;
 			b.min = pSourceLayout->Blocks[Index].Min;
 			b.size = pSourceLayout->Blocks[Index].Size;
 
@@ -640,7 +642,7 @@ namespace mu
 
         for ( int32 Index=0; Index<BlockCount; ++Index )
         {
-			pResult->Blocks[Index].Min = UE::Math::TIntVector2<uint16>(scratch.positions[Index]);
+			pResult->Blocks[Index].Min = scratch.positions[Index];
 			pResult->Blocks[Index].Size = scratch.blocks[Index];
 			pResult->Blocks[Index].Priority = scratch.priorities[Index];
 			pResult->Blocks[Index].bReduceBothAxes = scratch.ReduceBothAxes[Index];
diff --git a/Engine/Plugins/Experimental/Mutable/Source/MutableRuntime/Private/MuR/OpMeshApplyLayout.h b/Engine/Plugins/Experimental/Mutable/Source/MutableRuntime/Private/MuR/OpMeshApplyLayout.h
--- a/Engine/Plugins/Experimental/Mutable/Source/MutableRuntime/Private/MuR/OpMeshApplyLayout.h
+++ b/Engine/Plugins/Experimental/Mutable/Source/MutableRuntime/Private/MuR/OpMeshApplyLayout.h
@@ -54,7 +54,7 @@ namespace mu
 		{
 			FIntPoint grid = pLayout->GetGridSize();
 
-			box< UE::Math::TIntVector2<uint16> > block;
+			box< FIntVector2 > block;
 			block.min = pLayout->Blocks[b].Min;
 			block.size = pLayout->Blocks[b].Size;
 
diff --git a/Engine/Plugins/Experimental/Mutable/Source/MutableRuntime/Private/MuR/Raster.h b/Engine/Plugins/Experimental/Mutable/Source/MutableRuntime/Private/MuR/Raster.h
deleted file mode 100644
diff --git a/Engine/Plugins/Experimental/Mutable/Source/MutableRuntime/Public/MuR/Image.h b/Engine/Plugins/Experimental/Mutable/Source/MutableRuntime/Public/MuR/Image.h
--- a/Engine/Plugins/Experimental/Mutable/Source/MutableRuntime/Public/MuR/Image.h
+++ b/Engine/Plugins/Experimental/Mutable/Source/MutableRuntime/Public/MuR/Image.h
@@ -343,7 +343,7 @@ namespace mu
 		MUTABLERUNTIME_API bool ImageCrop(Image* Cropped, int32 CompressionQuality, const Image* Base, const box<FIntVector2>& Rect);
 
 		/** */
-		void ImageCompose(Image* Base, const Image* Block, const box< UE::Math::TIntVector2<uint16> >& Rect);
+		void ImageCompose(Image* Base, const Image* Block, const box<FIntVector2>& Rect);
 
 	};
 
diff --git a/Engine/Plugins/Experimental/Mutable/Source/MutableRuntime/Public/MuR/Layout.h b/Engine/Plugins/Experimental/Mutable/Source/MutableRuntime/Public/MuR/Layout.h
--- a/Engine/Plugins/Experimental/Mutable/Source/MutableRuntime/Public/MuR/Layout.h
+++ b/Engine/Plugins/Experimental/Mutable/Source/MutableRuntime/Public/MuR/Layout.h
@@ -33,8 +33,8 @@ namespace mu
 	{
 		static constexpr uint64 InvalidBlockId = TNumericLimits<uint64>::Max();
 
-		UE::Math::TIntVector2<uint16> Min = { 0, 0 };
-		UE::Math::TIntVector2<uint16> Size = { 0, 0 };
+		FIntVector2 Min = { 0, 0 };
+		FIntVector2 Size = { 0, 0 };
 
 		//! Absolute id used to control merging of various layouts
 		uint64 Id;
@@ -52,7 +52,7 @@ namespace mu
 		uint32 UnusedPadding : 30;
 
 		/** */
-		FLayoutBlock(UE::Math::TIntVector2<uint16> InMin = {}, UE::Math::TIntVector2<uint16> InSize = {})
+		FLayoutBlock(FIntVector2 InMin = {}, FIntVector2 InSize = {})
 		{
 			Min = InMin;
 			Size = InSize;
@@ -92,12 +92,12 @@ namespace mu
 	public:
 
 		//!
-		UE::Math::TIntVector2<uint16> Size = UE::Math::TIntVector2<uint16>(0, 0);
+		FIntVector2 Size = FIntVector2(0, 0);
 
 		/** Maximum size in layout blocks that this layout can grow to. From there on, blocks will shrink to fit. 
 		* If 0,0 then no maximum size applies.
 		*/
-		UE::Math::TIntVector2<uint16> MaxSize = UE::Math::TIntVector2<uint16>(0, 0);
+		FIntVector2 MaxSize = FIntVector2(0, 0);
 
 		//!
 		TArray<FLayoutBlock> Blocks;
diff --git a/Engine/Plugins/Experimental/Mutable/Source/MutableRuntime/Public/MuR/MutableMath.h b/Engine/Plugins/Experimental/Mutable/Source/MutableRuntime/Public/MuR/MutableMath.h
--- a/Engine/Plugins/Experimental/Mutable/Source/MutableRuntime/Public/MuR/MutableMath.h
+++ b/Engine/Plugins/Experimental/Mutable/Source/MutableRuntime/Public/MuR/MutableMath.h
@@ -16,6 +16,12 @@ namespace mu
 		return VECTOR::GetDim();
 	}
 
+	template<>
+	inline int32 GetDim<FVector2f>()
+	{
+		return 2;
+	}
+
 	template<>
 	inline int32 GetDim<FVector3f>()
 	{
@@ -93,17 +99,23 @@ namespace mu
             return result;
         }
 
-        //-----------------------------------------------------------------------------------------
-        inline bool Contains( const VECTOR& v ) const
-        {
-            return AllGreaterOrEqualThan(v,min) && AllSmallerThan(v, min+size);
-        }
+		//-----------------------------------------------------------------------------------------
+		inline bool Contains(const VECTOR& v) const
+		{
+			return AllGreaterOrEqualThan(v, min) && AllSmallerThan(v, min + size);
+		}
 
-        //-----------------------------------------------------------------------------------------
-        inline bool ContainsInclusive( const VECTOR& v ) const
-        {
-            return AllGreaterOrEqualThan(v, min) && AllSmallerOrEqualThan(v, min+size);
-        }
+		//-----------------------------------------------------------------------------------------
+		inline bool ContainsInclusive(const VECTOR& v) const
+		{
+			return AllGreaterOrEqualThan(v, min) && AllSmallerOrEqualThan(v, min + size);
+		}
+
+		//-----------------------------------------------------------------------------------------
+		inline bool Contains(const box<VECTOR>& v) const
+		{
+			return Contains(v.min) && ContainsInclusive(v.min+v.size);
+		}
 
         //-----------------------------------------------------------------------------------------
         //! It is inclusive: If the borders touch it is considered an intersection.
diff --git a/Engine/Plugins/Experimental/Mutable/Source/MutableTools/Private/MuT/ASTOpImageCompose.cpp b/Engine/Plugins/Experimental/Mutable/Source/MutableTools/Private/MuT/ASTOpImageCompose.cpp
--- a/Engine/Plugins/Experimental/Mutable/Source/MutableTools/Private/MuT/ASTOpImageCompose.cpp
+++ b/Engine/Plugins/Experimental/Mutable/Source/MutableTools/Private/MuT/ASTOpImageCompose.cpp
@@ -247,7 +247,7 @@ mu::Ptr<ASTOp> ASTOpImageCompose::OptimiseSemantic(const FModelOptimizationOptio
 		// Constant single-block full layout?
 		if (pLayout->GetBlockCount() == 1
 			&&
-			pLayout->Blocks[0].Min == UE::Math::TIntVector2<uint16>(0, 0)
+			pLayout->Blocks[0].Min == FIntVector2(0, 0)
 			&&
 			pLayout->Blocks[0].Size == pLayout->Size
 			&&
diff --git a/Engine/Plugins/Experimental/Mutable/Source/MutableTools/Private/MuT/CodeGenerator.cpp b/Engine/Plugins/Experimental/Mutable/Source/MutableTools/Private/MuT/CodeGenerator.cpp
--- a/Engine/Plugins/Experimental/Mutable/Source/MutableTools/Private/MuT/CodeGenerator.cpp
+++ b/Engine/Plugins/Experimental/Mutable/Source/MutableTools/Private/MuT/CodeGenerator.cpp
@@ -794,7 +794,7 @@ namespace mu
 	}
 
 
-	Ptr<Image> CodeGenerator::GenerateImageBlockPatchMask(const NodePatchImage* Patch, FIntPoint GridSize, int32 BlockPixelsX, int32 BlockPixelsY, box<UE::Math::TIntVector2<uint16>> RectInCells )
+	Ptr<Image> CodeGenerator::GenerateImageBlockPatchMask(const NodePatchImage* Patch, FIntPoint GridSize, int32 BlockPixelsX, int32 BlockPixelsY, box<FIntVector2> RectInCells )
 	{
 		// Create a patching mask for the block
 		Ptr<Image> PatchMask;
@@ -1372,7 +1372,7 @@ namespace mu
 								FIntPoint GridSize(FakeLayoutSize, FakeLayoutSize);
 								int32 BlockPixelsX = 1;
 								int32 BlockPixelsY = 1;
-								box< UE::Math::TIntVector2<uint16> > RectInCells;
+								box< FIntVector2 > RectInCells;
 								RectInCells.min = { 0,0 };
 								RectInCells.size = { FakeLayoutSize ,FakeLayoutSize };
 
@@ -1476,9 +1476,9 @@ namespace mu
 								imageAd = BlankImageOp;
 							}
 
-							auto UpdateBlockSize = [&BlockPixelsX, &BlockPixelsY, &FinalFormat, &bBlocksHaveMips, &bImageSizeWarning, &formatNode, &BlankImageOp, &node, surfaceNode, &t, &Options, this]( FImageDesc BlockDesc, UE::Math::TIntVector2<uint16> LayoutCellSize )
+							auto UpdateBlockSize = [&BlockPixelsX, &BlockPixelsY, &FinalFormat, &bBlocksHaveMips, &bImageSizeWarning, &formatNode, &BlankImageOp, &node, surfaceNode, &t, &Options, this]( FImageDesc BlockDesc, FIntVector2 LayoutCellSize )
 							{
-								if (BlockPixelsX == 0)
+								if (BlockPixelsX == 0 && LayoutCellSize.X>0 && LayoutCellSize.Y>0)
 								{
 									if (!bImageSizeWarning)
 									{
@@ -1547,7 +1547,7 @@ namespace mu
 								FImageDesc BlockDesc = blockAd->GetImageDesc(bReturnBestOption, nullptr);
 
 								// Block in layout grid units (cells)
-								box< UE::Math::TIntVector2<uint16> > RectInCells;
+								box< FIntVector2 > RectInCells;
 								RectInCells.min = pLayout->Blocks[BlockIndex].Min;
 								RectInCells.size = pLayout->Blocks[BlockIndex].Size;
 
@@ -1645,7 +1645,7 @@ namespace mu
 												Ptr<ASTOp> fragmentAd = ExtendResult.op;
 
 												// Block in layout grid units
-												box< UE::Math::TIntVector2<uint16> > rectInCells;
+												box< FIntVector2 > rectInCells;
 												rectInCells.min = pExtendLayout->Blocks[b].Min;
 												rectInCells.size = pExtendLayout->Blocks[b].Size;
 
diff --git a/Engine/Plugins/Experimental/Mutable/Source/MutableTools/Private/MuT/CodeGenerator.h b/Engine/Plugins/Experimental/Mutable/Source/MutableTools/Private/MuT/CodeGenerator.h
--- a/Engine/Plugins/Experimental/Mutable/Source/MutableTools/Private/MuT/CodeGenerator.h
+++ b/Engine/Plugins/Experimental/Mutable/Source/MutableTools/Private/MuT/CodeGenerator.h
@@ -134,12 +134,7 @@ namespace mu
 				return KeyHash;
 			}
 
-			bool operator==(const FGenericGenerationOptions& InKey) const
-			{
-				if (State != InKey.State) return false;
-				if (ActiveTags != InKey.ActiveTags) return false;
-				return true;
-			}
+			bool operator==(const FGenericGenerationOptions& InKey) const = default;
 
 			int32 State = -1;
 			TArray<FString> ActiveTags;
@@ -212,10 +207,7 @@ namespace mu
 				return KeyHash;
 			}
 
-			FORCEINLINE bool operator==(const FGeneratedCacheKey& Other) const
-			{
-				return Node == Other.Node && Options == Other.Options;
-			}
+			FORCEINLINE bool operator==(const FGeneratedCacheKey& Other) const = default;
 		};
 
 		typedef TMap<FGeneratedCacheKey, FGenericGenerationResult> FGeneratedGenericNodesMap;
@@ -513,7 +505,7 @@ namespace mu
 		Ptr<ASTOp> ApplyTiling(Ptr<ASTOp> Source, UE::Math::TIntVector2<int32> Size, EImageFormat Format);
 
 		/** Generate a layout block-sized image with a mask including all pixels in the blocks defined in the patch node. */
-		Ptr<Image> GenerateImageBlockPatchMask(const NodePatchImage* Patch, FIntPoint GridSize, int32 BlockPixelsX, int32 BlockPixelsY, box<UE::Math::TIntVector2<uint16>> RectInCells);
+		Ptr<Image> GenerateImageBlockPatchMask(const NodePatchImage* Patch, FIntPoint GridSize, int32 BlockPixelsX, int32 BlockPixelsY, box<FIntVector2> RectInCells);
 
 		/** Generate all the operations to apply the block patching on top of the BlockOp, and masking with PatchMask. */
 		Ptr<ASTOp> GenerateImageBlockPatch(Ptr<ASTOp> BlockOp, const NodePatchImage*, Ptr<Image> PatchMask, Ptr<ASTOp> ConditionOp, const FImageGenerationOptions&);
diff --git a/Engine/Plugins/Experimental/Mutable/Source/MutableTools/Private/MuT/CodeGenerator_Image.cpp b/Engine/Plugins/Experimental/Mutable/Source/MutableTools/Private/MuT/CodeGenerator_Image.cpp
--- a/Engine/Plugins/Experimental/Mutable/Source/MutableTools/Private/MuT/CodeGenerator_Image.cpp
+++ b/Engine/Plugins/Experimental/Mutable/Source/MutableTools/Private/MuT/CodeGenerator_Image.cpp
@@ -224,7 +224,7 @@ namespace mu
 			check( BlockIndex>=0 );
 
 			// Block in layout grid units
-			box< UE::Math::TIntVector2<uint16> > RectInCells;
+			box< FIntVector2 > RectInCells;
 			RectInCells.min = Options.LayoutToApply->Blocks[BlockIndex].Min;
 			RectInCells.size = Options.LayoutToApply->Blocks[BlockIndex].Size;
 
@@ -233,22 +233,40 @@ namespace mu
 			grid[1] = FMath::Max(1, grid[1]);
 
 			// Transform to pixels
-			box< UE::Math::TIntVector2<int32> > rect;
-			rect.min[0]  = (RectInCells.min[0]  * SourceImageSize[0]) / grid[0];
-			rect.min[1]  = (RectInCells.min[1]  * SourceImageSize[1]) / grid[1];
-			rect.size[0] = (RectInCells.size[0] * SourceImageSize[0]) / grid[0];
-			rect.size[1] = (RectInCells.size[1] * SourceImageSize[1]) / grid[1];
+			box< FIntVector2 > RectInPixels;
+			RectInPixels.min[0]  = (RectInCells.min[0]  * SourceImageSize[0]) / grid[0];
+			RectInPixels.min[1]  = (RectInCells.min[1]  * SourceImageSize[1]) / grid[1];
+			RectInPixels.size[0] = (RectInCells.size[0] * SourceImageSize[0]) / grid[0];
+			RectInPixels.size[1] = (RectInCells.size[1] * SourceImageSize[1]) / grid[1];
 
 			// Do we need to crop?
-			if (rect.min[0]!=0 || rect.min[1]!=0 || pImage->GetSizeX() != rect.size[0] || pImage->GetSizeY() != rect.size[1])
+			if (RectInPixels.min[0]!=0 || RectInPixels.min[1]!=0 || pImage->GetSizeX()!=RectInPixels.size[0] || pImage->GetSizeY()!=RectInPixels.size[1])
 			{
-				Ptr<ASTOpImageCrop> CropOp = new ASTOpImageCrop();
-				CropOp->Source = Result.op;
-				CropOp->Min[0] = rect.min[0];
-				CropOp->Min[1] = rect.min[1];
-				CropOp->Size[0] = rect.size[0];
-				CropOp->Size[1] = rect.size[1];
-				Result.op = CropOp;
+				// See if the rect belongs to a single texture tile
+				FIntVector2 TileMin(RectInPixels.min[0] / SourceImageSize[0], RectInPixels.min[1] / SourceImageSize[1]);
+				FIntVector2 TileMax((RectInPixels.min[0]+ RectInPixels.size[0]-1) / SourceImageSize[0], (RectInPixels.min[1]+ RectInPixels.size[1]-1) / SourceImageSize[1]);
+
+				if (TileMin != TileMax)
+				{
+					// Blocks spaning multiple texture tiles are not supported.
+					// To implement them, assemble a series of instructions to crop and compose the necessary rects from each tile into the final image.
+					ensure(false);
+
+					// Log an error message
+					ErrorLog->GetPrivate()->Add("A layout block goes across different texture tiles. This is not supported yet.", ELMT_ERROR, InNode->GetMessageContext());
+				}
+				else
+				{
+					Ptr<ASTOpImageCrop> CropOp = new ASTOpImageCrop();
+					CropOp->Source = Result.op;
+					
+					// Bring the crop rect to tile 0,0
+					CropOp->Min[0] = RectInPixels.min[0] - TileMin[0] * SourceImageSize[0];
+					CropOp->Min[1] = RectInPixels.min[1] - TileMin[1] * SourceImageSize[1];
+					CropOp->Size[0] = RectInPixels.size[0];
+					CropOp->Size[1] = RectInPixels.size[1];
+					Result.op = CropOp;
+				}
 			}
 		}
     }
@@ -1152,7 +1170,7 @@ namespace mu
 				check(BlockIndex >= 0);
 
 				// Block in layout grid units
-				box< UE::Math::TIntVector2<uint16> > RectInCells;
+				box< FIntVector2 > RectInCells;
 				RectInCells.min = Options.LayoutToApply->Blocks[BlockIndex].Min;
 				RectInCells.size = Options.LayoutToApply->Blocks[BlockIndex].Size;
 
@@ -1211,7 +1229,7 @@ namespace mu
 				check(BlockIndex >= 0);
 
 				// Block in layout grid units
-				box< UE::Math::TIntVector2<uint16> > RectInCells;
+				box< FIntVector2 > RectInCells;
 				RectInCells.min = Options.LayoutToApply->Blocks[BlockIndex].Min;
 				RectInCells.size = Options.LayoutToApply->Blocks[BlockIndex].Size;
 
@@ -1574,7 +1592,7 @@ namespace mu
 			check(BlockIndex >= 0);
 
 			// Block in layout grid units
-			box< UE::Math::TIntVector2<uint16> > RectInCells;
+			box< FIntVector2 > RectInCells;
 			RectInCells.min = Options.LayoutToApply->Blocks[BlockIndex].Min;
 			RectInCells.size = Options.LayoutToApply->Blocks[BlockIndex].Size;
 
diff --git a/Engine/Plugins/Experimental/Mutable/Source/MutableTools/Private/MuT/CodeGenerator_Mesh.cpp b/Engine/Plugins/Experimental/Mutable/Source/MutableTools/Private/MuT/CodeGenerator_Mesh.cpp
--- a/Engine/Plugins/Experimental/Mutable/Source/MutableTools/Private/MuT/CodeGenerator_Mesh.cpp
+++ b/Engine/Plugins/Experimental/Mutable/Source/MutableTools/Private/MuT/CodeGenerator_Mesh.cpp
@@ -75,190 +75,39 @@
 #include "MuT/Table.h"
 #include "MuT/TablePrivate.h"
 
-#include "Spatial/PointHashGrid3.h"
-
 namespace mu
 {
 	class Node;
 
-	
-	//---------------------------------------------------------------------------------------------
-    //! Create a map from vertices into vertices, collapsing vertices that have the same position. 
-	//! This version uses UE Containers to return.	
-    //---------------------------------------------------------------------------------------------
-	void MeshCreateCollapsedVertexMap(const mu::Mesh* Mesh, TArray<int32>& CollapsedVertices)
+	template<typename T>
+	struct TArray2D
 	{
-		MUTABLE_CPUPROFILER_SCOPE(LayoutUV_CreateCollapsedVertexMap);
-	
-		const int32 NumVertices = Mesh->GetVertexCount();
-		CollapsedVertices.Reserve(NumVertices);
-	
-		UE::Geometry::TPointHashGrid3f<int32> VertHash(0.01f, INDEX_NONE);
-		VertHash.Reserve(NumVertices);
-	
-		TArray<FVector3f> Vertices;
-		Vertices.SetNumUninitialized(NumVertices);
-	
-		mu::UntypedMeshBufferIteratorConst ItPosition = mu::UntypedMeshBufferIteratorConst(Mesh->GetVertexBuffers(), mu::MBS_POSITION);
-	
-		FVector3f* VertexData = Vertices.GetData();
-	
-		for (int32 VertexIndex = 0; VertexIndex < NumVertices; ++VertexIndex)
+		int32 SizeX = 0;
+		int32 SizeY = 0;
+		TArray<T> Data;
+
+		void Init(const T& Value, int32 InSizeX, int32 InSizeY)
 		{
-			*VertexData = ItPosition.GetAsVec3f();
-			VertHash.InsertPointUnsafe(VertexIndex, *VertexData);
-	
-			++ItPosition;
-			++VertexData;
+			SizeX = InSizeX;
+			SizeY = InSizeY;
+			Data.Init(Value, SizeX*SizeY );
 		}
-	
-		// Find unique vertices
-		CollapsedVertices.Init(INDEX_NONE, NumVertices);
-	
-		TArray<int32> NearbyVertices;
-		for (int32 VertexIndex = 0; VertexIndex < NumVertices; ++VertexIndex)
+
+		inline const T& Get(int32 X, int32 Y) const
 		{
-			if (CollapsedVertices[VertexIndex] != INDEX_NONE)
-			{
-				continue;
-			}
-	
-			const FVector3f& Vertex = Vertices[VertexIndex];
-	
-			NearbyVertices.Reset();
-			VertHash.FindPointsInBall(Vertex, 0.00001,
-				[&Vertex, &Vertices](const int32& Other) -> float {return FVector3f::DistSquared(Vertices[Other], Vertex); },
-				NearbyVertices);
-	
-			// Find equals
-			for (int32 NearbyVertexIndex : NearbyVertices)
-			{
-				CollapsedVertices[NearbyVertexIndex] = VertexIndex;
-			}
+			check(X >= 0 && X < SizeX);
+			check(Y >= 0 && Y < SizeY);
+			return Data[SizeX * Y + X];
 		}
-	}
-
-	// Helper used to connect triangles of a UV island
-	struct FTriangle
-	{
-		uint32 Indices[3];
-		uint32 CollapsedIndices[3];
-
-		uint16 BlockIndices[3];
 
-		bool bUVsFixed = false;
-	};
-
-	//---------------------------------------------------------------------------------------------
-	//! Fill an array with the indices of all triangles belonging to the same UV island as InFirstTriangle. 
-	//---------------------------------------------------------------------------------------------
-	void GetUVIsland(TArray<FTriangle>& InTriangles,
-		const uint32 InFirstTriangle,
-		TArray<uint32>& OutTriangleIndices,
-		const TArray<FVector2f>& InUVs,
-		const TMultiMap<int32, uint32>& InVertexToTriangleMap)
-	{
-		MUTABLE_CPUPROFILER_SCOPE(LayoutUV_GetUVIsland);
-
-		const uint32 NumTriangles = (uint32)InTriangles.Num();
-		
-		OutTriangleIndices.Reserve(NumTriangles);
-		OutTriangleIndices.Add(InFirstTriangle);
-
-		TArray<bool> SkipTrinalges;
-		SkipTrinalges.Init(false, NumTriangles);
-
-		TArray<uint32> PendingTriangles;
-		PendingTriangles.Reserve(NumTriangles / 64);
-		PendingTriangles.Add(InFirstTriangle);
-
-		while (!PendingTriangles.IsEmpty())
+		inline void Set(int32 X, int32 Y, const T& Value)
 		{
-			const uint32 TriangleIndex = PendingTriangles.Pop();
-
-			// Triangle about to be proccessed, mark as skip;
-			SkipTrinalges[TriangleIndex] = true;
-
-			bool ConnectedEdges[3] = { false, false, false };
-
-			const FTriangle& Triangle = InTriangles[TriangleIndex];
-
-			// Find Triangles connected to edges 0 and 2
-			int32 CollapsedVertex1 = Triangle.CollapsedIndices[1];
-			int32 CollapsedVertex2 = Triangle.CollapsedIndices[2];
-
-			TArray<uint32> FoundTriangleIndices;
-			InVertexToTriangleMap.MultiFind(Triangle.CollapsedIndices[0], FoundTriangleIndices);
-
-			for (uint32 OtherTriangleIndex : FoundTriangleIndices)
-			{
-				const FTriangle& OtherTriangle = InTriangles[OtherTriangleIndex];
-
-				for (int32 OtherIndex = 0; OtherIndex < 3; ++OtherIndex)
-				{
-					const int32 OtherCollapsedIndex = OtherTriangle.CollapsedIndices[OtherIndex];
-					if (OtherCollapsedIndex == CollapsedVertex1)
-					{
-						// Check if the vertex is in the same UV Island 
-						if (!SkipTrinalges[OtherTriangleIndex]
-							&& InUVs[Triangle.Indices[1]].Equals(InUVs[OtherTriangle.Indices[OtherIndex]], 0.00001f))
-						{
-							OutTriangleIndices.Add(OtherTriangleIndex);
-							PendingTriangles.Add(OtherTriangleIndex);
-							SkipTrinalges[OtherTriangleIndex] = true;
-						}
-
-						// Connected but already processed or in another island
-						break;
-					}
-
-					if (OtherCollapsedIndex  == CollapsedVertex2)
-					{
-						// Check if the vertex is in the same UV Island 
-						if (!SkipTrinalges[OtherTriangleIndex]
-							&& InUVs[Triangle.Indices[2]].Equals(InUVs[OtherTriangle.Indices[OtherIndex]], 0.00001f))
-						{
-							OutTriangleIndices.Add(OtherTriangleIndex);
-							PendingTriangles.Add(OtherTriangleIndex);
-							SkipTrinalges[OtherTriangleIndex] = true;
-						}
-
-						// Connected but already processed or in another UV Island
-						break;
-					}
-				}
-
-			}
-
-			// Find the triangle connected to edge 1
-			FoundTriangleIndices.Reset();
-			InVertexToTriangleMap.MultiFind(CollapsedVertex1, FoundTriangleIndices);
-
-			for (uint32 OtherTriangleIndex : FoundTriangleIndices)
-			{
-				const FTriangle& OtherTriangle = InTriangles[OtherTriangleIndex];
-
-				for (int32 OtherIndex = 0; OtherIndex < 3; ++OtherIndex)
-				{
-					const int32 OtherCollapsedIndex = OtherTriangle.CollapsedIndices[OtherIndex];
-					if (OtherCollapsedIndex == CollapsedVertex2)
-					{
-						// Check if the vertex belong to the same UV island
-						if (!SkipTrinalges[OtherTriangleIndex]
-							&& InUVs[Triangle.Indices[2]].Equals(InUVs[OtherTriangle.Indices[OtherIndex]], 0.00001f))
-						{
-							OutTriangleIndices.Add(OtherTriangleIndex);
-							PendingTriangles.Add(OtherTriangleIndex);
-							SkipTrinalges[OtherTriangleIndex] = true;
-						}
-
-						// Connected but already processed or in another island
-						break;
-					}
-				}
-			}
+			check(X >= 0 && X < SizeX);
+			check(Y >= 0 && Y < SizeY);
+			Data[SizeX * Y + X] = Value;
 		}
-	}
+
+	};
 
     //---------------------------------------------------------------------------------------------
 	void CodeGenerator::PrepareMeshForLayout( const FGeneratedLayout& GeneratedLayout,
@@ -286,12 +135,20 @@ namespace mu
 		const int32 NumVertices = Mesh->GetVertexCount();
 		const int32 NumBlocks = Layout->GetBlockCount();
 
-		bool bIsSingleFullBlock = (NumBlocks == 1) && (Layout->Blocks[0].Min == FImageSize(0, 0) && Layout->Blocks[0].Size == Layout->Size);
+		bool bIsSingleFullBlock = (NumBlocks == 1) && (Layout->Blocks[0].Min == FIntVector2(0, 0) && Layout->Blocks[0].Size == Layout->Size);
 
-		// Find block ids for each block in the grid
-		const FIntPoint Grid = Layout->GetGridSize();
-		TArray<int32> GridBlockBlockId;
-		GridBlockBlockId.Init(MAX_uint16, Grid.X * Grid.Y);
+		// Find block ids for each block in the grid. Calculate a grid size that contains all blocks
+		FIntPoint LayoutGrid = Layout->GetGridSize();
+		FIntPoint WorkingGrid = LayoutGrid;
+		for (const FSourceLayoutBlock& Block: GeneratedLayout.Source->Blocks)
+		{
+			WorkingGrid.X = FMath::Max(WorkingGrid.X, Block.Min.X + Block.Size.X );
+			WorkingGrid.Y = FMath::Max(WorkingGrid.Y, Block.Min.Y + Block.Size.Y );
+		}
+
+
+		TArray2D<int32> GridBlockBlockId;
+		GridBlockBlockId.Init(MAX_uint16, WorkingGrid.X, WorkingGrid.Y);
 
 		// 
 		TArray<box<FVector2f>> BlockRects;
@@ -304,14 +161,14 @@ namespace mu
 			bool bBlockHasMask = GeneratedLayout.Source->Blocks[BlockIndex].Mask.get() != nullptr;
 
 			// Fill the block rect
-			FImageSize Min = Layout->Blocks[BlockIndex].Min;
-			FImageSize Size = Layout->Blocks[BlockIndex].Size;
+			FIntVector2 Min = Layout->Blocks[BlockIndex].Min;
+			FIntVector2 Size = Layout->Blocks[BlockIndex].Size;
 
 			box<FVector2f>& BlockRect = BlockRects[BlockIndex];
-			BlockRect.min[0] = float(Min.X) / float(Grid.X);
-			BlockRect.min[1] = float(Min.Y) / float(Grid.Y);
-			BlockRect.size[0] = float(Size.X) / float(Grid.X);
-			BlockRect.size[1] = float(Size.Y) / float(Grid.Y);
+			BlockRect.min[0] = float(Min.X) / float(LayoutGrid.X);
+			BlockRect.min[1] = float(Min.Y) / float(LayoutGrid.Y);
+			BlockRect.size[0] = float(Size.X) / float(LayoutGrid.X);
+			BlockRect.size[1] = float(Size.Y) / float(LayoutGrid.Y);
 
 			// Fill the block index per cell array
 			// Ignore the block in this stage if it has a mask, because blocks with masks will very likely overlap other blocks
@@ -319,13 +176,11 @@ namespace mu
 			{
 				for (uint16 Y = Min.Y; Y < Min.Y + Size.Y; ++Y)
 				{
-					const uint16 PositionY = Y * Grid.X;
-
 					for (uint16 X = Min.X; X < Min.X + Size.X; ++X)
 					{
-						if (GridBlockBlockId[PositionY + X] == MAX_uint16)
+						if (GridBlockBlockId.Get(X,Y) == MAX_uint16)
 						{
-							GridBlockBlockId[PositionY + X] = BlockIndex;
+							GridBlockBlockId.Set(X,Y, BlockIndex);
 						}
 						else
 						{
@@ -405,7 +260,7 @@ namespace mu
 
 
 		const int32 NumTriangles = Mesh->GetIndexCount() / 3;
-		TArray<FTriangle> Triangles;
+		TArray<FTriangleInfo> Triangles;
 
 		// Vertices mapped to unique vertex index
 		TArray<int32> CollapsedVertices;
@@ -424,8 +279,8 @@ namespace mu
 
 		TArray<int32> ConflictiveTriangles;
 
-		const uint32 MaxGridX = MeshOptions.bNormalizeUVs ? MAX_uint32 : Grid.X - 1;
-		const uint32 MaxGridY = MeshOptions.bNormalizeUVs ? MAX_uint32 : Grid.Y - 1;
+		const uint32 MaxGridX = MeshOptions.bNormalizeUVs ? MAX_uint32 : WorkingGrid.X - 1;
+		const uint32 MaxGridY = MeshOptions.bNormalizeUVs ? MAX_uint32 : WorkingGrid.Y - 1;
 
 		// Allocate the per-vertex layout block data
 		TArray<uint16> LayoutData;
@@ -443,13 +298,22 @@ namespace mu
 			++ItIndices;
 
 			auto AssignOneVertex = 
-				[Grid, MaxGridX, MaxGridY, NumBlocks, Layout, &GeneratedLayout, &GridBlockBlockId, &TexCoords, &LayoutData]
+				[LayoutGrid, MaxGridX, MaxGridY, NumBlocks, Layout, &GeneratedLayout, &GridBlockBlockId, &TexCoords, &LayoutData]
 				(int32 VertexIndex)
 				{
 					uint16& BlockIndex = LayoutData[VertexIndex];
 
+					// Was it previously assigned?
+					if (BlockIndex != NullBlockId)
+					{
+						return BlockIndex;
+					}
+
 					FVector2f UV = TexCoords[VertexIndex];
 
+					int32 VertexWorkingGridX = FMath::Clamp(LayoutGrid.X * UV[0], 0, LayoutGrid.X - 1);
+					int32 VertexWorkingGridY = FMath::Clamp(LayoutGrid.Y * UV[1], 0, LayoutGrid.Y - 1);
+
 					// First: Assign the vertices to masked blocks in order
 					for (int32 CandidateBlockIndex = 0; CandidateBlockIndex < NumBlocks; ++CandidateBlockIndex)
 					{
@@ -457,19 +321,22 @@ namespace mu
 						if (Mask)
 						{
 							// First discard with block limits.
-							FImageSize Min = Layout->Blocks[CandidateBlockIndex].Min;
-							FImageSize Size = Layout->Blocks[CandidateBlockIndex].Size;
+							FIntVector2 Min = Layout->Blocks[CandidateBlockIndex].Min;
+							FIntVector2 Size = Layout->Blocks[CandidateBlockIndex].Size;
 
-							uint16 X = (uint16)FMath::Min<uint32>(MaxGridX, FMath::Max<uint32>(0, (uint32)(Grid.X * UV[0])));
-							uint16 Y = (uint16)FMath::Min<uint32>(MaxGridY, FMath::Max<uint32>(0, (uint32)(Grid.Y * UV[1])));
 							bool bInBlock = 
-								(X >= Min.X && X < Min.X + Size.X)
+								(VertexWorkingGridX >= Min.X && VertexWorkingGridX < Min.X + Size.X)
 								&&
-								(Y >= Min.Y && Y < Min.Y + Size.Y);
+								(VertexWorkingGridY >= Min.Y && VertexWorkingGridY < Min.Y + Size.Y);
 
 							if (bInBlock)
 							{
-								FVector4f MaskValue = Mask->Sample(UV);
+								// TODO: This always clamps the UVs
+								FVector2f SampleUV;
+								SampleUV.X = FMath::Fmod(UV.X, 1.0);
+								SampleUV.Y = FMath::Fmod(UV.Y, 1.0);
+
+								FVector4f MaskValue = Mask->Sample(SampleUV);
 								if (MaskValue.X > 0.5f)
 								{
 									BlockIndex = CandidateBlockIndex;
@@ -482,9 +349,9 @@ namespace mu
 					// Second: Assign to non-masked blocks if not assigned yet
 					if (BlockIndex == NullBlockId)
 					{
-						uint16 X = (uint16)FMath::Min<uint32>(MaxGridX, FMath::Max<uint32>(0, (uint32)(Grid.X * UV[0])));
-						uint16 Y = (uint16)FMath::Min<uint32>(MaxGridY, FMath::Max<uint32>(0, (uint32)(Grid.Y * UV[1])));
-						BlockIndex = (X < Grid.X && Y < Grid.Y) ? GridBlockBlockId[Y * Grid.X + X] : 0;
+						uint32 ClampedX = FMath::Min<uint32>(MaxGridX, FMath::Max<uint32>(0, VertexWorkingGridX));
+						uint32 ClampedY = FMath::Min<uint32>(MaxGridY, FMath::Max<uint32>(0, VertexWorkingGridY));
+						BlockIndex = GridBlockBlockId.Get(ClampedX, ClampedY);
 					}
 					return BlockIndex;
 				};
@@ -500,7 +367,7 @@ namespace mu
 					ConflictiveTriangles.Add(TriangleIndex);
 				}
 
-				FTriangle& Triangle = Triangles[TriangleIndex];
+				FTriangleInfo& Triangle = Triangles[TriangleIndex];
 
 				Triangle.Indices[0] = Index0;
 				Triangle.Indices[1] = Index1;
@@ -523,7 +390,7 @@ namespace mu
 		// Clamp UV islands to the predominant block of each island. Will only happen if bClampUVIslands is true.
 		for (int32 ConflictiveTriangleIndex : ConflictiveTriangles)
 		{
-			FTriangle& Triangle = Triangles[ConflictiveTriangleIndex];
+			FTriangleInfo& Triangle = Triangles[ConflictiveTriangleIndex];
 
 			// Skip the ones that have been fixed already
 			if (Triangle.bUVsFixed)
@@ -541,7 +408,7 @@ namespace mu
 
 			for (int32 TriangleIndex : TriangleIndices)
 			{
-				FTriangle& OtherTriangle = Triangles[TriangleIndex];
+				FTriangleInfo& OtherTriangle = Triangles[TriangleIndex];
 				for (int32 VertexIndex = 0; VertexIndex < 3; ++VertexIndex)
 				{
 					const uint16& BlockIndex = OtherTriangle.BlockIndices[VertexIndex];
@@ -567,15 +434,15 @@ namespace mu
 			const FLayoutBlock& LayoutBlock = Layout->Blocks[BlockIndex];
 
 			const float SmallNumber = 0.000001;
-			const float MinX = ((float)LayoutBlock.Min.X) / (float)Grid.X + SmallNumber;
-			const float MinY = ((float)LayoutBlock.Min.Y) / (float)Grid.Y + SmallNumber;
-			const float MaxX = (((float)LayoutBlock.Size.X + LayoutBlock.Min.X) / (float)Grid.X) - 2 * SmallNumber;
-			const float MaxY = (((float)LayoutBlock.Size.Y + LayoutBlock.Min.Y) / (float)Grid.Y) - 2 * SmallNumber;
+			const float MinX = ((float)LayoutBlock.Min.X) / (float)LayoutGrid.X + SmallNumber;
+			const float MinY = ((float)LayoutBlock.Min.Y) / (float)LayoutGrid.Y + SmallNumber;
+			const float MaxX = (((float)LayoutBlock.Size.X + LayoutBlock.Min.X) / (float)LayoutGrid.X) - 2 * SmallNumber;
+			const float MaxY = (((float)LayoutBlock.Size.Y + LayoutBlock.Min.Y) / (float)LayoutGrid.Y) - 2 * SmallNumber;
 
 			// Iterate triangles and clamp the UVs
 			for (int32 TriangleIndex : TriangleIndices)
 			{
-				FTriangle& OtherTriangle = Triangles[TriangleIndex];
+				FTriangleInfo& OtherTriangle = Triangles[TriangleIndex];
 
 				for (int8 VertexIndex = 0; VertexIndex < 3; ++VertexIndex)
 				{
@@ -586,7 +453,7 @@ namespace mu
 
 					OtherTriangle.BlockIndices[VertexIndex] = BlockIndex;
 
-					// Clamp UVs
+					// Clamp UVs to the block they are assigned to
 					const int32 UVIndex = OtherTriangle.Indices[VertexIndex];
 					FVector2f& UV = TexCoords[UVIndex];
 					UV[0] = FMath::Clamp(UV[0], MinX, MaxX);
@@ -660,14 +527,14 @@ namespace mu
 
 			for (int32 VertexIndex = 0; VertexIndex < NumVertices; ++VertexIndex)
 			{
-				FVector2f* UV = &TexCoords[VertexIndex];
+				FVector2f& UV = TexCoords[VertexIndex];
 
 				uint16 LayoutBlockIndex = LayoutData[VertexIndex];
 				if (Layout->Blocks.IsValidIndex(LayoutBlockIndex))
 				{
 					uint64 LayoutBlockId = Layout->Blocks[LayoutBlockIndex].Id;
 
-					*UV = BlockRects[LayoutBlockIndex].Homogenize(*UV);
+					UV = BlockRects[LayoutBlockIndex].Homogenize(UV);
 
 					// Replace block index by the actual id of the block
 					if (bUseAbsoluteBlockIds)
@@ -700,13 +567,13 @@ namespace mu
 				if (TexCoordsChannel.Format == MBF_FLOAT32)
 				{
 					FVector2f* pUV = reinterpret_cast<FVector2f*>(pVertices);
-					*pUV = *UV;
+					*pUV = UV;
 				}
 				else if (TexCoordsChannel.Format == MBF_FLOAT16)
 				{
 					FFloat16* pUV = reinterpret_cast<FFloat16*>(pVertices);
-					pUV[0] = FFloat16((*UV)[0]);
-					pUV[1] = FFloat16((*UV)[1]);
+					pUV[0] = FFloat16(UV[0]);
+					pUV[1] = FFloat16(UV[1]);
 				}
 
 				pVertices += elemSize;
diff --git a/Engine/Plugins/Experimental/Mutable/Source/MutableTools/Private/MuT/NodeLayout.cpp b/Engine/Plugins/Experimental/Mutable/Source/MutableTools/Private/MuT/NodeLayout.cpp
--- a/Engine/Plugins/Experimental/Mutable/Source/MutableTools/Private/MuT/NodeLayout.cpp
+++ b/Engine/Plugins/Experimental/Mutable/Source/MutableTools/Private/MuT/NodeLayout.cpp
@@ -2,12 +2,15 @@
 
 #include "MuT/NodeLayout.h"
 
-#include "Math/IntPoint.h"
-#include "Misc/AssertionMacros.h"
+#include "MuT/NodeMesh.h"
+#include "MuR/Raster.h"
 #include "MuR/ConvertData.h"
 #include "MuR/MeshBufferSet.h"
 #include "MuR/MeshPrivate.h"
 #include "MuR/MutableMath.h"
+#include "MuR/OpImageGrow.h"
+#include "Math/IntPoint.h"
+#include "Misc/AssertionMacros.h"
 
 
 namespace mu
@@ -16,204 +19,518 @@ namespace mu
 	FNodeType NodeLayout::StaticType = FNodeType(Node::EType::Layout, Node::GetStaticType() );
 
 
-	Ptr<NodeLayout> NodeLayout::GenerateLayoutBlocks(const Ptr<Mesh>& pMesh, int32 layoutIndex, int32 gridSizeX, int32 gridSizeY)
+	void NodeLayout::GenerateLayoutBlocks(const Ptr<Mesh>& Mesh, int32 LayoutIndex )
 	{
-		Ptr<NodeLayout> LayoutNode = nullptr;
+		if (!(Mesh && LayoutIndex >= 0 && Size.X > 0 && Size.Y > 0))
+		{
+			return;
+		}
 
-		if (pMesh && layoutIndex >=0 && gridSizeX+gridSizeY>0)
+		int32 IndexCount = Mesh->GetIndexCount();
+		TArray< FVector2f > UVs;
+
+		// Extract all the triangle edges
 		{
-			int32 indexCount = pMesh->GetIndexCount();
-			TArray< FVector2f > UVs;
-			UVs.SetNumUninitialized(indexCount * 2);
+			UVs.SetNumUninitialized(IndexCount * 2);
 
-			UntypedMeshBufferIteratorConst indexIt(pMesh->GetIndexBuffers(), MBS_VERTEXINDEX, 0);
-			UntypedMeshBufferIteratorConst texIt(pMesh->GetVertexBuffers(), MBS_TEXCOORDS, layoutIndex);
-			
-			//Getting UVs face by face
-			for (int32 v = 0; v < indexCount/3; ++v)
+			UntypedMeshBufferIteratorConst IndexIt(Mesh->GetIndexBuffers(), MBS_VERTEXINDEX, 0);
+			UntypedMeshBufferIteratorConst TexIt(Mesh->GetVertexBuffers(), MBS_TEXCOORDS, LayoutIndex);
+
+			int32 TriangleCount = IndexCount / 3;
+			for (int32 TriangleIndex = 0; TriangleIndex < TriangleCount; ++TriangleIndex)
 			{
-				uint32_t i_1 = indexIt.GetAsUINT32(); 
-				indexIt++;
-				uint32_t i_2 = indexIt.GetAsUINT32();
-				indexIt++;
-				uint32_t i_3 = indexIt.GetAsUINT32();
-				indexIt++;
-
-				float uv_1[2] = { 0.0f,0.0f };
-				ConvertData(0, uv_1, MBF_FLOAT32, (texIt + i_1).ptr(), texIt.GetFormat());
-				ConvertData(1, uv_1, MBF_FLOAT32, (texIt + i_1).ptr(), texIt.GetFormat());
-				
-				float uv_2[2] = { 0.0f,0.0f };
-				ConvertData(0, uv_2, MBF_FLOAT32, (texIt + i_2).ptr(), texIt.GetFormat());
-				ConvertData(1, uv_2, MBF_FLOAT32, (texIt + i_2).ptr(), texIt.GetFormat());
-
-				float uv_3[2] = { 0.0f,0.0f };
-				ConvertData(0, uv_3, MBF_FLOAT32, (texIt + i_3).ptr(), texIt.GetFormat());
-				ConvertData(1, uv_3, MBF_FLOAT32, (texIt + i_3).ptr(), texIt.GetFormat());
-
-				
-				UVs[v * 6 + 0][0] = uv_1[0];
-				UVs[v * 6 + 0][1] = uv_1[1];
-				UVs[v * 6 + 1][0] = uv_2[0];
-				UVs[v * 6 + 1][1] = uv_2[1];
-						
-				UVs[v * 6 + 2][0] = uv_2[0];
-				UVs[v * 6 + 2][1] = uv_2[1];
-				UVs[v * 6 + 3][0] = uv_3[0];
-				UVs[v * 6 + 3][1] = uv_3[1];
-						
-				UVs[v * 6 + 4][0] = uv_3[0]; 
-				UVs[v * 6 + 4][1] = uv_3[1]; 
-				UVs[v * 6 + 5][0] = uv_1[0];
-				UVs[v * 6 + 5][1] = uv_1[1];
-			}
+				uint32 VertexIndex1 = IndexIt.GetAsUINT32();
+				IndexIt++;
+				uint32 VertexIndex2 = IndexIt.GetAsUINT32();
+				IndexIt++;
+				uint32 VertexIndex3 = IndexIt.GetAsUINT32();
+				IndexIt++;
+
+				FVector2f UV1 = (TexIt + VertexIndex1).GetAsVec2f();
+				FVector2f UV2 = (TexIt + VertexIndex2).GetAsVec2f();
+				FVector2f UV3 = (TexIt + VertexIndex3).GetAsVec2f();
+
+				UVs[TriangleIndex * 6 + 0] = UV1;
+				UVs[TriangleIndex * 6 + 1] = UV2;
 
-			LayoutNode = new NodeLayout;
-			LayoutNode->Size = { uint16(gridSizeX), uint16(gridSizeY) };
-			LayoutNode->MaxSize = { uint16(gridSizeX), uint16(gridSizeY) };
-			LayoutNode->Strategy = EPackStrategy::Resizeable;
+				UVs[TriangleIndex * 6 + 2] = UV2;
+				UVs[TriangleIndex * 6 + 3] = UV3;
+
+				UVs[TriangleIndex * 6 + 4] = UV3;
+				UVs[TriangleIndex * 6 + 5] = UV1;
+			}
+		}
 			
-			TArray<box<FIntVector2>> blocks;
+		TArray<box<FIntVector2>> BlockRects;
+		BlockRects.Reserve(Blocks.Num());
+
+		// Get the rects of existing blocks
+		for (const FSourceLayoutBlock& Block : Blocks)
+		{ 
+			box<FIntVector2> Rect;
+			Rect.min = Block.Min;
+			Rect.size = Block.Size;
+			BlockRects.Add( Rect );
+		}
+
+		// Generate blocks by iterating all the edges
+		int32 EdgeCount = IndexCount;
+		for (int32 EdgeIndex = 0; EdgeIndex < EdgeCount; ++EdgeIndex)
+		{
+			FVector2f APosition = UVs[EdgeIndex * 2];
+			FVector2f BPosition = UVs[EdgeIndex * 2+1];
+
+			FIntVector2 AGrid, BGrid;
+			AGrid[0] = FMath::FloorToInt32(APosition[0] * Size.X);
+			AGrid[1] = FMath::FloorToInt32(APosition[1] * Size.Y);
+									  
+			BGrid[0] = FMath::FloorToInt32(BPosition[0] * Size.X);
+			BGrid[1] = FMath::FloorToInt32(BPosition[1] * Size.Y);
+
+			// TODO: handle cases of UVs on grid edges.
+			// floor of UV = 1*GridSize is GridSize which is not AGrid valid range
+			//if (AGrid[0] == GridSizeX){ AGrid[0] = GridSizeX-1; }
+			//if (AGrid[1] == GridSizeY){	AGrid[1] = GridSizeY-1;	}
+			//if (BPosition[0] == GridSizeX){	BPosition[0] = GridSizeX-1;	}
+			//if (BPosition[1] == GridSizeY){ BPosition[1] = GridSizeY-1;	}
 			
-			//Generating blocks
-			for (int32 i = 0; i < indexCount; ++i)
+			// AGrid and BPosition are in the same block
+			if (AGrid == BGrid)
 			{
-				FIntVector2 a, b;
-			
-				a[0] = (int32)floor(UVs[i * 2][0] * gridSizeX);
-				a[1] = (int32)floor(UVs[i * 2][1] * gridSizeY);
-									  
-				b[0] = (int32)floor(UVs[i * 2 +1][0] * gridSizeX);
-				b[1] = (int32)floor(UVs[i * 2 +1][1] * gridSizeY);
-
-				//floor of UV = 1*gridSize is gridSize which is not a valid range
-				if (a[0] == gridSizeX){ a[0] = gridSizeX-1; }
-				if (a[1] == gridSizeY){	a[1] = gridSizeY-1;	}
-				if (b[0] == gridSizeX){	b[0] = gridSizeX-1;	}
-				if (b[1] == gridSizeY){ b[1] = gridSizeY-1;	}
+				bool bIsContained = false;
 			
-				//a and b are in the same block
-				if (a == b)
+				for (int32 BlockIndex = 0; BlockIndex < BlockRects.Num(); ++BlockIndex)
 				{
-					bool contains = false;
-			
-					for (int32 it = 0; it < blocks.Num(); ++it)
+					if (BlockRects[BlockIndex].Contains(AGrid) || BlockRects[BlockIndex].Contains(BGrid))
 					{
-						if (blocks[it].Contains(a) || blocks[it].Contains(b))
-						{
-							contains = true;
-						}
+						bIsContained = true;
 					}
+				}
 					
-					//There is no block that contains them 
-					if (!contains)
-					{
-						box<FIntVector2> currBlock;
-						currBlock.min = a;
-						currBlock.size = FIntVector2(1, 1);
+				// There is no block that contains them 
+				if (!bIsContained)
+				{
+					box<FIntVector2> NewBlock;
+					NewBlock.min = AGrid;
+					NewBlock.size = FIntVector2(1, 1);
 			
-						blocks.Add(currBlock);
-					}
+					BlockRects.Add(NewBlock);
 				}
-				else //they are in different blocks
-				{
-					int32 idxA = -1;
-					int32 idxB = -1;
+			}
+
+			else // they are in different blocks
+			{
+				int32 ABlockIndex = -1;
+				int32 BBlockIndex = -1;
 					
-					//Getting the blocks that contain them
-					for (int32 it = 0; it < blocks.Num(); ++it)
+				// Get the blocks that contain them
+				for (int32 BlockIndex = 0; BlockIndex < BlockRects.Num(); ++BlockIndex)
+				{
+					if (BlockRects[BlockIndex].Contains(AGrid))
 					{
-						if (blocks[it].Contains(a))
-						{
-							idxA = (int32)it;
-						}
-						if (blocks[it].Contains(b))
-						{
-							idxB = (int32)it;
-						}
+						ABlockIndex = BlockIndex;
 					}
-					
-					//The blocks are not the same
-					if (idxA != idxB)
+					if (BlockRects[BlockIndex].Contains(BGrid))
 					{
-						box<FIntVector2> currBlock;
+						BBlockIndex = BlockIndex;
+					}
+				}
+					
+				// The blocks are not the same
+				if (ABlockIndex != BBlockIndex)
+				{
+					box<FIntVector2> NewBlock;
 						
-						//One of the blocks doesn't exist
-						if (idxA != -1 && idxB == -1)
-						{
-							currBlock.min = b;
-							currBlock.size = FIntVector2(1, 1);
-							blocks[idxA].Bound(currBlock);
-						}
-						else if (idxB != -1 && idxA == -1)
-						{
-							currBlock.min = a;
-							currBlock.size = FIntVector2(1, 1);
-							blocks[idxB].Bound(currBlock);
-						}
-						else //Both exist
-						{
-							blocks[idxA].Bound(blocks[idxB]);
-							blocks.RemoveAt(idxB);
-						}
+					//One of the blocks doesn't exist
+					if (ABlockIndex != -1 && BBlockIndex == -1)
+					{
+						NewBlock.min = BGrid;
+						NewBlock.size = FIntVector2(1, 1);
+						BlockRects[ABlockIndex].Bound(NewBlock);
 					}
-					else //the blocks doesn't exist
+					else if (BBlockIndex != -1 && ABlockIndex == -1)
 					{
-						if (idxA == -1)
-						{
-							box<FIntVector2> currBlockA;
-							box<FIntVector2> currBlockB;
+						NewBlock.min = AGrid;
+						NewBlock.size = FIntVector2(1, 1);
+						BlockRects[BBlockIndex].Bound(NewBlock);
+					}
+					else //Both exist
+					{
+						BlockRects[ABlockIndex].Bound(BlockRects[BBlockIndex]);
+						BlockRects.RemoveAt(BBlockIndex);
+					}
+				}
+				else // the block doesn't exist
+				{
+					if (ABlockIndex == -1)
+					{
+						box<FIntVector2> NewBlockA;
+						box<FIntVector2> NewBlockB;
 			
-							currBlockA.min = a;
-							currBlockB.min = b;
-							currBlockA.size = FIntVector2(1, 1);
-							currBlockB.size = FIntVector2(1, 1);
+						NewBlockA.min = AGrid;
+						NewBlockB.min = BGrid;
+						NewBlockA.size = FIntVector2(1, 1);
+						NewBlockB.size = FIntVector2(1, 1);
 			
-							currBlockA.Bound(currBlockB);
-							blocks.Add(currBlockA);
-						}
+						NewBlockA.Bound(NewBlockB);
+						BlockRects.Add(NewBlockA);
 					}
 				}
 			}
+		}
 			
-			bool intersections = true;
+		bool bHasIntersections = true;
 			
-			// Check if the blocks intersect with each other
-			while (intersections)
+		// Check if blocks intersect with each other or are null
+		while (bHasIntersections)
+		{
+			bHasIntersections = false;
+			
+			for (int32 i = 0; !bHasIntersections && i < BlockRects.Num(); ++i)
 			{
-				intersections = false;
+				for (int32 j = 0; j < BlockRects.Num(); ++j)
+				{
+					if (i != j && BlockRects[i].IntersectsExclusive(BlockRects[j]))
+					{
+						BlockRects[i].Bound(BlockRects[j]);
+						BlockRects.RemoveAt(j);
+						bHasIntersections = true;
+						break;
+					}
+				}
+
+				// Remove degenerated blocks.
+				if (BlockRects[i].size.X * BlockRects[i].size.Y == 0)
+				{
+					BlockRects.RemoveAt(i);
+					bHasIntersections = true;
+					break;
+				}
+			}
+		}
+			
+		int32 NumBlocks = BlockRects.Num();
+			
+		// Generate the layout blocks
+		if (NumBlocks > 0)
+		{
+			Blocks.SetNum(NumBlocks);
 			
-				for (int32 i = 0; !intersections && i < blocks.Num(); ++i)
+			for (int32 BlockIndex = 0; BlockIndex < NumBlocks; ++BlockIndex)
+			{
+				Blocks[BlockIndex].Min = BlockRects[BlockIndex].min;
+				Blocks[BlockIndex].Size = BlockRects[BlockIndex].size;
+			}
+		}
+	}
+
+
+	void NodeLayout::GenerateLayoutBlocksFromUVIslands(const Ptr<Mesh>& Mesh, int32 LayoutIndex, bool bMergeChildBlocks)
+	{
+		if (!(Mesh && LayoutIndex >= 0 && Size.X > 0 && Size.Y > 0))
+		{
+			return;
+		}
+
+		int32 IndexCount = Mesh->GetIndexCount();
+		const int32 NumTriangles = IndexCount / 3;
+		const int32 NumVertices = Mesh->GetVertexCount();
+
+		TArray<FTriangleInfo> Triangles;
+		Triangles.SetNumUninitialized(NumTriangles);
+
+		// Vertices mapped to unique vertex index
+		TArray<int32> CollapsedVertices;
+
+		// Vertex to face map used to speed up connectivity building
+		TMultiMap<int32, uint32> VertexToFaceMap;
+		VertexToFaceMap.Reserve(NumVertices*4);
+
+		// Find Unique Vertices
+		MeshCreateCollapsedVertexMap(Mesh.get(), CollapsedVertices);
+
+		UntypedMeshBufferIteratorConst ItIndices(Mesh->GetIndexBuffers(), MBS_VERTEXINDEX);
+		for (int32 TriangleIndex = 0; TriangleIndex < NumTriangles; ++TriangleIndex)
+		{
+			uint32 Index0 = ItIndices.GetAsUINT32();
+			++ItIndices;
+			uint32 Index1 = ItIndices.GetAsUINT32();
+			++ItIndices;
+			uint32 Index2 = ItIndices.GetAsUINT32();
+			++ItIndices;
+
+			FTriangleInfo& Triangle = Triangles[TriangleIndex];
+
+			Triangle.Indices[0] = Index0;
+			Triangle.Indices[1] = Index1;
+			Triangle.Indices[2] = Index2;
+			Triangle.CollapsedIndices[0] = CollapsedVertices[Index0];
+			Triangle.CollapsedIndices[1] = CollapsedVertices[Index1];
+			Triangle.CollapsedIndices[2] = CollapsedVertices[Index2];
+
+			Triangle.BlockIndices[0] = 0;
+			Triangle.BlockIndices[1] = 0;
+			Triangle.BlockIndices[2] = 0;
+			Triangle.bUVsFixed = false;
+
+			VertexToFaceMap.Add(Triangle.CollapsedIndices[0], TriangleIndex);
+			VertexToFaceMap.Add(Triangle.CollapsedIndices[1], TriangleIndex);
+			VertexToFaceMap.Add(Triangle.CollapsedIndices[2], TriangleIndex);
+		}
+
+		// Get a copy of the UVs as FVector2f to work with them. 
+		TArray<FVector2f> TexCoords;
+		{
+			int32 TexCoordsBufferIndex = -1;
+			int32 TexCoordsChannelIndex = -1;
+			Mesh->GetVertexBuffers().FindChannel(MBS_TEXCOORDS, LayoutIndex, &TexCoordsBufferIndex, &TexCoordsChannelIndex);
+			check(TexCoordsBufferIndex >= 0);
+			check(TexCoordsChannelIndex >= 0);
+
+			const FMeshBufferChannel& TexCoordsChannel = Mesh->VertexBuffers.Buffers[TexCoordsBufferIndex].Channels[TexCoordsChannelIndex];
+			check(TexCoordsChannel.Semantic == MBS_TEXCOORDS);
+
+			uint8* TexCoordData = Mesh->GetVertexBuffers().GetBufferData(TexCoordsBufferIndex);
+			int32 ElemSize = Mesh->GetVertexBuffers().GetElementSize(TexCoordsBufferIndex);
+
+			TexCoords.SetNumUninitialized(NumVertices);
+
+			const uint8* pVertices = TexCoordData;
+			for (int32 VertexIndex = 0; VertexIndex < NumVertices; ++VertexIndex)
+			{
+				FVector2f& UV = TexCoords[VertexIndex];
+				if (TexCoordsChannel.Format == MBF_FLOAT32)
+				{
+					UV = *((FVector2f*)pVertices);
+				}
+				else if (TexCoordsChannel.Format == MBF_FLOAT16)
+				{
+					const FFloat16* pUV = reinterpret_cast<const FFloat16*>(pVertices);
+					UV = FVector2f(float(pUV[0]), float(pUV[1]));
+				}
+
+				pVertices += ElemSize;
+			}
+		}
+
+		// Generate UV islands
+		TArray<int16> IslandPerTriangle;
+		TArray<box<FVector2f>> IslandBlocks;
+		{
+			int16 IslandCount = 0;
+			IslandPerTriangle.Init(-1, NumTriangles);
+			for (int32 TriangleIndex = 0; TriangleIndex < NumTriangles; ++TriangleIndex)
+			{
+				if (IslandPerTriangle[TriangleIndex] >= 0)
+				{
+					// Already assigned to an island
+					continue;
+				}
+
+				// Find triangles from the same UV Island
+				TArray<uint32> TriangleIndices;
+				GetUVIsland(Triangles, TriangleIndex, TriangleIndices, TexCoords, VertexToFaceMap);
+
+				box<FVector2f> IslandBlock;
+				for (int32 IslandTriangleIndexIndex = 0; IslandTriangleIndexIndex < TriangleIndices.Num(); ++IslandTriangleIndexIndex)
+				{
+					uint32 IslandTriangleIndex = TriangleIndices[IslandTriangleIndexIndex];
+					FTriangleInfo& Triangle = Triangles[IslandTriangleIndex];
+
+					// Mark the triangle as already assigned to an island.
+					IslandPerTriangle[IslandTriangleIndex] = IslandCount;
+
+					FVector2f UV0 = TexCoords[Triangle.Indices[0]];
+					FVector2f UV1 = TexCoords[Triangle.Indices[1]];
+					FVector2f UV2 = TexCoords[Triangle.Indices[2]];
+
+					if (IslandTriangleIndexIndex == 0)
+					{
+						IslandBlock.min = UV0;
+					}
+
+					IslandBlock.Bound(UV0);
+					IslandBlock.Bound(UV1);
+					IslandBlock.Bound(UV2);
+				}
+
+				IslandBlocks.Add(IslandBlock);
+
+				++IslandCount;
+			}
+		}
+
+
+		TArray<box<FIntVector2>> BlockRects;
+		BlockRects.Reserve(Blocks.Num());
+
+		// Get the rects of existing blocks
+		for (const FSourceLayoutBlock& Block : Blocks)
+		{
+			box<FIntVector2> Rect;
+			Rect.min = Block.Min;
+			Rect.size = Block.Size;
+			BlockRects.Add(Rect);
+		}
+
+		int32 IslandBlocksOffset = BlockRects.Num();
+
+		for (box<FVector2f>& UVBlock : IslandBlocks)
+		{
+			box<FIntVector2> IslandBlock;
+			IslandBlock.min.X = FMath::FloorToInt32(UVBlock.min.X * Size.X);
+			IslandBlock.min.Y = FMath::FloorToInt32(UVBlock.min.Y * Size.Y);
+			FVector2f MaxF = UVBlock.min + UVBlock.size;
+			FIntVector2 Max;
+			Max.X = FMath::CeilToInt32(MaxF.X * Size.X);
+			Max.Y = FMath::CeilToInt32(MaxF.Y * Size.Y);
+			IslandBlock.size.X = Max.X - IslandBlock.min.X;
+			IslandBlock.size.Y = Max.Y - IslandBlock.min.Y;
+			BlockRects.Add(IslandBlock);
+		}
+
+		TArray<bool> RemovedBlocks;
+		RemovedBlocks.Init(false, BlockRects.Num());
+
+		// Merge blocks if necessary
+		if (bMergeChildBlocks)
+		{
+			for (int32 BlockIndex = IslandBlocksOffset; BlockIndex < BlockRects.Num(); ++BlockIndex)
+			{
+				int32 IslandBlockIndex = BlockIndex - IslandBlocksOffset;
+				box<FIntVector2> ThisRect = BlockRects[BlockIndex];
+				for (int32 OtherBlockIndex = BlockIndex+1; OtherBlockIndex < BlockRects.Num(); ++OtherBlockIndex)
 				{
-					for (int32 j = 0; j < blocks.Num(); ++j)
+					box<FIntVector2> OtherRect = BlockRects[OtherBlockIndex];
+					if (!RemovedBlocks[OtherBlockIndex] && ThisRect.Contains(OtherRect))
 					{
-						if (i != j && blocks[i].IntersectsExclusive(blocks[j]))
+						RemovedBlocks[OtherBlockIndex] = true;
+
+						// Merge other into this 
+						int32 OtherIslandBlockIndex = OtherBlockIndex - IslandBlocksOffset;
+						for (int32 TriangleIndex = 0; TriangleIndex < NumTriangles; ++TriangleIndex)
 						{
-							blocks[i].Bound(blocks[j]);
-							blocks.RemoveAt(j);
-							intersections = true;
-							break;
+							if (IslandPerTriangle[TriangleIndex] == OtherIslandBlockIndex)
+							{
+								IslandPerTriangle[TriangleIndex] = IslandBlockIndex;
+							}
 						}
 					}
 				}
+
 			}
+		}
+
+		// Check if blocks intersect with each other or are null
+		// actually this code wouldn't work because we need to update IslandBlocksOffset too
+		//bool bHasIntersections = true;
+		//while (bHasIntersections)
+		//{
+		//	bHasIntersections = false;
+		//	
+		//	for (int32 i = 0; !bHasIntersections && i < BlockRects.Num(); ++i)
+		//	{
+		//		for (int32 j = 0; j < BlockRects.Num(); ++j)
+		//		{
+		//			if (i != j && BlockRects[i].IntersectsExclusive(BlockRects[j]))
+		//			{
+		//				BlockRects[i].Bound(BlockRects[j]);
+		//				BlockRects.RemoveAt(j);
+		//				bHasIntersections = true;
+		//				break;
+		//			}
+		//		}
+
+		//		// Remove degenerated blocks.
+		//		if (BlockRects[i].size.X * BlockRects[i].size.Y == 0)
+		//		{
+		//			BlockRects.RemoveAt(i);
+		//			bHasIntersections = true;
+		//			break;
+		//		}
+		//	}
+		//}
 			
-			int32 NumBlocks = blocks.Num();
-			
-			//Generating layout blocks
-			if (NumBlocks > 0)
+		// Raster all the faces
+		class WhitePixelProcessor
+		{
+		public:
+			inline void ProcessPixel(uint8* pBufferPos, float[1]) const
 			{
-				LayoutNode->Blocks.SetNum(NumBlocks);
-			
-				for (int32 BlockIndex = 0; BlockIndex < NumBlocks; ++BlockIndex)
+				pBufferPos[0] = 255;
+			}
+
+			inline void operator()(uint8* BufferPos, float Interpolators[1]) const
+			{
+				ProcessPixel(BufferPos, Interpolators);
+			}
+		};
+
+		WhitePixelProcessor pixelProc;
+
+		// Generate the layout blocks
+		Blocks.Reserve(BlockRects.Num());
+		for (int32 IslandBlockIndex = 0; IslandBlockIndex < IslandBlocks.Num(); ++IslandBlockIndex)
+		{
+			int32 BlockIndex = IslandBlockIndex + IslandBlocksOffset;
+
+			if (RemovedBlocks[BlockIndex])
+			{
+				continue;
+			}
+
+			FSourceLayoutBlock& Block = Blocks.Emplace_GetRef();
+			Block.Min = BlockRects[BlockIndex].min;
+			Block.Size = BlockRects[BlockIndex].size;
+
+			// Generate the block mask
+			// TODO: Size?
+			int32 SizeX = 1024;
+			int32 SizeY = 1024;
+			Ptr<Image> Mask = new Image(SizeX, SizeY, 1, EImageFormat::IF_L_UBYTE, EInitializationType::Black);
+			const TArrayView<uint8> ImageData = Mask->DataStorage.GetLOD(0);
+
+			for (int32 TriangleIndex = 0; TriangleIndex < NumTriangles; ++TriangleIndex)
+			{
+				if (IslandPerTriangle[TriangleIndex] != IslandBlockIndex)
 				{
-					LayoutNode->Blocks[BlockIndex].Min = UE::Math::TIntVector2<uint16>(blocks[BlockIndex].min);
-					LayoutNode->Blocks[BlockIndex].Size = UE::Math::TIntVector2<uint16>(blocks[BlockIndex].size);
+					continue;
 				}
+
+				FTriangleInfo& ThisTriangle = Triangles[TriangleIndex];
+				FVector2f UV0 = TexCoords[ThisTriangle.Indices[0]];
+				FVector2f UV1 = TexCoords[ThisTriangle.Indices[1]];
+				FVector2f UV2 = TexCoords[ThisTriangle.Indices[2]];
+
+				// TODO Modulo doesn't work with cross-tile blocks
+				UV0.X = FMath::Fmod(UV0.X, 1.0);
+				UV0.Y = FMath::Fmod(UV0.Y, 1.0);
+				UV1.X = FMath::Fmod(UV1.X, 1.0);
+				UV1.Y = FMath::Fmod(UV1.Y, 1.0);
+				UV2.X = FMath::Fmod(UV2.X, 1.0);
+				UV2.Y = FMath::Fmod(UV2.Y, 1.0);
+
+				RasterVertex<1> V0(UV0.X* SizeX, UV0.Y* SizeY);
+				RasterVertex<1> V1(UV1.X* SizeX, UV1.Y* SizeY);
+				RasterVertex<1> V2(UV2.X* SizeX, UV2.Y* SizeY);
+
+				constexpr int32 NumInterpolators = 1;
+				Triangle<NumInterpolators>(ImageData.GetData(), ImageData.Num(),
+					SizeX, SizeY,
+					1,
+					V0,V1,V2,
+					pixelProc,
+					false);
+
 			}
+
+			// TODO: Clamp UV islands always?
+			ImageGrow(Mask.get());
+			ImageGrow(Mask.get());
+
+			Block.Mask = Mask;
 		}
 
-		return LayoutNode;
 	}
 
 }
diff --git a/Engine/Plugins/Experimental/Mutable/Source/MutableTools/Private/MuT/NodeMesh.cpp b/Engine/Plugins/Experimental/Mutable/Source/MutableTools/Private/MuT/NodeMesh.cpp
--- a/Engine/Plugins/Experimental/Mutable/Source/MutableTools/Private/MuT/NodeMesh.cpp
+++ b/Engine/Plugins/Experimental/Mutable/Source/MutableTools/Private/MuT/NodeMesh.cpp
@@ -2,32 +2,193 @@
 
 #include "MuT/NodeMesh.h"
 
+#include "MuR/MeshPrivate.h"
+#include "MuR/MutableTrace.h"
 #include "Misc/AssertionMacros.h"
+#include "Spatial/PointHashGrid3.h"
 
 
 namespace mu
 {
 
-	//---------------------------------------------------------------------------------------------
-	// Static initialisation
-	//---------------------------------------------------------------------------------------------
 	static FNodeType s_nodeMeshType = FNodeType(Node::EType::Mesh , Node::GetStaticType());
 
 
-	//---------------------------------------------------------------------------------------------
 	const FNodeType* NodeMesh::GetType() const
 	{
 		return GetStaticType();
 	}
 
 
-	//---------------------------------------------------------------------------------------------
 	const FNodeType* NodeMesh::GetStaticType()
 	{
 		return &s_nodeMeshType;
 	}
 
 
+	void MeshCreateCollapsedVertexMap(const mu::Mesh* Mesh, TArray<int32>& CollapsedVertices)
+	{
+		MUTABLE_CPUPROFILER_SCOPE(LayoutUV_CreateCollapsedVertexMap);
+
+		const int32 NumVertices = Mesh->GetVertexCount();
+		CollapsedVertices.Reserve(NumVertices);
+
+		UE::Geometry::TPointHashGrid3f<int32> VertHash(0.01f, INDEX_NONE);
+		VertHash.Reserve(NumVertices);
+
+		TArray<FVector3f> Vertices;
+		Vertices.SetNumUninitialized(NumVertices);
+
+		mu::UntypedMeshBufferIteratorConst ItPosition = mu::UntypedMeshBufferIteratorConst(Mesh->GetVertexBuffers(), mu::MBS_POSITION);
+
+		FVector3f* VertexData = Vertices.GetData();
+
+		for (int32 VertexIndex = 0; VertexIndex < NumVertices; ++VertexIndex)
+		{
+			*VertexData = ItPosition.GetAsVec3f();
+			VertHash.InsertPointUnsafe(VertexIndex, *VertexData);
+
+			++ItPosition;
+			++VertexData;
+		}
+
+		// Find unique vertices
+		CollapsedVertices.Init(INDEX_NONE, NumVertices);
+
+		TArray<int32> NearbyVertices;
+		for (int32 VertexIndex = 0; VertexIndex < NumVertices; ++VertexIndex)
+		{
+			if (CollapsedVertices[VertexIndex] != INDEX_NONE)
+			{
+				continue;
+			}
+
+			const FVector3f& Vertex = Vertices[VertexIndex];
+
+			NearbyVertices.Reset();
+			VertHash.FindPointsInBall(Vertex, 0.00001,
+				[&Vertex, &Vertices](const int32& Other) -> float {return FVector3f::DistSquared(Vertices[Other], Vertex); },
+				NearbyVertices);
+
+			// Find equals
+			for (int32 NearbyVertexIndex : NearbyVertices)
+			{
+				CollapsedVertices[NearbyVertexIndex] = VertexIndex;
+			}
+		}
+	}
+
+
+
+	void GetUVIsland(TArray<FTriangleInfo>& InTriangles,
+		const uint32 InFirstTriangle,
+		TArray<uint32>& OutTriangleIndices,
+		const TArray<FVector2f>& InUVs,
+		const TMultiMap<int32, uint32>& InVertexToTriangleMap)
+	{
+		MUTABLE_CPUPROFILER_SCOPE(GetUVIsland);
+
+		const uint32 NumTriangles = (uint32)InTriangles.Num();
+
+		OutTriangleIndices.Reserve(NumTriangles);
+		OutTriangleIndices.Add(InFirstTriangle);
+
+		TArray<bool> SkipTriangles;
+		SkipTriangles.Init(false, NumTriangles);
+
+		TArray<uint32> PendingTriangles;
+		PendingTriangles.Reserve(NumTriangles / 64);
+		PendingTriangles.Add(InFirstTriangle);
+
+		while (!PendingTriangles.IsEmpty())
+		{
+			const uint32 TriangleIndex = PendingTriangles.Pop();
+
+			// Triangle about to be proccessed, mark as skip;
+			SkipTriangles[TriangleIndex] = true;
+
+			bool ConnectedEdges[3] = { false, false, false };
+
+			const FTriangleInfo& Triangle = InTriangles[TriangleIndex];
+
+			// Find Triangles connected to edges 0 and 2
+			int32 CollapsedVertex1 = Triangle.CollapsedIndices[1];
+			int32 CollapsedVertex2 = Triangle.CollapsedIndices[2];
+
+			TArray<uint32> FoundTriangleIndices;
+			InVertexToTriangleMap.MultiFind(Triangle.CollapsedIndices[0], FoundTriangleIndices);
+
+			for (uint32 OtherTriangleIndex : FoundTriangleIndices)
+			{
+				const FTriangleInfo& OtherTriangle = InTriangles[OtherTriangleIndex];
+
+				for (int32 OtherIndex = 0; OtherIndex < 3; ++OtherIndex)
+				{
+					const int32 OtherCollapsedIndex = OtherTriangle.CollapsedIndices[OtherIndex];
+					if (OtherCollapsedIndex == CollapsedVertex1)
+					{
+						// Check if the vertex is in the same UV Island 
+						if (!SkipTriangles[OtherTriangleIndex]
+							&& InUVs[Triangle.Indices[1]].Equals(InUVs[OtherTriangle.Indices[OtherIndex]], 0.00001f))
+						{
+							OutTriangleIndices.Add(OtherTriangleIndex);
+							PendingTriangles.Add(OtherTriangleIndex);
+							SkipTriangles[OtherTriangleIndex] = true;
+						}
+
+						// Connected but already processed or in another island
+						break;
+					}
+
+					if (OtherCollapsedIndex == CollapsedVertex2)
+					{
+						// Check if the vertex is in the same UV Island 
+						if (!SkipTriangles[OtherTriangleIndex]
+							&& InUVs[Triangle.Indices[2]].Equals(InUVs[OtherTriangle.Indices[OtherIndex]], 0.00001f))
+						{
+							OutTriangleIndices.Add(OtherTriangleIndex);
+							PendingTriangles.Add(OtherTriangleIndex);
+							SkipTriangles[OtherTriangleIndex] = true;
+						}
+
+						// Connected but already processed or in another UV Island
+						break;
+					}
+				}
+
+			}
+
+			// Find the triangle connected to edge 1
+			FoundTriangleIndices.Reset();
+			InVertexToTriangleMap.MultiFind(CollapsedVertex1, FoundTriangleIndices);
+
+			for (uint32 OtherTriangleIndex : FoundTriangleIndices)
+			{
+				const FTriangleInfo& OtherTriangle = InTriangles[OtherTriangleIndex];
+
+				for (int32 OtherIndex = 0; OtherIndex < 3; ++OtherIndex)
+				{
+					const int32 OtherCollapsedIndex = OtherTriangle.CollapsedIndices[OtherIndex];
+					if (OtherCollapsedIndex == CollapsedVertex2)
+					{
+						// Check if the vertex belong to the same UV island
+						if (!SkipTriangles[OtherTriangleIndex]
+							&& InUVs[Triangle.Indices[2]].Equals(InUVs[OtherTriangle.Indices[OtherIndex]], 0.00001f))
+						{
+							OutTriangleIndices.Add(OtherTriangleIndex);
+							PendingTriangles.Add(OtherTriangleIndex);
+							SkipTriangles[OtherTriangleIndex] = true;
+						}
+
+						// Connected but already processed or in another island
+						break;
+					}
+				}
+			}
+		}
+	}
+
+
 }
 
 
diff --git a/Engine/Plugins/Experimental/Mutable/Source/MutableTools/Public/MuT/NodeLayout.h b/Engine/Plugins/Experimental/Mutable/Source/MutableTools/Public/MuT/NodeLayout.h
--- a/Engine/Plugins/Experimental/Mutable/Source/MutableTools/Public/MuT/NodeLayout.h
+++ b/Engine/Plugins/Experimental/Mutable/Source/MutableTools/Public/MuT/NodeLayout.h
@@ -20,17 +20,17 @@ namespace mu
 		/** Optional mask image that selects the vertices to include in the block. */
 		Ptr<Image> Mask;
 
-		UE::Math::TIntVector2<uint16> Min = { 0, 0 };
-		UE::Math::TIntVector2<uint16> Size = { 0, 0 };
+		FIntVector2 Min = { 0, 0 };
+		FIntVector2 Size = { 0, 0 };
 
 		//! Priority value to control the shrink texture layout strategy
-		int32 Priority;
+		int32 Priority = 0;
 
 		//! Value to control the method to reduce the block
-		uint32 bReduceBothAxes : 1;
+		uint32 bReduceBothAxes : 1 = 0;
 
 		//! Value to control if a block has to be reduced by two in an unitary reduction strategy
-		uint32 bReduceByTwo : 1;
+		uint32 bReduceByTwo : 1 = 0;
 	};
 
 
@@ -40,12 +40,12 @@ namespace mu
 	public:
 
 		//!
-		UE::Math::TIntVector2<uint16> Size = UE::Math::TIntVector2<uint16>(0, 0);
+		FIntVector2 Size = FIntVector2(0, 0);
 
 		/** Maximum size in layout blocks that this layout can grow to. From there on, blocks will shrink to fit.
 		* If 0,0 then no maximum size applies.
 		*/
-		UE::Math::TIntVector2<uint16> MaxSize = UE::Math::TIntVector2<uint16>(0, 0);
+		FIntVector2 MaxSize = FIntVector2(0, 0);
 
 		//!
 		TArray<FSourceLayoutBlock> Blocks;
@@ -67,8 +67,11 @@ namespace mu
 
 		// Own interface
 
-		//! Generate the blocks of a layout using the UV of the meshes
-		static Ptr<NodeLayout> GenerateLayoutBlocks(const Ptr<Mesh>& pMesh, int32 layoutIndex, int32 gridSizeX, int32 gridSizeY);
+		/** Generate the blocks of a layout using the UV of the meshes. 
+		* A list of existing blocks may be provided to specify starting blocks that shouldn't be modified unless it is mandatory because they contain partial face islands. 
+		*/
+		void GenerateLayoutBlocks(const Ptr<Mesh>&, int32 LayoutIndex );
+		void GenerateLayoutBlocksFromUVIslands(const Ptr<Mesh>&, int32 LayoutIndex, bool bMergeChildBlocks);
 
 	protected:
 
diff --git a/Engine/Plugins/Experimental/Mutable/Source/MutableTools/Public/MuT/NodeMesh.h b/Engine/Plugins/Experimental/Mutable/Source/MutableTools/Public/MuT/NodeMesh.h
--- a/Engine/Plugins/Experimental/Mutable/Source/MutableTools/Public/MuT/NodeMesh.h
+++ b/Engine/Plugins/Experimental/Mutable/Source/MutableTools/Public/MuT/NodeMesh.h
@@ -6,38 +6,67 @@
 #include "MuR/Ptr.h"
 #include "MuR/RefCounted.h"
 #include "MuT/Node.h"
+#include "Containers/Array.h"
+#include "Containers/Map.h"
 
 namespace mu
 {
 
 	// Forward definitions
+	class Mesh;
 	class NodeMesh;
 	typedef Ptr<NodeMesh> NodeMeshPtr;
 	typedef Ptr<const NodeMesh> NodeMeshPtrConst;
 
 
-    //! %Base class of any node that outputs a mesh.
+    /** Base class of any node that outputs a mesh. */
 	class MUTABLETOOLS_API NodeMesh : public Node
 	{
 	public:
 
-		//-----------------------------------------------------------------------------------------
 		// Node Interface
-		//-----------------------------------------------------------------------------------------
         const FNodeType* GetType() const override;
 		static const FNodeType* GetStaticType();
 
-		//-----------------------------------------------------------------------------------------
 		// Interface pattern
-		//-----------------------------------------------------------------------------------------
 		class Private;
 
 	protected:
 
-		//! Forbidden. Manage with the Ptr<> template.
+		// Forbidden. Manage with the Ptr<> template.
 		inline ~NodeMesh() {}
 	};
 
 
+	/** Helper structs for mesh utility methods below. */
+	struct FTriangleInfo
+	{
+		/** Vertex indices in the original mesh. */
+		uint32 Indices[3];
+
+		/** Vertex indices in the collapsed vertex list of the mesh. */
+		uint32 CollapsedIndices[3];
+
+		/** Optional data with layout block indices. */
+		uint16 BlockIndices[3];
+
+		/** Optional data with a flag indicating the UVs have changed during layout for this trioangle. */
+		bool bUVsFixed = false;
+	};
+
+
+	/** Fill an array with the indices of all triangles belonging to the same UV island as InFirstTriangle. 
+	*/
+	extern void GetUVIsland(TArray<FTriangleInfo>& InTriangles,
+		const uint32 InFirstTriangle,
+		TArray<uint32>& OutTriangleIndices,
+		const TArray<FVector2f>& InUVs,
+		const TMultiMap<int32, uint32>& InVertexToTriangleMap);
+
+
+	/** Create a map from vertices into vertices, collapsing vertices that have the same position.
+	*/
+	extern void MeshCreateCollapsedVertexMap(const mu::Mesh* Mesh, TArray<int32>& CollapsedVertices);
+
 }
 
