[mutable] Remove duplicated blocks when converting previously incorrect legacy data.
#rnx

diff --git a/Engine/Plugins/Experimental/Mutable/Source/CustomizableObjectEditor/Private/MuCOE/Nodes/CustomizableObjectNodeEditMaterial.cpp b/Engine/Plugins/Experimental/Mutable/Source/CustomizableObjectEditor/Private/MuCOE/Nodes/CustomizableObjectNodeEditMaterial.cpp
--- a/Engine/Plugins/Experimental/Mutable/Source/CustomizableObjectEditor/Private/MuCOE/Nodes/CustomizableObjectNodeEditMaterial.cpp
+++ b/Engine/Plugins/Experimental/Mutable/Source/CustomizableObjectEditor/Private/MuCOE/Nodes/CustomizableObjectNodeEditMaterial.cpp
@@ -160,7 +160,7 @@ void UCustomizableObjectNodeEditMaterial::BackwardsCompatibleFixup()
 
 		if (!ParentLayouts.IsValidIndex(ParentLayoutIndex))
 		{
-			UE_LOG(LogMutable, Warning, TEXT("[%s] UCustomizableObjectNodeRemoveMeshBlocks refers to an invalid texture layout index %d. Parent node has %d layouts."),
+			UE_LOG(LogMutable, Warning, TEXT("[%s] UCustomizableObjectNodeEditMaterial refers to an invalid texture layout index %d. Parent node has %d layouts."),
 				*GetOutermost()->GetName(), ParentLayoutIndex, ParentLayouts.Num());
 		}
 		else
@@ -174,10 +174,29 @@ void UCustomizableObjectNodeEditMaterial::BackwardsCompatibleFixup()
 			{
 				for (const FGuid& BlockId : BlockIds_DEPRECATED)
 				{
+					bool bSkipBlock = false;
+					for (const FCustomizableObjectLayoutBlock& ExistingBlock : Layout->Blocks)
+					{
+						if (ExistingBlock.Id == BlockId)
+						{
+							bSkipBlock = true;
+							UE_LOG(LogMutable, Warning, TEXT("[%s] UCustomizableObjectNodeEditMaterial has a duplicated layout block id. One has been ignored during version upgrade."), *GetOutermost()->GetName());
+							break;
+						}
+					}
+
+					if (bSkipBlock)
+					{
+						continue;
+					}
+
+					bool bFoundInParent = false;
 					for (const FCustomizableObjectLayoutBlock& ParentBlock : ParentLayout->Blocks)
 					{
 						if (ParentBlock.Id == BlockId)
 						{
+							bFoundInParent = true;
+
 							FCustomizableObjectLayoutBlock NewBlock;
 							NewBlock = ParentBlock;
 
@@ -187,14 +206,14 @@ void UCustomizableObjectNodeEditMaterial::BackwardsCompatibleFixup()
 							NewBlock.Priority = 0;
 
 							Layout->Blocks.Add(NewBlock);
+							break;
 						}
 					}
-				}
 
-				if (Layout->Blocks.Num() != BlockIds_DEPRECATED.Num())
-				{
-					UE_LOG(LogMutable, Warning, TEXT("[%s] UCustomizableObjectNodeRemoveMeshBlocks refers to %d invalid layout block. It has been ignored during version upgrade."),
-						*GetOutermost()->GetName(), int32(BlockIds_DEPRECATED.Num() - Layout->Blocks.Num()));
+					if (!bFoundInParent)
+					{
+						UE_LOG(LogMutable, Warning, TEXT("[%s] UCustomizableObjectNodeEditMaterial refers to and invalid layout block. It has been ignored during version upgrade."), *GetOutermost()->GetName());
+					}
 				}
 			}
 		}
diff --git a/Engine/Plugins/Experimental/Mutable/Source/CustomizableObjectEditor/Private/MuCOE/Nodes/CustomizableObjectNodeRemoveMeshBlocks.cpp b/Engine/Plugins/Experimental/Mutable/Source/CustomizableObjectEditor/Private/MuCOE/Nodes/CustomizableObjectNodeRemoveMeshBlocks.cpp
--- a/Engine/Plugins/Experimental/Mutable/Source/CustomizableObjectEditor/Private/MuCOE/Nodes/CustomizableObjectNodeRemoveMeshBlocks.cpp
+++ b/Engine/Plugins/Experimental/Mutable/Source/CustomizableObjectEditor/Private/MuCOE/Nodes/CustomizableObjectNodeRemoveMeshBlocks.cpp
@@ -89,10 +89,29 @@ void UCustomizableObjectNodeRemoveMeshBlocks::BackwardsCompatibleFixup()
 			{
 				for (const FGuid& BlockId : BlockIds_DEPRECATED)
 				{
+					bool bSkipBlock = false;
+					for (const FCustomizableObjectLayoutBlock& ExistingBlock : Layout->Blocks)
+					{
+						if (ExistingBlock.Id == BlockId)
+						{
+							bSkipBlock = true;
+							UE_LOG(LogMutable, Warning, TEXT("[%s] UCustomizableObjectNodeRemoveMeshBlocks has a duplicated layout block id. One has been ignored during version upgrade."), *GetOutermost()->GetName());
+							break;
+						}
+					}
+
+					if (bSkipBlock)
+					{
+						continue;
+					}
+
+					bool bFoundInParent = false;
 					for (const FCustomizableObjectLayoutBlock& ParentBlock : ParentLayout->Blocks)
 					{
 						if (ParentBlock.Id == BlockId)
 						{
+							bFoundInParent = true;
+
 							FCustomizableObjectLayoutBlock NewBlock;
 							NewBlock = ParentBlock;
 
@@ -102,14 +121,14 @@ void UCustomizableObjectNodeRemoveMeshBlocks::BackwardsCompatibleFixup()
 							NewBlock.Priority = 0;
 
 							Layout->Blocks.Add(NewBlock);
+							break;
 						}
 					}
-				}
 
-				if (Layout->Blocks.Num() != BlockIds_DEPRECATED.Num())
-				{
-					UE_LOG(LogMutable, Warning, TEXT("[%s] UCustomizableObjectNodeRemoveMeshBlocks refers to %d invalid layout block. It has been ignored during version upgrade."),
-						*GetOutermost()->GetName(), int32(BlockIds_DEPRECATED.Num() - Layout->Blocks.Num()));
+					if (!bFoundInParent)
+					{
+						UE_LOG(LogMutable, Warning, TEXT("[%s] UCustomizableObjectNodeRemoveMeshBlocks refers to and invalid layout block. It has been ignored during version upgrade."), *GetOutermost()->GetName());
+					}
 				}
 			}
 		}
