[Mutable] Split mesh data in multiple roms, geometry, pose and externaly referenced data with buffer description. Add an initial pass to generate parts of the mesh that are not geometry.

- This new pass executes mesh code but only generates skeletons, physics bodies, external references and a description of the generated mesh.
- The description will only be informative since some operations can change format and number of vertices of the final mesh based on the geometry data.

#perflight
#rnx

diff --git a/Engine/Plugins/Experimental/Mutable/Source/CustomizableObject/Internal/MuCO/CustomizableObjectPrivate.h b/Engine/Plugins/Experimental/Mutable/Source/CustomizableObject/Internal/MuCO/CustomizableObjectPrivate.h
--- a/Engine/Plugins/Experimental/Mutable/Source/CustomizableObject/Internal/MuCO/CustomizableObjectPrivate.h
+++ b/Engine/Plugins/Experimental/Mutable/Source/CustomizableObject/Internal/MuCO/CustomizableObjectPrivate.h
@@ -1315,6 +1315,8 @@ public:
 
 		LayoutBlocksAsInt32,
 
+		MeshDataStorageSplit,
+
 		// -----<new versions can be added above this line>--------
 		LastCustomizableObjectVersion
 	};
diff --git a/Engine/Plugins/Experimental/Mutable/Source/CustomizableObject/Internal/MuCO/CustomizableObjectSystemPrivate.h b/Engine/Plugins/Experimental/Mutable/Source/CustomizableObject/Internal/MuCO/CustomizableObjectSystemPrivate.h
--- a/Engine/Plugins/Experimental/Mutable/Source/CustomizableObject/Internal/MuCO/CustomizableObjectSystemPrivate.h
+++ b/Engine/Plugins/Experimental/Mutable/Source/CustomizableObject/Internal/MuCO/CustomizableObjectSystemPrivate.h
@@ -363,7 +363,20 @@ struct FInstanceUpdateData
 	struct FLOD
 	{		
 		mu::FResourceID MeshID;
-		mu::Ptr<const mu::Mesh> Mesh;
+		
+		/** 
+		 * Mesh generated only loading the geometry data. It may contain Pose and other data information
+		 * as a side result, but should not be used. 
+		 */
+		mu::Ptr<const mu::Mesh> MeshGeometry;
+
+		/**
+		 * Mesh generated without executing geometry operations. It contains Pose information,
+		 * external references and tags. Its buffers are imprecise descriptors, e.i., don't have 
+		 * data, the number of elements is an upper bound of the final result and its format may 
+		 * vary in bit width for some channels.
+		 */
+		mu::Ptr<const mu::Mesh> MeshComps;
 
 		/** Range in the Surfaces array */
 		uint16 FirstSurface = 0;
diff --git a/Engine/Plugins/Experimental/Mutable/Source/CustomizableObject/Private/MuCO/CustomizableObjectInstance.cpp b/Engine/Plugins/Experimental/Mutable/Source/CustomizableObject/Private/MuCO/CustomizableObjectInstance.cpp
--- a/Engine/Plugins/Experimental/Mutable/Source/CustomizableObject/Private/MuCO/CustomizableObjectInstance.cpp
+++ b/Engine/Plugins/Experimental/Mutable/Source/CustomizableObject/Private/MuCO/CustomizableObjectInstance.cpp
@@ -1896,18 +1896,18 @@ bool UCustomizableInstancePrivate::DoComponentsNeedUpdate(UCustomizableObjectIns
 		{
 			const FInstanceUpdateData::FLOD& LOD = OperationData->InstanceUpdateData.LODs[Component.FirstLOD + LODIndex];
 
-			if (!LOD.bGenerated || !LOD.Mesh )
+			if (!LOD.bGenerated || !LOD.MeshGeometry)
 			{
 				continue;
 			}
 
-			if (LOD.SurfaceCount == 0 && !LOD.Mesh->IsReference())
+			if (LOD.SurfaceCount == 0 && !LOD.MeshGeometry->IsReference())
 			{
 				continue;
 			}
 
 			// Unreal does not support empty sections.
-			if (!LOD.Mesh->IsReference() && LOD.Mesh->GetVertexCount() == 0) // else
+			if (!LOD.MeshGeometry->IsReference() && LOD.MeshGeometry->GetVertexCount() == 0) // else
 			{
 				UE_LOG(LogMutable, Error, TEXT("Failed to generate SkeletalMesh for CO Instance [%s]. CO [%s] has invalid geometry for LOD [%d] Component [%d]."),
 					*Public->GetName(), *CustomizableObject->GetName(),
@@ -1953,7 +1953,7 @@ bool UCustomizableInstancePrivate::DoComponentsNeedUpdate(UCustomizableObjectIns
 		}
 
 		const FInstanceUpdateData::FLOD& LOD = OperationData->InstanceUpdateData.LODs[Component.FirstLOD];
-		bool bIsReferenced = LOD.Mesh && LOD.Mesh->IsReference();
+		bool bIsReferenced = LOD.MeshGeometry && LOD.MeshGeometry->IsReference();
 		if (!bIsReferenced)
 		{
 			if (ComponentWithMesh[InstanceComponentIndex] && MeshIDs[InstanceComponentIndex * MAX_MESH_LOD_COUNT + OperationData->NumLODsAvailablePerComponent[ObjectComponentIndex] - 1] == MAX_uint64)
@@ -2197,7 +2197,7 @@ bool UCustomizableInstancePrivate::UpdateSkeletalMesh_PostBeginUpdate0(UCustomiz
 			mu::Ptr<const mu::Mesh> FirstMesh;
 			if (Component.LODCount > 0)
 			{
-				FirstMesh = OperationData->InstanceUpdateData.LODs[Component.FirstLOD].Mesh;
+				FirstMesh = OperationData->InstanceUpdateData.LODs[Component.FirstLOD].MeshGeometry;
 			}
 
 			if (FirstMesh && FirstMesh->IsReference())
@@ -2247,7 +2247,7 @@ bool UCustomizableInstancePrivate::UpdateSkeletalMesh_PostBeginUpdate0(UCustomiz
 		mu::Ptr<const mu::Mesh> ComponentMesh;
 		for (int32 FirstValidLODIndex = Component.FirstLOD; FirstValidLODIndex < OperationData->InstanceUpdateData.LODs.Num() && !ComponentMesh; ++FirstValidLODIndex)
 		{
-			ComponentMesh = OperationData->InstanceUpdateData.LODs[FirstValidLODIndex].Mesh;
+			ComponentMesh = OperationData->InstanceUpdateData.LODs[FirstValidLODIndex].MeshGeometry;
 		}
 		
 		if (ComponentMesh)
@@ -2371,7 +2371,7 @@ bool UCustomizableInstancePrivate::UpdateSkeletalMesh_PostBeginUpdate0(UCustomiz
 				mu::Ptr<const mu::Mesh> FirstMesh;
 				if (Component.LODCount > 0)
 				{
-					FirstMesh = OperationData->InstanceUpdateData.LODs[Component.FirstLOD].Mesh;
+					FirstMesh = OperationData->InstanceUpdateData.LODs[Component.FirstLOD].MeshGeometry;
 				}
 				if (FirstMesh && FirstMesh->IsReference())
 				{
@@ -4079,7 +4079,7 @@ void UCustomizableInstancePrivate::BuildOrCopyClothingData(const TSharedRef<FUpd
 				continue;
 			}
 
-			if (mu::MeshPtrConst MutableMesh = LOD.Mesh)
+			if (mu::Ptr<const mu::Mesh> MutableMesh = LOD.MeshGeometry)
 			{
 				const mu::FMeshBufferSet& MeshSet = MutableMesh->GetVertexBuffers();
 				const mu::FMeshBufferSet& IndicesSet = MutableMesh->GetIndexBuffers();
@@ -4799,7 +4799,7 @@ void UCustomizableInstancePrivate::BuildOrCopyClothingData(const TSharedRef<FUpd
 		{
 			const FInstanceUpdateData::FLOD& LOD = OperationData->InstanceUpdateData.LODs[Component.FirstLOD+LODIndex];
 
-			if (mu::MeshPtrConst MutableMesh = LOD.Mesh)
+			if (mu::Ptr<const mu::Mesh> MutableMesh = LOD.MeshGeometry)
 			{
 				TArray<FMeshToMeshVertData>& MappingData = LodMappingData[LODIndex];
 				TArray<FClothBufferIndexMapping>& ClothingIndexMapping = LodClothingIndexMapping[LODIndex];
@@ -5007,7 +5007,7 @@ bool UCustomizableInstancePrivate::BuildOrCopyRenderData(const TSharedRef<FUpdat
 		const FInstanceUpdateData::FLOD& LOD = OperationData->InstanceUpdateData.LODs[Component.FirstLOD+LODIndex];
 
 		// There could be components without a mesh in LODs
-		if (!LOD.bGenerated || !LOD.Mesh || LOD.SurfaceCount == 0)
+		if (!LOD.bGenerated || !LOD.MeshGeometry || LOD.SurfaceCount == 0)
 		{
 			UE_LOG(LogMutable, Warning, TEXT("Building instance: generated mesh [%s] has LOD [%d] of object component index [%d] with no mesh.")
 				, *SkeletalMesh->GetName()
@@ -5043,13 +5043,13 @@ bool UCustomizableInstancePrivate::BuildOrCopyRenderData(const TSharedRef<FUpdat
 
 
 		// Find referenced surface metadata.
-		const int32 MeshNumSurfaces = LOD.Mesh->Surfaces.Num();
+		const int32 MeshNumSurfaces = LOD.MeshGeometry->Surfaces.Num();
 		TArray<const FMutableSurfaceMetadata*> MeshSurfacesMetadata;
 		MeshSurfacesMetadata.Init(nullptr, MeshNumSurfaces);
 
 		for (int32 MeshSectionIndex = 0; MeshSectionIndex < MeshNumSurfaces; ++MeshSectionIndex)
 		{
-			uint32 MeshSurfaceId = LOD.Mesh->GetSurfaceId(MeshSectionIndex);
+			uint32 MeshSurfaceId = LOD.MeshGeometry->GetSurfaceId(MeshSectionIndex);
 			int32 InstanceSurfaceIndex = OperationData->MutableInstance->FindSurfaceById(InstanceComponentIndex, LODIndex, MeshSurfaceId);
 			
 			if (InstanceSurfaceIndex < 0)
@@ -5064,7 +5064,7 @@ bool UCustomizableInstancePrivate::BuildOrCopyRenderData(const TSharedRef<FUpdat
 		// Set RenderSections
 		UnrealConversionUtils::SetupRenderSections(
 			LODResource,
-			LOD.Mesh,
+			LOD.MeshGeometry,
 			OperationData->InstanceUpdateData.BoneMaps,
 			BoneInfoMap,
 			LOD.FirstBoneMap,
@@ -5075,11 +5075,11 @@ bool UCustomizableInstancePrivate::BuildOrCopyRenderData(const TSharedRef<FUpdat
 			// Copy Vertices
 			UnrealConversionUtils::CopyMutableVertexBuffers(
 				LODResource,
-				LOD.Mesh,
+				LOD.MeshGeometry,
 				SkeletalMesh->GetLODInfo(LODIndex)->bAllowCPUAccess);
 
 			// Copy indices.
-			if (!UnrealConversionUtils::CopyMutableIndexBuffers(LODResource, LOD.Mesh))
+			if (!UnrealConversionUtils::CopyMutableIndexBuffers(LODResource, LOD.MeshGeometry))
 			{
 				// End with failure
 				return false;
@@ -5088,7 +5088,7 @@ bool UCustomizableInstancePrivate::BuildOrCopyRenderData(const TSharedRef<FUpdat
 			// Copy SkinWeightProfiles
 			if (!ModelResources.SkinWeightProfilesInfo.IsEmpty())
 			{
-				const mu::FMeshBufferSet& MutableMeshVertexBuffers = LOD.Mesh->GetVertexBuffers();
+				const mu::FMeshBufferSet& MutableMeshVertexBuffers = LOD.MeshGeometry->GetVertexBuffers();
 
 				bool bHasSkinWeightProfiles = false;
 			
@@ -5143,11 +5143,11 @@ bool UCustomizableInstancePrivate::BuildOrCopyRenderData(const TSharedRef<FUpdat
 			// Init VertexBuffers for streaming
 			UnrealConversionUtils::InitVertexBuffersWithDummyData(
 				LODResource,
-				LOD.Mesh,
+				LOD.MeshGeometry,
 				SkeletalMesh->GetLODInfo(LODIndex)->bAllowCPUAccess);
 
 			// Init IndexBuffers for streaming
-			UnrealConversionUtils::InitIndexBuffersWithDummyData(LODResource, LOD.Mesh);
+			UnrealConversionUtils::InitIndexBuffersWithDummyData(LODResource, LOD.MeshGeometry);
 
 			// SkinWeightProfilesInfo Not supported yet
 		}
@@ -5308,15 +5308,15 @@ UE::Tasks::FTask UCustomizableInstancePrivate::LoadAdditionalAssetsAndData(
 			{
 				const FInstanceUpdateData::FLOD& LOD = OperationData->InstanceUpdateData.LODs[Component.FirstLOD + LODIndex];
 
-				mu::Ptr<const mu::Mesh> MutableMesh = LOD.Mesh;
-				if (!MutableMesh)
+				mu::Ptr<const mu::Mesh> MutableMeshComps = LOD.MeshComps;
+				if (!MutableMeshComps)
 				{
 					continue;
 				}
 
 				FCustomizableInstanceComponentData* ComponentData = GetComponentData(ComponentIndex);
 
-				const TArray<uint64>& StreamedResources = MutableMesh->GetStreamedResources();
+				const TArray<uint64>& StreamedResources = MutableMeshComps->GetStreamedResources();
 
 				for (uint64 ResourceId : StreamedResources)
 				{
@@ -5361,7 +5361,7 @@ UE::Tasks::FTask UCustomizableInstancePrivate::LoadAdditionalAssetsAndData(
 						const TMap<uint32, FClothingStreamable>& ClothingStreamables = ModelStreamableBulkData->ClothingStreamables;
 						if (const FClothingStreamable* ClothingStreamable = ClothingStreamables.Find(TypedResourceId.Id))
 						{
-							// TODO: Add async loading of ClothingAsset Data. This could be loaded as an streamead resource similar to and the asset user data.
+							// TODO: Add async loading of ClothingAsset Data. This could be loaded as an streamead resource similar to the asset user data.
 							int32 ClothingAssetIndex = ClothingStreamable->ClothingAssetIndex; 
 							int32 PhysicsAssetIndex = ClothingStreamable->PhysicsAssetIndex;					
 							const TSoftObjectPtr<UPhysicsAsset>& PhysicsAsset = ModelResources.PhysicsAssets.IsValidIndex(PhysicsAssetIndex) 
@@ -5401,9 +5401,9 @@ UE::Tasks::FTask UCustomizableInstancePrivate::LoadAdditionalAssetsAndData(
 
 				const bool bReplacePhysicsAssets = HasCOInstanceFlags(ReplacePhysicsAssets);
 
-				for (int32 TagIndex = 0; TagIndex < MutableMesh->GetTagCount(); ++TagIndex)
+				for (int32 TagIndex = 0; TagIndex < MutableMeshComps->GetTagCount(); ++TagIndex)
 				{
-					FString Tag = MutableMesh->GetTag(TagIndex);
+					FString Tag = MutableMeshComps->GetTag(TagIndex);
 					if (Tag.RemoveFromStart("__PA:"))
 					{
 						const int32 AssetIndex = FCString::Atoi(*Tag);
@@ -5488,10 +5488,10 @@ UE::Tasks::FTask UCustomizableInstancePrivate::LoadAdditionalAssetsAndData(
 	#endif
 				}
 
-				const int32 AdditionalPhysicsNum = MutableMesh->AdditionalPhysicsBodies.Num();
+				const int32 AdditionalPhysicsNum = MutableMeshComps->AdditionalPhysicsBodies.Num();
 				for (int32 I = 0; I < AdditionalPhysicsNum; ++I)
 				{
-					const int32 ExternalId = MutableMesh->AdditionalPhysicsBodies[I]->CustomId;
+					const int32 ExternalId = MutableMeshComps->AdditionalPhysicsBodies[I]->CustomId;
 					
 					ComponentData->PhysicsAssets.AdditionalPhysicsAssetsToLoad.Add(ExternalId);
 					AssetsToStream.Add(ModelResources.AnimBpOverridePhysiscAssetsInfo[ExternalId].SourceAsset.ToSoftObjectPath());
diff --git a/Engine/Plugins/Experimental/Mutable/Source/CustomizableObject/Private/MuCO/CustomizableObjectSystem.cpp b/Engine/Plugins/Experimental/Mutable/Source/CustomizableObject/Private/MuCO/CustomizableObjectSystem.cpp
--- a/Engine/Plugins/Experimental/Mutable/Source/CustomizableObject/Private/MuCO/CustomizableObjectSystem.cpp
+++ b/Engine/Plugins/Experimental/Mutable/Source/CustomizableObject/Private/MuCO/CustomizableObjectSystem.cpp
@@ -54,6 +54,7 @@ DECLARE_CYCLE_STAT(TEXT("MutableTask"), STAT_MutableTask, STATGROUP_Game);
 
 #define UE_MUTABLE_UPDATE_REGION TEXT("Mutable Update")
 #define UE_TASK_MUTABLE_GETMESHES_REGION TEXT("Task_Mutable_GetMeshes")
+#define UE_TASK_MUTABLE_GETMESHES_COMPS_REGION TEXT("Task_Mutable_GetMeshesComponents")
 #define UE_TASK_MUTABLE_GETIMAGES_REGION TEXT("Task_Mutable_GetImages")
 
 
@@ -1563,8 +1564,7 @@ namespace impl
 		double StartTime,
 		uint32 StartCycles,
 		const TSharedRef<TArray<FGetImageData>>& GetImagesData,
-		int32 GetImageIndex);
-	
+		int32 GetImageIndex);	
 
 	struct FGetMeshData
 	{
@@ -1572,6 +1572,13 @@ namespace impl
 		mu::FResourceID MeshID;
 	};
 
+	/** Process the next Mesh. If there are no more Meshes, go to the process Images loop. */
+	void Task_Mutable_GetMeshes_GetMeshComps_Loop(
+		const TSharedRef<FUpdateContextPrivate>& OperationData,
+		double StartTime,
+		uint32 StartCycles,
+		const TSharedRef<TArray<FGetMeshData>>& GetMeshesData,
+		int32 GetMeshIndex);
 
 	/** Process the next Mesh. If there are no more Meshes, go to the process Images loop. */
 	void Task_Mutable_GetMeshes_GetMesh_Loop(
@@ -1581,9 +1588,10 @@ namespace impl
 		const TSharedRef<TArray<FGetMeshData>>& GetMeshesData,
 		int32 GetMeshIndex);
 
-
-	/** Call GetImage.
-	  * Once GetImage is called, the task must end. Following code will be in a subsequent TaskGraph task. */
+	/** 
+	 * Call GetImage.
+	 * Once GetImage is called, the task must end. Following code will be in a subsequent TaskGraph task. 
+	 */
 	void Task_Mutable_GetImages_GetImage(
 		const TSharedRef<FUpdateContextPrivate>& OperationData,
 		double StartTime,
@@ -1601,7 +1609,7 @@ namespace impl
 		const TSharedPtr<TArray<mu::FResourceID>>& ImagesInThisInstance,
 		int32 ImageIndex);
 
-	
+
 	void Subtask_Mutable_UpdateParameterRelevancy(const TSharedRef<FUpdateContextPrivate>& OperationData)
 	{
 		MUTABLE_CPUPROFILER_SCOPE(Subtask_Mutable_UpdateParameterRelevancy)
@@ -1783,7 +1791,7 @@ namespace impl
 			{
 				FInstanceUpdateData::FLOD& LOD = OperationData->InstanceUpdateData.LODs[Component.FirstLOD+LODIndex];
 
-				mu::Ptr<const mu::Mesh> Mesh = LOD.Mesh;
+				mu::Ptr<const mu::Mesh> Mesh = LOD.MeshComps;
 				if (!Mesh || Mesh->IsReference())
 				{
 					continue;
@@ -1910,9 +1918,9 @@ namespace impl
 				const FInstanceUpdateData::FComponent& Component = UpdateData.Components[ComponentIndex];
 				const FInstanceUpdateData::FLOD& LOD = UpdateData.LODs[Component.FirstLOD+LODIndex];
 				check(LOD.bGenerated);
-				check(LOD.Mesh);
+				check(LOD.MeshGeometry);
 
-				const mu::FMeshBufferSet& MeshSet = LOD.Mesh->GetVertexBuffers();
+				const mu::FMeshBufferSet& MeshSet = LOD.MeshGeometry->GetVertexBuffers();
 
 				int32 VertexMorphsInfoIndexAndCountBufferIndex, VertexMorphsInfoIndexAndCountBufferChannel;
 				MeshSet.FindChannel(mu::MBS_OTHER, 0, &VertexMorphsInfoIndexAndCountBufferIndex, &VertexMorphsInfoIndexAndCountBufferChannel);
@@ -1931,7 +1939,7 @@ namespace impl
 				const uint32* const VertexMorphsResourceIdBuffer = reinterpret_cast<const uint32*>(MeshSet.GetBufferData(VertexMorphsResourceIdBufferIndex));
 				TArrayView<const uint32> VertexMorphsResourceIdView(VertexMorphsResourceIdBuffer, MeshSet.GetElementCount());
 
-				const int32 SurfaceCount = LOD.Mesh->GetSurfaceCount();
+				const int32 SurfaceCount = LOD.MeshGeometry->GetSurfaceCount();
 				for (int32 Section = 0; Section < SurfaceCount; ++Section)
 				{
 					// Reset SectionMorphTargets.
@@ -1941,7 +1949,7 @@ namespace impl
 					}
 
 					int32 FirstVertex, VerticesCount, FirstIndex, IndiciesCount, UnusedBoneIndex, UnusedBoneCount;
-					LOD.Mesh->GetSurface(Section, FirstVertex, VerticesCount, FirstIndex, IndiciesCount, UnusedBoneIndex, UnusedBoneCount);
+					LOD.MeshGeometry->GetSurface(Section, FirstVertex, VerticesCount, FirstIndex, IndiciesCount, UnusedBoneIndex, UnusedBoneCount);
 
 					for (int32 VertexIdx = FirstVertex; VertexIdx < FirstVertex + VerticesCount;)
 					{
@@ -2072,7 +2080,7 @@ namespace impl
 		uint32 StartCycles)
 	{
 		MUTABLE_CPUPROFILER_SCOPE(Task_Mutable_GetMeshes_End)
-
+		
 		// TODO: Not strictly mutable: move to another worker thread task to free mutable access?
 		Subtask_Mutable_PrepareSkeletonData(OperationData);
 		if (OperationData->GetCapturedDescriptor().GetBuildParameterRelevancy())
@@ -2108,7 +2116,6 @@ namespace impl
 		TRACE_END_REGION(UE_TASK_MUTABLE_GETMESHES_REGION);
 	}
 
-
 	/** TaskGraph task after GetImage has completed. */
 	void Task_Mutable_GetMeshes_GetImage_Post(
 		const TSharedRef<FUpdateContextPrivate>& OperationData,
@@ -2210,7 +2217,7 @@ namespace impl
 				LOD.FirstSurface = OperationData->InstanceUpdateData.Surfaces.Num();
 				LOD.SurfaceCount = 0;
 
-				if (!LOD.Mesh)
+				if (!LOD.MeshGeometry)
 				{
 					continue;
 				}
@@ -2332,7 +2339,7 @@ namespace impl
 
 				// If the mesh is a reference mesh, it won't have the surface information in the mutable mesh. We need to get it from the instance
 				// and all defined surfaces will be present. 
-				if (LOD.Mesh->IsReference())
+				if (LOD.MeshGeometry->IsReference())
 				{
 					const int32 SurfaceCount = MutableInstance->GetSurfaceCount(InstanceComponentIndex, MutableLODIndex);
 					for (int32 SurfaceIndex = 0; SurfaceIndex < SurfaceCount; ++SurfaceIndex)
@@ -2345,13 +2352,13 @@ namespace impl
 				// If the mesh is a not a reference mesh, we have to add only the materials of the surfaces that appear in the actual final mesh. 
 				else
 				{
-					const int32 SurfaceCount = LOD.Mesh->GetSurfaceCount();
+					const int32 SurfaceCount = LOD.MeshGeometry->GetSurfaceCount();
 					for (int32 MeshSurfaceIndex = 0; MeshSurfaceIndex < SurfaceCount; ++MeshSurfaceIndex)
 					{
-						const uint32 SurfaceId = LOD.Mesh->GetSurfaceId(MeshSurfaceIndex);
+						const uint32 SurfaceId = LOD.MeshGeometry->GetSurfaceId(MeshSurfaceIndex);
 
 						const int32 InstanceSurfaceIndex = MutableInstance->FindSurfaceById(InstanceComponentIndex, MutableLODIndex, SurfaceId);
-						check(LOD.Mesh->GetVertexCount() > 0 || InstanceSurfaceIndex >= 0);
+						check(LOD.MeshGeometry->GetVertexCount() > 0 || InstanceSurfaceIndex >= 0);
 
 						if (InstanceSurfaceIndex >= 0)
 						{
@@ -2381,15 +2388,15 @@ namespace impl
 		const int32 LODIndex = (*GetMeshesData)[GetMeshIndex].LODIndex;
 		FInstanceUpdateData::FLOD& LOD = OperationData->InstanceUpdateData.LODs[LODIndex];
 
-		LOD.Mesh = GetMeshTask.GetResult();
+		LOD.MeshGeometry = GetMeshTask.GetResult();
 
-		if (LOD.Mesh->IsReference())
+		if (LOD.MeshGeometry->IsReference())
 		{
 			const UCustomizableObjectInstance* Instance = OperationData->Instance.Get();
 			const UCustomizableObject* CustomizableObject = Instance->GetCustomizableObject();
 			const FModelResources& ModelResources = CustomizableObject->GetPrivate()->GetModelResources();
 
-			uint32 ReferenceID = LOD.Mesh->GetReferencedMesh();
+			uint32 ReferenceID = LOD.MeshGeometry->GetReferencedMesh();
 
 			if (ModelResources.PassThroughMeshes.IsValidIndex(ReferenceID))
 			{
@@ -2405,7 +2412,6 @@ namespace impl
 			
 		Task_Mutable_GetMeshes_GetMesh_Loop(OperationData, StartTime, StartCycles, GetMeshesData, ++GetMeshIndex);
 	}
-	
 
 	/** See declaration. */
 	void Task_Mutable_GetMeshes_GetMesh_Loop(
@@ -2424,8 +2430,8 @@ namespace impl
 		}
 		
 		const FGetMeshData& MeshData = (*GetMeshesData)[GetMeshIndex];
+
 		UE::Tasks::TTask<mu::Ptr<const mu::Mesh>> GetMeshTask = OperationData->MutableSystem->GetMesh(OperationData->InstanceID, MeshData.MeshID);
-		
 		UE::Tasks::AddNested(UE::Tasks::Launch(TEXT("Task_MutableGetMeshes_GetMesh_Post"), [=]()
 		{
 			Task_MutableGetMeshes_GetMesh_Post(OperationData, StartTime, StartCycles, GetMeshesData, GetMeshIndex, GetMeshTask);
@@ -2433,11 +2439,79 @@ namespace impl
 		GetMeshTask));
 	}
 
+	void Task_MutableGetMeshes_GetMeshComps_Post(
+		const TSharedRef<FUpdateContextPrivate>& OperationData,
+		double StartTime,
+		uint32 StartCycles,
+		const TSharedRef<TArray<FGetMeshData>>& GetMeshesData,
+		int32 GetMeshIndex,
+		UE::Tasks::TTask<mu::Ptr<const mu::Mesh>> GetMeshCompTask)
+	{
+		MUTABLE_CPUPROFILER_SCOPE(Task_MutableGetMeshes_GetMeshComps_Post)
+
+		const int32 LODIndex = (*GetMeshesData)[GetMeshIndex].LODIndex;
+		FInstanceUpdateData::FLOD& LOD = OperationData->InstanceUpdateData.LODs[LODIndex];
+
+		LOD.MeshComps = GetMeshCompTask.GetResult();
+
+		// For now Mesh references are processed as a normal mesh but this could be resolved earlier in the
+		// component gather phase.
+		//if (LOD.MeshComps->IsReference())
+		//{
+		//	const UCustomizableObjectInstance* Instance = OperationData->Instance.Get();
+		//	const UCustomizableObject* CustomizableObject = Instance->GetCustomizableObject();
+		//	const FModelResources& ModelResources = CustomizableObject->GetPrivate()->GetModelResources();
+
+		//	uint32 ReferenceID = LOD.Mesh->GetReferencedMesh();
+
+		//	if (ModelResources.PassThroughMeshes.IsValidIndex(ReferenceID))
+		//	{
+		//		TSoftObjectPtr<USkeletalMesh> Ref = ModelResources.PassThroughMeshes[ReferenceID];
+		//		Instance->GetPrivate()->PassThroughMeshesToLoad.Add(Ref);
+		//	}
+		//	else
+		//	{
+		//		// internal error.
+		//		UE_LOG(LogMutable, Error, TEXT("Referenced mesh [%d] was not stored in the resource array."), ReferenceID);
+		//	}
+		//}
+			
+		Task_Mutable_GetMeshes_GetMeshComps_Loop(OperationData, StartTime, StartCycles, GetMeshesData, ++GetMeshIndex);
+	}
+
+	void Task_Mutable_GetMeshes_GetMeshComps_Loop(
+		const TSharedRef<FUpdateContextPrivate>& OperationData,
+		double StartTime,
+		uint32 StartCycles,
+		const TSharedRef<TArray<FGetMeshData>>& GetMeshesData,
+		int32 GetMeshIndex)
+	{
+		MUTABLE_CPUPROFILER_SCOPE(Task_Mutable_GetMeshes_GetMeshCompoenents_Loop)
+
+		if (GetMeshIndex >= GetMeshesData->Num())
+		{
+			TRACE_END_REGION(UE_TASK_MUTABLE_GETMESHES_COMPS_REGION);
+			// After processing mesh components start meshes loop.
+			// TODO: We need to be able to generate only the resident lods for mesh streaming.
+			Task_Mutable_GetMeshes_GetMesh_Loop(OperationData, StartTime, StartCycles, GetMeshesData, 0);
+			return;
+		}
+		
+		const FGetMeshData& MeshData = (*GetMeshesData)[GetMeshIndex];
+
+		UE::Tasks::TTask<mu::Ptr<const mu::Mesh>> GetMeshCompTask = OperationData->MutableSystem->GetMeshComponents(OperationData->InstanceID, MeshData.MeshID);
+		UE::Tasks::AddNested(UE::Tasks::Launch(TEXT("Task_MutableGetMeshes_GetMeshComps_Post"), [=]()
+		{
+			Task_MutableGetMeshes_GetMeshComps_Post(OperationData, StartTime, StartCycles, GetMeshesData, GetMeshIndex, GetMeshCompTask);
+		},
+		GetMeshCompTask));
+	}
 
 	namespace Impl
 	{
-		/** Start of the GetMeshes tasks.
-		  * Gathers all GetMeshes that has to be called. */
+		/** 
+		* Start of the GetMeshes tasks.  
+		*/
 		void Task_Mutable_GetMeshes(const TSharedRef<FUpdateContextPrivate>& OperationData)
 		{
 			MUTABLE_CPUPROFILER_SCOPE(Task_Mutable_GetMeshes)
@@ -2523,17 +2597,17 @@ namespace impl
 						{
 							LOD.bGenerated = true;
 
-							GetMeshesData->Add({ UpdateDataLODIndex, LOD.MeshID});
+							GetMeshesData->Add({UpdateDataLODIndex, LOD.MeshID});
 						}
 					}
 				}
 			}
 
-			Task_Mutable_GetMeshes_GetMesh_Loop(OperationData, StartTime, StartCycles, GetMeshesData, 0);
+			TRACE_BEGIN_REGION(UE_TASK_MUTABLE_GETMESHES_COMPS_REGION);
+			Task_Mutable_GetMeshes_GetMeshComps_Loop(OperationData, StartTime, StartCycles, GetMeshesData, 0);
 		}
 	}
 
-
 	void Task_Mutable_GetMeshes(const TSharedRef<FUpdateContextPrivate>& OperationData)
 	{
 		Impl::Task_Mutable_GetMeshes(OperationData);
diff --git a/Engine/Plugins/Experimental/Mutable/Source/CustomizableObjectEditor/Private/MuCOE/SMutableCodeViewer.cpp b/Engine/Plugins/Experimental/Mutable/Source/CustomizableObjectEditor/Private/MuCOE/SMutableCodeViewer.cpp
--- a/Engine/Plugins/Experimental/Mutable/Source/CustomizableObjectEditor/Private/MuCOE/SMutableCodeViewer.cpp
+++ b/Engine/Plugins/Experimental/Mutable/Source/CustomizableObjectEditor/Private/MuCOE/SMutableCodeViewer.cpp
@@ -2359,7 +2359,7 @@ bool SMutableCodeViewer::IsConstantResourceUsedByOperation(const int32 IndexOnCo
 			{
 				if (OperationType == mu::OP_TYPE::ME_CONSTANT)
 				{
-					bResourceLocated = IndexOnConstantsArray == InProgram.GetOpArgs<mu::OP::MeshConstantArgs>(OperationAddress).value;
+					bResourceLocated = IndexOnConstantsArray == InProgram.GetOpArgs<mu::OP::MeshConstantArgs>(OperationAddress).Value;
 				}
 				break;
 			}
@@ -2426,7 +2426,7 @@ bool SMutableCodeViewer::IsConstantResourceUsedByOperation(const int32 IndexOnCo
 			{
 				if (OperationType == mu::OP_TYPE::ME_CONSTANT)
 				{
-					bResourceLocated = IndexOnConstantsArray == InProgram.GetOpArgs<mu::OP::MeshConstantArgs>(OperationAddress).skeleton;
+					bResourceLocated = IndexOnConstantsArray == InProgram.GetOpArgs<mu::OP::MeshConstantArgs>(OperationAddress).Skeleton;
 				}
 				break;
 			}
diff --git a/Engine/Plugins/Experimental/Mutable/Source/CustomizableObjectEditor/Private/MuCOE/SMutableConstantsWidget.cpp b/Engine/Plugins/Experimental/Mutable/Source/CustomizableObjectEditor/Private/MuCOE/SMutableConstantsWidget.cpp
--- a/Engine/Plugins/Experimental/Mutable/Source/CustomizableObjectEditor/Private/MuCOE/SMutableConstantsWidget.cpp
+++ b/Engine/Plugins/Experimental/Mutable/Source/CustomizableObjectEditor/Private/MuCOE/SMutableConstantsWidget.cpp
@@ -1228,8 +1228,9 @@ void SMutableConstantsWidget::LoadConstantMeshes()
 	
 	for (int32 MeshIndex = 0; MeshIndex < ConstantsCount; MeshIndex++)
 	{
-		TSharedPtr< FMutableConstantMeshElement> ConstantMeshElement = MakeShared<FMutableConstantMeshElement>();
-		ConstantMeshElement->MeshPtr = MutableProgramPtr->ConstantMeshes[MeshIndex].Value;
+		TSharedPtr<FMutableConstantMeshElement> ConstantMeshElement = MakeShared<FMutableConstantMeshElement>();
+
+		MutableProgramPtr->GetConstant(MeshIndex, ConstantMeshElement->MeshPtr);
 		ConstantMeshElement->IndexOnSourceVector = MeshIndex;
 
 		// Actual core disk size would be:
@@ -1276,7 +1277,7 @@ void SMutableConstantsWidget::LoadConstantLayouts()
 void SMutableConstantsWidget::LoadConstantSkeletons()
 {
 	check (MutableProgramPtr);
-	const int32 ConstantsCount = MutableProgramPtr->m_constantSkeletons.Num();
+	const int32 ConstantsCount = MutableProgramPtr->ConstantSkeletons.Num();
 	ConstantSkeletonElements.Empty(ConstantsCount);
 
 	mu::OutputMemoryStream Stream;
@@ -1285,7 +1286,7 @@ void SMutableConstantsWidget::LoadConstantSkeletons()
 	for (int32 SkeletonIndex = 0; SkeletonIndex < ConstantsCount; SkeletonIndex++)
 	{
 		TSharedPtr<FMutableConstantSkeletonElement> ConstantSkeletonElement = MakeShared<FMutableConstantSkeletonElement>();
-		ConstantSkeletonElement->Skeleton = MutableProgramPtr->m_constantSkeletons[SkeletonIndex];
+		ConstantSkeletonElement->Skeleton = MutableProgramPtr->ConstantSkeletons[SkeletonIndex];
 		ConstantSkeletonElement->IndexOnSourceVector = SkeletonIndex;
 		
 		ConstantSkeletonElement->Skeleton->Serialise(Archive);
diff --git a/Engine/Plugins/Experimental/Mutable/Source/MutableRuntime/Internal/MuR/ModelPrivate.h b/Engine/Plugins/Experimental/Mutable/Source/MutableRuntime/Internal/MuR/ModelPrivate.h
--- a/Engine/Plugins/Experimental/Mutable/Source/MutableRuntime/Internal/MuR/ModelPrivate.h
+++ b/Engine/Plugins/Experimental/Mutable/Source/MutableRuntime/Internal/MuR/ModelPrivate.h
@@ -69,6 +69,13 @@ namespace mu
 	};
 	MUTABLE_DEFINE_POD_SERIALISABLE(FImageLODRange);
 
+	struct FMeshRange
+	{
+		int32 FirstIndex = 0;
+		int32 NumMeshes = 0;
+	};
+	MUTABLE_DEFINE_POD_SERIALISABLE(FMeshRange);
+
 	struct FExtensionDataConstant
 	{
 		// This should always be valid, but if the state is Unloaded it won't be usable.
@@ -235,7 +242,13 @@ namespace mu
 		TArray<FImageLODRange> ConstantImages;
 
         //! Constant mesh data: the first is the index in m_roms for each mesh or -1 if it is always loaded.
-		TArray<TPair<int32, Ptr<const Mesh>>> ConstantMeshes;
+		TArray<TPair<int32, Ptr<const Mesh>>> ConstantMeshData;
+	
+		/** Indices pointing at ConstantMeshData */
+		TArray<int32> ConstantMeshIndices;
+
+		/** Span in the ConstantMeshIndices indicating the indices of the components of a mesh constant. */
+		TArray<FMeshRange> ConstantMeshes;
 
 		//! Constant ExtensionData
 		TArray<FExtensionDataConstant> m_constantExtensionData;
@@ -259,10 +272,10 @@ namespace mu
 		TArray<FRichCurve> ConstantCurves;
 
         //! Constant skeletons
-		TArray<Ptr<const Skeleton>> m_constantSkeletons;
+		TArray<Ptr<const Skeleton>> ConstantSkeletons;
 
 		//! Constant Physics Bodies
-		TArray<Ptr<const PhysicsBody>> m_constantPhysicsBodies;
+		TArray<Ptr<const PhysicsBody>> ConstantPhysicsBodies;
 
         //! Parameters of the model.
         //! The value stored here is the default value.
@@ -282,53 +295,57 @@ namespace mu
 		bool bIsValid = true;
 #endif
         //!
-        void Serialise( OutputArchive& arch ) const
+        void Serialise(OutputArchive& Arch) const
         {
-            arch << m_opAddress;
-            arch << m_byteCode;
-            arch << m_states;
-			arch << m_roms;
-			arch << ConstantImageLODs;
-			arch << ConstantImageLODIndices;
-			arch << ConstantImages;
-			arch << ConstantMeshes;
-			arch << m_constantExtensionData;
-			arch << m_constantStrings;
-            arch << m_constantLayouts;
-            arch << m_constantProjectors;
-			arch << m_constantMatrices;
-			arch << m_constantShapes;
-            arch << ConstantCurves;
-            arch << m_constantSkeletons;
-			arch << m_constantPhysicsBodies;
-            arch << m_parameters;
-            arch << m_ranges;
-            arch << m_parameterLists;
+            Arch << m_opAddress;
+            Arch << m_byteCode;
+            Arch << m_states;
+			Arch << m_roms;
+			Arch << ConstantImageLODs;
+			Arch << ConstantImageLODIndices;
+			Arch << ConstantImages;
+			Arch << ConstantMeshes;
+			Arch << ConstantMeshIndices;
+			Arch << ConstantMeshData;
+			Arch << m_constantExtensionData;
+			Arch << m_constantStrings;
+            Arch << m_constantLayouts;
+            Arch << m_constantProjectors;
+			Arch << m_constantMatrices;
+			Arch << m_constantShapes;
+            Arch << ConstantCurves;
+            Arch << ConstantSkeletons;
+			Arch << ConstantPhysicsBodies;
+            Arch << m_parameters;
+            Arch << m_ranges;
+            Arch << m_parameterLists;
         }
 
         //!
-        void Unserialise( InputArchive& arch )
+        void Unserialise(InputArchive& Arch)
         {
-            arch >> m_opAddress;
-            arch >> m_byteCode;
-            arch >> m_states;
-			arch >> m_roms;
-			arch >> ConstantImageLODs;
-			arch >> ConstantImageLODIndices;
-			arch >> ConstantImages;
-			arch >> ConstantMeshes;
-			arch >> m_constantExtensionData;
-			arch >> m_constantStrings;
-            arch >> m_constantLayouts;
-            arch >> m_constantProjectors;
-			arch >> m_constantMatrices;
-			arch >> m_constantShapes;
-            arch >> ConstantCurves;
-            arch >> m_constantSkeletons;
-			arch >> m_constantPhysicsBodies;
-            arch >> m_parameters;
-            arch >> m_ranges;
-            arch >> m_parameterLists;
+            Arch >> m_opAddress;
+            Arch >> m_byteCode;
+            Arch >> m_states;
+			Arch >> m_roms;
+			Arch >> ConstantImageLODs;
+			Arch >> ConstantImageLODIndices;
+			Arch >> ConstantImages;
+			Arch >> ConstantMeshes;
+			Arch >> ConstantMeshIndices;
+			Arch >> ConstantMeshData;
+			Arch >> m_constantExtensionData;
+			Arch >> m_constantStrings;
+            Arch >> m_constantLayouts;
+            Arch >> m_constantProjectors;
+			Arch >> m_constantMatrices;
+			Arch >> m_constantShapes;
+            Arch >> ConstantCurves;
+            Arch >> ConstantSkeletons;
+			Arch >> ConstantPhysicsBodies;
+            Arch >> m_parameters;
+            Arch >> m_ranges;
+            Arch >> m_parameterLists;
         }
 
         //! Debug method that sanity-checks the program with a variety of tests.
@@ -345,7 +362,7 @@ namespace mu
 				case DT_IMAGE: 
 					return ConstantImageLODs[m_roms[RomIndex].ResourceIndex].Value.get() != nullptr;
 				case DT_MESH: 
-					return ConstantMeshes[m_roms[RomIndex].ResourceIndex].Value.get() != nullptr;
+					return ConstantMeshData[m_roms[RomIndex].ResourceIndex].Value.get() != nullptr;
 				default:
 					check(false);
 					break;
@@ -377,10 +394,10 @@ namespace mu
 				}
 				break;
 			case DT_MESH:
-				if (ConstantMeshes[m_roms[RomIndex].ResourceIndex].Value)
+				if (ConstantMeshData[m_roms[RomIndex].ResourceIndex].Value)
 				{
-					RomSize = ConstantMeshes[m_roms[RomIndex].ResourceIndex].Value->GetDataSize();
-					ConstantMeshes[m_roms[RomIndex].ResourceIndex].Value = nullptr;
+					RomSize = ConstantMeshData[m_roms[RomIndex].ResourceIndex].Value->GetDataSize();
+					ConstantMeshData[m_roms[RomIndex].ResourceIndex].Value = nullptr;
 				}
 				break;
 			default:
@@ -396,7 +413,7 @@ namespace mu
 			check(m_roms[RomIndex].ResourceType == DT_MESH);
 			
 			LoadedMemTrackedRoms.EmplaceAt(RomIndex, (uint8)m_roms[RomIndex].ResourceType);
-			ConstantMeshes[m_roms[RomIndex].ResourceIndex].Value = Value;
+			ConstantMeshData[m_roms[RomIndex].ResourceIndex].Value = Value;
 		}
 
 		FORCEINLINE void SetImageRomValue(int32 RomIndex, const Ptr<Image>& Value)
@@ -407,13 +424,6 @@ namespace mu
 			ConstantImageLODs[m_roms[RomIndex].ResourceIndex].Value = Value;
 		}
 
-
-		int32 AddConstant(Ptr<const Mesh> pMesh)
-		{
-			// Uniques needs to be ensured outside
-			return ConstantMeshes.Add(TPair<int32, Ptr<const Mesh>>( -1, pMesh.get() ));
-		}
-
 		OP::ADDRESS AddConstant(Ptr<const ExtensionData> Data)
 		{
 			// Ensure unique
@@ -454,42 +464,36 @@ namespace mu
             return index;
         }
 
-        OP::ADDRESS AddConstant( Ptr<const Skeleton> pSkeleton )
+        OP::ADDRESS AddConstant(Ptr<const Skeleton> SkeletonPtr)
         {
             // Ensure unique
-            for ( SIZE_T i=0; i<m_constantSkeletons.Num(); ++i)
+            for (int32 SkeletonIndex = 0; SkeletonIndex < ConstantSkeletons.Num(); ++SkeletonIndex)
             {
-                if ( m_constantSkeletons[i]==pSkeleton
-                     ||
-                     *m_constantSkeletons[i]==*pSkeleton
-                     )
+                if (ConstantSkeletons[SkeletonIndex] == SkeletonPtr || *ConstantSkeletons[SkeletonIndex] == *SkeletonPtr)
                 {
-                    return (OP::ADDRESS)i;
+                    return (OP::ADDRESS)SkeletonIndex;
                 }
             }
 
-            OP::ADDRESS index = OP::ADDRESS( m_constantSkeletons.Num() );
-            m_constantSkeletons.Add( pSkeleton );
-            return index;
+            OP::ADDRESS NewSkeletonIndex = OP::ADDRESS(ConstantSkeletons.Num());
+            ConstantSkeletons.Add(SkeletonPtr);
+            return NewSkeletonIndex;
         }
 
-        OP::ADDRESS AddConstant( Ptr<const PhysicsBody> pPhysicsBody )
+        OP::ADDRESS AddConstant(Ptr<const PhysicsBody> PhysicsBodyPtr)
         {
             // Ensure unique
-            for ( SIZE_T i=0; i<m_constantPhysicsBodies.Num(); ++i)
+            for (int32 BodyIndex = 0; BodyIndex < ConstantPhysicsBodies.Num(); ++BodyIndex)
             {
-                if ( m_constantPhysicsBodies[i]==pPhysicsBody
-                     ||
-                     *m_constantPhysicsBodies[i]==*pPhysicsBody
-                     )
+                if (ConstantPhysicsBodies[BodyIndex] == PhysicsBodyPtr || *ConstantPhysicsBodies[BodyIndex] == *PhysicsBodyPtr)
                 {
-                    return (OP::ADDRESS)i;
+                    return (OP::ADDRESS)BodyIndex;
                 }
             }
 
-            OP::ADDRESS index = OP::ADDRESS( m_constantPhysicsBodies.Num() );
-            m_constantPhysicsBodies.Add( pPhysicsBody );
-            return index;
+            OP::ADDRESS NewPhysicsBodyIndex = OP::ADDRESS(ConstantPhysicsBodies.Num());
+            ConstantPhysicsBodies.Add(PhysicsBodyPtr);
+            return NewPhysicsBodyIndex;
         }
 
         OP::ADDRESS AddConstant( const FString& str )
@@ -643,9 +647,17 @@ namespace mu
 			}
 		}
 
-        void GetConstant(int32 ConstantIndex, MeshPtrConst& res) const
+        void GetConstant(int32 ConstantIndex, Ptr<const Mesh>& OutResult) const
         {
-			res = ConstantMeshes[ConstantIndex].Value;
+			// TODO: This should probably accept load mesh options and make a composite of the requested 
+			// data. Currently only used in Mutable Debugger code. Other code paths that need to get the 
+			// constant manage the merge themselves accessing directly to the program.
+			MUTABLE_CPUPROFILER_SCOPE(ComposeConstantMesh);
+
+			const FMeshRange MeshRange = ConstantMeshes[ConstantIndex];
+	
+			// Only the geometry part.
+			OutResult = ConstantMeshData[ConstantMeshIndices[MeshRange.FirstIndex + 0]].Value;
 		}
 
 		void GetExtensionDataConstant(int32 ConstantIndex, ExtensionDataPtrConst& Result) const
diff --git a/Engine/Plugins/Experimental/Mutable/Source/MutableRuntime/Internal/MuR/OpMeshMerge.h b/Engine/Plugins/Experimental/Mutable/Source/MutableRuntime/Internal/MuR/OpMeshMerge.h
--- a/Engine/Plugins/Experimental/Mutable/Source/MutableRuntime/Internal/MuR/OpMeshMerge.h
+++ b/Engine/Plugins/Experimental/Mutable/Source/MutableRuntime/Internal/MuR/OpMeshMerge.h
@@ -39,7 +39,17 @@ namespace mu
 
 			const int32 FirstCount = pFirst->GetIndexBuffers().GetElementCount();
 			const int32 SecondCount = pSecond->GetIndexBuffers().GetElementCount();
-			Result->GetIndexBuffers().SetElementCount(FirstCount + SecondCount);
+
+			const bool bIsMergingDescriptor = pFirst->IndexBuffers.IsDescriptor() || pSecond->IndexBuffers.IsDescriptor();
+
+			if (bIsMergingDescriptor)
+			{
+				Result->IndexBuffers.ElementCount = FirstCount + SecondCount;
+			}
+			else
+			{
+				Result->GetIndexBuffers().SetElementCount(FirstCount + SecondCount);
+			}
 
 			check(pFirst->GetIndexBuffers().GetBufferCount() <= 1);
 			check(pSecond->GetIndexBuffers().GetBufferCount() <= 1);
@@ -64,13 +74,13 @@ namespace mu
 				check(FirstIndexBuffer.ElementSize == SecondIndexBuffer.ElementSize);
 
 				// We need to know the total number of vertices in case we need to adjust the index buffer format.
-				const uint64 totalVertexCount = pFirst->GetVertexBuffers().GetElementCount() + pSecond->GetVertexBuffers().GetElementCount();
-				const uint64 maxValueBits = GetMeshFormatData(pFirst->GetIndexBuffers().Buffers[0].Channels[0].Format).MaxValueBits;
-				const uint64 maxSupportedVertices = uint64(1) << maxValueBits;
+				const uint64 TotalVertexCount = pFirst->GetVertexBuffers().GetElementCount() + pSecond->GetVertexBuffers().GetElementCount();
+				const uint64 MaxValueBits = GetMeshFormatData(pFirst->GetIndexBuffers().Buffers[0].Channels[0].Format).MaxValueBits;
+				const uint64 MaxSupportedVertices = uint64(1) << MaxValueBits;
 				
-				if (totalVertexCount > maxSupportedVertices)
+				if (TotalVertexCount > MaxSupportedVertices)
 				{
-					IndexBufferFormat = totalVertexCount > MAX_uint16 ? MBF_UINT32 : MBF_UINT16;
+					IndexBufferFormat = TotalVertexCount > MAX_uint16 ? MBF_UINT32 : MBF_UINT16;
 				}
 
 			}
@@ -99,74 +109,83 @@ namespace mu
 				ResultIndexBuffer.ElementSize = SecondIndexBuffer.ElementSize;
 			}
 
-			ResultIndexBuffer.Data.SetNum(ResultIndexBuffer.ElementSize * (FirstCount + SecondCount));
-
 			check(ResultIndexBuffer.Channels.Num() == 1);
 			check(ResultIndexBuffer.Channels[0].Semantic == MBS_VERTEXINDEX);
 
-			if (!ResultIndexBuffer.Data.IsEmpty())
+			if (!bIsMergingDescriptor)
 			{
-				if (FirstCount)
+				ResultIndexBuffer.Data.SetNum(ResultIndexBuffer.ElementSize * (FirstCount + SecondCount));
+				
+				if (!ResultIndexBuffer.Data.IsEmpty())
 				{
-					if (IndexBufferFormat == MBF_NONE
-						|| IndexBufferFormat == FirstIndexBuffer.Channels[0].Format)
-					{
-						FMemory::Memcpy(&ResultIndexBuffer.Data[0],
-							&FirstIndexBuffer.Data[0],
-							FirstIndexBuffer.ElementSize * FirstCount);
-					}
-					else
+					if (FirstCount)
 					{
-						// Conversion required
-						const uint8_t* pSource = &FirstIndexBuffer.Data[0];
-						uint8_t* pDest = &ResultIndexBuffer.Data[0];
-						switch (IndexBufferFormat)
+						if (IndexBufferFormat == MBF_NONE
+							|| IndexBufferFormat == FirstIndexBuffer.Channels[0].Format)
 						{
-						case MBF_UINT32:
+							FMemory::Memcpy(&ResultIndexBuffer.Data[0],
+								&FirstIndexBuffer.Data[0],
+								FirstIndexBuffer.ElementSize * FirstCount);
+						}
+						else
 						{
-							switch (FirstIndexBuffer.Channels[0].Format)
+							// Conversion required
+							const uint8* pSource = &FirstIndexBuffer.Data[0];
+							uint8* pDest = &ResultIndexBuffer.Data[0];
+							switch (IndexBufferFormat)
 							{
-							case MBF_UINT16:
+							case MBF_UINT32:
 							{
-								for (int v = 0; v < FirstCount; ++v)
+								switch (FirstIndexBuffer.Channels[0].Format)
 								{
-									*(uint32_t*)pDest = *(const uint16*)pSource;
-									pSource += FirstIndexBuffer.ElementSize;
-									pDest += ResultIndexBuffer.ElementSize;
+								case MBF_UINT16:
+								{
+									for (int32 v = 0; v < FirstCount; ++v)
+									{
+										*(uint32*)pDest = *(const uint16*)pSource;
+										pSource += FirstIndexBuffer.ElementSize;
+										pDest += ResultIndexBuffer.ElementSize;
+									}
+									break;
 								}
-								break;
-							}
 
-							case MBF_UINT8:
-							{
-								for (int v = 0; v < FirstCount; ++v)
+								case MBF_UINT8:
 								{
-									*(uint32_t*)pDest = *(const uint8_t*)pSource;
-									pSource += FirstIndexBuffer.ElementSize;
-									pDest += ResultIndexBuffer.ElementSize;
+									for (int32 v = 0; v < FirstCount; ++v)
+									{
+										*(uint32*)pDest = *(const uint8*)pSource;
+										pSource += FirstIndexBuffer.ElementSize;
+										pDest += ResultIndexBuffer.ElementSize;
+									}
+									break;
 								}
-								break;
-							}
 
-							default:
-								checkf(false, TEXT("Format not supported."));
+								default:
+									checkf(false, TEXT("Format not supported."));
+									break;
+								}
 								break;
 							}
-							break;
-						}
 
-						case MBF_UINT16:
-						{
-							switch (FirstIndexBuffer.Channels[0].Format)
+							case MBF_UINT16:
 							{
+								switch (FirstIndexBuffer.Channels[0].Format)
+								{
 
-							case MBF_UINT8:
-							{
-								for (int v = 0; v < FirstCount; ++v)
+								case MBF_UINT8:
 								{
-									*(uint16*)pDest = *(const uint8_t*)pSource;
-									pSource += FirstIndexBuffer.ElementSize;
-									pDest += ResultIndexBuffer.ElementSize;
+									for (int32 v = 0; v < FirstCount; ++v)
+									{
+										*(uint16*)pDest = *(const uint8*)pSource;
+										pSource += FirstIndexBuffer.ElementSize;
+										pDest += ResultIndexBuffer.ElementSize;
+									}
+									break;
+								}
+
+								default:
+									checkf(false, TEXT("Format not supported."));
+									break;
 								}
 								break;
 							}
@@ -175,93 +194,43 @@ namespace mu
 								checkf(false, TEXT("Format not supported."));
 								break;
 							}
-							break;
-						}
-
-						default:
-							checkf(false, TEXT("Format not supported."));
-							break;
 						}
 					}
-				}
-
-				if (SecondCount)
-				{
-					const uint8_t* pSource = &SecondIndexBuffer.Data[0];
-					uint8_t* pDest = &ResultIndexBuffer.Data[ResultIndexBuffer.ElementSize * FirstCount];
 
-					uint32_t firstVertexCount = pFirst->GetVertexBuffers().GetElementCount();
-
-					if (IndexBufferFormat == MBF_NONE
-						|| IndexBufferFormat == SecondIndexBuffer.Channels[0].Format)
+					if (SecondCount)
 					{
-						switch (SecondIndexBuffer.Channels[0].Format)
-						{
-						case MBF_INT32:
-						case MBF_UINT32:
-						case MBF_NINT32:
-						case MBF_NUINT32:
-						{
-							for (int v = 0; v < SecondCount; ++v)
-							{
-								*(uint32_t*)pDest = firstVertexCount + *(const uint32_t*)pSource;
-								pSource += SecondIndexBuffer.ElementSize;
-								pDest += ResultIndexBuffer.ElementSize;
-							}
-							break;
-						}
+						const uint8* pSource = &SecondIndexBuffer.Data[0];
+						uint8* pDest = &ResultIndexBuffer.Data[ResultIndexBuffer.ElementSize * FirstCount];
 
-						case MBF_INT16:
-						case MBF_UINT16:
-						case MBF_NINT16:
-						case MBF_NUINT16:
-						{
-							for (int v = 0; v < SecondCount; ++v)
-							{
-								*(uint16*)pDest = uint16(firstVertexCount) + *(const uint16*)pSource;
-								pSource += SecondIndexBuffer.ElementSize;
-								pDest += ResultIndexBuffer.ElementSize;
-							}
-							break;
-						}
+						uint32 firstVertexCount = pFirst->GetVertexBuffers().GetElementCount();
 
-						case MBF_INT8:
-						case MBF_UINT8:
-						case MBF_NINT8:
-						case MBF_NUINT8:
+						if (IndexBufferFormat == MBF_NONE
+							|| IndexBufferFormat == SecondIndexBuffer.Channels[0].Format)
 						{
-							for (int v = 0; v < SecondCount; ++v)
+							switch (SecondIndexBuffer.Channels[0].Format)
+							{
+							case MBF_INT32:
+							case MBF_UINT32:
+							case MBF_NINT32:
+							case MBF_NUINT32:
 							{
-								*(uint8_t*)pDest = uint8_t(firstVertexCount) + *(const uint8_t*)pSource;
-								pSource += SecondIndexBuffer.ElementSize;
-								pDest += ResultIndexBuffer.ElementSize;
+								for (int32 v = 0; v < SecondCount; ++v)
+								{
+									*(uint32*)pDest = firstVertexCount + *(const uint32*)pSource;
+									pSource += SecondIndexBuffer.ElementSize;
+									pDest += ResultIndexBuffer.ElementSize;
+								}
+								break;
 							}
-							break;
-						}
-
-						default:
-							checkf(false, TEXT("Format not supported."));
-							break;
-						}
-					}
-					else
-					{
-						// Format conversion required
-						switch (IndexBufferFormat)
-						{
 
-						case MBF_UINT32:
-						{
-							switch (SecondIndexBuffer.Channels[0].Format)
-							{
 							case MBF_INT16:
 							case MBF_UINT16:
 							case MBF_NINT16:
 							case MBF_NUINT16:
 							{
-								for (int v = 0; v < SecondCount; ++v)
+								for (int32 v = 0; v < SecondCount; ++v)
 								{
-									*(uint32_t*)pDest = uint32_t(firstVertexCount) + *(const uint16*)pSource;
+									*(uint16*)pDest = uint16(firstVertexCount) + *(const uint16*)pSource;
 									pSource += SecondIndexBuffer.ElementSize;
 									pDest += ResultIndexBuffer.ElementSize;
 								}
@@ -273,9 +242,9 @@ namespace mu
 							case MBF_NINT8:
 							case MBF_NUINT8:
 							{
-								for (int v = 0; v < SecondCount; ++v)
+								for (int32 v = 0; v < SecondCount; ++v)
 								{
-									*(uint32_t*)pDest = uint32_t(firstVertexCount) + *(const uint8_t*)pSource;
+									*(uint8*)pDest = uint8(firstVertexCount) + *(const uint8*)pSource;
 									pSource += SecondIndexBuffer.ElementSize;
 									pDest += ResultIndexBuffer.ElementSize;
 								}
@@ -286,40 +255,84 @@ namespace mu
 								checkf(false, TEXT("Format not supported."));
 								break;
 							}
-
-							break;
 						}
-
-						case MBF_UINT16:
+						else 
 						{
-							switch (SecondIndexBuffer.Channels[0].Format)
+							// Format conversion required
+							switch (IndexBufferFormat)
 							{
-							case MBF_INT8:
-							case MBF_UINT8:
-							case MBF_NINT8:
-							case MBF_NUINT8:
+
+							case MBF_UINT32:
 							{
-								for (int v = 0; v < SecondCount; ++v)
+								switch (SecondIndexBuffer.Channels[0].Format)
 								{
-									*(uint16*)pDest = uint16(firstVertexCount) + *(const uint8_t*)pSource;
-									pSource += SecondIndexBuffer.ElementSize;
-									pDest += ResultIndexBuffer.ElementSize;
+								case MBF_INT16:
+								case MBF_UINT16:
+								case MBF_NINT16:
+								case MBF_NUINT16:
+								{
+									for (int32 v = 0; v < SecondCount; ++v)
+									{
+										*(uint32*)pDest = uint32(firstVertexCount) + *(const uint16*)pSource;
+										pSource += SecondIndexBuffer.ElementSize;
+										pDest += ResultIndexBuffer.ElementSize;
+									}
+									break;
 								}
+
+								case MBF_INT8:
+								case MBF_UINT8:
+								case MBF_NINT8:
+								case MBF_NUINT8:
+								{
+									for (int32 v = 0; v < SecondCount; ++v)
+									{
+										*(uint32*)pDest = uint32(firstVertexCount) + *(const uint8*)pSource;
+										pSource += SecondIndexBuffer.ElementSize;
+										pDest += ResultIndexBuffer.ElementSize;
+									}
+									break;
+								}
+
+								default:
+									checkf(false, TEXT("Format not supported."));
+									break;
+								}
+
 								break;
 							}
 
-							default:
-								checkf(false, TEXT("Format not supported."));
+							case MBF_UINT16:
+							{
+								switch (SecondIndexBuffer.Channels[0].Format)
+								{
+								case MBF_INT8:
+								case MBF_UINT8:
+								case MBF_NINT8:
+								case MBF_NUINT8:
+								{
+									for (int32 v = 0; v < SecondCount; ++v)
+									{
+										*(uint16*)pDest = uint16(firstVertexCount) + *(const uint8*)pSource;
+										pSource += SecondIndexBuffer.ElementSize;
+										pDest += ResultIndexBuffer.ElementSize;
+									}
+									break;
+								}
+
+								default:
+									checkf(false, TEXT("Format not supported."));
+									break;
+								}
+
 								break;
 							}
 
-							break;
-						}
-
-						default:
-							checkf(false, TEXT("Format not supported."));
-							break;
+							default:
+								checkf(false, TEXT("Format not supported."));
+								break;
 
+							}
 						}
 					}
 				}
@@ -537,6 +550,7 @@ namespace mu
 
 		// Pose
 		//---------------------------
+		if (Result->GetSkeleton())
 		{
 			MUTABLE_CPUPROFILER_SCOPE(Pose);
 
@@ -894,7 +908,7 @@ namespace mu
 								int32 offset = 0;
 								for (int32 AuxChannelIndex = 0; AuxChannelIndex < ChannelsCount; ++AuxChannelIndex)
 								{
-									result.Channels[AuxChannelIndex].Offset = (uint8_t)offset;
+									result.Channels[AuxChannelIndex].Offset = (uint8)offset;
 									offset += result.Channels[AuxChannelIndex].ComponentCount
 										*
 										GetMeshFormatData(result.Channels[AuxChannelIndex].Format).SizeInBytes;
@@ -914,11 +928,18 @@ namespace mu
 			}
 
 
-			// Step 2: Fill the result buffers
-			//-----------------------------------------------------------------------
+			// Step 2: Fill the result buffers. In case we are merging descriptors only update the element count.
+			//--------------------------------------------------------------------------------------------------
+			if (pFirst->VertexBuffers.IsDescriptor() || pSecond->VertexBuffers.IsDescriptor())
+			{
+				Result->VertexBuffers.ElementCount = FirstVertexCount + SecondVertexCount;
+			}
+			else
 			{
 				MUTABLE_CPUPROFILER_SCOPE(ResultFill);
 
+				check(!pFirst->VertexBuffers.IsDescriptor() && !pSecond->VertexBuffers.IsDescriptor());
+				
 				// We have the final result vertex buffers structure: allocate the memory for them.
 				Result->VertexBuffers.SetElementCount(FirstVertexCount + SecondVertexCount);
 
@@ -1129,12 +1150,12 @@ namespace mu
 		}
 
         // Fix bone parent indices of the bones added from pOther
-        for ( int b=initialBones;b<pBase->GetBoneCount(); ++b)
+        for (int32 b = initialBones; b<pBase->GetBoneCount(); ++b)
         {
-            int16_t sourceIndex = pBase->BoneParents[b];
+            int16 sourceIndex = pBase->BoneParents[b];
             if (sourceIndex>=0)
             {
-                pBase->BoneParents[b] = (int16_t)otherToResult[ sourceIndex ];
+                pBase->BoneParents[b] = (int16)otherToResult[ sourceIndex ];
             }
         }
     }
diff --git a/Engine/Plugins/Experimental/Mutable/Source/MutableRuntime/Internal/MuR/Operations.h b/Engine/Plugins/Experimental/Mutable/Source/MutableRuntime/Internal/MuR/Operations.h
--- a/Engine/Plugins/Experimental/Mutable/Source/MutableRuntime/Internal/MuR/Operations.h
+++ b/Engine/Plugins/Experimental/Mutable/Source/MutableRuntime/Internal/MuR/Operations.h
@@ -419,15 +419,15 @@ namespace mu
         struct MeshConstantArgs
         {
             // Index of the mesh in the mesh constant array
-            ADDRESS value;
+            ADDRESS Value;
 
             // If not negative, index of the skeleton to set to the mesh from the skeleton
             // constant array.
-            int32 skeleton;
+            int32 Skeleton;
 
             // If not negative, index of the physics body to set to the mesh from the physics body
             // constant array.
-			int32 physicsBody;
+			int32 PhysicsBody;
         };
 
         struct ParameterArgs
diff --git a/Engine/Plugins/Experimental/Mutable/Source/MutableRuntime/Internal/MuR/SystemPrivate.h b/Engine/Plugins/Experimental/Mutable/Source/MutableRuntime/Internal/MuR/SystemPrivate.h
--- a/Engine/Plugins/Experimental/Mutable/Source/MutableRuntime/Internal/MuR/SystemPrivate.h
+++ b/Engine/Plugins/Experimental/Mutable/Source/MutableRuntime/Internal/MuR/SystemPrivate.h
@@ -52,8 +52,19 @@ namespace mu
 #endif
 	}
 
-	constexpr uint64 AllParametersMask = TNumericLimits<uint64> ::Max();
+	constexpr uint64 AllParametersMask = TNumericLimits<uint64>::Max();
 
+	enum class EMeshExecutionOptions : uint8
+	{
+		None               = 0,
+		LoadGeometryData   = 1 << 0,
+		LoadPoseData       = 1 << 1,
+		LoadComponentsData = 1 << 2,
+
+		AllFlags 	       = 0xFF,
+	};
+
+	ENUM_CLASS_FLAGS(EMeshExecutionOptions);
 
 	/** ExecutinIndex stores the location inside all ranges of the execution of a specific
 	* operation. The first integer on each pair is the dimension/range index in the program
@@ -146,25 +157,23 @@ namespace mu
 		//! Internal stage of the operation.
 		//! Stage 0 is usually scheduling of children, and 1 is execution. Some instructions
 		//! may have more steges to schedule children that are optional for execution, etc.
-		uint8 Stage : 7;
+		uint8 Stage : 6;
 
 		//! Type of calculation we are requesting for this operation.
 		enum class EType : uint8
 		{
-			//! Execute the operation to calculate the full result
+			// Execute the operation to calculate the full result
 			Full,
 
-			//! Execute the operation to obtain the descriptor of an image.
-			ImageDesc
-		};
-		EType Type : 1;
-	};
+			// Execute the operation to obtain the descriptor of an image.
+			ImageDesc,
 
-	inline uint32 GetTypeHash(const FScheduledOp& Op)
-	{
-		return HashCombine(::GetTypeHash(Op.At), HashCombine(::GetTypeHash(Op.Stage), ::GetTypeHash(Op.ExecutionIndex)));
-	}
+			// Execute the operation to generate a mesh components. Skeletons, referenced ids...
+			MeshComp,
+		};
 
+		EType Type : 2;
+	};
 
 	/** A cache address is the operation plus the context of execution(iteration indices, etc...). */
 	struct FCacheAddress
@@ -175,192 +184,175 @@ namespace mu
 		uint8 ExecutionOptions = 0;
 		FScheduledOp::EType Type = FScheduledOp::EType::Full;
 
-		FCacheAddress() {}
+		FCacheAddress() 
+		{
+		}
 
-		FCacheAddress(OP::ADDRESS InAt, uint16 InExecutionIndex, uint8 InExecutionOptions)
+		FCacheAddress(OP::ADDRESS InAt, uint16 InExecutionIndex, uint8 InExecutionOptions, FScheduledOp::EType InType = FScheduledOp::EType::Full)
+			: At(InAt), ExecutionIndex(InExecutionIndex), ExecutionOptions(InExecutionOptions), Type(InType)
 		{
-			At = InAt;
-			ExecutionIndex = InExecutionIndex;
-			ExecutionOptions = InExecutionOptions;
 		}
 
 		FCacheAddress(OP::ADDRESS InAt, const FScheduledOp& Item)
+			: At(InAt), ExecutionIndex(Item.ExecutionIndex), ExecutionOptions(Item.ExecutionOptions), Type(Item.Type)
 		{
-			At = InAt;
-			ExecutionIndex = Item.ExecutionIndex;
-			ExecutionOptions = Item.ExecutionOptions;
-			Type = Item.Type;
 		}
 
 		FCacheAddress(const FScheduledOp& Item)
+			: At(Item.At), ExecutionIndex(Item.ExecutionIndex), ExecutionOptions(Item.ExecutionOptions), Type(Item.Type)
 		{
-			At = Item.At;
-			ExecutionIndex = Item.ExecutionIndex;
-			ExecutionOptions = Item.ExecutionOptions;
-			Type = Item.Type;
 		}
 	};
 
+	static_assert(std::has_unique_object_representations_v<FCacheAddress>);
+	static_assert(sizeof(FCacheAddress) == sizeof(uint64));
 
-	inline uint32 GetTypeHash(const FCacheAddress& a)
+	inline bool operator==(const FCacheAddress& A, const FCacheAddress& B)
 	{
-		return HashCombine(::GetTypeHash(a.At), a.ExecutionIndex);
+		return BitCast<uint64>(A) == BitCast<uint64>(B); 
 	}
 
+	inline uint32 GetTypeHash(const FCacheAddress& CacheAddress)
+	{
+		return ::GetTypeHash(BitCast<uint64>(CacheAddress));
+	}
 
-	inline uint32 GetTypeHash(const Ptr<const Resource>& V)
+	inline uint32 GetTypeHash(const Ptr<const Resource>& Value)
 	{
-		return ::GetTypeHash(V.get());
+		return ::GetTypeHash(Value.get());
 	}
 
 
 	//! Container that stores data per executable code operation (indexed by address and execution
 	//! index).
-	template<class DATA>
+	template<class DataType>
 	class CodeContainer
 	{
 		using MemoryCounter = MemoryCounters::FInternalMemoryCounter;
-		using ArrayDataContainerType = TArray<DATA, FDefaultMemoryTrackingAllocator<MemoryCounter>>;
-		using MapDataContainerType = TMap<FCacheAddress, DATA, FDefaultMemoryTrackingSetAllocator<MemoryCounter>>;
+		using ArrayDataContainerType = TArray<DataType, FDefaultMemoryTrackingAllocator<MemoryCounter>>;
+		using MapDataContainerType = TMap<FCacheAddress, DataType, FDefaultMemoryTrackingSetAllocator<MemoryCounter>>;
 
 	public:
-		void resize(size_t s)
+		void Resize(int32 NewSize)
 		{
-			m_index0.SetNumZeroed(s);
+			Index0.SetNumZeroed(NewSize);
 		}
 
-		uint32 size_code() const
+		uint32 CodeSize() const
 		{
-			return uint32(m_index0.Num());
+			return uint32(Index0.Num());
 		}
 
-		void clear()
+		void Clear()
 		{
-			m_index0.Empty();
-			m_otherIndex.Empty();
+			Index0.Empty();
+			OtherIndex.Empty();
 		}
 
-		inline void erase(const FCacheAddress& at)
+		inline void Erase(const FCacheAddress& At)
 		{
-			if (at.ExecutionIndex == 0 && at.ExecutionOptions == 0)
+			if (At.Type == FScheduledOp::EType::Full && At.ExecutionIndex == 0 && At.ExecutionOptions == 0)
 			{
-				m_index0[at.At] = nullptr;
+				Index0[At.At] = nullptr;
 			}
 			else
 			{
-				m_otherIndex.Remove(at);
+				OtherIndex.Remove(At);
 			}
 		}
 
-		inline DATA get(const FCacheAddress& at) const
+		inline DataType Get(const FCacheAddress& At) const
 		{
-			if (at.ExecutionIndex == 0 && at.ExecutionOptions == 0)
+			if (At.Type == FScheduledOp::EType::Full && At.ExecutionIndex == 0 && At.ExecutionOptions == 0)
 			{
-				if (at.At < uint32(m_index0.Num()))
+				if (At.At < OP::ADDRESS(Index0.Num()))
 				{
-					return m_index0[at.At];
-				}
-				else
-				{
-					return 0;
+					return Index0[At.At];
 				}
 			}
 			else
 			{
-				const DATA* it = m_otherIndex.Find(at);
-				if (it)
+				const DataType* Found = OtherIndex.Find(At);
+				if (Found)
 				{
-					return *it;
+					return *Found;
 				}
 			}
-			return 0;
+
+			return DataType{};
 		}
 
-		inline DATA* get_ptr(const FCacheAddress& at)
+		inline DataType* GetPtr(const FCacheAddress& At)
 		{
-			if (at.ExecutionIndex == 0 && at.ExecutionOptions == 0)
+			if (At.Type == FScheduledOp::EType::Full && At.ExecutionIndex == 0 && At.ExecutionOptions == 0)
 			{
-				if (at.At < uint32(m_index0.Num()))
-				{
-					return &m_index0[at.At];
-				}
-				else
+				if (At.At < OP::ADDRESS(Index0.Num()))
 				{
-					return nullptr;
+					return &Index0[At.At];
 				}
 			}
 			else
 			{
-				DATA* it = m_otherIndex.Find(at);
-				if (it)
-				{
-					return it;
-				}
+				return &OtherIndex.FindOrAdd(At);
 			}
+
 			return nullptr;
 		}
 
-		inline const DATA* get_ptr(const FCacheAddress& at) const
+		inline const DataType* GetPtr(const FCacheAddress& At) const
 		{
-			if (at.ExecutionIndex == 0 && at.ExecutionOptions == 0)
+			if (At.Type == FScheduledOp::EType::Full && At.ExecutionIndex == 0 && At.ExecutionOptions == 0)
 			{
-				if (at.At < uint32(m_index0.Num()))
-				{
-					return &m_index0[at.At];
-				}
-				else
+				if (At.At < uint32(Index0.Num()))
 				{
-					return nullptr;
+					return &Index0[At.At];
 				}
 			}
 			else
 			{
-				const DATA* it = m_otherIndex.Find(at);
-				if (it)
-				{
-					return it;
-				}
+				return OtherIndex.Find(At);
 			}
+
 			return nullptr;
 		}
 
-		inline DATA& operator[](const FCacheAddress& at)
+		inline DataType& operator[](const FCacheAddress& At)
 		{
-			if (at.ExecutionIndex == 0 && at.ExecutionOptions == 0)
+			if (At.Type == FScheduledOp::EType::Full && At.ExecutionIndex == 0 && At.ExecutionOptions == 0)
 			{
-				return m_index0[at.At];
+				return Index0[At.At];
 			}
 			else
 			{
-				return m_otherIndex.FindOrAdd(at);
+				return OtherIndex.FindOrAdd(At);
 			}
 		}
 
-		inline const DATA& operator[](const FCacheAddress& at) const
+		inline const DataType& operator[](const FCacheAddress& At) const
 		{
-			if (at.ExecutionIndex == 0 && at.ExecutionOptions == 0)
+			if (At.Type == FScheduledOp::EType::Full && At.ExecutionIndex == 0 && At.ExecutionOptions == 0)
 			{
-				return m_index0[at.At];
+				return Index0[At.At];
 			}
 			else
 			{
-				return m_otherIndex[at];
+				return OtherIndex[At];
 			}
 		}
 
-		struct iterator
+		struct TIterator
 		{
 		private:
-			friend class CodeContainer<DATA>;
+			friend class CodeContainer<DataType>;
 
-			const CodeContainer<DATA>* container;
-			typename CodeContainer<DATA>::ArrayDataContainerType::TIterator it0;
-			typename CodeContainer<DATA>::MapDataContainerType::TIterator it1;
+			const CodeContainer<DataType>* Container;
+			typename CodeContainer<DataType>::ArrayDataContainerType::TIterator Iter0;
+			typename CodeContainer<DataType>::MapDataContainerType::TIterator Iter1;
 
-			iterator(CodeContainer<DATA>* InContainer)
-				: container(InContainer)
-				, it0(InContainer->m_index0.CreateIterator())
-				, it1(InContainer->m_otherIndex.CreateIterator())
+			TIterator(CodeContainer<DataType>* InContainer)
+				: Container(InContainer)
+				, Iter0(InContainer->Index0.CreateIterator())
+				, Iter1(InContainer->OtherIndex.CreateIterator())
 			{
 			}
 
@@ -368,81 +360,81 @@ namespace mu
 
 			inline void operator++(int)
 			{
-				if (it0)
+				if (Iter0)
 				{
-					++it0;
+					++Iter0;
 				}
 				else
 				{
-					++it1;
+					++Iter1;
 				}
 			}
 
-			iterator operator++()
+			TIterator operator++()
 			{
-				if (it0)
+				if (Iter0)
 				{
-					++it0;
+					++Iter0;
 				}
 				else
 				{
-					++it1;
+					++Iter1;
 				}
+
 				return *this;
 			}
 
-			inline bool operator!=(const iterator& o) const
+			inline bool operator != (const TIterator& Other) const
 			{
-				return it0 != o.it0 || it1 != o.it1;
+				return Iter0 != Other.Iter0 || Iter1 != Other.Iter1;
 			}
 
-			inline DATA& operator*()
+			inline DataType& operator*()
 			{
-				if (it0)
+				if (Iter0)
 				{
-					return *it0;
+					return *Iter0;
 				}
 				else
 				{
-					return it1->Value;
+					return Iter1->Value;
 				}
 			}
 
-			inline FCacheAddress get_address() const
+			inline FCacheAddress GetAddress() const
 			{
-				if (it0)
+				if (Iter0)
 				{
-					return { uint32_t(it0.GetIndex()), 0, 0};
+					return {uint32(Iter0.GetIndex()), 0, 0};
 				}
 				else
 				{
-					return it1->Key;
+					return Iter1->Key;
 				}
 			}
 
 			inline bool IsValid() const
 			{
-				return bool(it0) || bool(it1);
+				return bool(Iter0) || bool(Iter1);
 			}
 		};
 
-		inline iterator begin()
+		inline TIterator Begin()
 		{
-			iterator it(this);
-			return it;
+			TIterator Iter(this);
+			return Iter;
 		}
 
 		inline int32 GetAllocatedSize() const
 		{
-			return m_index0.GetAllocatedSize()
-				+ m_otherIndex.GetAllocatedSize();
+			return Index0.GetAllocatedSize() + OtherIndex.GetAllocatedSize();
 		}
 	private:
 		// For index 0
-		ArrayDataContainerType m_index0;
+		ArrayDataContainerType Index0;
 
-		// For index>0
-		MapDataContainerType m_otherIndex;
+		// For index > 0
+		MapDataContainerType OtherIndex;
 	};
 
 
@@ -453,10 +445,10 @@ namespace mu
 
 		using AllocType = FDefaultMemoryTrackingAllocator<MemoryCounters::FInternalMemoryCounter>;
 
-		template<class Type, class Alloc = AllocType >
+		template<class Type, class Alloc = AllocType>
 		using TMemoryTrackedArray = TArray<Type, Alloc>;
 		
-		TMemoryTrackedArray<ExecutionIndex, TInlineAllocator<4, AllocType>> m_usedRangeIndices;
+		TMemoryTrackedArray<ExecutionIndex, TInlineAllocator<4, AllocType>> UsedRangeIndices;
 
 		/** Runtime data for each program op. */
 		struct FOpExecutionData
@@ -465,16 +457,17 @@ namespace mu
 
 			/** Enabled if the op descriptor has been calculated. */
 			uint8 IsDescCacheValid : 1;
-			uint8 IsValueValid : 1;
-			uint8 IsCacheLocked : 1;
+			uint8 IsValueValid     : 1;
+			uint8 IsCacheLocked    : 1;
 
 			/** The operation DATA_TYPE. */
 			uint8 DataType;
 
-			/** The position in the type-specific array where the result data is stored. 
-			* 0 means index not valid.
-			* For small types like bool, int and float, the actual value is the index, instead of having their own array.
-			*/
+			/** 
+			 * The position in the type-specific array where the result data is stored. 
+			 * 0 means index not valid.
+			 * For small types like bool, int and float, the actual value is the index, instead of having their own array.
+			 */
 			union
 			{
 				int32 DataTypeIndex;
@@ -513,49 +506,49 @@ namespace mu
 		TMemoryTrackedArray<Ptr<const ExtensionData>> ExtensionDataResults;
 
 		/** */
-		inline const ExecutionIndex& GetRangeIndex(uint32_t i)
+		inline const ExecutionIndex& GetRangeIndex(uint32 Index)
 		{
 			// Make sure we have the default element.
-			if (m_usedRangeIndices.IsEmpty())
+			if (UsedRangeIndices.IsEmpty())
 			{
-				m_usedRangeIndices.Push(ExecutionIndex());
+				UsedRangeIndices.Push(ExecutionIndex());
 			}
 
-			check(i < uint32_t(m_usedRangeIndices.Num()));
-			return m_usedRangeIndices[i];
+			check(Index < uint32(UsedRangeIndices.Num()));
+			return UsedRangeIndices[Index];
 		}
 
 		//!
-		inline uint32 GetRangeIndexIndex(const ExecutionIndex& rangeIndex)
+		inline uint32 GetRangeIndexIndex(const ExecutionIndex& RangeIndex)
 		{
-			if (rangeIndex.IsEmpty())
+			if (RangeIndex.IsEmpty())
 			{
 				return 0;
 			}
 
 			// Make sure we have the default element.
-			if (m_usedRangeIndices.IsEmpty())
+			if (UsedRangeIndices.IsEmpty())
 			{
-				m_usedRangeIndices.Push(ExecutionIndex());
+				UsedRangeIndices.Push(ExecutionIndex());
 			}
 
 			// Look for or add the new element.
-			int32 ElemIndex = m_usedRangeIndices.Find(rangeIndex);
+			int32 ElemIndex = UsedRangeIndices.Find(RangeIndex);
 			if (ElemIndex != INDEX_NONE)
 			{
 				return ElemIndex;
 			}
 
-			m_usedRangeIndices.Push(rangeIndex);
-			return uint32_t(m_usedRangeIndices.Num()) - 1;
+			UsedRangeIndices.Push(RangeIndex);
+			return uint32(UsedRangeIndices.Num()) - 1;
 		}
 
 		void Init(uint32 Size)
 		{
 			// This clear would prevent live update cache reusal
-			//OpExecutionData.clear();
+			//OpExecutionData.Clear();
 
-			OpExecutionData.resize(Size);
+			OpExecutionData.Resize(Size);
 
 			if (ColorResults.IsEmpty())
 			{
@@ -605,21 +598,21 @@ namespace mu
 		}
 
 
-		bool IsValid(FCacheAddress at) const
+		bool IsValid(FCacheAddress At) const
 		{
-			if (at.At == 0 || at.At>=OpExecutionData.size_code() )
+			if ((At.At == 0) || (At.At >= OP::ADDRESS(OpExecutionData.CodeSize())))
 			{
 				return false;
 			}
 
-			const FOpExecutionData* Data = OpExecutionData.get_ptr(at);
+			const FOpExecutionData* Data = OpExecutionData.GetPtr(At);
 			if (!Data)
 			{
 				return false;
 			}
 
 			// Is it a desc data query?
-			if (at.Type == FScheduledOp::EType::ImageDesc)
+			if (At.Type == FScheduledOp::EType::ImageDesc)
 			{
 				return Data->IsDescCacheValid;
 			}
@@ -636,8 +629,8 @@ namespace mu
 			//MUTABLE_CPUPROFILER_SCOPE(CheckHitCountsCleared);
 
 			//int32 IncorrectCount = 0;
-			//CodeContainer<int>::iterator it = m_opHitCount.begin();
-			//for (; it.IsValid(); ++it)
+			//CodeContainer<int32>::TIterator Iter = OpHitCount.begin();
+			//for (; Iter.IsValid(); ++it)
 			//{
 			//	int32 Count = *it;
 			//	if (Count>0 && Count < UE_MUTABLE_CACHE_COUNT_LIMIT)
@@ -665,9 +658,9 @@ namespace mu
 		{
 			MUTABLE_CPUPROFILER_SCOPE(ProgramCacheClear);
 
-			uint32 CodeSize = OpExecutionData.size_code();
-			OpExecutionData.clear();
-			OpExecutionData.resize(CodeSize);
+			uint32 CodeSize = OpExecutionData.CodeSize();
+			OpExecutionData.Clear();
+			OpExecutionData.Resize(CodeSize);
 		}
 
 
@@ -675,69 +668,113 @@ namespace mu
 		{
 			MUTABLE_CPUPROFILER_SCOPE(ProgramDescCacheClear);
 
-			CodeContainer<FProgramCache::FOpExecutionData>::iterator it = OpExecutionData.begin();
-			for (; it.IsValid(); ++it)
+			CodeContainer<FProgramCache::FOpExecutionData>::TIterator Iter = OpExecutionData.Begin();
+			for (; Iter.IsValid(); ++Iter)
 			{
-				(*it).IsDescCacheValid = false;
+				FProgramCache::FOpExecutionData& Data = *Iter;
+				Data.IsDescCacheValid = false;
 			}
 		}
 
 
-		bool GetBool(FCacheAddress at)
+		bool GetBool(FCacheAddress At)
 		{
-			if (!at.At) return false;
-			const FOpExecutionData* Data = OpExecutionData.get_ptr(at);
-			if (!Data) return false;
+			if (!At.At)
+			{
+				return false;
+			}
 
-			check(Data->DataType==DATATYPE::DT_BOOL);
-			return Data->DataTypeIndex!=0;
+			const FOpExecutionData* Data = OpExecutionData.GetPtr(At);
+			if (!Data)
+			{
+				return false;
+			}
+
+			check(Data->DataType == DATATYPE::DT_BOOL);
+			return Data->DataTypeIndex != 0;
 		}
 
-		float GetScalar(FCacheAddress at)
+		float GetScalar(FCacheAddress At)
 		{
-			if (!at.At) return 0.0f;
-			const FOpExecutionData* Data = OpExecutionData.get_ptr(at);
-			if (!Data) return 0.0f;
+			if (!At.At)
+			{
+				return 0.0f;
+			}
+
+			const FOpExecutionData* Data = OpExecutionData.GetPtr(At);
+			if (!Data) 
+			{
+				return 0.0f;
+			}
 
 			check(Data->DataType == DATATYPE::DT_SCALAR);
 			return Data->ScalarResult;
 		}
 
-		int32 GetInt(FCacheAddress at)
+		int32 GetInt(FCacheAddress At)
 		{
-			if (!at.At) return 0;
-			const FOpExecutionData* Data = OpExecutionData.get_ptr(at);
-			if (!Data) return 0;
+			if (!At.At) 
+			{
+				return 0;
+			}
+
+			const FOpExecutionData* Data = OpExecutionData.GetPtr(At);
+			if (!Data) 
+			{
+				return 0;
+			}
 
 			check(Data->DataType == DATATYPE::DT_INT);
 			return Data->DataTypeIndex;
 		}
 
-		FVector4f GetColour(FCacheAddress at)
+		FVector4f GetColour(FCacheAddress At)
 		{
-			if (!at.At) return FVector4f();
-			const FOpExecutionData* Data = OpExecutionData.get_ptr(at);
-			if (!Data) return FVector4f();
+			if (!At.At)
+			{
+				return FVector4f();
+			}
+
+			const FOpExecutionData* Data = OpExecutionData.GetPtr(At);
+			if (!Data) 
+			{
+				return FVector4f();
+			}
 
 			check(Data->DataType == DATATYPE::DT_COLOUR);
 			return ColorResults[Data->DataTypeIndex];
 		}
 
-		FProjector GetProjector(FCacheAddress at)
+		FProjector GetProjector(FCacheAddress At)
 		{
-			if (!at.At) return FProjector();
-			const FOpExecutionData* Data = OpExecutionData.get_ptr(at);
-			if (!Data) return FProjector();
+			if (!At.At)
+			{
+				return FProjector();
+			}
+
+			const FOpExecutionData* Data = OpExecutionData.GetPtr(At);
+
+			if (!Data) 
+			{
+				return FProjector();
+			}
 
 			check(Data->DataType == DATATYPE::DT_PROJECTOR);
 			return ProjectorResults[Data->DataTypeIndex];
 		}
 
-		Ptr<const Instance> GetInstance(FCacheAddress at)
+		Ptr<const Instance> GetInstance(FCacheAddress At)
 		{
-			if (!at.At) return nullptr;
-			FOpExecutionData* Data = OpExecutionData.get_ptr(at);
-			if (!Data) return nullptr;
+			if (!At.At) 
+			{
+				return nullptr;
+			}
+
+			FOpExecutionData* Data = OpExecutionData.GetPtr(At);
+			if (!Data)
+			{
+				return nullptr;
+			}
 
 			check(Data->DataType == DATATYPE::DT_INSTANCE);
 			Ptr<const Instance> Result = InstanceResults[Data->DataTypeIndex];
@@ -755,14 +792,25 @@ namespace mu
 		}
 
 
-		Ptr<const Image> GetImage(FCacheAddress at, bool& bIsLastReference)
+		Ptr<const Image> GetImage(FCacheAddress At, bool& bIsLastReference)
 		{
 			bIsLastReference = false;
 
-			if (!at.At) return nullptr;
-			if (at.At >= OpExecutionData.size_code()) return nullptr;
-			FOpExecutionData* Data = OpExecutionData.get_ptr(at);
-			if (!Data) return nullptr;
+			if (!At.At) 
+			{
+				return nullptr;
+			}
+
+			if (At.At >= OpExecutionData.CodeSize())
+			{
+				return nullptr;
+			}
+
+			FOpExecutionData* Data = OpExecutionData.GetPtr(At);
+			if (!Data) 
+			{
+				return nullptr;
+			}
 
 			check(Data->DataType == DATATYPE::DT_IMAGE);
 			Ptr<const Image> Result = ImageResults[Data->DataTypeIndex].Value;
@@ -781,14 +829,28 @@ namespace mu
 		}
 
 
-		Ptr<const Mesh> GetMesh(FCacheAddress at, bool& bIsLastReference)
+		Ptr<const Mesh> GetMesh(FCacheAddress At, bool& bIsLastReference)
 		{
 			bIsLastReference = false;
 
-			if (!at.At) return nullptr;
-			if (at.At >= OpExecutionData.size_code()) return nullptr;
-			FOpExecutionData* Data = OpExecutionData.get_ptr(at);
-			if (!Data) return nullptr;
+			if (!At.At)
+			{
+				return nullptr;
+			}
+
+			if (At.At >= OpExecutionData.CodeSize())
+			{
+				return nullptr;
+			}
+
+			FOpExecutionData* Data = OpExecutionData.GetPtr(At);
+			if (!Data)
+			{
+				return nullptr;
+			}
+
+			// Chache lock forcing is only done at base address, without options, index or type. See SetForceCached().
+			FOpExecutionData& ForceLockCacheData = OpExecutionData[FCacheAddress(At.At, 0, 0)];
 
 			check(Data->DataType == DATATYPE::DT_MESH);
 			Ptr<const Mesh> Result = MeshResults[Data->DataTypeIndex].Value;
@@ -797,7 +859,7 @@ namespace mu
 			check(Data->OpHitCount > 0);
 
 			--Data->OpHitCount;
-			if (Data->OpHitCount == 0 && !Data->IsCacheLocked)
+			if (Data->OpHitCount == 0 && !Data->IsCacheLocked && !ForceLockCacheData.IsCacheLocked)
 			{
 				SetUnused(*Data);
 				bIsLastReference = true;
@@ -807,80 +869,101 @@ namespace mu
 		}
 
 
-		Ptr<const Layout> GetLayout(FCacheAddress at)
+		Ptr<const Layout> GetLayout(FCacheAddress At)
 		{
-			if (!at.At) return nullptr;
-			FOpExecutionData* Data = OpExecutionData.get_ptr(at);
-			if (!Data) return nullptr;
+			if (!At.At) 
+			{
+				return nullptr;
+			}
+
+			FOpExecutionData* Data = OpExecutionData.GetPtr(At);
+			if (!Data) 
+			{
+				return nullptr;
+			}
 
 			check(Data->DataType == DATATYPE::DT_LAYOUT);
 			return LayoutResults[Data->DataTypeIndex];
 		}
 
 
-		Ptr<const String> GetString(FCacheAddress at)
+		Ptr<const String> GetString(FCacheAddress At)
 		{
-			if (!at.At) return nullptr;
-			FOpExecutionData* Data = OpExecutionData.get_ptr(at);
-			if (!Data) return nullptr;
+			if (!At.At) 
+			{
+				return nullptr;
+			}
+
+			FOpExecutionData* Data = OpExecutionData.GetPtr(At);
+			if (!Data) 
+			{
+				return nullptr;
+			}
 
 			check(Data->DataType == DATATYPE::DT_STRING);
 			return StringResults[Data->DataTypeIndex];
 		}
 
 
-		Ptr<const ExtensionData> GetExtensionData(FCacheAddress at)
+		Ptr<const ExtensionData> GetExtensionData(FCacheAddress At)
 		{
-			if (!at.At) return nullptr;
-			FOpExecutionData* Data = OpExecutionData.get_ptr(at);
-			if (!Data) return nullptr;
+			if (!At.At) 
+			{
+				return nullptr;
+			}
+
+			FOpExecutionData* Data = OpExecutionData.GetPtr(At);
+			if (!Data) 
+			{
+				return nullptr;
+			}
 
 			check(Data->DataType == DATATYPE::DT_EXTENSION_DATA);
 			return ExtensionDataResults[Data->DataTypeIndex];
 		}
 
-		void SetValidDesc(FCacheAddress at)
+		void SetValidDesc(FCacheAddress At)
 		{
-			check(at.Type == FScheduledOp::EType::ImageDesc);
-			check(at.At < OpExecutionData.size_code());
-			FOpExecutionData* Data = OpExecutionData.get_ptr(at);
+			check(At.Type == FScheduledOp::EType::ImageDesc);
+			check(At.At < OpExecutionData.CodeSize());
+			FOpExecutionData* Data = OpExecutionData.GetPtr(At);
 			Data->IsDescCacheValid = true;
 		}
 
-		void SetBool(FCacheAddress at, bool v)
+		void SetBool(FCacheAddress At, bool Value)
 		{
-			check(at.At < OpExecutionData.size_code());
-			FOpExecutionData* Data = OpExecutionData.get_ptr(at);
+			check(At.At < OpExecutionData.CodeSize());
+			FOpExecutionData* Data = OpExecutionData.GetPtr(At);
 			check(Data->DataType == DATATYPE::DT_BOOL || Data->DataType == DATATYPE::DT_NONE);
 			Data->DataType = DATATYPE::DT_BOOL;
-			Data->DataTypeIndex = v;
+			Data->DataTypeIndex = Value;
 			Data->IsValueValid = true;
 		}
 
-		void SetInt(FCacheAddress at, int32 v)
+		void SetInt(FCacheAddress At, int32 Value)
 		{
-			check(at.At < OpExecutionData.size_code());
-			FOpExecutionData* Data = OpExecutionData.get_ptr(at);
+			check(At.At < OpExecutionData.CodeSize());
+			FOpExecutionData* Data = OpExecutionData.GetPtr(At);
 			check(Data->DataType == DATATYPE::DT_INT || Data->DataType == DATATYPE::DT_NONE);
 			Data->DataType = DATATYPE::DT_INT;
-			Data->DataTypeIndex = v;
+			Data->DataTypeIndex = Value;
 			Data->IsValueValid = true;
 		}
 
-		void SetScalar(FCacheAddress at, float v)
+		void SetScalar(FCacheAddress At, float Value)
 		{
-			check(at.At < OpExecutionData.size_code());
-			FOpExecutionData* Data = OpExecutionData.get_ptr(at);
+			check(At.At < OpExecutionData.CodeSize());
+			FOpExecutionData* Data = OpExecutionData.GetPtr(At);
 			check(Data->DataType == DATATYPE::DT_SCALAR || Data->DataType == DATATYPE::DT_NONE);
 			Data->DataType = DATATYPE::DT_SCALAR;
-			Data->ScalarResult = v;
+			Data->ScalarResult = Value;
 			Data->IsValueValid = true;
 		}
 
-		void SetColour(FCacheAddress at, const FVector4f& v)
+		void SetColour(FCacheAddress At, const FVector4f& Value)
 		{
-			check(at.At < OpExecutionData.size_code());
-			FOpExecutionData* Data = OpExecutionData.get_ptr(at);
+			check(At.At < OpExecutionData.CodeSize());
+			FOpExecutionData* Data = OpExecutionData.GetPtr(At);
 			check(Data->DataType == DATATYPE::DT_COLOUR || Data->DataType == DATATYPE::DT_NONE);
 			Data->DataType = DATATYPE::DT_COLOUR;
 			Data->IsValueValid = true;
@@ -888,19 +971,19 @@ namespace mu
 			if (!Data->DataTypeIndex)
 			{
 				Data->DataTypeIndex = ColorResults.Num();
-				ColorResults.Add(v);
+				ColorResults.Add(Value);
 			}
 			else
 			{
-				ColorResults[Data->DataTypeIndex] = v;
+				ColorResults[Data->DataTypeIndex] = Value;
 			}
 			check(Data->DataTypeIndex != 0);
 		}
 
-		void SetProjector(FCacheAddress at, const FProjector& v)
+		void SetProjector(FCacheAddress At, const FProjector& Value)
 		{
-			check(at.At < OpExecutionData.size_code());
-			FOpExecutionData* Data = OpExecutionData.get_ptr(at);
+			check(At.At < OpExecutionData.CodeSize());
+			FOpExecutionData* Data = OpExecutionData.GetPtr(At);
 			check(Data->DataType == DATATYPE::DT_PROJECTOR || Data->DataType == DATATYPE::DT_NONE);
 			Data->DataType = DATATYPE::DT_PROJECTOR;
 			Data->IsValueValid = true;
@@ -908,19 +991,19 @@ namespace mu
 			if (!Data->DataTypeIndex)
 			{
 				Data->DataTypeIndex = ProjectorResults.Num();
-				ProjectorResults.Add(v);
+				ProjectorResults.Add(Value);
 			}
 			else
 			{
-				ProjectorResults[Data->DataTypeIndex] = v;
+				ProjectorResults[Data->DataTypeIndex] = Value;
 			}
 			check(Data->DataTypeIndex != 0);
 		}
 
-		void SetInstance(FCacheAddress at, Ptr<const Instance> v)
+		void SetInstance(FCacheAddress At, Ptr<const Instance> Value)
 		{
-			check(at.At < OpExecutionData.size_code());
-			FOpExecutionData* Data = OpExecutionData.get_ptr(at);
+			check(At.At < OpExecutionData.CodeSize());
+			FOpExecutionData* Data = OpExecutionData.GetPtr(At);
 			check(Data->DataType == DATATYPE::DT_INSTANCE || Data->DataType == DATATYPE::DT_NONE);
 			Data->DataType = DATATYPE::DT_INSTANCE;
 			Data->IsValueValid = true;
@@ -928,19 +1011,19 @@ namespace mu
 			if (!Data->DataTypeIndex)
 			{
 				Data->DataTypeIndex = InstanceResults.Num();
-				InstanceResults.Add(v);
+				InstanceResults.Add(Value);
 			}
 			else
 			{
-				InstanceResults[Data->DataTypeIndex] = v;
+				InstanceResults[Data->DataTypeIndex] = Value;
 			}
 			check(Data->DataTypeIndex != 0);
 		}
 
-		void SetExtensionData(FCacheAddress at, Ptr<const ExtensionData> v)
+		void SetExtensionData(FCacheAddress At, Ptr<const ExtensionData> Value)
 		{
-			check(at.At < OpExecutionData.size_code());
-			FOpExecutionData* Data = OpExecutionData.get_ptr(at);
+			check(At.At < OpExecutionData.CodeSize());
+			FOpExecutionData* Data = OpExecutionData.GetPtr(At);
 			check(Data->DataType == DATATYPE::DT_EXTENSION_DATA || Data->DataType == DATATYPE::DT_NONE);
 			Data->DataType = DATATYPE::DT_EXTENSION_DATA;
 			Data->IsValueValid = true;
@@ -948,19 +1031,19 @@ namespace mu
 			if (!Data->DataTypeIndex)
 			{
 				Data->DataTypeIndex = ExtensionDataResults.Num();
-				ExtensionDataResults.Add(v);
+				ExtensionDataResults.Add(Value);
 			}
 			else
 			{
-				ExtensionDataResults[Data->DataTypeIndex] = v;
+				ExtensionDataResults[Data->DataTypeIndex] = Value;
 			}
 			check(Data->DataTypeIndex != 0);
 		}
 
 		void SetImage(FCacheAddress At, Ptr<const Image> Value)
 		{
-			check(At.At < OpExecutionData.size_code());
-			FOpExecutionData* Data = OpExecutionData.get_ptr(At);
+			check(At.At < OpExecutionData.CodeSize());
+			FOpExecutionData* Data = OpExecutionData.GetPtr(At);
 			check(Data->DataType == DATATYPE::DT_IMAGE || Data->DataType == DATATYPE::DT_NONE);
 			Data->DataType = DATATYPE::DT_IMAGE;
 			Data->IsValueValid = true;
@@ -981,9 +1064,9 @@ namespace mu
 
 		void SetMesh(FCacheAddress At, Ptr<const Mesh> Value)
 		{
-			check(At.At < OpExecutionData.size_code());
+			check(At.At < OpExecutionData.CodeSize());
 
-			FOpExecutionData* Data = OpExecutionData.get_ptr(At);
+			FOpExecutionData* Data = OpExecutionData.GetPtr(At);
 
 			check(Data->DataType == DATATYPE::DT_MESH || Data->DataType == DATATYPE::DT_NONE);
 			Data->DataType = DATATYPE::DT_MESH;
@@ -1003,10 +1086,10 @@ namespace mu
 			mu::UpdateLLMStats();
 		}
 
-		void SetLayout(FCacheAddress at, Ptr<const Layout> v)
+		void SetLayout(FCacheAddress At, Ptr<const Layout> Value)
 		{
-			check(at.At < OpExecutionData.size_code());
-			FOpExecutionData* Data = OpExecutionData.get_ptr(at);
+			check(At.At < OpExecutionData.CodeSize());
+			FOpExecutionData* Data = OpExecutionData.GetPtr(At);
 			check(Data->DataType == DATATYPE::DT_LAYOUT || Data->DataType == DATATYPE::DT_NONE);
 			Data->DataType = DATATYPE::DT_LAYOUT;
 			Data->IsValueValid = true;
@@ -1014,21 +1097,21 @@ namespace mu
 			if (!Data->DataTypeIndex)
 			{
 				Data->DataTypeIndex = LayoutResults.Num();
-				LayoutResults.Add(v);
+				LayoutResults.Add(Value);
 			}
 			else
 			{
-				LayoutResults[Data->DataTypeIndex] = v;
+				LayoutResults[Data->DataTypeIndex] = Value;
 			}
 			check(Data->DataTypeIndex != 0);
 
 			mu::UpdateLLMStats();
 		}
 
-		void SetString(FCacheAddress at, Ptr<const String> v)
+		void SetString(FCacheAddress At, Ptr<const String> Value)
 		{
-			check(at.At < OpExecutionData.size_code());
-			FOpExecutionData* Data = OpExecutionData.get_ptr(at);
+			check(At.At < OpExecutionData.CodeSize());
+			FOpExecutionData* Data = OpExecutionData.GetPtr(At);
 			check(Data->DataType == DATATYPE::DT_STRING || Data->DataType == DATATYPE::DT_NONE);
 			Data->DataType = DATATYPE::DT_STRING;
 			Data->IsValueValid = true;
@@ -1036,67 +1119,45 @@ namespace mu
 			if (!Data->DataTypeIndex)
 			{
 				Data->DataTypeIndex = StringResults.Num();
-				StringResults.Add(v);
+				StringResults.Add(Value);
 			}
 			else
 			{
-				StringResults[Data->DataTypeIndex] = v;
+				StringResults[Data->DataTypeIndex] = Value;
 			}
-			check(Data->DataTypeIndex!=0);
+			check(Data->DataTypeIndex != 0);
 		}
 
 
-		inline void IncreaseHitCount(FCacheAddress at)
+		inline void IncreaseHitCount(FCacheAddress At)
 		{
 			// Don't count hits for instruction 0, which is always null. It is usually already
 			// check that At is not 0, and then it is not requested, generating a stray non-zero count
 			// at its position.
-			if (at.At)
+			if (At.At)
 			{
-				check(at.At < OpExecutionData.size_code());
-				FOpExecutionData& Data = OpExecutionData[at];
+				check(At.At < OpExecutionData.CodeSize());
+				FOpExecutionData& Data = OpExecutionData[At];
 				Data.OpHitCount++;
 			}
 		}
 
-		inline void SetForceCached(OP::ADDRESS at)
+		inline void SetForceCached(OP::ADDRESS At)
 		{
 			// \TODO: It only locks at,0,0
-			if (at)
+			if (At)
 			{
-				check(at < OpExecutionData.size_code());
-				FOpExecutionData& Data = OpExecutionData[FCacheAddress(at, 0, 0)];
+				check(At < OpExecutionData.CodeSize());
+				FOpExecutionData& Data = OpExecutionData[FCacheAddress(At, 0, 0)];
 				Data.IsCacheLocked = true;
 			}
 		}
 	};
 
-
-	inline bool operator==(const FCacheAddress& a, const FCacheAddress& b)
-	{
-		return a.At == b.At
-			&&
-			a.ExecutionIndex == b.ExecutionIndex
-			&&
-			a.ExecutionOptions == b.ExecutionOptions
-			&&
-			a.Type == b.Type;
-	}
-
-	inline bool operator<(const FCacheAddress& a, const FCacheAddress& b)
-	{
-		if (a.At < b.At) return true;
-		if (a.At > b.At) return false;
-		if (a.ExecutionIndex < b.ExecutionIndex) return true;
-		if (a.ExecutionIndex > b.ExecutionIndex) return false;
-		if (a.ExecutionOptions < b.ExecutionOptions) return true;
-		if (a.ExecutionOptions > b.ExecutionOptions) return false;
-		return a.Type < b.Type;
-	}
-
-	/** Data for an instance that is currently being processed in the mutable system. This means it is
-	* between a BeginUpdate and EndUpdate, or during an "atomic" operation (like generate a single resource).
-	*/
+	/** 
+	 * Data for an instance that is currently being processed in the mutable system. This means it is
+	 * between a BeginUpdate and EndUpdate, or during an "atomic" operation (like generate a single resource).
+	 */
 	struct FLiveInstance
 	{
 		Instance::ID InstanceID;
@@ -1245,7 +1306,7 @@ namespace mu
 
 			// Look for an unused image in the pool that can be reused
 			int32 PooledImageCount = PooledImages.Num();
-			for (int Index = 0; DataSize>0 && Index<PooledImageCount; ++Index)
+			for (int32 Index = 0; DataSize > 0 && Index < PooledImageCount; ++Index)
 			{
 				Ptr<Image>& Candidate = PooledImages[Index];
 				if (Candidate->GetFormat() == Format
@@ -1596,7 +1657,7 @@ namespace mu
 						Result += Rom.Value->GetDataSize();
 					}
 				}
-				for (const TPair<int32, Ptr<const Mesh>>& Rom : Program.ConstantMeshes)
+				for (const TPair<int32, Ptr<const Mesh>>& Rom : Program.ConstantMeshData)
 				{
 					if (Rom.Value && Rom.Key >= 0)
 					{
@@ -1629,25 +1690,27 @@ namespace mu
 
 			for (const FLiveInstance& Instance : LiveInstances)
 			{
-				CodeContainer<FProgramCache::FOpExecutionData>::iterator it = Instance.Cache->OpExecutionData.begin();
-				for (; it.IsValid(); ++it)
+				CodeContainer<FProgramCache::FOpExecutionData>::TIterator Iter = Instance.Cache->OpExecutionData.Begin();
+				for (; Iter.IsValid(); ++Iter)
 				{
-					if (!(*it).DataTypeIndex)
+					FProgramCache::FOpExecutionData& Data = *Iter;
+
+					if (!Data.DataTypeIndex)
 					{
 						continue;
 					}
 
 					TSet<const Resource*>* TargetSet = &Cache0Unique;
-					if ((*it).IsCacheLocked)
+					if (Data.IsCacheLocked)
 					{
 						TargetSet = &Cache1Unique;
 					}
 
 					const Resource* Value = nullptr;
-					switch ((*it).DataType)
+					switch (Data.DataType)
 					{
 					case DATATYPE::DT_IMAGE:
-						Value = Instance.Cache->ImageResults[(*it).DataTypeIndex].Value.get();
+						Value = Instance.Cache->ImageResults[Data.DataTypeIndex].Value.get();
 						if (Value)
 						{
 							TargetSet->Add(Value);
@@ -1655,7 +1718,7 @@ namespace mu
 						break;
 
 					case DATATYPE::DT_MESH:
-						Value = Instance.Cache->MeshResults[(*it).DataTypeIndex].Value.get();
+						Value = Instance.Cache->MeshResults[Data.DataTypeIndex].Value.get();
 						if (Value)
 						{
 							TargetSet->Add(Value);
@@ -1689,13 +1752,12 @@ namespace mu
 
 			MUTABLE_CPUPROFILER_SCOPE(ClearLayer0);
 
-			CodeContainer<FProgramCache::FOpExecutionData>::iterator it = CurrentInstanceCache->OpExecutionData.begin();
-			for (; it.IsValid(); ++it)
+			CodeContainer<FProgramCache::FOpExecutionData>::TIterator Iter = CurrentInstanceCache->OpExecutionData.Begin();
+			for (; Iter.IsValid(); ++Iter)
 			{
-				FProgramCache::FOpExecutionData& Data = *it;
-				if (!Data.DataTypeIndex
-					||
-					Data.IsCacheLocked)
+				FProgramCache::FOpExecutionData& Data = *Iter;
+				
+				if (!Data.DataTypeIndex || Data.IsCacheLocked)
 				{
 					continue;
 				}
@@ -1749,10 +1811,11 @@ namespace mu
 		{
 			MUTABLE_CPUPROFILER_SCOPE(ClearLayer1);
 
-			CodeContainer<FProgramCache::FOpExecutionData>::iterator it = CurrentInstanceCache->OpExecutionData.begin();
-			for (; it.IsValid(); ++it)
+			CodeContainer<FProgramCache::FOpExecutionData>::TIterator Iter = CurrentInstanceCache->OpExecutionData.Begin();
+			for (; Iter.IsValid(); ++Iter)
 			{
-				FProgramCache::FOpExecutionData& Data = *it;
+				FProgramCache::FOpExecutionData& Data = *Iter;
+
 				if (!Data.DataTypeIndex)
 				{
 					continue;
@@ -1900,16 +1963,16 @@ namespace mu
 		MUTABLERUNTIME_API void BeginBuild(const TSharedPtr<const Model>&);
 		MUTABLERUNTIME_API void EndBuild();
 
-		MUTABLERUNTIME_API bool BuildBool(const TSharedPtr<const Model>&, const Parameters*, OP::ADDRESS) ;
-		MUTABLERUNTIME_API int32 BuildInt(const TSharedPtr<const Model>&, const Parameters*, OP::ADDRESS) ;
-		MUTABLERUNTIME_API float BuildScalar(const TSharedPtr<const Model>&, const Parameters*, OP::ADDRESS) ;
-		MUTABLERUNTIME_API FVector4f BuildColour(const TSharedPtr<const Model>&, const Parameters*, OP::ADDRESS) ;
-		MUTABLERUNTIME_API Ptr<const String> BuildString(const TSharedPtr<const Model>&, const Parameters*, OP::ADDRESS) ;
-		MUTABLERUNTIME_API Ptr<const Image> BuildImage(const TSharedPtr<const Model>&, const Parameters*, OP::ADDRESS, int32 MipsToSkip, int32 LOD) ;
-		MUTABLERUNTIME_API Ptr<const Mesh> BuildMesh(const TSharedPtr<const Model>&, const Parameters*, OP::ADDRESS) ;
+		MUTABLERUNTIME_API bool BuildBool(const TSharedPtr<const Model>&, const Parameters*, OP::ADDRESS);
+		MUTABLERUNTIME_API int32 BuildInt(const TSharedPtr<const Model>&, const Parameters*, OP::ADDRESS);
+		MUTABLERUNTIME_API float BuildScalar(const TSharedPtr<const Model>&, const Parameters*, OP::ADDRESS);
+		MUTABLERUNTIME_API FVector4f BuildColour(const TSharedPtr<const Model>&, const Parameters*, OP::ADDRESS);
+		MUTABLERUNTIME_API Ptr<const String> BuildString(const TSharedPtr<const Model>&, const Parameters*, OP::ADDRESS);
+		MUTABLERUNTIME_API Ptr<const Image> BuildImage(const TSharedPtr<const Model>&, const Parameters*, OP::ADDRESS, int32 MipsToSkip, int32 LOD);
+		MUTABLERUNTIME_API Ptr<const Mesh> BuildMesh(const TSharedPtr<const Model>&, const Parameters*, OP::ADDRESS, uint8 ExecutionOptions = 0);
 		MUTABLERUNTIME_API Ptr<const Instance> BuildInstance(const TSharedPtr<const Model>&, const Parameters*, OP::ADDRESS);
-		MUTABLERUNTIME_API Ptr<const Layout> BuildLayout(const TSharedPtr<const Model>&, const Parameters*, OP::ADDRESS) ;
-    	MUTABLERUNTIME_API FProjector BuildProjector(const TSharedPtr<const Model>&, const Parameters*, OP::ADDRESS) ;
+		MUTABLERUNTIME_API Ptr<const Layout> BuildLayout(const TSharedPtr<const Model>&, const Parameters*, OP::ADDRESS);
+    	MUTABLERUNTIME_API FProjector BuildProjector(const TSharedPtr<const Model>&, const Parameters*, OP::ADDRESS);
 
 		ExtensionDataStreamer* GetExtensionDataStreamer() const { return ExtensionDataStreamer.Get(); }
 
@@ -1944,7 +2007,7 @@ namespace mu
         bool CheckUpdatedParameters( const FLiveInstance*, const Ptr<const Parameters>&, uint64& OutUpdatedParameters );
 
 
-		void RunCode(const TSharedPtr<const Model>&, const Parameters*, OP::ADDRESS at, uint32 LODs = System::AllLODs, uint8 executionOptions = 0, int32 LOD = 0);
+		void RunCode(const TSharedPtr<const Model>&, const Parameters*, OP::ADDRESS at, uint32 LODs = System::AllLODs, uint8 ExecutionOptions = 0, int32 LOD = 0);
 
 		//!
 		void PrepareCache(const Model*, int32 State);
diff --git a/Engine/Plugins/Experimental/Mutable/Source/MutableRuntime/Private/MuR/CodeRunner.cpp b/Engine/Plugins/Experimental/Mutable/Source/MutableRuntime/Private/MuR/CodeRunner.cpp
--- a/Engine/Plugins/Experimental/Mutable/Source/MutableRuntime/Private/MuR/CodeRunner.cpp
+++ b/Engine/Plugins/Experimental/Mutable/Source/MutableRuntime/Private/MuR/CodeRunner.cpp
@@ -124,8 +124,8 @@
 		EExecutionStrategy InExecutionStrategy,
 		const TSharedPtr<const Model>& InModel,
 		const Parameters* InParams,
-		OP::ADDRESS at,
-		uint32 InLodMask, uint8 executionOptions, int32 InImageLOD, FScheduledOp::EType Type )
+		OP::ADDRESS At,
+		uint32 InLodMask, uint8 ExecutionOptions, int32 InImageLOD, FScheduledOp::EType Type)
 		: m_pSettings(InSettings)
 		, RunnerCompletionEvent(TEXT("CodeRunnerCompletioneEventInit"))
 		, ExecutionStrategy(InExecutionStrategy)
@@ -134,13 +134,15 @@
 		, m_pParams(InParams)
 		, m_lodMask(InLodMask)
 	{
-		const FProgram& program = m_pModel->GetPrivate()->m_program;
-		ScheduledStagePerOp.resize(program.m_opAddress.Num());
+		const FProgram& Program = m_pModel->GetPrivate()->m_program;
+		ScheduledStagePerOp.Resize(Program.m_opAddress.Num());
 
 		// We will read this in the end, so make sure we keep it.
-   		if (Type == FScheduledOp::EType::Full)
+   		if (Type != FScheduledOp::EType::ImageDesc)
    		{
-			GetMemory().IncreaseHitCount(FCacheAddress(at, 0, executionOptions));
+			FCacheAddress RootCacheAddress = FCacheAddress(At, 0, ExecutionOptions);
+			RootCacheAddress.Type = Type;
+			GetMemory().IncreaseHitCount(RootCacheAddress);
 		}
     
 		// Start with a completed Event. This is checked at StartRun() to make sure StartRun is not called while there is 
@@ -150,11 +152,12 @@
 		ImageLOD = InImageLOD;
 	
 		// Push the root operation
-		FScheduledOp rootOp;
-		rootOp.At = at;
-		rootOp.ExecutionOptions = executionOptions;
-		rootOp.Type = Type;
-		AddOp(rootOp);
+		FScheduledOp RootOp;
+		RootOp.At = At;
+		RootOp.ExecutionOptions = ExecutionOptions;
+		RootOp.Type = Type;
+		
+		AddOp(RootOp);
 	}
 
 
@@ -922,60 +925,130 @@
         }
     }
 
-
-    //---------------------------------------------------------------------------------------------
-    bool CodeRunner::RunCode_ConstantResource(const FScheduledOp& item, const Model* pModel )
+    bool CodeRunner::RunCode_ConstantResource(const FScheduledOp& Item, const Model* pModel)
     {
 		MUTABLE_CPUPROFILER_SCOPE(RunCode_Constant);
 
 		const FProgram& Program = m_pModel->GetPrivate()->m_program;
 
-		OP_TYPE type = Program.GetOpType(item.At);
-        switch (type)
+		OP_TYPE Type = Program.GetOpType(Item.At);
+        switch (Type)
         {
 
         case OP_TYPE::ME_CONSTANT:
         {
-			OP::MeshConstantArgs args = Program.GetOpArgs<OP::MeshConstantArgs>(item.At);
+			OP::MeshConstantArgs Args = Program.GetOpArgs<OP::MeshConstantArgs>(Item.At);
 
-            OP::ADDRESS cat = args.value;
+			const EMeshExecutionOptions ExecutionOptions = static_cast<EMeshExecutionOptions>(Item.ExecutionOptions);
 
-            // Assume the ROM has been loaded previously
-            check(Program.ConstantMeshes[cat].Value)
+			const FMeshRange MeshIndexRange = Program.ConstantMeshes[Args.Value];
+			check(MeshIndexRange.NumMeshes == 3);
+			
+			uint32 EstimatedMeshDataSize = 0;
+			Ptr<const Mesh> GeometryMesh;
+			if (EnumHasAnyFlags(ExecutionOptions, EMeshExecutionOptions::LoadGeometryData))
+			{	
+				int32 GeometryRomIndex = Program.ConstantMeshIndices[MeshIndexRange.FirstIndex + 0];
+            	check(Program.ConstantMeshData[GeometryRomIndex].Value)
 
-            Ptr<const Mesh> SourceConst;
-			Program.GetConstant(cat, SourceConst);
+				GeometryMesh = Program.ConstantMeshData[GeometryRomIndex].Value;
 
-			check(SourceConst);
-			Ptr<Mesh> Source = CreateMesh(SourceConst->GetDataSize());
-			Source->CopyFrom(*SourceConst);
+				check(GeometryMesh);
 
-            // Set the separate skeleton if necessary
-            if (args.skeleton >= 0)
-            {
-                check(Program.m_constantSkeletons.Num() > size_t(args.skeleton));
-                Ptr<const Skeleton> pSkeleton = Program.m_constantSkeletons[args.skeleton];
-                Source->SetSkeleton(pSkeleton);
-            }
-
-			if (args.physicsBody >= 0)
-			{
-                check(Program.m_constantPhysicsBodies.Num() > size_t(args.physicsBody));
-                Ptr<const PhysicsBody> pPhysicsBody = Program.m_constantPhysicsBodies[args.physicsBody];
-                Source->SetPhysicsBody(pPhysicsBody);
+				EstimatedMeshDataSize += GeometryMesh->GetDataSize();
 			}
 
-            StoreMesh(item, Source);
+			Ptr<const Mesh> PoseMesh;
+			if (EnumHasAnyFlags(ExecutionOptions, EMeshExecutionOptions::LoadPoseData | EMeshExecutionOptions::LoadGeometryData))
+			{	
+				int32 PoseRomIndex = Program.ConstantMeshIndices[MeshIndexRange.FirstIndex + 1];
+            	check(Program.ConstantMeshData[PoseRomIndex].Value)
+				PoseMesh = Program.ConstantMeshData[PoseRomIndex].Value;
+				
+				check(PoseMesh);
+
+				EstimatedMeshDataSize += PoseMesh->GetDataSize();
+			}
+
+			Ptr<const Mesh> ComponentsMesh;
+			if (EnumHasAnyFlags(ExecutionOptions, EMeshExecutionOptions::LoadComponentsData))
+			{
+				int32 ComponentsRomIndex = Program.ConstantMeshIndices[MeshIndexRange.FirstIndex + 2];
+            	check(Program.ConstantMeshData[ComponentsRomIndex].Value)
+				
+				ComponentsMesh = Program.ConstantMeshData[ComponentsRomIndex].Value;
+				check(ComponentsMesh);
+				
+				EstimatedMeshDataSize += ComponentsMesh->GetDataSize();
+			}
+
+			Ptr<Mesh> Result = CreateMesh(EstimatedMeshDataSize);
+		
+			if (GeometryMesh)
+			{
+				Result->CopyFrom(*GeometryMesh);
+			}
+
+			if (ComponentsMesh)
+			{
+				// ComponentMeshes may have data like buffer descriptors that would clash with the
+				// proper data in the Geometry mesh. If There is no Geometry mesh, then the mesh can be copied directly
+				// otherwise only take the missing bits.
+				if (!GeometryMesh)
+				{
+					Result->CopyFrom(*ComponentsMesh);
+				}
+				else
+				{
+					Result->SkeletonIDs = ComponentsMesh->SkeletonIDs;
+					Result->AdditionalPhysicsBodies = ComponentsMesh->AdditionalPhysicsBodies;
+					Result->Tags = ComponentsMesh->Tags;
+					Result->StreamedResources = ComponentsMesh->StreamedResources;
+					
+					// Add additional buffers.
+					for (const TPair<EMeshBufferType, FMeshBufferSet>& AdditionalBuffer : ComponentsMesh->AdditionalBuffers)
+					{
+						Result->AdditionalBuffers.Add(AdditionalBuffer);
+					}
+				}
+			}
+
+			if (PoseMesh)
+			{		
+				Result->BonePoses = PoseMesh->BonePoses;
+				Result->BoneMap = PoseMesh->BoneMap;
+			}
+
+			// Shared Physics and Skeleton only get loaded if LoadComponentsData option is present.
+			if (EnumHasAnyFlags(ExecutionOptions, EMeshExecutionOptions::LoadComponentsData))
+			{
+				if (Args.PhysicsBody >= 0)
+				{
+					check(Program.ConstantPhysicsBodies.Num() > Args.PhysicsBody);
+					Ptr<const PhysicsBody> PhysicsBodyConstant = Program.ConstantPhysicsBodies[Args.PhysicsBody];
+					Result->SetPhysicsBody(PhysicsBodyConstant);
+				}
+
+				// Set the separate skeleton if necessary.
+				if (Args.Skeleton >= 0)
+				{
+					check(Program.ConstantSkeletons.Num() > Args.Skeleton);
+					Ptr<const Skeleton> SkeletonConstant = Program.ConstantSkeletons[Args.Skeleton];
+					Result->SetSkeleton(SkeletonConstant);
+				}
+			}
+
+            StoreMesh(Item, Result);
 			//UE_LOG(LogMutableCore, Log, TEXT("Set mesh constant %d."), item.At);
             break;
         }
 
         case OP_TYPE::IM_CONSTANT:
         {
-			OP::ResourceConstantArgs args = Program.GetOpArgs<OP::ResourceConstantArgs>(item.At);
+			OP::ResourceConstantArgs args = Program.GetOpArgs<OP::ResourceConstantArgs>(Item.At);
             OP::ADDRESS cat = args.value;
 
-			int32 MipsToSkip = item.ExecutionOptions;
+			int32 MipsToSkip = Item.ExecutionOptions;
             Ptr<const Image> Source;
 			Program.GetConstant(cat, Source, MipsToSkip, [this](int32 x, int32 y, int32 m, EImageFormat f, EInitializationType i)
 				{
@@ -988,14 +1061,14 @@
 				return false;
 			}
 
-            StoreImage( item, Source );
+            StoreImage(Item, Source);
 			//UE_LOG(LogMutableCore, Log, TEXT("Set image constant %d."), item.At);
 			break;
         }
 
 		case OP_TYPE::ED_CONSTANT:
 		{
-			OP::ResourceConstantArgs Args = Program.GetOpArgs<OP::ResourceConstantArgs>(item.At);
+			OP::ResourceConstantArgs Args = Program.GetOpArgs<OP::ResourceConstantArgs>(Item.At);
 
 			// Assume the ROM has been loaded previously
 			ExtensionDataPtrConst SourceConst;
@@ -1003,12 +1076,12 @@
 
 			check(SourceConst);
 
-            StoreExtensionData(item, SourceConst);
+            StoreExtensionData(Item, SourceConst);
             break;
 		}
 
         default:
-            if (type!=OP_TYPE::NONE)
+            if (Type != OP_TYPE::NONE)
             {
                 // Operation not implemented
                 check( false );
@@ -1020,4 +1093,812 @@
 		return true;
     }
 
+	void CodeRunner::RunCodeMeshComponents(const FScheduledOp& Item, const Parameters* ParamsPtr, const Model* ModelPtr, uint32 LodMask)
+    {
+		MUTABLE_CPUPROFILER_SCOPE(RunCode_MeshComponents);
+
+		check(Item.Type == FScheduledOp::EType::MeshComp);
+
+		const FProgram& Program = m_pModel->GetPrivate()->m_program;
+
+		auto ScheduleOrRunPassthroughOp = [this](const FScheduledOp& Item, OP::ADDRESS SourceAddress)
+		{
+			switch (Item.Stage)
+			{
+			case 0:
+			{
+				if (SourceAddress)
+				{
+					AddOp(FScheduledOp(Item.At, Item, 1), FScheduledOp(SourceAddress, Item));
+				}
+				else
+				{
+					StoreMesh(Item, nullptr);
+				}
+				break;
+			}
+            case 1:
+            {
+                StoreMesh(Item, LoadMesh(FCacheAddress(SourceAddress, Item)));
+                break;
+            }
+
+            default:
+                check(false);
+            }
+		};
+
+
+		OP_TYPE Type = Program.GetOpType(Item.At);
+        switch (Type)
+        {
+		case OP_TYPE::ME_CONDITIONAL:
+		{
+			OP::ConditionalArgs Args = Program.GetOpArgs<OP::ConditionalArgs>(Item.At);
+			switch (Item.Stage)
+			{
+			case 0:
+			{
+				// We need to run the full condition result
+				FScheduledOp FullConditionOp(Args.condition, Item);
+				FullConditionOp.Type = FScheduledOp::EType::Full;
+				AddOp(FScheduledOp(Item.At, Item, 1), FullConditionOp);
+				break;
+			}
+
+			case 1:
+			{
+				const bool Value = LoadBool(FCacheAddress(Args.condition, Item.ExecutionIndex, Item.ExecutionOptions));
+				const OP::ADDRESS ResultAddress = Value ? Args.yes : Args.no;
+				AddOp(FScheduledOp(Item.At, Item, 2, ResultAddress), FScheduledOp(ResultAddress, Item));
+				break;
+			}
+
+			case 2:
+			{
+				const uint32 ResultAddress = Item.CustomState;
+				if (ResultAddress)
+				{
+					StoreMesh(Item, LoadMesh(FCacheAddress(ResultAddress, Item)));
+				}
+				else
+				{
+					StoreMesh(Item, nullptr);
+				}
+				break;
+			}
+			default: check(false);
+			}
+			break;
+		}	
+		case OP_TYPE::ME_SWITCH:
+		{
+			uint8 const * const Data = Program.GetOpArgsPointer(Item.At);
+
+			constexpr uint32 DefAddressOffset = sizeof(OP::ADDRESS);
+			constexpr uint32 CaseCountOffset = DefAddressOffset + sizeof(OP::ADDRESS);
+			constexpr uint32 SwitchFixedOffset = CaseCountOffset + sizeof(uint32);		
+
+			switch (Item.Stage)
+			{
+			case 0:
+			{
+				OP::ADDRESS VarAddress;
+				FMemory::Memcpy(&VarAddress, Data + 0, sizeof(OP::ADDRESS));
+				
+				if (VarAddress)
+				{
+					FScheduledOp FullVariableOp(VarAddress, Item);
+					FullVariableOp.Type = FScheduledOp::EType::Full;
+
+					AddOp(FScheduledOp(Item.At, Item, 1), FullVariableOp);
+				}
+				else
+				{
+					check(GetOpDataType(Type) == DT_MESH);
+					StoreMesh(Item, nullptr);
+				}
+				break;
+			}
+			case 1:
+			{
+				OP::ADDRESS VarAddress;
+				FMemory::Memcpy(&VarAddress, Data + 0, sizeof(OP::ADDRESS));
+			
+				OP::ADDRESS DefAddress;
+				FMemory::Memcpy(&DefAddress, Data + DefAddressOffset, sizeof(OP::ADDRESS));
+
+				uint32 CaseCount;
+				FMemory::Memcpy(&CaseCount, Data + CaseCountOffset, sizeof(uint32));
+				
+				// Get the variable result
+				int32 Variable = LoadInt(FCacheAddress(VarAddress, Item.ExecutionIndex, Item.ExecutionOptions));
+				
+				constexpr uint32 SwitchCaseSize = sizeof(int32) + sizeof(OP::ADDRESS);
+
+				OP::ADDRESS ValueAt = DefAddress;
+				for (uint32 C = 0; C < CaseCount; ++C)
+				{
+					const uint32 CaseOffset = SwitchFixedOffset + SwitchCaseSize*C;
+
+					int32 Condition;
+					FMemory::Memcpy(&Condition, Data + CaseOffset, sizeof(int32));
+
+					OP::ADDRESS At;
+					FMemory::Memcpy(&At, Data + CaseOffset + sizeof(int32), sizeof(OP::ADDRESS));
+
+					if (At && Variable == Condition)
+					{
+						ValueAt = At;
+						break; 
+					}
+				}
+
+				// Schedule the end of this instruction if necessary
+				AddOp(FScheduledOp(Item.At, Item, 2, ValueAt), FScheduledOp(ValueAt, Item));
+				break;
+			}
+
+			case 2:
+			{
+				OP::ADDRESS ResultAt = OP::ADDRESS(Item.CustomState);
+				
+				check(GetOpDataType(Type) == DT_MESH);
+				StoreMesh(FCacheAddress(Item), LoadMesh(FCacheAddress(ResultAt, Item)));
+				break;
+			}
+
+			default:
+				check(false);
+			}
+			break;
+		}
+		case OP_TYPE::ME_CONSTANT:
+		{	
+            RunCode_ConstantResource(Item, ModelPtr);
+			break;
+		}
+		case OP_TYPE::ME_REFERENCE:
+		{
+			OP::ResourceReferenceArgs Args = Program.GetOpArgs<OP::ResourceReferenceArgs>(Item.At);
+			switch (Item.Stage)
+			{
+			case 0:
+			{
+				Ptr<Mesh> Result;
+				if (Args.ForceLoad)
+				{
+					// This should never be reached because it should have been caught as a Task in IssueOp
+					check(false);
+				}
+				else
+				{
+					Result = Mesh::CreateAsReference(Args.ID, false);
+				}
+				StoreMesh(Item, Result);
+				break;
+			}
+
+			default:
+				check(false);
+			}
+
+			break;
+		}
+        case OP_TYPE::ME_APPLYLAYOUT:
+        {
+			OP::MeshApplyLayoutArgs Args = Program.GetOpArgs<OP::MeshApplyLayoutArgs>(Item.At);
+			ScheduleOrRunPassthroughOp(Item, Args.mesh);
+
+            break;
+        }
+
+        case OP_TYPE::ME_DIFFERENCE:
+        {
+			const uint8* Data = Program.GetOpArgsPointer(Item.At);
+
+			OP::ADDRESS BaseAt = 0;
+			FMemory::Memcpy(&BaseAt, Data, sizeof(OP::ADDRESS)); 
+			ScheduleOrRunPassthroughOp(Item, BaseAt);
+
+            break;
+        }
+
+        case OP_TYPE::ME_MORPH:
+        {
+			const uint8* Data = Program.GetOpArgsPointer(Item.At);
+	
+			OP::ADDRESS BaseAt = 0;
+			FMemory::Memcpy(&BaseAt, Data + sizeof(OP::ADDRESS), sizeof(OP::ADDRESS)); 
+			ScheduleOrRunPassthroughOp(Item, BaseAt);
+
+            break;
+        }
+
+        case OP_TYPE::ME_MERGE:
+        {
+			OP::MeshMergeArgs Args = Program.GetOpArgs<OP::MeshMergeArgs>(Item.At);
+            switch (Item.Stage)
+            {
+            case 0:
+			{
+				AddOp(FScheduledOp(Item.At, Item, 1), FScheduledOp(Args.base, Item), FScheduledOp(Args.added, Item));
+				break;
+			}
+            case 1:
+            {
+                Ptr<const Mesh> AMesh = LoadMesh(FCacheAddress(Args.base, Item));
+                Ptr<const Mesh> BMesh = LoadMesh(FCacheAddress(Args.added, Item));
+
+				// TODO: Optimize for mesh component pass.
+                if (AMesh && BMesh && AMesh->GetVertexCount() && BMesh->GetVertexCount())
+                {
+					check(!AMesh->IsReference() && !BMesh->IsReference());
+
+					FMeshMergeScratchMeshes Scratch;
+					Scratch.FirstReformat = CreateMesh();
+					Scratch.SecondReformat = CreateMesh();
+
+					Ptr<Mesh> Result = CreateMesh();
+
+					MeshMerge(Result.get(), AMesh.get(), BMesh.get(), !Args.newSurfaceID, Scratch);
+
+					Release(Scratch.FirstReformat);
+					Release(Scratch.SecondReformat);
+
+                    if (Args.newSurfaceID)
+                    {
+						check(BMesh->GetSurfaceCount() == 1);
+						Result->Surfaces.Last().Id = Args.newSurfaceID;
+                    }
+
+					Release(AMesh);
+					Release(BMesh);
+					StoreMesh(Item, Result);
+                }
+                else if (AMesh && (AMesh->GetVertexCount() || AMesh->IsReference()))
+                {
+					Release(BMesh);
+					StoreMesh(Item, AMesh);
+                }
+                else if (BMesh && (BMesh->GetVertexCount() || BMesh->IsReference()))
+                {
+					Ptr<Mesh> Result = CloneOrTakeOver(BMesh);
+
+                    check(Result->IsReference() || (Result->GetSurfaceCount() == 1));
+
+                    if (Result->GetSurfaceCount() > 0 && Args.newSurfaceID)
+                    {
+                        Result->Surfaces.Last().Id = Args.newSurfaceID;
+                    }
+
+					Release(AMesh);
+					StoreMesh(Item, Result);
+                }
+                else
+                {
+					Release(AMesh);
+					Release(BMesh);
+					StoreMesh(Item, CreateMesh());
+                }
+                break;
+            }
+
+            default:
+                check(false);
+            }
+
+            break;
+        }
+
+        case OP_TYPE::ME_INTERPOLATE:
+        {
+			OP::MeshInterpolateArgs Args = Program.GetOpArgs<OP::MeshInterpolateArgs>(Item.At);
+			ScheduleOrRunPassthroughOp(Item, Args.base);
+            break;
+        }
+
+        case OP_TYPE::ME_MASKCLIPMESH:
+        {
+			OP::MeshMaskClipMeshArgs Args = Program.GetOpArgs<OP::MeshMaskClipMeshArgs>(Item.At);
+			ScheduleOrRunPassthroughOp(Item, Args.source);
+            break;
+        }
+
+		case OP_TYPE::ME_MASKCLIPUVMASK:
+		{
+			OP::MeshMaskClipUVMaskArgs Args = Program.GetOpArgs<OP::MeshMaskClipUVMaskArgs>(Item.At);
+			ScheduleOrRunPassthroughOp(Item, Args.Source);
+			break;
+		}
+
+        case OP_TYPE::ME_FORMAT:
+        {
+			OP::MeshFormatArgs Args = Program.GetOpArgs<OP::MeshFormatArgs>(Item.At);
+            switch (Item.Stage)
+            {
+            case 0:
+			{
+				if (Args.source && Args.format)
+				{
+					AddOp(FScheduledOp(Item.At, Item, 1), FScheduledOp(Args.source, Item), FScheduledOp(Args.format, Item));
+				}
+				else
+				{
+					StoreMesh(Item, nullptr);
+				}
+				break;
+			}
+            case 1:
+            {
+            	MUTABLE_CPUPROFILER_SCOPE(ME_FORMAT_Gather);
+
+				Ptr<const Mesh> Source = LoadMesh(FCacheAddress(Args.source, Item));
+                Ptr<const Mesh> Format = LoadMesh(FCacheAddress(Args.format, Item));
+
+				if (Source && Source->IsReference())
+				{
+					Release(Format);
+					StoreMesh(Item, Source);
+				}
+				else if (Source)
+				{
+					uint8 Flags = Args.Flags;
+					if (!Format && !(Flags & OP::MeshFormatArgs::ResetBufferIndices))
+					{
+						StoreMesh(Item, Source);
+					}
+					else if (!Format)
+					{
+						Ptr<Mesh> Result = CloneOrTakeOver(Source);
+
+						if (Flags & OP::MeshFormatArgs::ResetBufferIndices)
+						{
+							Result->ResetBufferIndices();
+						}
+
+						StoreMesh(Item, Result);
+					}
+					else
+					{
+						Ptr<Mesh> Result = CreateMesh();
+
+						bool bOutSuccess = false;
+						MeshFormat(Result.get(), Source.get(), Format.get(),
+							true,
+							(Flags & OP::MeshFormatArgs::Vertex) != 0,
+							(Flags & OP::MeshFormatArgs::Index) != 0,
+							(Flags & OP::MeshFormatArgs::IgnoreMissing) != 0,
+							bOutSuccess);
+
+						check(bOutSuccess);
+
+						if (Flags & OP::MeshFormatArgs::ResetBufferIndices)
+						{
+							Result->ResetBufferIndices();
+						}
+
+						if (Flags & OP::MeshFormatArgs::OptimizeBuffers)
+						{
+							MUTABLE_CPUPROFILER_SCOPE(MeshOptimizeBuffers)
+							MeshOptimizeBuffers(Result.get());
+						}
+
+						Release(Source);
+						Release(Format);
+						StoreMesh(Item, Result);
+					}
+				}
+				else
+				{
+					Release(Format);
+					StoreMesh(Item, nullptr);
+				}
+                break;
+            }
+
+            default:
+                check(false);
+            }
+
+            break;
+        }
+
+        case OP_TYPE::ME_CLIPMORPHPLANE:
+        {
+			OP::MeshClipMorphPlaneArgs Args = Program.GetOpArgs<OP::MeshClipMorphPlaneArgs>(Item.At);
+			ScheduleOrRunPassthroughOp(Item, Args.source);
+            break;
+        }
+
+		case OP_TYPE::ME_CLIPDEFORM:
+		{
+			OP::MeshClipDeformArgs Args = Program.GetOpArgs<OP::MeshClipDeformArgs>(Item.At);
+			ScheduleOrRunPassthroughOp(Item, Args.mesh);
+			break;
+		}
+
+        case OP_TYPE::ME_APPLYPOSE:
+        {
+			OP::MeshApplyPoseArgs Args = Program.GetOpArgs<OP::MeshApplyPoseArgs>(Item.At);
+
+			constexpr uint8 PoseExecutionOptions = static_cast<uint8>(EMeshExecutionOptions::LoadPoseData);
+            switch (Item.Stage)
+            {
+            case 0:
+			{
+				if (Args.base)
+				{
+					// Execute the pose as a full operation with specialized options for the pose
+					FScheduledOp PoseOp = FScheduledOp::FromOpAndOptions(Args.pose, Item, PoseExecutionOptions);
+					PoseOp.Type = FScheduledOp::EType::Full;
+
+					AddOp(FScheduledOp(Item.At, Item, 1), FScheduledOp(Args.base, Item), PoseOp);
+				}
+				else
+				{
+					StoreMesh(Item, nullptr);
+				}
+				break;
+			}
+            case 1:
+            {
+          		MUTABLE_CPUPROFILER_SCOPE(ME_APPLYPOSE_Gather);
+
+                Ptr<const Mesh> BaseMesh = LoadMesh(FCacheAddress(Args.base, Item));
+                Ptr<const Mesh> PoseMesh = LoadMesh(FCacheAddress(Args.pose, Item.ExecutionIndex, PoseExecutionOptions));
+
+                // Only if both are valid.
+                if (BaseMesh && PoseMesh)
+                {
+					Ptr<Mesh> Result = CloneOrTakeOver(BaseMesh);
+
+					SetPoseAsReference(Result.get(), PoseMesh.get());
+					Release(PoseMesh);
+					StoreMesh(Item, Result);
+                }
+                else
+                {
+					Release(PoseMesh);
+					StoreMesh(Item, BaseMesh);
+                }
+
+                break;
+            }
+
+            default:
+                check(false);
+            }
+
+            break;
+        }
+
+		case OP_TYPE::ME_BINDSHAPE:
+		{
+			OP::MeshBindShapeArgs Args = Program.GetOpArgs<OP::MeshBindShapeArgs>(Item.At);
+			const uint8* Data = Program.GetOpArgsPointer(Item.At);
+
+			switch (Item.Stage)
+			{
+			case 0:
+			{
+				if (Args.mesh)
+				{
+					// The Geometry of the Shape is needed, switch to Full execution policy.
+					FScheduledOp ShapeFullSheduledOp = FScheduledOp(Args.shape, Item);
+					ShapeFullSheduledOp.Type = FScheduledOp::EType::Full;
+
+					AddOp(FScheduledOp(Item.At, Item, 1), FScheduledOp(Args.mesh, Item), ShapeFullSheduledOp);
+				}
+				else
+				{
+					StoreMesh(Item, nullptr);
+				}
+				break;
+			}
+			case 1:
+			{
+				MUTABLE_CPUPROFILER_SCOPE(ME_BINDSHAPE_1)
+				Ptr<const Mesh> BaseMesh = LoadMesh(FCacheAddress(Args.mesh, Item));
+				Ptr<const Mesh> Shape = LoadMesh(FCacheAddress(Args.shape, Item));
+				
+				EShapeBindingMethod BindingMethod = static_cast<EShapeBindingMethod>(Args.bindingMethod); 
+
+				if (BindingMethod == EShapeBindingMethod::ReshapeClosestProject)
+				{ 
+					// Bones are stored after the args
+					Data += sizeof(Args);
+
+					// Rebuilding array of bone names ----
+					int32 NumBones;
+					FMemory::Memcpy(&NumBones, Data, sizeof(int32)); 
+					Data += sizeof(int32);
+					
+					TArray<FBoneName> BonesToDeform;
+					BonesToDeform.SetNumUninitialized(NumBones);
+					FMemory::Memcpy(BonesToDeform.GetData(), Data, NumBones * sizeof(FBoneName));
+					Data += NumBones * sizeof(FBoneName);
+
+					int32 NumPhysicsBodies;
+					FMemory::Memcpy(&NumPhysicsBodies, Data, sizeof(int32)); 
+					Data += sizeof(int32);
+
+					TArray<FBoneName> PhysicsToDeform;
+					PhysicsToDeform.SetNumUninitialized(NumPhysicsBodies);
+					FMemory::Memcpy(PhysicsToDeform.GetData(), Data, NumPhysicsBodies * sizeof(FBoneName));
+					Data += NumPhysicsBodies * sizeof(FBoneName);
+
+					const EMeshBindShapeFlags BindFlags = static_cast<EMeshBindShapeFlags>(Args.flags);
+
+					FMeshBindColorChannelUsages ColorChannelUsages;
+					FMemory::Memcpy(&ColorChannelUsages, &Args.ColorUsage, sizeof(ColorChannelUsages));
+					static_assert(sizeof(ColorChannelUsages) == sizeof(Args.ColorUsage));
+
+					Ptr<Mesh> BindMeshResult = CreateMesh();
+
+					bool bOutSuccess = false;
+					MeshBindShapeReshape(BindMeshResult.get(), BaseMesh.get(), Shape.get(), BonesToDeform, PhysicsToDeform, BindFlags, ColorChannelUsages, bOutSuccess);
+				
+					Release(Shape);
+					// not success indicates nothing has bond so the base mesh can be reused.
+					if (!bOutSuccess)
+					{
+						Release(BindMeshResult);
+						StoreMesh(Item, BaseMesh);
+					}
+					else
+					{
+						if (!EnumHasAnyFlags(BindFlags, EMeshBindShapeFlags::ReshapeVertices))
+						{
+							Ptr<Mesh> BindMeshNoVertsResult = CloneOrTakeOver(BaseMesh);
+							BindMeshNoVertsResult->AdditionalBuffers = MoveTemp(BindMeshResult->AdditionalBuffers);
+
+							Release(BaseMesh);
+							Release(BindMeshResult);
+							StoreMesh(Item, BindMeshNoVertsResult);
+						}
+						else
+						{
+							Release(BaseMesh);
+							StoreMesh(Item, BindMeshResult);
+						}
+					}
+				}	
+				else
+				{
+					Ptr<Mesh> Result = CreateMesh();
+
+					bool bOutSuccess = false;
+					MeshBindShapeClipDeform(Result.get(), BaseMesh.get(), Shape.get(), BindingMethod, bOutSuccess);
+
+					Release(Shape);
+					if (!bOutSuccess)
+					{
+						Release(Result);
+						StoreMesh(Item, BaseMesh);
+					}
+					else
+					{
+						Release(BaseMesh);
+						StoreMesh(Item, Result);
+					}
+				}
+
+				break;
+			}
+
+			default:
+				check(false);
+			}
+
+			break;
+		}
+
+
+		case OP_TYPE::ME_APPLYSHAPE:
+		{
+			OP::MeshApplyShapeArgs Args = Program.GetOpArgs<OP::MeshApplyShapeArgs>(Item.At);
+			switch (Item.Stage)
+			{
+			case 0:
+			{
+				if (Args.mesh)
+				{
+					// Shape needs the geomtry, schedule as Full execution policy.
+					FScheduledOp ShapeFullScheduledOp = FScheduledOp(Args.shape, Item);
+					ShapeFullScheduledOp.Type = FScheduledOp::EType::Full;
+					
+					AddOp(FScheduledOp(Item.At, Item, 1), FScheduledOp(Args.mesh, Item), ShapeFullScheduledOp);
+				}
+				else
+				{
+					StoreMesh(Item, nullptr);
+				}
+				break;
+			}
+			case 1:
+			{
+				MUTABLE_CPUPROFILER_SCOPE(ME_APPLYSHAPE_1)
+					
+				Ptr<const Mesh> BaseMesh = LoadMesh(FCacheAddress(Args.mesh, Item));
+				Ptr<const Mesh> Shape = LoadMesh(FCacheAddress(Args.shape, Item));
+
+				const EMeshBindShapeFlags ReshapeFlags = static_cast<EMeshBindShapeFlags>(Args.flags);
+				const bool bReshapeVertices = EnumHasAnyFlags(ReshapeFlags, EMeshBindShapeFlags::ReshapeVertices);
+
+				Ptr<Mesh> ReshapedMeshResult = CreateMesh(BaseMesh ? BaseMesh->GetDataSize() : 0);
+
+				bool bOutSuccess = false;
+				MeshApplyShape(ReshapedMeshResult.get(), BaseMesh.get(), Shape.get(), ReshapeFlags, bOutSuccess);
+
+				Release(Shape);
+				
+				if (!bOutSuccess)
+				{
+					Release(ReshapedMeshResult);
+					StoreMesh(Item, BaseMesh);
+				}
+				else
+				{
+					if (!bReshapeVertices)
+					{
+						// Clone without Skeleton, Physics or Poses 
+						EMeshCopyFlags CopyFlags = ~(
+							EMeshCopyFlags::WithSkeleton |
+							EMeshCopyFlags::WithPhysicsBody |
+							EMeshCopyFlags::WithAdditionalPhysics |
+							EMeshCopyFlags::WithPoses);
+
+						Ptr<Mesh> NoVerticesReshpedMesh = CloneOrTakeOver(BaseMesh);
+
+						NoVerticesReshpedMesh->SetSkeleton(ReshapedMeshResult->GetSkeleton().get());
+						NoVerticesReshpedMesh->SetPhysicsBody(ReshapedMeshResult->GetPhysicsBody().get());
+						NoVerticesReshpedMesh->AdditionalPhysicsBodies = ReshapedMeshResult->AdditionalPhysicsBodies;
+						NoVerticesReshpedMesh->BonePoses = ReshapedMeshResult->BonePoses;
+
+						Release(BaseMesh);
+						Release(ReshapedMeshResult);
+						StoreMesh(Item, NoVerticesReshpedMesh);
+					}
+					else
+					{
+						Release(BaseMesh);
+						StoreMesh(Item, ReshapedMeshResult);
+					}
+				}
+				break;
+			}
+
+			default:
+				check(false);
+			}
+
+			break;
+		}
+
+		case OP_TYPE::ME_MORPHRESHAPE:
+		{
+			OP::MeshMorphReshapeArgs Args = Program.GetOpArgs<OP::MeshMorphReshapeArgs>(Item.At);
+			ScheduleOrRunPassthroughOp(Item, Args.Reshape);
+			break;
+		}
+
+        case OP_TYPE::ME_SETSKELETON:
+        {
+			checkf(false, TEXT("Unsupported operation."));
+            break;
+        }
+
+        case OP_TYPE::ME_REMOVEMASK:
+        {
+       		MUTABLE_CPUPROFILER_SCOPE(ME_REMOVEMASK)
+        		
+            // Decode op
+            // TODO: Partial decode for each stage
+            const uint8* Data = Program.GetOpArgsPointer(Item.At);
+
+            OP::ADDRESS Source;
+            FMemory::Memcpy(&Source, Data, sizeof(OP::ADDRESS)); 
+			ScheduleOrRunPassthroughOp(Item, Source);
+            break;
+        }
+
+        case OP_TYPE::ME_ADDTAGS:
+		{
+			MUTABLE_CPUPROFILER_SCOPE(ME_ADDTAGS)
+
+			// Decode op
+			// TODO: Partial decode for each stage
+			const uint8* Data = Program.GetOpArgsPointer(Item.At);
+
+			OP::ADDRESS Source;
+			FMemory::Memcpy(&Source, Data, sizeof(OP::ADDRESS));
+			Data += sizeof(OP::ADDRESS);
+
+			// Schedule next stages
+			switch (Item.Stage)
+			{
+			case 0:
+			{
+				if (Source)
+				{
+					// Request the source
+					AddOp(FScheduledOp(Item.At, Item, 1), FScheduledOp(Source, Item));
+				}
+				else
+				{
+					StoreMesh(Item, nullptr);
+				}
+				break;
+			}
+
+			case 1:
+			{
+				MUTABLE_CPUPROFILER_SCOPE(ME_ADDTAGS_2)
+
+				Ptr<const Mesh> SourceMesh = LoadMesh(FCacheAddress(Source, Item));
+
+				if (!SourceMesh)
+				{
+					StoreMesh(Item, nullptr);
+				}
+				else
+				{
+					Ptr<Mesh> Result = CloneOrTakeOver(SourceMesh);
+
+					// Decode the tags
+					uint16 TagCount;
+					FMemory::Memcpy(&TagCount, Data, sizeof(uint16));
+					Data += sizeof(uint16);
+
+					int32 FirstMeshTagIndex = Result->Tags.Num();
+					Result->Tags.SetNum(FirstMeshTagIndex+TagCount);
+					for (uint16 TagIndex = 0; TagIndex < TagCount; ++TagIndex)
+					{
+						OP::ADDRESS TagConstant;
+						FMemory::Memcpy(&TagConstant, Data, sizeof(OP::ADDRESS));
+						Data += sizeof(OP::ADDRESS);
+
+						check(TagConstant < (uint32)ModelPtr->GetPrivate()->m_program.m_constantStrings.Num());
+						const FString& Name = Program.m_constantStrings[TagConstant];
+						Result->Tags[FirstMeshTagIndex+TagIndex] = Name;
+					}
+
+					StoreMesh(Item, Result);
+				}
+
+				break;
+			}
+
+			default:
+				check(false);
+			}
+
+			break;
+		}
+
+		case OP_TYPE::ME_PROJECT:
+        {
+			OP::MeshProjectArgs Args = Program.GetOpArgs<OP::MeshProjectArgs>(Item.At);
+			ScheduleOrRunPassthroughOp(Item, Args.mesh);
+            break;
+        }
+
+		case OP_TYPE::ME_OPTIMIZESKINNING:
+		{
+			OP::MeshOptimizeSkinningArgs Args = Program.GetOpArgs<OP::MeshOptimizeSkinningArgs>(Item.At);
+			ScheduleOrRunPassthroughOp(Item, Args.source);
+			break;
+		}
+
+        default:
+            if (Type != OP_TYPE::NONE)
+            {
+                // Operation not implemented
+                check(false);
+            }
+            break;
+        }
+    }
+
     //---------------------------------------------------------------------------------------------
@@ -1069,7 +1950,7 @@
 			{
 				AddOp(FScheduledOp(item.At, item, 1),
 					FScheduledOp(args.mesh, item),
-					FScheduledOp(args.layout, item));
+					FScheduledOp::FromOpAndOptions(args.layout, item, 0));
 				break;
 			}
             case 1:
@@ -1082,7 +1963,7 @@
                 {
 					Ptr<Mesh> Result = CloneOrTakeOver(pBase);
 
-                    Ptr<const Layout> pLayout = LoadLayout(FCacheAddress(args.layout, item));
+                    Ptr<const Layout> pLayout = LoadLayout(FCacheAddress(args.layout, item.ExecutionIndex, 0, item.Type));
                     int texCoordsSet = args.channel;
 
                     MeshApplyLayout(Result.get(), pLayout.get(), texCoordsSet);
@@ -1661,7 +2542,7 @@
 			{
 				AddOp(FScheduledOp(item.At, item, 1),
 					FScheduledOp(args.Source, item),
-					FScheduledOp(args.Mask, item));
+					FScheduledOp::FromOpAndOptions(args.Mask, item, 0));
 				break;
 			}
 			case 1:
@@ -1669,7 +2550,7 @@
 				MUTABLE_CPUPROFILER_SCOPE(ME_MASKCLIPUVMASK_1)
 
 				Ptr<const Mesh> Source = LoadMesh(FCacheAddress(args.Source, item));
-				Ptr<const Image> Mask = LoadImage(FCacheAddress(args.Mask, item));
+				Ptr<const Image> Mask = LoadImage(FCacheAddress(args.Mask, item.ExecutionIndex, item.ExecutionOptions, item.Type));
 
 				// Only if both are valid.
 				if (Source.get() && Mask.get())
@@ -1994,6 +2875,14 @@
         case OP_TYPE::ME_CLIPMORPHPLANE:
         {
 			OP::MeshClipMorphPlaneArgs args = Program.GetOpArgs<OP::MeshClipMorphPlaneArgs>(item.At);
+
+			// In case of bone hierarchy selection the skeleton needs to be loaded.
+			// For now we fully load the mesh. We may want this to be an extra option where
+			// only skeleton and pose data is added. 
+			const uint8 MeshExecutionOptions = args.vertexSelectionType == OP::MeshClipMorphPlaneArgs::VS_BONE_HIERARCHY
+					? item.ExecutionOptions | static_cast<uint8>(EMeshExecutionOptions::LoadPoseData) 
+					: item.ExecutionOptions;
+
             switch (item.Stage)
             {
             case 0:
@@ -2001,7 +2890,7 @@
 				if (args.source)
 				{
 					AddOp(FScheduledOp(item.At, item, 1),
-						FScheduledOp(args.source, item));
+						FScheduledOp::FromOpAndOptions(args.source, item, MeshExecutionOptions));
 				}
 				else
 				{
@@ -2013,7 +2902,7 @@
             {
            		MUTABLE_CPUPROFILER_SCOPE(ME_CLIPMORPHPLANE_1)
             		
-                Ptr<const Mesh> Source = LoadMesh(FCacheAddress(args.source, item));
+                Ptr<const Mesh> Source = LoadMesh(FCacheAddress(args.source, item.ExecutionIndex, MeshExecutionOptions, item.Type));
 
                 check(args.morphShape < (uint32)pModel->GetPrivate()->m_program.m_constantShapes.Num());
 
@@ -2232,15 +3121,19 @@
         case OP_TYPE::ME_APPLYPOSE:
         {
 			OP::MeshApplyPoseArgs args = Program.GetOpArgs<OP::MeshApplyPoseArgs>(item.At);
+
+			// Pose needs to load the components part of the mesh as it needs skeleton and pose info. It does not need geometry.
+			constexpr uint8 PoseExecutionOptions = static_cast<uint8>(EMeshExecutionOptions::LoadPoseData);
             switch (item.Stage)
             {
             case 0:
 			{
 				if (args.base)
 				{
+
 					AddOp(FScheduledOp(item.At, item, 1),
 						FScheduledOp(args.base, item),
-						FScheduledOp(args.pose, item));
+						FScheduledOp::FromOpAndOptions(args.pose, item, PoseExecutionOptions));
 				}
 				else
 				{
@@ -2253,7 +3146,7 @@
           		MUTABLE_CPUPROFILER_SCOPE(ME_APPLYPOSE_1)
 
                 Ptr<const Mesh> pBase = LoadMesh(FCacheAddress(args.base, item));
-                Ptr<const Mesh> pPose = LoadMesh(FCacheAddress(args.pose, item));
+                Ptr<const Mesh> pPose = LoadMesh(FCacheAddress(args.pose, item.ExecutionIndex, PoseExecutionOptions, item.Type));
 
                 // Only if both are valid.
                 if (pBase && pPose)
@@ -2304,8 +3197,8 @@
 					AddOp(FScheduledOp(item.At, item, 1),
 						FScheduledOp(args.meshA, item),
 						FScheduledOp(args.meshB, item),
-						FScheduledOp(args.scalarA, item),
-						FScheduledOp(args.scalarB, item));
+						FScheduledOp::FromOpAndOptions(args.scalarA, item, 0),
+						FScheduledOp::FromOpAndOptions(args.scalarB, item, 0));
 				}
 				else
 				{
@@ -2319,8 +3212,8 @@
 				
 				Ptr<const Mesh> MeshA = LoadMesh(FCacheAddress(args.meshA, item));
 				Ptr<const Mesh> MeshB = LoadMesh(FCacheAddress(args.meshB, item));
-				float ScalarA = LoadScalar(FCacheAddress(args.scalarA, item));
-				float ScalarB = LoadScalar(FCacheAddress(args.scalarB, item));
+				float ScalarA = LoadScalar(FCacheAddress(args.scalarA, item.ExecutionIndex, 0, item.Type));
+				float ScalarB = LoadScalar(FCacheAddress(args.scalarB, item.ExecutionIndex, 0, item.Type));
 
 				Ptr<Mesh> Result = CreateMesh(MeshA ? MeshA->GetDataSize() : 0);
 
@@ -2613,6 +3506,7 @@
         case OP_TYPE::ME_SETSKELETON:
         {
 			OP::MeshSetSkeletonArgs args = Program.GetOpArgs<OP::MeshSetSkeletonArgs>(item.At);
+			constexpr uint8 SkeletonExecutionOptions = static_cast<uint8>(EMeshExecutionOptions::LoadComponentsData);
             switch (item.Stage)
             {
             case 0:
@@ -2621,7 +3515,7 @@
 				{
 					AddOp(FScheduledOp(item.At, item, 1),
 						FScheduledOp(args.source, item),
-						FScheduledOp(args.skeleton, item));
+						FScheduledOp::FromOpAndOptions(args.skeleton, item, SkeletonExecutionOptions));
 				}
 				else
 				{
@@ -2634,14 +3528,12 @@
             	MUTABLE_CPUPROFILER_SCOPE(ME_SETSKELETON_1)
             		
                 Ptr<const Mesh> Source = LoadMesh(FCacheAddress(args.source, item));
-                Ptr<const Mesh> pSkeleton = LoadMesh(FCacheAddress(args.skeleton, item));
+                Ptr<const Mesh> pSkeleton = LoadMesh(FCacheAddress(args.skeleton, item.ExecutionIndex, SkeletonExecutionOptions, item.Type));
 
                 // Only if both are valid.
                 if (Source && pSkeleton)
                 {
-                    if ( Source->GetSkeleton()
-                         &&
-                         Source->GetSkeleton()->GetBoneCount() > 0 )
+                    if (Source->GetSkeleton() && Source->GetSkeleton()->GetBoneCount() > 0)
                     {
                         // For some reason we already have bone data, so we can't just overwrite it
                         // or the skinning may break. This may happen because of a problem in the
@@ -2714,13 +3606,13 @@
 			FMemory::Memcpy(&removes,data,sizeof(uint16)); 
 			data += sizeof(uint16);
 
-            for( uint16 r=0; r<removes; ++r)
+            for (uint16 r = 0; r < removes; ++r)
             {
                 OP::ADDRESS condition;
 				FMemory::Memcpy(&condition,data,sizeof(OP::ADDRESS)); 
 				data += sizeof(OP::ADDRESS);
                 
-				conditions.Emplace(condition, item);
+				conditions.Add(FScheduledOp::FromOpAndOptions(condition, item, 0));
 
                 OP::ADDRESS mask;
 				FMemory::Memcpy(&mask,data,sizeof(OP::ADDRESS)); 
@@ -2754,13 +3646,13 @@
                 // \todo: store condition values in heap?
                 TArray<FScheduledOp> deps;
                 deps.Emplace( source, item );
-                for( size_t r=0; source && r<conditions.Num(); ++r )
+                for (int32 r = 0; source && r < conditions.Num(); ++r)
                 {
                     // If there is no expression, we'll assume true.
                     bool value = true;
                     if (conditions[r].At)
                     {
-                        value = LoadBool(FCacheAddress(conditions[r].At, item));
+                        value = LoadBool(FCacheAddress(conditions[r].At, item.ExecutionIndex, 0, item.Type));
                     }
 
                     if (value)
@@ -2796,7 +3688,7 @@
 						bool value = true;
 						if (conditions[r].At)
 						{
-							value = LoadBool(FCacheAddress(conditions[r].At, item));
+							value = LoadBool(FCacheAddress(conditions[r].At, item.ExecutionIndex, 0, item.Type));
 						}
 
 						if (value)
@@ -2877,16 +3769,19 @@
 
 				Ptr<const Mesh> SourceMesh = LoadMesh(FCacheAddress(Source, item));
 
-				if (!SourceMesh)
+				// In case the we are executing only for geometry, skip this operation.
+				const bool bSkipOperation = 
+						!EnumHasAnyFlags(static_cast<EMeshExecutionOptions>(item.ExecutionOptions), EMeshExecutionOptions::LoadComponentsData);
+				if (!SourceMesh || bSkipOperation)
 				{
-					StoreMesh(item, nullptr);
+					StoreMesh(item, SourceMesh);
 				}
 				else
 				{
+					// Decode the tags.
 					Ptr<Mesh> Result = CloneOrTakeOver(SourceMesh);
-
-					// Decode the tags
-					uint16 TagCount;
+					
+					uint16 TagCount = 0;
 					FMemory::Memcpy(&TagCount, Data, sizeof(uint16));
 					Data += sizeof(uint16);
 
@@ -2902,7 +3797,6 @@
 						const FString& Name = Program.m_constantStrings[TagConstant];
 						Result->Tags[FirstMeshTagIndex+TagIndex] = Name;
 					}
-
 					StoreMesh(item, Result);
 				}
 
@@ -2927,7 +3821,7 @@
 				{
 					AddOp(FScheduledOp(item.At, item, 1),
 						FScheduledOp(args.mesh, item),
-						FScheduledOp(args.projector, item));
+						FScheduledOp::FromOpAndOptions(args.projector, item, 0));
 				}
 				else
 				{
@@ -2939,8 +3833,8 @@
             {
 				MUTABLE_CPUPROFILER_SCOPE(ME_PROJECT_1)
 
-                Ptr<const Mesh> pMesh = LoadMesh(FCacheAddress(args.mesh,item));
-                const FProjector Projector = LoadProjector(FCacheAddress(args.projector, item));
+                Ptr<const Mesh> pMesh = LoadMesh(FCacheAddress(args.mesh, item));
+                const FProjector Projector = LoadProjector(FCacheAddress(args.projector, item.ExecutionIndex, 0, item.Type));
 
                 // Only if both are valid.
                 if (pMesh && pMesh->GetVertexBuffers().GetBufferCount() > 0)
@@ -3033,10 +3927,10 @@
 		}
 
         default:
-            if (type!=OP_TYPE::NONE)
+            if (type != OP_TYPE::NONE)
             {
                 // Operation not implemented
-                check( false );
+                check(false);
             }
             break;
         }
@@ -4450,19 +5344,20 @@
         case OP_TYPE::IM_RASTERMESH:
         {
 			OP::ImageRasterMeshArgs args = Program.GetOpArgs<OP::ImageRasterMeshArgs>(item.At);
+			constexpr uint8 MeshExecutionOptions = static_cast<uint8>(EMeshExecutionOptions::LoadGeometryData);
             switch (item.Stage)
             {
             case 0:
 				if (args.image)
 				{
 					AddOp(FScheduledOp(item.At, item, 1),
-						FScheduledOp::FromOpAndOptions(args.mesh, item, 0),
+						FScheduledOp::FromOpAndOptions(args.mesh, item, MeshExecutionOptions),
 						FScheduledOp::FromOpAndOptions(args.projector, item, 0));
 				}
 				else
 				{
 					AddOp(FScheduledOp(item.At, item, 1),
-						FScheduledOp::FromOpAndOptions(args.mesh, item, 0));
+						FScheduledOp::FromOpAndOptions(args.mesh, item, MeshExecutionOptions));
 				}
                 break;
 
@@ -4470,7 +5365,7 @@
 			{
 				MUTABLE_CPUPROFILER_SCOPE(IM_RASTERMESH_1)
 
-				Ptr<const Mesh> pMesh = LoadMesh(FScheduledOp::FromOpAndOptions(args.mesh, item, 0));
+				Ptr<const Mesh> pMesh = LoadMesh(FScheduledOp::FromOpAndOptions(args.mesh, item, MeshExecutionOptions));
 
 				// If no image, we are generating a flat mesh UV raster. This is the final stage in this case.
 				if (!args.image)
@@ -5824,11 +6719,10 @@
         case OP_TYPE::LA_CONSTANT:
         {
 			OP::ResourceConstantArgs args = Program.GetOpArgs<OP::ResourceConstantArgs>(item.At);
-            check( args.value < (uint32)pModel->GetPrivate()->m_program.m_constantLayouts.Num() );
+            check(args.value < (uint32)pModel->GetPrivate()->m_program.m_constantLayouts.Num());
 
-            Ptr<const Layout> pResult = Program.m_constantLayouts
-                    [ args.value ];
-            StoreLayout( item, pResult );
+            Ptr<const Layout> pResult = Program.m_constantLayouts[args.value];
+            StoreLayout(item, pResult);
             break;
         }
 
@@ -5911,17 +6805,18 @@
 		case OP_TYPE::LA_FROMMESH:
 		{
 			OP::LayoutFromMeshArgs args = Program.GetOpArgs<OP::LayoutFromMeshArgs>(item.At);
+			
+			constexpr uint8 MeshExecutionOptions = static_cast<uint8>(EMeshExecutionOptions::LoadGeometryData);
 			switch (item.Stage)
 			{
 			case 0:
 				AddOp(FScheduledOp(item.At, item, 1),
-					FScheduledOp(args.Mesh, item));
+						FScheduledOp::FromOpAndOptions(args.Mesh, item, MeshExecutionOptions));
 				break;
 
 			case 1:
 			{
-				Ptr<const Mesh> Mesh = LoadMesh(FCacheAddress(args.Mesh, item));
-
+				Ptr<const Mesh> Mesh = LoadMesh(FCacheAddress(args.Mesh, item.ExecutionIndex, MeshExecutionOptions, item.Type));
 				Ptr<const Layout> Result = LayoutFromMesh_RemoveBlocks(Mesh.get(), args.LayoutIndex);
 
 				Release(Mesh);
@@ -6094,7 +6989,6 @@
         }
     }
 
-	//---------------------------------------------------------------------------------------------
 	void CodeRunner::RunCodeImageDesc(const FScheduledOp& item, const Parameters* pParams, const Model* pModel,  uint32 lodMask )
 	{
 		MUTABLE_CPUPROFILER_SCOPE(RunCodeImageDesc);
@@ -6166,7 +7060,10 @@
 
 			case 1:
 			{
-				bool value = LoadBool(FCacheAddress(args.condition, item.ExecutionIndex, item.ExecutionOptions));
+				FCacheAddress BoolLoadAddress = FCacheAddress(args.condition, item);
+				BoolLoadAddress.Type = FScheduledOp::EType::Full;
+
+				bool value = LoadBool(BoolLoadAddress);
 				OP::ADDRESS resultAt = value ? args.yes : args.no;
 				AddOp(FScheduledOp(item.At, item, 2), FScheduledOp(resultAt, item));
 				break;
@@ -6180,19 +7077,19 @@
 
 		case OP_TYPE::IM_SWITCH:
 		{
-			const uint8* data = Program.GetOpArgsPointer(item.At);
+			const uint8* Data = Program.GetOpArgsPointer(item.At);
 		
 			OP::ADDRESS VarAddress;
-			FMemory::Memcpy( &VarAddress, data, sizeof(OP::ADDRESS));
-			data += sizeof(OP::ADDRESS);
+			FMemory::Memcpy(&VarAddress, Data, sizeof(OP::ADDRESS));
+			Data += sizeof(OP::ADDRESS);
 
 			OP::ADDRESS DefAddress;
-			FMemory::Memcpy( &DefAddress, data, sizeof(OP::ADDRESS));
-			data += sizeof(OP::ADDRESS);
+			FMemory::Memcpy(&DefAddress, Data, sizeof(OP::ADDRESS));
+			Data += sizeof(OP::ADDRESS);
 
 			uint32 CaseCount;
-			FMemory::Memcpy( &CaseCount, data, sizeof(uint32));
-			data += sizeof(uint32);
+			FMemory::Memcpy(&CaseCount, Data, sizeof(uint32));
+			Data += sizeof(uint32);
 	
 			switch (item.Stage)
 			{
@@ -6215,28 +7112,31 @@
 			case 1:
 			{
 				// Get the variable result
-				int var = LoadInt(FCacheAddress(VarAddress, item));
+				FCacheAddress IntLoadAddress = FCacheAddress(VarAddress, item);
+				IntLoadAddress.Type = FScheduledOp::EType::Full;
+				
+				int32 Var = LoadInt(IntLoadAddress);
 
-				OP::ADDRESS valueAt = DefAddress;
+				OP::ADDRESS ValueAt = DefAddress;
 				for (uint32 C = 0; C < CaseCount; ++C)
 				{
 					int32 Condition;
-					FMemory::Memcpy( &Condition, data, sizeof(int32) );
-					data += sizeof(int32);
+					FMemory::Memcpy( &Condition, Data, sizeof(int32) );
+					Data += sizeof(int32);
 					
 					OP::ADDRESS At;
-					FMemory::Memcpy( &At, data, sizeof(OP::ADDRESS) );
-					data += sizeof(OP::ADDRESS);
+					FMemory::Memcpy(&At, Data, sizeof(OP::ADDRESS));
+					Data += sizeof(OP::ADDRESS);
 
-					if (At && var == (int)Condition)
+					if (At && Var == (int32)Condition)
 					{
-						valueAt = At;
+						ValueAt = At;
 						break;
 					}
 				}
 
-				AddOp(FScheduledOp(item.At, item, 2, valueAt),
-					  FScheduledOp(valueAt, item));
+				AddOp(FScheduledOp(item.At, item, 2, ValueAt),
+					  FScheduledOp(ValueAt, item));
 
 				break;
 			}
@@ -6469,7 +7369,9 @@
 
 			case 1:
 			{
-				Ptr<const Layout> pLayout = LoadLayout(FCacheAddress(args.layout, item));
+				FCacheAddress LayoutLoadAddress = FCacheAddress(args.layout, item);
+				LayoutLoadAddress.Type =  FScheduledOp::EType::Full; 
+				Ptr<const Layout> pLayout = LoadLayout(LayoutLoadAddress);
 
 				FIntPoint SizeInBlocks = pLayout->GetGridSize();
 				FIntPoint BlockSizeInPixels(args.blockSize[0], args.blockSize[1]);
diff --git a/Engine/Plugins/Experimental/Mutable/Source/MutableRuntime/Private/MuR/CodeRunner.h b/Engine/Plugins/Experimental/Mutable/Source/MutableRuntime/Private/MuR/CodeRunner.h
--- a/Engine/Plugins/Experimental/Mutable/Source/MutableRuntime/Private/MuR/CodeRunner.h
+++ b/Engine/Plugins/Experimental/Mutable/Source/MutableRuntime/Private/MuR/CodeRunner.h
@@ -42,15 +42,15 @@ namespace  mu
     {
 		// The private token allows only members or friends to call MakeShared.
 		struct FPrivateToken { explicit FPrivateToken() = default; };
-
 	public:
+		
 		static TSharedRef<CodeRunner> Create(
 				const Ptr<const Settings>&, 
 				class System::Private*, 
 				EExecutionStrategy,
 				const TSharedPtr<const Model>&, 
-				const Parameters* pParams,
-				OP::ADDRESS at, uint32 lodMask, uint8 executionOptions, int32 InImageLOD, FScheduledOp::EType);
+				const Parameters* Params,
+				OP::ADDRESS At, uint32 LodMask, uint8 ExecutionOptions, int32 InImageLOD, FScheduledOp::EType);
 
 		// Private constructor to prevent stack allocation. In general we can not call AsShared() if the lifetime is
 		// bounded.
@@ -138,6 +138,12 @@ namespace  mu
         void RunCode_Projector(const FScheduledOp&, const Parameters*, const Model* );
 
 		void RunCodeImageDesc(const FScheduledOp&, const Parameters*, const Model*, uint32 LodMask);
+	
+		/**
+	 	* Run code for meshes components. Any resource that is referenced or owned by a mesh that is not
+	 	* geometry data. 
+	 	*/
+		void RunCodeMeshComponents(const FScheduledOp&, const Parameters*, const Model*, uint32 LodMask);
 
     public:
 
@@ -260,10 +266,12 @@ namespace  mu
 		class FLoadMeshRomTask : public CodeRunner::FIssuedTask
 		{
 		public:
-			FLoadMeshRomTask( const FScheduledOp& InOp, int32 InRomIndex)
+			FLoadMeshRomTask( const FScheduledOp& InOp, int32 InGeometryRomIndex, int32 InPoseRomIndex, int32 InComponentsRomIndex)
 				: FIssuedTask(InOp)
+				, GeometryRomIndex(InGeometryRomIndex)
+				, PoseRomIndex(InPoseRomIndex)
+				, ComponentsRomIndex(InComponentsRomIndex)
 			{
-				RomIndex = InRomIndex;
 			}
 
 			// FIssuedTask interface
@@ -271,7 +279,11 @@ namespace  mu
 			virtual bool Complete(CodeRunner*) override;
 
 		private:
-			int32 RomIndex = -1;
+			int32 GeometryRomIndex = -1;
+			int32 PoseRomIndex = -1;
+			int32 ComponentsRomIndex = -1;
+
+			TArray<int32, TInlineAllocator<3>> RomIndices;
 		};
 
 		class FLoadExtensionDataTask : public CodeRunner::FIssuedTask
@@ -693,21 +705,23 @@ namespace  mu
     	
     	FRomLoadOps RomLoadOps = FRomLoadOps(*this);
     	    	
-		inline void AddChildren(const FScheduledOp& dep)
+		inline void AddChildren(const FScheduledOp& Dep)
 		{
-			FCacheAddress at(dep);
-			if (dep.At && !GetMemory().IsValid(at))
+			FCacheAddress At(Dep);
+			if (Dep.At && !GetMemory().IsValid(At))
 			{
-				if (ScheduledStagePerOp.get(at) <= dep.Stage)
+				if (ScheduledStagePerOp.Get(At) <= Dep.Stage)
 				{
-					OpenTasks.Add(dep);
-					ScheduledStagePerOp[at] = dep.Stage + 1;
+					OpenTasks.Add(Dep);
+					ScheduledStagePerOp[At] = Dep.Stage + 1;
 				}
 			}
 
-			if (dep.Type == FScheduledOp::EType::Full)
+			// ImageDesc is the exception. Will not decrease the hit counter because results are not
+			// stored in the cache.
+			if (Dep.Type != FScheduledOp::EType::ImageDesc)
 			{
-				m_pSystem->WorkingMemoryManager.CurrentInstanceCache->IncreaseHitCount(at);
+				m_pSystem->WorkingMemoryManager.CurrentInstanceCache->IncreaseHitCount(At);
 			}
 		}
 
diff --git a/Engine/Plugins/Experimental/Mutable/Source/MutableRuntime/Private/MuR/CodeRunnerTasked.cpp b/Engine/Plugins/Experimental/Mutable/Source/MutableRuntime/Private/MuR/CodeRunnerTasked.cpp
--- a/Engine/Plugins/Experimental/Mutable/Source/MutableRuntime/Private/MuR/CodeRunnerTasked.cpp
+++ b/Engine/Plugins/Experimental/Mutable/Source/MutableRuntime/Private/MuR/CodeRunnerTasked.cpp
@@ -435,7 +435,18 @@ namespace mu
 				else
 				{
 					// Run immediately
-					RunCode(item, m_pParams, m_pModel, m_lodMask);
+					if (item.Type == FScheduledOp::EType::Full)
+					{
+						RunCode(item, m_pParams, m_pModel, m_lodMask);
+					}
+					else if (item.Type == FScheduledOp::EType::MeshComp)
+					{
+						RunCodeMeshComponents(item, m_pParams, m_pModel.Get(), m_lodMask);
+					}
+					else
+					{
+						checkf(false, TEXT("Unexpected execution type."));
+					}
 
 					if (ScheduledStagePerOp[item] == item.Stage + 1)
 					{
@@ -2091,60 +2102,85 @@ namespace mu
 		bOutFailed = false;
 
 		const FProgram& Program = Runner->m_pModel->GetPrivate()->m_program;
-
-		check(RomIndex < Program.m_roms.Num());
-		
 		FWorkingMemoryManager::FModelCacheEntry* ModelCache = Runner->m_pSystem->WorkingMemoryManager.FindModelCache(Runner->m_pModel.Get());
-		++ModelCache->PendingOpsPerRom[RomIndex];
+		
+		TArray<UE::Tasks::FTask, TInlineAllocator<3>> ReadCompleteEvents;
+		ReadCompleteEvents.Reserve(3); 
 
-		if (Program.IsRomLoaded(RomIndex))
+		const TArray<int32, TFixedAllocator<3>> TaskRoms = {GeometryRomIndex, PoseRomIndex, ComponentsRomIndex};
+		for (int32 RomIndex : TaskRoms)
 		{
-			return false;
-		}
+			if (RomIndex < 0)
+			{
+				continue;
+			}
 
-		if (const FRomLoadOp* Result = Runner->RomLoadOps.Find(RomIndex))
-		{
-			Event = Result->Event;  // Wait for the read operation started by other task
-			return false;
-		}
+			check(RomIndex < Program.m_roms.Num());
 
-		FRomLoadOp& RomLoadOp = Runner->RomLoadOps.Create(RomIndex);
-		
-		check(Runner->m_pSystem->StreamInterface);
+			++ModelCache->PendingOpsPerRom[RomIndex];
 
-		const uint32 RomSize = Program.m_roms[RomIndex].Size;
-		check(RomSize > 0);
+			if (DebugRom && (DebugRomAll || RomIndex == DebugRomIndex))
+			{
+				UE_LOG(LogMutableCore, Log, TEXT("Preparing rom %d, now peding ops is %d."), RomIndex, ModelCache->PendingOpsPerRom[RomIndex]);
+			}
 
-		// Free roms if necessary
-		{
-			MUTABLE_CPUPROFILER_SCOPE(FreeingRoms);
+			if (Program.IsRomLoaded(RomIndex))
+			{
+				continue;
+			}
 
-			Runner->m_pSystem->WorkingMemoryManager.MarkRomUsed(RomIndex, Runner->m_pModel);
-			Runner->m_pSystem->WorkingMemoryManager.EnsureBudgetBelow(RomSize);
-		}
+			RomIndices.Add(RomIndex);
+
+			if (const FRomLoadOp* Result = Runner->RomLoadOps.Find(RomIndex))
+			{
+				ReadCompleteEvents.Add(Result->Event); // Wait for the read operation started by other task
+				continue;
+			}
 
-		const int32 SizeBefore = RomLoadOp.m_streamBuffer.GetAllocatedSize();
-		RomLoadOp.m_streamBuffer.SetNumUninitialized(RomSize);
-		const int32 SizeAfter = RomLoadOp.m_streamBuffer.GetAllocatedSize();
+			FRomLoadOp& RomLoadOp = Runner->RomLoadOps.Create(RomIndex);
+			
+			check(Runner->m_pSystem->StreamInterface);
 
-	 	UE::Tasks::FTaskEvent ReadCompletionEvent = UE::Tasks::FTaskEvent(TEXT("FLoadMeshRomTask"));	
-		RomLoadOp.Event = ReadCompletionEvent;
-		
-		TFunction<void(bool)> Callback = [ReadCompletionEvent](bool bSuccess) mutable
-		{
-			ReadCompletionEvent.Trigger();
-		};
-		
-		const uint32 RomId = Program.m_roms[RomIndex].Id;
-		RomLoadOp.m_streamID = Runner->m_pSystem->StreamInterface->BeginReadBlock(Runner->m_pModel.Get(), RomId, RomLoadOp.m_streamBuffer.GetData(), RomSize, &Callback);
-		if (RomLoadOp.m_streamID < 0)
-		{
-			bOutFailed = true;
-			return false;
+			const uint32 RomSize = Program.m_roms[RomIndex].Size;
+			check(RomSize > 0);
+
+			// Free roms if necessary
+			{
+				Runner->m_pSystem->WorkingMemoryManager.MarkRomUsed(RomIndex, Runner->m_pModel);
+				Runner->m_pSystem->WorkingMemoryManager.EnsureBudgetBelow(RomSize);
+			}
+
+			const int32 SizeBefore = RomLoadOp.m_streamBuffer.GetAllocatedSize();
+			RomLoadOp.m_streamBuffer.SetNumUninitialized(RomSize);
+			const int32 SizeAfter = RomLoadOp.m_streamBuffer.GetAllocatedSize();
+
+			UE::Tasks::FTaskEvent ReadCompletionEvent(TEXT("FLoadMeshRomsTaskRom"));
+			ReadCompleteEvents.Add(ReadCompletionEvent);
+			RomLoadOp.Event = ReadCompletionEvent;
+
+			TFunction<void(bool)> Callback = [ReadCompletionEvent](bool bSuccess) mutable // Mutable due Trigger not being const
+			{
+				ReadCompletionEvent.Trigger();
+			};
+			
+			const uint32 RomId = Program.m_roms[RomIndex].Id;
+			RomLoadOp.m_streamID = Runner->m_pSystem->StreamInterface->BeginReadBlock(Runner->m_pModel.Get(), RomId, RomLoadOp.m_streamBuffer.GetData(), RomSize, &Callback);
+			if (RomLoadOp.m_streamID < 0)
+			{
+				bOutFailed = true;
+				return false;
+			}
 		}
 
-		Event = ReadCompletionEvent; // Wait for read operation to end
+		// Wait for all read operations to end
+		UE::Tasks::FTaskEvent GatherReadsCompletionEvent(TEXT("FLoadMeshRomsTask"));
+		GatherReadsCompletionEvent.AddPrerequisites(ReadCompleteEvents);
+		GatherReadsCompletionEvent.Trigger();
+
+		Event = GatherReadsCompletionEvent;
+			
 		return false; // No worker thread work
+
 	}
 	
 
@@ -2160,38 +2196,67 @@ namespace mu
 		}
 
 		FProgram& Program = Runner->m_pModel->GetPrivate()->m_program;
-
-		// Since task could be reordered, we need to make sure we end the rom read before continuing
-		if (FRomLoadOp* RomLoadOp = Runner->RomLoadOps.Find(RomIndex))
+		FWorkingMemoryManager::FModelCacheEntry* ModelCache = Runner->m_pSystem->WorkingMemoryManager.FindModelCache(Runner->m_pModel.Get());
+		
+		bool bSomeMissingData = false;
+		for (const int32 RomIndex : RomIndices)
 		{
-			Runner->m_pSystem->StreamInterface->EndRead(RomLoadOp->m_streamID);									
+			// Since task could be reordered, we need to make sure we end the rom read before continuing
+			if (FRomLoadOp* RomLoadOp = Runner->RomLoadOps.Find(RomIndex))
+			{				
+				bool bSuccess = Runner->m_pSystem->StreamInterface->EndRead(RomLoadOp->m_streamID);
 
-			const int32 ResIndex = Program.m_roms[RomIndex].ResourceIndex;
-			check(!Program.ConstantMeshes[ResIndex].Value)
-			MUTABLE_CPUPROFILER_SCOPE(Unserialise);
+				MUTABLE_CPUPROFILER_SCOPE(Unserialise);
 
-			InputMemoryStream Stream(RomLoadOp->m_streamBuffer.GetData(), RomLoadOp->m_streamBuffer.Num());
-			InputArchive Arch(&Stream);
+				if (bSuccess)
+				{
+					InputMemoryStream Stream(RomLoadOp->m_streamBuffer.GetData(), RomLoadOp->m_streamBuffer.Num());
+					InputArchive Arch(&Stream);
 
-			check(!Program.ConstantMeshes[ResIndex].Value);
-			Ptr<Mesh> Value = Mesh::StaticUnserialise(Arch);
-			Program.SetMeshRomValue(RomIndex, Value);
+					const int32 ResourceIndex = Program.m_roms[RomIndex].ResourceIndex;
+						
+					check(!Program.ConstantMeshData[ResourceIndex].Value);
+					Ptr<Mesh> Value = Mesh::StaticUnserialise(Arch);
 
-			check(Program.ConstantMeshes[ResIndex].Value);
+					Program.SetMeshRomValue(RomIndex, Value);
+					check(Program.ConstantMeshData[ResourceIndex].Value);
+				}
+				else
+				{
+					bSomeMissingData = true;
+				}
+				
+				Runner->RomLoadOps.Remove(*RomLoadOp);
+			}
+		}
 
-			Runner->RomLoadOps.Remove(*RomLoadOp);
+		if (bSomeMissingData)
+		{
+			// Some data may be missing. We can try to go on if some requested mesh parts are there. 
+			UE_LOG(LogMutableCore, Verbose, TEXT("FLoadMeshRomsTask::Complete failed: missing data?"));
 		}
 
-		// Process the constant op normally, now that the rom is loaded.
-		Runner->RunCode(Op, Runner->m_pParams, Runner->m_pModel, Runner->m_lodMask);
+		// Process the constant op normally, now that the rom is loaded.	
+		bool bSuccess = Runner->RunCode_ConstantResource(Op, Runner->m_pModel.Get());
 
-		FWorkingMemoryManager::FModelCacheEntry* ModelCache = Runner->m_pSystem->WorkingMemoryManager.FindModelCache(Runner->m_pModel.Get());
+		if (bSuccess)
+		{
+			for (int32 RomIndex : RomIndices)
+			{
+				check(RomIndex < Program.m_roms.Num());
 
-		Runner->m_pSystem->WorkingMemoryManager.MarkRomUsed(RomIndex, Runner->m_pModel);
-		--ModelCache->PendingOpsPerRom[RomIndex];
+				Runner->m_pSystem->WorkingMemoryManager.MarkRomUsed(RomIndex, Runner->m_pModel);
+				--ModelCache->PendingOpsPerRom[RomIndex];
+
+				if (DebugRom && (DebugRomAll || RomIndex == DebugRomIndex))
+				{
+					UE_LOG(LogMutableCore, Log, TEXT("FLoadMeshRomsTask::Complete rom %d, now peding ops is %d."), RomIndex, ModelCache->PendingOpsPerRom[RomIndex]);
+				}
+			}
+		}
 
-		bool bSuccess = true;
 		return bSuccess;
+
 	}
 
 
@@ -2602,39 +2667,59 @@ namespace mu
 	}
 
 
-	//---------------------------------------------------------------------------------------------
-	//---------------------------------------------------------------------------------------------
-	//---------------------------------------------------------------------------------------------
-	TSharedPtr<CodeRunner::FIssuedTask> CodeRunner::IssueOp(FScheduledOp item)
+	TSharedPtr<CodeRunner::FIssuedTask> CodeRunner::IssueOp(FScheduledOp Item)
 	{
 		TSharedPtr<FIssuedTask> Issued;
 
-		FProgram& program = m_pModel->GetPrivate()->m_program;
+		FProgram& Program = m_pModel->GetPrivate()->m_program;
 
-		OP_TYPE type = program.GetOpType(item.At);
+		OP_TYPE Type = Program.GetOpType(Item.At);
 
-		switch (type)
+		switch (Type)
 		{
 		case OP_TYPE::ME_CONSTANT:
 		{
-			OP::MeshConstantArgs args = program.GetOpArgs<OP::MeshConstantArgs>(item.At);
-			int32 RomIndex = program.ConstantMeshes[args.value].Key;
-			if (RomIndex >= 0 && !program.ConstantMeshes[args.value].Value)
+			// Mesh constants have three parts. Geometry, pose and components. 
+
+			OP::MeshConstantArgs Args = Program.GetOpArgs<OP::MeshConstantArgs>(Item.At);
+
+			int32 GeometryRomIndex = -1;
+			int32 PoseRomIndex = -1;
+			int32 ComponentsRomIndex = -1;
+		
+			const FMeshRange MeshRange = Program.ConstantMeshes[Args.Value];
+			const EMeshExecutionOptions ExecutionOptions = static_cast<EMeshExecutionOptions>(Item.ExecutionOptions);
+			
+			if (EnumHasAnyFlags(ExecutionOptions, EMeshExecutionOptions::LoadGeometryData))
 			{
-				Issued = MakeShared<FLoadMeshRomTask>(item, RomIndex);
+				int32 GeometryConstantIndex = Program.ConstantMeshIndices[MeshRange.FirstIndex + 0];
+				GeometryRomIndex = Program.ConstantMeshData[GeometryConstantIndex].Key;
 			}
-			else
+	
+			// Pose data will be loaded with geometry as some geometry operations need skinning data.
+			// This could be tuned on a Op to Op basis but it might not be worth it.
+			if (EnumHasAnyFlags(ExecutionOptions, EMeshExecutionOptions::LoadPoseData | EMeshExecutionOptions::LoadGeometryData))
 			{
-				// If already available, the rest of the constant code will run right away.
+				int32 PoseConstantIndex = Program.ConstantMeshIndices[MeshRange.FirstIndex + 1];
+				PoseRomIndex = Program.ConstantMeshData[PoseConstantIndex].Key;
+			}
+
+			if (EnumHasAnyFlags(ExecutionOptions, EMeshExecutionOptions::LoadComponentsData))
+			{
+				int32 ComponentsConstantIndex = Program.ConstantMeshIndices[MeshRange.FirstIndex + 2];
+				ComponentsRomIndex = Program.ConstantMeshData[ComponentsConstantIndex].Key;
 			}
+
+			Issued = MakeShared<FLoadMeshRomTask>(Item, GeometryRomIndex, PoseRomIndex, ComponentsRomIndex);
+
 			break;
 		}
 		
 		case OP_TYPE::ED_CONSTANT:
 		{
-			OP::ResourceConstantArgs Args = program.GetOpArgs<OP::ResourceConstantArgs>(item.At);
+			OP::ResourceConstantArgs Args = Program.GetOpArgs<OP::ResourceConstantArgs>(Item.At);
 
-			const FExtensionDataConstant::ELoadState LoadState = program.m_constantExtensionData[Args.value].LoadState;
+			const FExtensionDataConstant::ELoadState LoadState = Program.m_constantExtensionData[Args.value].LoadState;
 
 			check(LoadState != FExtensionDataConstant::ELoadState::Invalid);
 			
@@ -2645,19 +2730,19 @@ namespace mu
 			}
 			else
 			{
-				Issued = MakeShared<FLoadExtensionDataTask>(item, program.m_constantExtensionData[Args.value].Data, Args.value);
+				Issued = MakeShared<FLoadExtensionDataTask>(Item, Program.m_constantExtensionData[Args.value].Data, Args.value);
 			}
 			break;
 		}
 
 		case OP_TYPE::IM_CONSTANT:
 		{
-			OP::ResourceConstantArgs args = program.GetOpArgs<OP::ResourceConstantArgs>(item.At);
-			int32 MipsToSkip = item.ExecutionOptions;
+			OP::ResourceConstantArgs args = Program.GetOpArgs<OP::ResourceConstantArgs>(Item.At);
+			int32 MipsToSkip = Item.ExecutionOptions;
 			int32 ImageIndex = args.value;
-			int32 ReallySkip = FMath::Min(MipsToSkip, program.ConstantImages[ImageIndex].LODCount - 1);
-			int32 LODIndexIndex = program.ConstantImages[ImageIndex].FirstIndex + ReallySkip;
-			int32 LODIndexCount = program.ConstantImages[ImageIndex].LODCount - ReallySkip;
+			int32 ReallySkip = FMath::Min(MipsToSkip, Program.ConstantImages[ImageIndex].LODCount - 1);
+			int32 LODIndexIndex = Program.ConstantImages[ImageIndex].FirstIndex + ReallySkip;
+			int32 LODIndexCount = Program.ConstantImages[ImageIndex].LODCount - ReallySkip;
 			check(LODIndexCount > 0);
 
 			// We always need to follow this path, or roms may not be protected for long enough and might be unloaded 
@@ -2666,8 +2751,8 @@ namespace mu
 			//bool bAnyMissing = false;
 			//for (int32 i=0; i<LODIndexCount; ++i)
 			//{
-			//	uint32 LODIndex = program.ConstantImageLODIndices[LODIndexIndex+i];
-			//	if ( !program.ConstantImageLODs[LODIndex].Value )
+			//	uint32 LODIndex = Program.ConstantImageLODIndices[LODIndexIndex+i];
+			//	if ( !Program.ConstantImageLODs[LODIndex].Value )
 			//	{
 			//		bAnyMissing = true;
 			//		break;
@@ -2676,7 +2761,7 @@ namespace mu
 
 			if (bAnyMissing)
 			{
-				Issued = MakeShared<FLoadImageRomsTask>(item, LODIndexIndex, LODIndexCount);
+				Issued = MakeShared<FLoadImageRomsTask>(Item, LODIndexIndex, LODIndexCount);
 
 				if (DebugRom && (DebugRomAll || ImageIndex == DebugImageIndex))
 					UE_LOG(LogMutableCore, Log, TEXT("Issuing image %d skipping %d ."), ImageIndex, ReallySkip);
@@ -2692,37 +2777,37 @@ namespace mu
 
 		case OP_TYPE::IM_PARAMETER:
 		{
-			OP::ParameterArgs args = program.GetOpArgs<OP::ParameterArgs>(item.At);
-			Ptr<RangeIndex> Index = BuildCurrentOpRangeIndex(item, m_pParams, m_pModel.Get(), args.variable);
+			OP::ParameterArgs args = Program.GetOpArgs<OP::ParameterArgs>(Item.At);
+			Ptr<RangeIndex> Index = BuildCurrentOpRangeIndex(Item, m_pParams, m_pModel.Get(), args.variable);
 
 			const FName Id = m_pParams->GetImageValue(args.variable, Index);
 
 			check(ImageLOD < TNumericLimits<uint8>::Max() && ImageLOD >= 0);
-			check(ImageLOD + static_cast<int32>(item.ExecutionOptions) < TNumericLimits<uint8>::Max());
+			check(ImageLOD + static_cast<int32>(Item.ExecutionOptions) < TNumericLimits<uint8>::Max());
 
-			const uint8 MipmapsToSkip = item.ExecutionOptions + static_cast<uint8>(ImageLOD);
+			const uint8 MipmapsToSkip = Item.ExecutionOptions + static_cast<uint8>(ImageLOD);
 
 			CodeRunner::FExternalResourceId FullId;
 			FullId.ParameterId = Id;
-			Issued = MakeShared<FImageExternalLoadTask>(item, MipmapsToSkip, FullId);
+			Issued = MakeShared<FImageExternalLoadTask>(Item, MipmapsToSkip, FullId);
 
 			break;
 		}
 
 		case OP_TYPE::IM_REFERENCE:
 		{
-			OP::ResourceReferenceArgs Args = m_pModel->GetPrivate()->m_program.GetOpArgs<OP::ResourceReferenceArgs>(item.At);
+			OP::ResourceReferenceArgs Args = m_pModel->GetPrivate()->m_program.GetOpArgs<OP::ResourceReferenceArgs>(Item.At);
 
 			// We only convert references to images if indicated in the operation.
 			if (Args.ForceLoad)
 			{
-				check(item.Stage==0);
+				check(Item.Stage==0);
 
-				const uint8 MipmapsToSkip = item.ExecutionOptions + static_cast<uint8>(ImageLOD);
+				const uint8 MipmapsToSkip = Item.ExecutionOptions + static_cast<uint8>(ImageLOD);
 
 				FExternalResourceId FullId;
 				FullId.ReferenceResourceId = Args.ID;
-				Issued = MakeShared<FImageExternalLoadTask>(item, MipmapsToSkip, FullId);
+				Issued = MakeShared<FImageExternalLoadTask>(Item, MipmapsToSkip, FullId);
 			}
 
 			break;
@@ -2730,16 +2815,17 @@ namespace mu
 
 		case OP_TYPE::ME_REFERENCE:
 		{
-			OP::ResourceReferenceArgs Args = m_pModel->GetPrivate()->m_program.GetOpArgs<OP::ResourceReferenceArgs>(item.At);
+			OP::ResourceReferenceArgs Args = m_pModel->GetPrivate()->m_program.GetOpArgs<OP::ResourceReferenceArgs>(Item.At);
 
 			// We only convert references to meshes if indicated in the operation.
 			if (Args.ForceLoad)
 			{
-				check(item.Stage == 0);
+				check(Item.Stage == 0);
+				checkf(Item.Type != FScheduledOp::EType::MeshComp, TEXT("Operation not implemented with mesh components execution mode."));
 
 				FExternalResourceId FullId;
 				FullId.ReferenceResourceId = Args.ID;
-				Issued = MakeShared<FMeshExternalLoadTask>(item, FullId);
+				Issued = MakeShared<FMeshExternalLoadTask>(Item, FullId);
 			}
 
 			break;
@@ -2747,105 +2833,105 @@ namespace mu
 
 		case OP_TYPE::IM_PIXELFORMAT:
 		{
-			if (item.Stage == 1)
+			if (Item.Stage == 1)
 			{
-				OP::ImagePixelFormatArgs Args = program.GetOpArgs<OP::ImagePixelFormatArgs>(item.At);
-				Issued = MakeShared<FImagePixelFormatTask>(item, Args);
+				OP::ImagePixelFormatArgs Args = Program.GetOpArgs<OP::ImagePixelFormatArgs>(Item.At);
+				Issued = MakeShared<FImagePixelFormatTask>(Item, Args);
 			}
 			break;
 		}
 
 		case OP_TYPE::IM_LAYERCOLOUR:
 		{
-			if (item.Stage == 1)
+			if (Item.Stage == 1)
 			{
-				OP::ImageLayerColourArgs Args = program.GetOpArgs<OP::ImageLayerColourArgs>(item.At);
-				Issued = MakeShared<FImageLayerColourTask>(item, Args);
+				OP::ImageLayerColourArgs Args = Program.GetOpArgs<OP::ImageLayerColourArgs>(Item.At);
+				Issued = MakeShared<FImageLayerColourTask>(Item, Args);
 			}
 			break;
 		}
 
 		case OP_TYPE::IM_LAYER:
 		{
-			if ((ExecutionStrategy == EExecutionStrategy::MinimizeMemory && item.Stage == 2)
+			if ((ExecutionStrategy == EExecutionStrategy::MinimizeMemory && Item.Stage == 2)
 				||
-				(ExecutionStrategy != EExecutionStrategy::MinimizeMemory && item.Stage == 1)
+				(ExecutionStrategy != EExecutionStrategy::MinimizeMemory && Item.Stage == 1)
 				)
 			{
-				OP::ImageLayerArgs Args = program.GetOpArgs<OP::ImageLayerArgs>(item.At);
-				Issued = MakeShared<FImageLayerTask>(item, Args);
+				OP::ImageLayerArgs Args = Program.GetOpArgs<OP::ImageLayerArgs>(Item.At);
+				Issued = MakeShared<FImageLayerTask>(Item, Args);
 			}
 			break;
 		}
 
 		case OP_TYPE::IM_MIPMAP:
 		{
-			if (item.Stage == 1)
+			if (Item.Stage == 1)
 			{
-				OP::ImageMipmapArgs Args = program.GetOpArgs<OP::ImageMipmapArgs>(item.At);
-				Issued = MakeShared<FImageMipmapTask>(item, Args);
+				OP::ImageMipmapArgs Args = Program.GetOpArgs<OP::ImageMipmapArgs>(Item.At);
+				Issued = MakeShared<FImageMipmapTask>(Item, Args);
 			}
 			break;
 		}
 
 		case OP_TYPE::IM_SWIZZLE:
 		{
-			if (item.Stage == 1)
+			if (Item.Stage == 1)
 			{
-				OP::ImageSwizzleArgs Args = program.GetOpArgs<OP::ImageSwizzleArgs>(item.At);
-				Issued = MakeShared<FImageSwizzleTask>(item, Args);
+				OP::ImageSwizzleArgs Args = Program.GetOpArgs<OP::ImageSwizzleArgs>(Item.At);
+				Issued = MakeShared<FImageSwizzleTask>(Item, Args);
 			}
 			break;
 		}
 
 		case OP_TYPE::IM_SATURATE:
 		{
-			if (item.Stage == 1)
+			if (Item.Stage == 1)
 			{
-				OP::ImageSaturateArgs Args = program.GetOpArgs<OP::ImageSaturateArgs>(item.At);
-				Issued = MakeShared<FImageSaturateTask>(item, Args);
+				OP::ImageSaturateArgs Args = Program.GetOpArgs<OP::ImageSaturateArgs>(Item.At);
+				Issued = MakeShared<FImageSaturateTask>(Item, Args);
 			}
 			break;
 		}
 
 		case OP_TYPE::IM_INVERT:
 		{
-			if (item.Stage == 1)
+			if (Item.Stage == 1)
 			{
-				OP::ImageInvertArgs Args = program.GetOpArgs<OP::ImageInvertArgs>(item.At);
-				Issued = MakeShared<FImageInvertTask>(item, Args);
+				OP::ImageInvertArgs Args = Program.GetOpArgs<OP::ImageInvertArgs>(Item.At);
+				Issued = MakeShared<FImageInvertTask>(Item, Args);
 			}
 			break;
 		}
 
 		case OP_TYPE::IM_RESIZE:
 		{
-			if (item.Stage == 1)
+			if (Item.Stage == 1)
 			{
-				OP::ImageResizeArgs Args = program.GetOpArgs<OP::ImageResizeArgs>(item.At);
-				Issued = MakeShared<FImageResizeTask>(item, Args);
+				OP::ImageResizeArgs Args = Program.GetOpArgs<OP::ImageResizeArgs>(Item.At);
+				Issued = MakeShared<FImageResizeTask>(Item, Args);
 			}
 			break;
 		}
 
 		case OP_TYPE::IM_RESIZEREL:
 		{
-			if (item.Stage == 1)
+			if (Item.Stage == 1)
 			{
-				OP::ImageResizeRelArgs Args = program.GetOpArgs<OP::ImageResizeRelArgs>(item.At);
-				Issued = MakeShared<FImageResizeRelTask>(item, Args);
+				OP::ImageResizeRelArgs Args = Program.GetOpArgs<OP::ImageResizeRelArgs>(Item.At);
+				Issued = MakeShared<FImageResizeRelTask>(Item, Args);
 			}
 			break;
 		}
 
 		case OP_TYPE::IM_COMPOSE:
 		{
-			if ((ExecutionStrategy == EExecutionStrategy::MinimizeMemory && item.Stage == 3) ||
-				(ExecutionStrategy != EExecutionStrategy::MinimizeMemory && item.Stage == 2))
+			if ((ExecutionStrategy == EExecutionStrategy::MinimizeMemory && Item.Stage == 3) ||
+				(ExecutionStrategy != EExecutionStrategy::MinimizeMemory && Item.Stage == 2))
 			{
-				OP::ImageComposeArgs Args = program.GetOpArgs<OP::ImageComposeArgs>(item.At);
-				Ptr<const Layout> ComposeLayout = static_cast<const Layout*>( m_heapData[item.CustomState].Resource.get());
-				Issued = MakeShared<FImageComposeTask>(item, Args, ComposeLayout);
+				OP::ImageComposeArgs Args = Program.GetOpArgs<OP::ImageComposeArgs>(Item.At);
+				Ptr<const Layout> ComposeLayout = static_cast<const Layout*>( m_heapData[Item.CustomState].Resource.get());
+				Issued = MakeShared<FImageComposeTask>(Item, Args, ComposeLayout);
 			}
 			break;
 		}
diff --git a/Engine/Plugins/Experimental/Mutable/Source/MutableRuntime/Private/MuR/Mesh.cpp b/Engine/Plugins/Experimental/Mutable/Source/MutableRuntime/Private/MuR/Mesh.cpp
--- a/Engine/Plugins/Experimental/Mutable/Source/MutableRuntime/Private/MuR/Mesh.cpp
+++ b/Engine/Plugins/Experimental/Mutable/Source/MutableRuntime/Private/MuR/Mesh.cpp
@@ -73,51 +73,12 @@ uint32 Mesh::GetReferencedMesh() const
 	return ReferenceID;
 }
 
-
 Ptr<Mesh> Mesh::Clone() const
 {
     //MUTABLE_CPUPROFILER_SCOPE(MeshClone);
-	LLM_SCOPE_BYNAME(TEXT("MutableRuntime"));
-
-    Ptr<Mesh> Result = new Mesh();
-
-    Result->InternalId = InternalId;
-	Result->Flags = Flags;
-	Result->ReferenceID = ReferenceID;
-	Result->Surfaces = Surfaces;
-	Result->Skeleton = Skeleton;
-	Result->PhysicsBody = PhysicsBody;
-	Result->Tags = Tags;
-	Result->StreamedResources = StreamedResources;
-	Result->MeshIDPrefix = MeshIDPrefix;
-
-    // Clone the main buffers
-    Result->VertexBuffers = VertexBuffers;
-    Result->IndexBuffers = IndexBuffers;
-
-	// Clone additional buffers
-	Result->AdditionalBuffers = AdditionalBuffers;
-
-    // Clone the layouts
-	Result->Layouts = Layouts;
-
-    // The skeleton is not cloned because it is not owned by this mesh and it is always assumed
-    // to be shared.
-
-	// physics body doen't need to be deep cloned either as they are also assumed to be shared.
-	
-	// Clone bone poses
-	Result->BonePoses = BonePoses;
-	Result->BoneMap = BoneMap;
-
-	// Clone SkeletonIDs
-	Result->SkeletonIDs = SkeletonIDs;
-	Result->AdditionalPhysicsBodies = AdditionalPhysicsBodies;
-
-    return Result;
+    return Clone(EMeshCopyFlags::AllFlags);
 }
 
-
 Ptr<Mesh> Mesh::Clone(EMeshCopyFlags InFlags) const
 {
     //MUTABLE_CPUPROFILER_SCOPE(MeshClone);
@@ -125,100 +86,21 @@ Ptr<Mesh> Mesh::Clone(EMeshCopyFlags InFlags) const
 
 	Ptr<Mesh> Result = new Mesh();
 
-    Result->InternalId = InternalId;
-	Result->MeshIDPrefix = MeshIDPrefix;
-	Result->Flags = Flags;
-	Result->ReferenceID = ReferenceID;
-
-	if (EnumHasAnyFlags(InFlags, EMeshCopyFlags::WithSurfaces))
-	{
-		Result->Surfaces = Surfaces;
-	}
-
-	if (EnumHasAnyFlags(InFlags, EMeshCopyFlags::WithSkeleton))
-	{
-		Result->Skeleton = Skeleton;
-	}
-
-	if (EnumHasAnyFlags(InFlags, EMeshCopyFlags::WithPhysicsBody))
-	{
-		Result->PhysicsBody = PhysicsBody;
-	}
-
-	if (EnumHasAnyFlags(InFlags, EMeshCopyFlags::WithTags))
-	{
-		Result->Tags = Tags;
-	}
-
-	if (EnumHasAnyFlags(InFlags, EMeshCopyFlags::WithStreamedResources))
-	{
-		Result->StreamedResources = StreamedResources;
-	}
-
-    // Clone the main buffers
-	if (EnumHasAnyFlags(InFlags, EMeshCopyFlags::WithVertexBuffers))
-    {
-		Result->VertexBuffers = VertexBuffers;
-	}
-
-	if (EnumHasAnyFlags(InFlags, EMeshCopyFlags::WithIndexBuffers))
-    {
-		Result->IndexBuffers = IndexBuffers;
-	}
-
-	// Clone additional buffers
-	if (EnumHasAnyFlags(InFlags, EMeshCopyFlags::WithAdditionalBuffers))
-	{
-		Result->AdditionalBuffers = AdditionalBuffers;
-	}
-
-    // Clone the layout	
-	if (EnumHasAnyFlags(InFlags, EMeshCopyFlags::WithLayouts))
-	{
-		Result->Layouts = Layouts;
-	}
-    // The skeleton is not cloned because it is not owned by this mesh and it is always assumed
-    // to be shared.
-
-	// physics body doen't need to be deep cloned either as they are also assumed to be shared.
-	
-	// Clone bone poses
-	if (EnumHasAnyFlags(InFlags, EMeshCopyFlags::WithPoses))
-	{
-		Result->BonePoses = BonePoses;
-	}
-
-	// Clone BoneMap
-	if (EnumHasAnyFlags(InFlags, EMeshCopyFlags::WithBoneMap))
-	{
-		Result->BoneMap = BoneMap;
-	}
-
-	// Clone SkeletonIDs
-	if (EnumHasAnyFlags(InFlags, EMeshCopyFlags::WithSkeletonIDs))
-	{
-		Result->SkeletonIDs = SkeletonIDs;
-	}
-
-	if (EnumHasAnyFlags(InFlags, EMeshCopyFlags::WithAdditionalPhysics))
-	{
-		Result->AdditionalPhysicsBodies = AdditionalPhysicsBodies;
-	}
-
-    return Result;
+	Result->CopyFrom(*this, InFlags);
+    
+	return Result;
 }
 
 
 void Mesh::CopyFrom(const Mesh& From, EMeshCopyFlags InFlags)
 {
     //MUTABLE_CPUPROFILER_SCOPE(CopyFrom);
-
+	
     InternalId = From.InternalId;
+	MeshIDPrefix = From.MeshIDPrefix;
 	Flags = From.Flags;
 	ReferenceID = From.ReferenceID;
 
-	MeshIDPrefix = From.MeshIDPrefix;
-
 	if (EnumHasAnyFlags(InFlags, EMeshCopyFlags::WithSurfaces))
 	{
 		Surfaces = From.Surfaces;
diff --git a/Engine/Plugins/Experimental/Mutable/Source/MutableRuntime/Private/MuR/MeshBufferSet.cpp b/Engine/Plugins/Experimental/Mutable/Source/MutableRuntime/Private/MuR/MeshBufferSet.cpp
--- a/Engine/Plugins/Experimental/Mutable/Source/MutableRuntime/Private/MuR/MeshBufferSet.cpp
+++ b/Engine/Plugins/Experimental/Mutable/Source/MutableRuntime/Private/MuR/MeshBufferSet.cpp
@@ -190,6 +190,8 @@ namespace mu
 
 		FMeshBuffer& Buffer = Buffers[BufferIndex];
 
+		const bool bWasNotInitialized = Buffer.ElementSize == 0 && Buffer.Channels.IsEmpty() && Buffer.Data.IsEmpty();
+
 		int32 MinElemSize = 0;
 		Buffer.Channels.SetNum(ChannelCount);
 		for (int32 ChannelIndex = 0; ChannelIndex < ChannelCount; ++ChannelIndex)
@@ -208,18 +210,23 @@ namespace mu
 		// Set the user specified element size, or enlarge it if it was too small
 		Buffer.ElementSize = FMath::Max(ElementSize, MinElemSize);
 
+		bool bAllocateMemory = !IsDescriptor() || bWasNotInitialized; 
+
 		// Update the buffer data
-		if (MemoryInitPolicy == EMemoryInitPolicy::Uninitialized)
+		if (bAllocateMemory)
 		{
-			Buffer.Data.SetNumUninitialized(Buffer.ElementSize * ElementCount, EAllowShrinking::No);
-		}
-		else if (MemoryInitPolicy == EMemoryInitPolicy::Zeroed)
-		{
-			Buffer.Data.SetNumZeroed(Buffer.ElementSize * ElementCount, EAllowShrinking::No);
-		}
-		else
-		{
-			check(false);
+			if (MemoryInitPolicy == EMemoryInitPolicy::Uninitialized)
+			{
+				Buffer.Data.SetNumUninitialized(Buffer.ElementSize * ElementCount, EAllowShrinking::No);
+			}
+			else if (MemoryInitPolicy == EMemoryInitPolicy::Zeroed)
+			{
+				Buffer.Data.SetNumZeroed(Buffer.ElementSize * ElementCount, EAllowShrinking::No);
+			}
+			else
+			{
+				check(false);
+			}
 		}
 	}
 
@@ -373,10 +380,17 @@ namespace mu
 	{
 		int32 Result = 0;
 
+		const bool bIsDescriptor = IsDescriptor();
+
 		const int32 NumBuffers = Buffers.Num();
 		for (int32 BufferIndex = 0; BufferIndex < NumBuffers; ++BufferIndex)
 		{
-			Result += Buffers[BufferIndex].ElementSize * ElementCount;
+			Result += sizeof(FMeshBufferChannel) * Buffers[BufferIndex].Channels.Num();
+
+			if (!bIsDescriptor)
+			{
+				Result += Buffers[BufferIndex].ElementSize * ElementCount;
+			}
 		}
 
 		return Result;
@@ -615,6 +629,8 @@ namespace mu
 
 	void FMeshBufferSet::UpdateOffsets(int32 BufferIndex)
 	{
+		checkf(Buffers[BufferIndex].Data.IsEmpty(), TEXT("UpdateOffsets called on a non empty buffer set. This is not supported."));
+		
 		uint32 Offset = 0;
 		for (FMeshBufferChannel& Channel : Buffers[BufferIndex].Channels)
 		{
@@ -656,6 +672,20 @@ namespace mu
 		return false;
 	}
 
+
+	bool FMeshBufferSet::IsDescriptor() const
+	{
+		for (const FMeshBuffer& Buffer : Buffers)
+		{
+			if (!Buffer.Data.IsEmpty())
+			{
+				return false;
+			}
+		}
+
+		return ElementCount > 0;
+	}
+	
 	void FMeshBuffer::Serialise(OutputArchive& Arch) const
 	{
 		Arch << Channels;
diff --git a/Engine/Plugins/Experimental/Mutable/Source/MutableRuntime/Private/MuR/Model.cpp b/Engine/Plugins/Experimental/Mutable/Source/MutableRuntime/Private/MuR/Model.cpp
--- a/Engine/Plugins/Experimental/Mutable/Source/MutableRuntime/Private/MuR/Model.cpp
+++ b/Engine/Plugins/Experimental/Mutable/Source/MutableRuntime/Private/MuR/Model.cpp
@@ -24,6 +24,7 @@ namespace mu
 {
     MUTABLE_IMPLEMENT_POD_SERIALISABLE(FRomData);
 	MUTABLE_IMPLEMENT_POD_SERIALISABLE(FImageLODRange);
+	MUTABLE_IMPLEMENT_POD_SERIALISABLE(FMeshRange);
 
 	
     //---------------------------------------------------------------------------------------------
@@ -150,7 +151,7 @@ namespace mu
 		if (!bDropData)
 		{
 			InitialImages = program.ConstantImageLODs;
-			InitialMeshes = program.ConstantMeshes;
+			InitialMeshes = program.ConstantMeshData;
 		}
 
 		OutputMemoryStream MemStream(16 * 1024 * 1024);
@@ -186,9 +187,9 @@ namespace mu
 		}
 
 		// Save meshes and unload from memory
-		for (int32 ResourceIndex = 0; ResourceIndex < program.ConstantMeshes.Num(); ++ResourceIndex)
+		for (int32 ResourceIndex = 0; ResourceIndex < program.ConstantMeshData.Num(); ++ResourceIndex)
 		{
-			TPair<int32, mu::MeshPtrConst>& ResData = program.ConstantMeshes[ResourceIndex];
+			TPair<int32, mu::MeshPtrConst>& ResData = program.ConstantMeshData[ResourceIndex];
 
 			// This shouldn't have been serialised with rom support before.
 			if (ResData.Key < 0)
@@ -230,7 +231,7 @@ namespace mu
 		if (!bDropData)
 		{
 			program.ConstantImageLODs = InitialImages;
-			program.ConstantMeshes = InitialMeshes;
+			program.ConstantMeshData = InitialMeshes;
 		}
 	}
 
@@ -269,11 +270,11 @@ namespace mu
 			}
 		}
 
-		for (int32 ResIndex = 0; ResIndex < m_pD->m_program.ConstantMeshes.Num(); ++ResIndex)
+		for (int32 ResIndex = 0; ResIndex < m_pD->m_program.ConstantMeshData.Num(); ++ResIndex)
 		{
-			if (m_pD->m_program.ConstantMeshes[ResIndex].Key >= 0)
+			if (m_pD->m_program.ConstantMeshData[ResIndex].Key >= 0)
 			{
-				m_pD->m_program.ConstantMeshes[ResIndex].Value = nullptr;
+				m_pD->m_program.ConstantMeshData[ResIndex].Value = nullptr;
 			}
 		}
 	}
diff --git a/Engine/Plugins/Experimental/Mutable/Source/MutableRuntime/Private/MuR/OpMeshApplyPose.h b/Engine/Plugins/Experimental/Mutable/Source/MutableRuntime/Private/MuR/OpMeshApplyPose.h
--- a/Engine/Plugins/Experimental/Mutable/Source/MutableRuntime/Private/MuR/OpMeshApplyPose.h
+++ b/Engine/Plugins/Experimental/Mutable/Source/MutableRuntime/Private/MuR/OpMeshApplyPose.h
@@ -14,13 +14,47 @@ namespace mu
 	/**
     * Reference version
     */
+
+	inline void SetPoseAsReference(Mesh* InOutResult, const Mesh* PoseMesh)
+	{
+		// Set the pose as the Result reference pose.
+		// PoseMesh PoseXform can be decomposed to PoseXform = ModelPoseXform * ModelRefXform^-1, 
+		// BaseMesh Poses have ModelRefXform, to get the new ModelRefXform to be ModelPoseXform we need 
+		// to multiply PoseXform * ModelRefXform. 
+		// (Note, transform multiplication application is from right to left for this comment, but TTransform::operator* is reversed).
+		const int32 ResultNumBones = InOutResult->GetBonePoseCount();
+		for (int32 BoneIndex = 0; BoneIndex < ResultNumBones; ++BoneIndex)
+		{
+			const FBoneName BoneName = InOutResult->GetBonePoseId(BoneIndex);
+			const int32 FoundPoseIndex = PoseMesh->FindBonePose(BoneName);
+			if (FoundPoseIndex == INDEX_NONE)
+			{
+				continue;
+			}
+
+			FTransform3f ModelRefTransform;
+			InOutResult->GetBonePoseTransform(BoneIndex, ModelRefTransform);
+
+			FTransform3f PoseTransform;
+			PoseMesh->GetBonePoseTransform(FoundPoseIndex, PoseTransform); 
+
+			// Adding the Reshaped flag will prioritize this bone over other bones without the flag
+			// when merging poses with the same bone.
+			// TODO: Add another flag to indicate this case or generalize the Reshaped flag. 
+			EBoneUsageFlags UsageFlags = InOutResult->GetBoneUsageFlags(BoneIndex);
+			EnumAddFlags(UsageFlags, EBoneUsageFlags::Reshaped);
+			
+			InOutResult->SetBonePose(BoneIndex, BoneName, ModelRefTransform * PoseTransform, UsageFlags);
+		}
+	}
+
     inline void MeshApplyPose(Mesh* Result, const Mesh* BaseMesh, const Mesh* PoseMesh, bool& bOutSuccess)
     {
         MUTABLE_CPUPROFILER_SCOPE(MeshApplyPose);
 
 		bOutSuccess = true;
 
-		Ptr<const Skeleton> Skeleton = BaseMesh->GetSkeleton();
+		Ptr<const Skeleton> Skeleton = PoseMesh->GetSkeleton();
 		if (!Skeleton)
 		{
 			bOutSuccess = false;
@@ -105,34 +139,7 @@ namespace mu
 
 		Result->CopyFrom(*BaseMesh);
 
-		// Set the pose as the Result reference pose.
-		// PoseMesh PoseXform can be decomposed to PoseXform = ModelPoseXform * ModelRefXform^-1, 
-		// BaseMesh Poses have ModelRefXform, to get the new ModelRefXform to be ModelPoseXform we need 
-		// to multiply PoseXform * ModelRefXform. 
-		// (Note, transform multiplication application is from right to left for this comment, but TTransform::operator* is reversed).
-		const int32 ResultNumBones = Result->GetBonePoseCount();
-		for (int32 BoneIndex = 0; BoneIndex < ResultNumBones; ++BoneIndex)
-		{
-			const FBoneName BoneName = Result->GetBonePoseId(BoneIndex);
-			const int32 FoundPoseIndex = PoseMesh->FindBonePose(BoneName);
-			if (FoundPoseIndex == INDEX_NONE)
-			{
-				continue;
-			}
-
-			FTransform3f ModelRefTransform;
-			Result->GetBonePoseTransform(BoneIndex, ModelRefTransform);
-
-			FTransform3f PoseTransform;
-			PoseMesh->GetBonePoseTransform(FoundPoseIndex, PoseTransform); 
-
-			// Reshape flag is added so the bone is prioritized in case of merge conflict.
-			// TODO: Add another flag to indicate this case or generalize the Reshaped flag. 
-			EBoneUsageFlags UsageFlags = BaseMesh->GetBoneUsageFlags(BoneIndex);
-			EnumAddFlags(UsageFlags, EBoneUsageFlags::Reshaped);
-			
-			Result->SetBonePose(BoneIndex, BoneName, ModelRefTransform * PoseTransform, UsageFlags);
-		}
+		SetPoseAsReference(Result, PoseMesh);
 
 		MeshBufferIterator<MBF_FLOAT32, float, 3> TargetPositionIterBegin(Result->VertexBuffers, MBS_POSITION, 0);
 		
diff --git a/Engine/Plugins/Experimental/Mutable/Source/MutableRuntime/Private/MuR/OpMeshFormat.cpp b/Engine/Plugins/Experimental/Mutable/Source/MutableRuntime/Private/MuR/OpMeshFormat.cpp
--- a/Engine/Plugins/Experimental/Mutable/Source/MutableRuntime/Private/MuR/OpMeshFormat.cpp
+++ b/Engine/Plugins/Experimental/Mutable/Source/MutableRuntime/Private/MuR/OpMeshFormat.cpp
@@ -466,14 +466,20 @@ namespace mu
 		}
 
 
-		// For every vertex buffer in result
-		int32 vCount = Source.GetElementCount();
-		Result.SetElementCount(vCount);
-		for (int32 b = 0; b < Result.GetBufferCount(); ++b)
+		if (!Source.IsDescriptor())
 		{
-			MeshFormatBuffer(Source, Result.Buffers[b], 0, bIsVertexBuffer, IDPrefix);
+			// For every vertex buffer in result
+			int32 vCount = Source.GetElementCount();
+			Result.SetElementCount(vCount);
+			for (int32 b = 0; b < Result.GetBufferCount(); ++b)
+			{
+				MeshFormatBuffer(Source, Result.Buffers[b], 0, bIsVertexBuffer, IDPrefix);
+			}
+		}
+		else
+		{
+			Result.ElementCount = Source.GetElementCount();
 		}
-
 
 		// Detect internal system buffers and clone them unmodified.
 		if (bKeepSystemBuffers)
@@ -518,21 +524,105 @@ namespace mu
 				}
 			}
 		}
+	}
+
+	int32 FindHighestBoneIndexInUse(const Mesh* InMesh, const FMeshBufferChannel& Channel)
+	{
+		const FMeshBufferSet& VertexBuffers = InMesh->GetVertexBuffers();
+
+		// If InMesh buffers are descriptors, use the bone map to get an upper bound of the highest index
+		// we will find in the mesh. Otherwise, find it from the vertex buffers data.
+		if (VertexBuffers.IsDescriptor())
+		{
+			return InMesh->BoneMap.Num() - 1;
+		}
+
+		const UntypedMeshBufferIteratorConst BoneIndicesBegin(VertexBuffers, MBS_BONEINDICES, Channel.SemanticIndex);
+
+		if (!BoneIndicesBegin.ptr())
+		{
+			return -1;
+		}
+
+		const int32 NumVertices = VertexBuffers.GetElementCount();
+		const int32 NumInfluences = BoneIndicesBegin.GetComponents();
+		
+		int32 HighestBoneIndex = -1;
+		for (int32 VertexIndex = 0; VertexIndex < NumVertices; ++VertexIndex)
+		{
+			// If MAX_TOTAL_INFLUENCES ever changed, the next line would no longer work or compile and 
+			// GetAsVec12i would need to be changed accordingly
+			int32 VertexInfluences[MAX_TOTAL_INFLUENCES];
+			(BoneIndicesBegin + VertexIndex).GetAsInt32Vec(VertexInfluences, MAX_TOTAL_INFLUENCES);
+			
+			for (int32 InfluenceIndex = 0; InfluenceIndex < NumInfluences; ++InfluenceIndex)
+			{
+				HighestBoneIndex = FMath::Max(HighestBoneIndex, VertexInfluences[InfluenceIndex]);
+			}
+		}
 
+		return HighestBoneIndex;
 	}
 
+	/** Updates the InOutFormat bone index semantic format if it cannot represent the highest bone index in use. */
+	void EnsureBoneIndexFormat(const Mesh* InMesh, FMeshBufferSet& InOutFormatBufferSet)
+	{
+		const FMeshBufferSet& VertexBuffers = InMesh->GetVertexBuffers();
 
+		const int32 BufferCount = VertexBuffers.GetBufferCount();
+		for (int32 BufferIndex = 0; BufferIndex < BufferCount; ++BufferIndex)
+		{
+			const int32 ChannelCount = VertexBuffers.GetBufferChannelCount(BufferIndex);
+			for (int32 ChannelIndex = 0; ChannelIndex < ChannelCount; ++ChannelIndex)
+			{
+				const FMeshBufferChannel& Channel = VertexBuffers.Buffers[BufferIndex].Channels[ChannelIndex];
+
+				if (Channel.Semantic == MBS_BONEINDICES)
+				{
+					int32 ResultBuf = 0;
+					int32 ResultChan = 0;
+					InOutFormatBufferSet.FindChannel(MBS_BONEINDICES, Channel.SemanticIndex, &ResultBuf, &ResultChan);
+					if (ResultBuf >= 0)
+					{
+						int32 HighestBoneIndexInUse = FindHighestBoneIndexInUse(InMesh, Channel);
+						check(HighestBoneIndexInUse >= 0);
+
+						EMeshBufferFormat& BoneIndexFormat = InOutFormatBufferSet.Buffers[ResultBuf].Channels[ResultChan].Format;
+						if (HighestBoneIndexInUse > 0xffff && (BoneIndexFormat == MBF_UINT8 || BoneIndexFormat == MBF_UINT16))
+						{
+							BoneIndexFormat = MBF_UINT32;
+							InOutFormatBufferSet.UpdateOffsets(ResultBuf);
+						}
+						else if (HighestBoneIndexInUse > 0x7fff && (BoneIndexFormat == MBF_INT8 || BoneIndexFormat == MBF_INT16))
+						{
+							BoneIndexFormat = MBF_UINT32;
+							InOutFormatBufferSet.UpdateOffsets(ResultBuf);
+						}
+						else if (HighestBoneIndexInUse > 0xff && BoneIndexFormat == MBF_UINT8)
+						{
+							BoneIndexFormat = MBF_UINT16;
+							InOutFormatBufferSet.UpdateOffsets(ResultBuf);
+						}
+						else if (HighestBoneIndexInUse > 0x7f && BoneIndexFormat == MBF_INT8)
+						{
+							BoneIndexFormat = MBF_INT16;
+							InOutFormatBufferSet.UpdateOffsets(ResultBuf);
+						}
+					}
+				}
+			}
+		}
+	}
 
-	//-------------------------------------------------------------------------------------------------
 	void MeshFormat
 	(
 		Mesh* Result, 
 		const Mesh* PureSource,
 		const Mesh* Format,
-		bool keepSystemBuffers,
-		bool formatVertices,
-		bool formatIndices,
-		bool ignoreMissingChannels,
+		bool bKeepSystemBuffers,
+		bool bFormatVertices,
+		bool bFormatIndices,
+		bool bIgnoreMissingChannels,
 		bool& bOutSuccess
 	)
 	{
@@ -558,84 +648,22 @@ namespace mu
 		Result->CopyFrom(*Format);
 		Result->MeshIDPrefix = Source->MeshIDPrefix;
 
-		// Make sure that the bone indices will fit in this format, or extend it.
-		if (formatVertices)
-		{
-			const FMeshBufferSet& VertexBuffers = Source->GetVertexBuffers();
-
-			const int32 BufferCount = VertexBuffers.GetBufferCount();
-			for (int32 BufferIndex = 0; BufferIndex < BufferCount; ++BufferIndex)
-			{
-				const int32 ChannelCount = VertexBuffers.GetBufferChannelCount(BufferIndex);
-				for (int32 ChannelIndex = 0; ChannelIndex < ChannelCount; ++ChannelIndex)
-				{
-					const FMeshBufferChannel& Channel = VertexBuffers.Buffers[BufferIndex].Channels[ChannelIndex];
-
-					if (Channel.Semantic == MBS_BONEINDICES)
-					{
-						int32 resultBuf = 0;
-						int32 resultChan = 0;
-						FMeshBufferSet& formBuffs = Result->GetVertexBuffers();
-						formBuffs.FindChannel(MBS_BONEINDICES, Channel.SemanticIndex, &resultBuf, &resultChan);
-						if (resultBuf >= 0)
-						{
-							UntypedMeshBufferIteratorConst it(VertexBuffers, MBS_BONEINDICES, Channel.SemanticIndex);
-							int32 maxBoneIndex = 0;
-							for (int32 v = 0; v < VertexBuffers.GetElementCount(); ++v)
-							{
-								// If MAX_TOTAL_INFLUENCES ever changed, the next line would no longer work or compile and 
-								// GetAsVec12i would need to be changed accordingly
-								int32 va[MAX_TOTAL_INFLUENCES];
-								it.GetAsInt32Vec(va, MAX_TOTAL_INFLUENCES);
-								for (int32 c = 0; c < it.GetComponents(); ++c)
-								{
-									maxBoneIndex = FMath::Max(maxBoneIndex, va[c]);
-								}
-								++it;
-							}
-
-							EMeshBufferFormat& format = formBuffs.Buffers[resultBuf].Channels[resultChan].Format;
-							if (maxBoneIndex > 0xffff && (format == MBF_UINT8 || format == MBF_UINT16))
-							{
-								format = MBF_UINT32;
-								formBuffs.UpdateOffsets(resultBuf);
-							}
-							else if (maxBoneIndex > 0x7fff && (format == MBF_INT8 || format == MBF_INT16))
-							{
-								format = MBF_UINT32;
-								formBuffs.UpdateOffsets(resultBuf);
-							}
-							else if (maxBoneIndex > 0xff && format == MBF_UINT8)
-							{
-								format = MBF_UINT16;
-								formBuffs.UpdateOffsets(resultBuf);
-							}
-							else if (maxBoneIndex > 0x7f && format == MBF_INT8)
-							{
-								format = MBF_INT16;
-								formBuffs.UpdateOffsets(resultBuf);
-							}
-						}
-					}
-				}
-			}
-		}
-
-		if (formatVertices)
+		if (bFormatVertices)
 		{
+			EnsureBoneIndexFormat(Source.get(), Result->GetVertexBuffers());
 			FormatBufferSet(Source->GetVertexBuffers(), Result->GetVertexBuffers(),
-				keepSystemBuffers, ignoreMissingChannels, true, Source->MeshIDPrefix);
+				bKeepSystemBuffers, bIgnoreMissingChannels, true, Source->MeshIDPrefix);
 		}
 		else
 		{
 			Result->VertexBuffers = Source->GetVertexBuffers();
 		}
 
-		if (formatIndices)
+		if (bFormatIndices)
 		{
 			// \todo Make sure that the vertex indices will fit in this format, or extend it.
-			FormatBufferSet(Source->GetIndexBuffers(), Result->GetIndexBuffers(), keepSystemBuffers,
-				ignoreMissingChannels, false);
+			FormatBufferSet(Source->GetIndexBuffers(), Result->GetIndexBuffers(), bKeepSystemBuffers,
+				bIgnoreMissingChannels, false);
 		}
 		else
 		{
@@ -672,7 +700,7 @@ namespace mu
 	}
 
 
-	void MeshOptimizeBuffers( Mesh* InMesh )
+	void MeshOptimizeBuffers(Mesh* InMesh)
 	{
 		if (!InMesh)
 		{
@@ -681,6 +709,12 @@ namespace mu
 
 		FMeshBufferSet& VertexBuffers = InMesh->VertexBuffers;
 
+		// Ignore the operation if the MeshBuffersSet is a descriptor.
+		if (VertexBuffers.IsDescriptor())
+		{
+			return;
+		}
+
 		// Reduce the number of influences if possible
 		constexpr int32 SemanticIndex = 0;
 		
diff --git a/Engine/Plugins/Experimental/Mutable/Source/MutableRuntime/Private/MuR/System.cpp b/Engine/Plugins/Experimental/Mutable/Source/MutableRuntime/Private/MuR/System.cpp
--- a/Engine/Plugins/Experimental/Mutable/Source/MutableRuntime/Private/MuR/System.cpp
+++ b/Engine/Plugins/Experimental/Mutable/Source/MutableRuntime/Private/MuR/System.cpp
@@ -467,7 +467,6 @@ namespace mu
 		m_pD->WorkingMemoryManager.CurrentInstanceCache = nullptr;
 	}
 
-
 	UE::Tasks::TTask<FImageDesc> System::GetImageDesc(Instance::ID instanceID, FResourceID ImageId)
 	{
 		LLM_SCOPE_BYNAME(TEXT("MutableRuntime"));
@@ -525,6 +524,112 @@ namespace mu
 
 	}
 
+	Ptr<const Mesh> System::GetMeshComponentsInline(Instance::ID InstanceID, FResourceID MeshId)
+	{
+		LLM_SCOPE_BYNAME(TEXT("MutableRuntime"));
+		MUTABLE_CPUPROFILER_SCOPE(SystemGetMesh);
+
+		Ptr<const Mesh> Result;
+
+		// Find the live instance
+		FLiveInstance* LiveInstancePtr = m_pD->FindLiveInstance(InstanceID);
+		check(LiveInstancePtr);
+		m_pD->WorkingMemoryManager.CurrentInstanceCache = LiveInstancePtr->Cache;
+		m_pD->WorkingMemoryManager.BeginRunnerThread();
+
+		OP::ADDRESS RootAddress = GetResourceIDRoot(MeshId);
+		
+		
+		constexpr uint8 ExecutionsOptions = static_cast<uint8>(
+                EMeshExecutionOptions::LoadComponentsData | EMeshExecutionOptions::LoadPoseData);
+		TSharedRef<CodeRunner> Runner = CodeRunner::Create(
+				m_pD->Settings, m_pD, EExecutionStrategy::MinimizeMemory, LiveInstancePtr->Model, 
+				LiveInstancePtr->OldParameters.get(), RootAddress, System::AllLODs, ExecutionsOptions, 0, FScheduledOp::EType::MeshComp);
+
+		constexpr bool bForceInlineExecutution = true;
+		UE::Tasks::FTask RunnerCompletionEvent = Runner->StartRun(bForceInlineExecutution);
+		check(RunnerCompletionEvent.IsCompleted());
+
+		m_pD->bUnrecoverableError = Runner->bUnrecoverableError;
+		if (!Runner->bUnrecoverableError)
+		{
+			const FCacheAddress ResultAddress = FCacheAddress(RootAddress, 0, ExecutionsOptions, FScheduledOp::EType::MeshComp);
+			Result = m_pD->WorkingMemoryManager.LoadMesh(ResultAddress, true);
+		}
+
+		if (!Result)
+		{
+			Result = new Mesh();
+		}
+
+		m_pD->WorkingMemoryManager.EndRunnerThread();
+		m_pD->WorkingMemoryManager.CurrentInstanceCache = nullptr;
+
+		return Result;
+	}
+
+	UE::Tasks::TTask<Ptr<const Mesh>> System::GetMeshComponents(Instance::ID InstanceID, FResourceID MeshId)
+	{
+		LLM_SCOPE_BYNAME(TEXT("MutableRuntime"));
+		MUTABLE_CPUPROFILER_SCOPE(SystemGetMeshComponents);
+
+		Ptr<const Mesh> ResultMesh;
+
+		// Find the live instance
+		FLiveInstance* LiveInstancePtr = m_pD->FindLiveInstance(InstanceID);
+		check(LiveInstancePtr);
+		m_pD->WorkingMemoryManager.CurrentInstanceCache = LiveInstancePtr->Cache;
+		m_pD->WorkingMemoryManager.BeginRunnerThread();
+		
+		OP::ADDRESS RootAddress = GetResourceIDRoot(MeshId);
+			
+		mu::OP_TYPE OpType = LiveInstancePtr->Model->GetPrivate()->m_program.GetOpType(RootAddress);
+		if (GetOpDataType(OpType) != DT_MESH)
+		{
+			m_pD->WorkingMemoryManager.EndRunnerThread();
+			m_pD->WorkingMemoryManager.CurrentInstanceCache = nullptr;
+
+			return UE::Tasks::MakeCompletedTask<Ptr<const Mesh>>(new Mesh());
+		}
+			
+		constexpr uint8 ExecutionsOptions = static_cast<uint8>(
+                EMeshExecutionOptions::LoadComponentsData | EMeshExecutionOptions::LoadPoseData);
+
+		TSharedRef<CodeRunner> Runner = CodeRunner::Create(
+				m_pD->Settings, m_pD, EExecutionStrategy::MinimizeMemory, LiveInstancePtr->Model, 
+				LiveInstancePtr->OldParameters.get(), RootAddress, System::AllLODs, ExecutionsOptions, 0, FScheduledOp::EType::MeshComp);
+		
+		constexpr bool bForceInlineExecution = false;
+		UE::Tasks::FTask RunnerCompletionEvent = Runner->StartRun(bForceInlineExecution);
+
+		return UE::Tasks::Launch(TEXT("System::GetMeshComponentsResultTask"),
+				[SystemPrivate = m_pD, Runner, RootAddress]() -> Ptr<const Mesh>
+				{
+					Ptr<const Mesh> Result;
+
+					SystemPrivate->bUnrecoverableError = Runner->bUnrecoverableError;
+					if (!Runner->bUnrecoverableError)
+					{
+						const FCacheAddress ResultAddress = FCacheAddress(RootAddress, 0, ExecutionsOptions, FScheduledOp::EType::MeshComp);
+						Result = SystemPrivate->WorkingMemoryManager.LoadMesh(ResultAddress, true);
+					}
+
+					if (!Result)
+					{
+						Result = new Mesh();
+					}
+
+					SystemPrivate->WorkingMemoryManager.EndRunnerThread();
+					SystemPrivate->WorkingMemoryManager.CurrentInstanceCache = nullptr;
+					
+					return Result;
+				},
+				UE::Tasks::Prerequisites(RunnerCompletionEvent),
+				UE::Tasks::ETaskPriority::Inherit,
+				UE::Tasks::EExtendedTaskPriority::Inline);
+	}
+
+
 
     Ptr<const Mesh> System::GetMeshInline(Instance::ID instanceID, FResourceID MeshId)
     {
@@ -554,7 +659,7 @@ namespace mu
 	UE::Tasks::TTask<Ptr<const Mesh>> System::GetMesh(Instance::ID instanceID, FResourceID MeshId)
     {
 		LLM_SCOPE_BYNAME(TEXT("MutableRuntime"));
-		MUTABLE_CPUPROFILER_SCOPE(SystemGetImage);
+		MUTABLE_CPUPROFILER_SCOPE(SystemGetMesh);
 
 		Ptr<const Mesh> ResultMesh;
 
@@ -574,9 +679,14 @@ namespace mu
 
 			return UE::Tasks::MakeCompletedTask<Ptr<const Mesh>>(new Mesh());
 		}
-		
+	
+		// For now the pose data is loaded for all constants since some operations need it. The idea is to set this option on
+		// an Op to Op basis if it needs the data.	
+		constexpr uint8 ExecutionOptions = static_cast<uint8>(EMeshExecutionOptions::LoadGeometryData | EMeshExecutionOptions::LoadPoseData);
 		TSharedRef<CodeRunner> Runner = CodeRunner::Create(
-				m_pD->Settings, m_pD, EExecutionStrategy::MinimizeMemory, pLiveInstance->Model, pLiveInstance->OldParameters.get(), RootAddress, System::AllLODs, 0, 0, FScheduledOp::EType::Full);
+				m_pD->Settings, m_pD, EExecutionStrategy::MinimizeMemory, 
+				pLiveInstance->Model, pLiveInstance->OldParameters.get(), 
+				RootAddress, System::AllLODs, ExecutionOptions, 0, FScheduledOp::EType::Full);
 		
 		constexpr bool bForceInlineExecution = false;
 		UE::Tasks::FTask RunnerCompletionEvent = Runner->StartRun(bForceInlineExecution);
@@ -589,7 +699,7 @@ namespace mu
 					SystemPrivate->bUnrecoverableError = Runner->bUnrecoverableError;
 					if (!Runner->bUnrecoverableError)
 					{
-						Result = SystemPrivate->WorkingMemoryManager.LoadMesh(FCacheAddress(RootAddress, 0, 0), true);
+						Result = SystemPrivate->WorkingMemoryManager.LoadMesh(FCacheAddress(RootAddress, 0, ExecutionOptions), true);
 					}
 
 					if (!Result)
@@ -849,10 +959,10 @@ namespace mu
 
 	//---------------------------------------------------------------------------------------------
 	void System::Private::RunCode(const TSharedPtr<const Model>& InModel,
-		const Parameters* InParameters, OP::ADDRESS InCodeRoot, uint32 InLODs, uint8 executionOptions, int32 InImageLOD)
+		const Parameters* InParameters, OP::ADDRESS InCodeRoot, uint32 InLODs, uint8 ExecutionOptions, int32 InImageLOD)
 	{
 		TSharedRef<CodeRunner> Runner = CodeRunner::Create(Settings, this, EExecutionStrategy::MinimizeMemory, InModel, InParameters, InCodeRoot, InLODs,
-			executionOptions, InImageLOD, FScheduledOp::EType::Full);
+			ExecutionOptions, InImageLOD, FScheduledOp::EType::Full);
 		
 		constexpr bool bForceInlineExecutution = true;
 		UE::Tasks::FTask RunnerCompletionEvent = Runner->StartRun(bForceInlineExecutution);
@@ -982,8 +1092,7 @@ namespace mu
 		return Result;
 	}
 
-
-	Ptr<const Mesh> System::Private::BuildMesh(const TSharedPtr<const Model>& pModel, const Parameters* Params, OP::ADDRESS at)
+	Ptr<const Mesh> System::Private::BuildMesh(const TSharedPtr<const Model>& pModel, const Parameters* Params, OP::ADDRESS at, uint8 ExecutionOptions)
 	{
 		WorkingMemoryManager.BeginRunnerThread();
 
@@ -992,10 +1101,10 @@ namespace mu
 		mu::OP_TYPE opType = pModel->GetPrivate()->m_program.GetOpType(at);
 		if (GetOpDataType(opType) == DT_MESH)
 		{
-			RunCode(pModel, Params, at);
+			RunCode(pModel, Params, at, System::AllLODs, ExecutionOptions);
 			if (!bUnrecoverableError)
 			{
-				Result = WorkingMemoryManager.LoadMesh(FCacheAddress(at, 0, 0), true);
+				Result = WorkingMemoryManager.LoadMesh(FCacheAddress(at, 0, ExecutionOptions), true);
 			}	
 		}
 
@@ -1089,15 +1198,18 @@ namespace mu
 		WorkingMemoryManager.CurrentInstanceCache->Init(OpCount);
 
 		// Clear cache flags of existing data
-		CodeContainer<FProgramCache::FOpExecutionData>::iterator It = WorkingMemoryManager.CurrentInstanceCache->OpExecutionData.begin();
-		for (; It.IsValid(); ++It)
+		CodeContainer<FProgramCache::FOpExecutionData>::TIterator Iter = 
+				WorkingMemoryManager.CurrentInstanceCache->OpExecutionData.Begin();
+		for (; Iter.IsValid(); ++Iter)
 		{
-			(*It).OpHitCount = 0; // This should already be 0, but just in case.
-			(*It).IsCacheLocked = false;
+			FProgramCache::FOpExecutionData& Data = *Iter;
+
+			Data.OpHitCount = 0; // This should already be 0, but just in case.
+			Data.IsCacheLocked = false;
 		}
 
 		// Mark the resources that have to be cached to update the instance in this state
-		if (InState >= 0 && InState< Program.m_states.Num())
+		if (InState >= 0 && InState < Program.m_states.Num())
 		{
 			const FProgram::FState& State = Program.m_states[InState];
 			for (uint32 Address : State.m_updateCache)
diff --git a/Engine/Plugins/Experimental/Mutable/Source/MutableRuntime/Public/MuR/Mesh.h b/Engine/Plugins/Experimental/Mutable/Source/MutableRuntime/Public/MuR/Mesh.h
--- a/Engine/Plugins/Experimental/Mutable/Source/MutableRuntime/Public/MuR/Mesh.h
+++ b/Engine/Plugins/Experimental/Mutable/Source/MutableRuntime/Public/MuR/Mesh.h
@@ -162,7 +162,6 @@ namespace mu
 
 		/** The mesh is a reference to an external resource mesh and must be loaded when first referenced. */
 		IsResourceForceLoad = 1 << 3,
-
 	};
 
 
@@ -438,8 +437,8 @@ namespace mu
 			EBoneUsageFlags BoneUsageFlags = EBoneUsageFlags::None;
 			FTransform3f BoneTransform;
 
-			inline void Serialise(OutputArchive& arch) const;
-			inline void Unserialise(InputArchive& arch);
+			inline void Serialise(OutputArchive& Arch) const;
+			inline void Unserialise(InputArchive& Arch);
 
 			inline bool operator==(const FBonePose& Other) const
 			{
diff --git a/Engine/Plugins/Experimental/Mutable/Source/MutableRuntime/Public/MuR/MeshBufferSet.h b/Engine/Plugins/Experimental/Mutable/Source/MutableRuntime/Public/MuR/MeshBufferSet.h
--- a/Engine/Plugins/Experimental/Mutable/Source/MutableRuntime/Public/MuR/MeshBufferSet.h
+++ b/Engine/Plugins/Experimental/Mutable/Source/MutableRuntime/Public/MuR/MeshBufferSet.h
@@ -13,7 +13,7 @@ namespace mu::MemoryCounters
 {
 	struct MUTABLERUNTIME_API FMeshMemoryCounter
 	{
-		alignas(8) static inline std::atomic<SSIZE_T> Counter{0};
+		alignas(8) static inline std::atomic<SSIZE_T> Counter {0};
 	};
 }
 
@@ -194,8 +194,16 @@ namespace mu
 		inline bool operator==(const FMeshBuffer& Other) const
 		{
 			bool bEqual = (Channels == Other.Channels);
-			bEqual = bEqual && (ElementSize == Other.ElementSize);
-			bEqual = bEqual && (Data == Other.Data);
+			
+            if (bEqual) 
+            {
+                bEqual = (ElementSize == Other.ElementSize);
+            }
+
+			if (bEqual)
+            {
+                bEqual = (Data == Other.Data);
+            }
 
 			return bEqual;
 		}
@@ -235,8 +243,8 @@ namespace mu
 	class MUTABLERUNTIME_API FMeshBufferSet
 	{
 	public:
-
 		uint32 ElementCount = 0;
+
 		TArray<FMeshBuffer> Buffers;
 
 		void Serialise(OutputArchive& Arch) const;
@@ -431,6 +439,9 @@ namespace mu
 		
 		/** Check that all channels of a specific semantic use the provided format. */
 		bool HasAnySemanticWithDifferentFormat(EMeshBufferSemantic Semantic, EMeshBufferFormat ExpectedFormat) const;
+
+		/** Check if this buffer set is only a format descriptor, e.i. reports num elements larger than 0 but does not hold any data*/ 
+		bool IsDescriptor() const;
 	};	
 
 	
diff --git a/Engine/Plugins/Experimental/Mutable/Source/MutableRuntime/Public/MuR/System.h b/Engine/Plugins/Experimental/Mutable/Source/MutableRuntime/Public/MuR/System.h
--- a/Engine/Plugins/Experimental/Mutable/Source/MutableRuntime/Public/MuR/System.h
+++ b/Engine/Plugins/Experimental/Mutable/Source/MutableRuntime/Public/MuR/System.h
@@ -171,6 +171,9 @@ namespace mu
         //! Only valid between BeginUpdate and EndUpdate
 		UE::Tasks::TTask<Ptr<const Mesh>> GetMesh(Instance::ID InstanceID, FResourceID MeshId);
 
+		UE::Tasks::TTask<Ptr<const Mesh>> GetMeshComponents(Instance::ID InstanceID, FResourceID MeshId);
+
+
 		//! Only valid between BeginUpdate and EndUpdate
 		//! Calculate the description of an image, without generating it.
 		void GetImageDescInline(Instance::ID InstanceID, FResourceID ImageId, FImageDesc& OutDesc);
@@ -184,6 +187,8 @@ namespace mu
         //! Only valid between BeginUpdate and EndUpdate
 		Ptr<const Mesh> GetMeshInline(Instance::ID InstanceID, FResourceID MeshId);
 
+		Ptr<const Mesh> GetMeshComponentsInline(Instance::ID InstanceID, FResourceID MeshId);
+
         //! Invalidate and free the last Instance data returned by a call to BeginUpdate with
         //! the same instance index. After a call to this method, that Instance cannot be used any
         //! more and its content is undefined.
diff --git a/Engine/Plugins/Experimental/Mutable/Source/MutableTools/Private/MuT/ASTOpConstantResource.cpp b/Engine/Plugins/Experimental/Mutable/Source/MutableTools/Private/MuT/ASTOpConstantResource.cpp
--- a/Engine/Plugins/Experimental/Mutable/Source/MutableTools/Private/MuT/ASTOpConstantResource.cpp
+++ b/Engine/Plugins/Experimental/Mutable/Source/MutableTools/Private/MuT/ASTOpConstantResource.cpp
@@ -279,6 +279,175 @@ namespace mu
 
 	namespace
 	{
+		/** Adds a constant mesh data to a program and returns its constant index. */
+		int32 AddConstantMesh(FProgram& Program, const Ptr<const Mesh>& MeshData, FLinkerOptions& Options)
+		{
+			// TODO: Add LinkerOptions to control how the mesh is split.
+		
+			// Split generated mesh data in two parts. Geometry and other components. 
+			// Bone maps are needed for some geoemtry only operations, duplicate the data for now.
+			
+			int32 FirstIndex = Program.ConstantMeshIndices.Num();
+			
+			// GeometryMesh
+			{
+				const EMeshCopyFlags GeometryDataCopyFlags = 
+						//EMeshCopyFlags::WithBoneMap       |
+						//EMeshCopyFlags::WithPoses         |
+						EMeshCopyFlags::WithSurfaces      | 
+						EMeshCopyFlags::WithVertexBuffers |
+						EMeshCopyFlags::WithIndexBuffers  |
+						EMeshCopyFlags::WithLayouts;
+
+				Ptr<Mesh> MeshGeometryData = MeshData->Clone(GeometryDataCopyFlags);
+
+				// Copy geometry related additional buffers.
+				for (const TPair<EMeshBufferType, FMeshBufferSet>& AdditionalBuffer : MeshData->AdditionalBuffers)
+				{
+					const bool bIsGeometryBufferType = 
+							AdditionalBuffer.Key == EMeshBufferType::MeshLaplacianData    ||
+							AdditionalBuffer.Key == EMeshBufferType::MeshLaplacianOffsets ||
+							AdditionalBuffer.Key == EMeshBufferType::UniqueVertexMap;
+
+					if (bIsGeometryBufferType)
+					{
+						MeshGeometryData->AdditionalBuffers.Emplace(AdditionalBuffer);
+					}
+				}
+
+				// Use a map-based deduplication
+	
+				int32 GeometryIndex = -1;
+				mu::Ptr<const mu::Mesh> GeometryKey = MeshGeometryData;
+				const int32* GeometryIndexPtr = Options.MeshConstantMap.Find(GeometryKey);
+				if (!GeometryIndexPtr)
+				{
+					GeometryIndex = Program.ConstantMeshData.Add(TPair<int32, Ptr<const Mesh>>(-1, MeshGeometryData));
+					Options.MeshConstantMap.Add(MeshGeometryData.get(), GeometryIndex);
+				}
+				else
+				{
+					GeometryIndex = *GeometryIndexPtr;
+				}
+
+				Program.ConstantMeshIndices.Add(GeometryIndex);
+			}
+
+			// Pose Mesh
+			{
+				Ptr<Mesh> MeshPoseData = new Mesh();
+				MeshPoseData->BonePoses = MeshData->BonePoses;
+				MeshPoseData->BoneMap =  MeshData->BoneMap;
+
+				int32 PoseIndex = -1;
+				mu::Ptr<const mu::Mesh> PoseKey = MeshPoseData;
+				const int32* PoseIndexPtr = Options.MeshConstantMap.Find(PoseKey);
+				if (!PoseIndexPtr)
+				{
+					PoseIndex = Program.ConstantMeshData.Add(TPair<int32, Ptr<const Mesh>>(-1, MeshPoseData));
+					Options.MeshConstantMap.Add(MeshPoseData.get(), PoseIndex);
+				}
+				else
+				{
+					PoseIndex = *PoseIndexPtr;
+				}
+				
+				Program.ConstantMeshIndices.Add(PoseIndex);
+			}
+
+			// Components Mesh
+			{
+				const EMeshCopyFlags ComponentsDataCopyFlags = 
+						//EMeshCopyFlags::WithBoneMap           |
+						//EMeshCopyFlags::WithPoses             |
+						EMeshCopyFlags::WithSurfaces          |
+						EMeshCopyFlags::WithTags              |
+						EMeshCopyFlags::WithSkeletonIDs       |
+						EMeshCopyFlags::WithAdditionalPhysics |
+						EMeshCopyFlags::WithStreamedResources;
+
+				Ptr<Mesh> MeshComponentsData = MeshData->Clone(ComponentsDataCopyFlags);
+			
+				// Add an empty MeshBufferSet to the Component Part to have formating info.
+				// Probably not needed, it needs special handling when data is empty but
+				// element count is non zero.
+				{
+					FMeshBufferSet VertexMeshFormat;
+					const FMeshBufferSet& VertexBufferSet = MeshData->VertexBuffers;
+					
+					VertexMeshFormat.ElementCount = VertexBufferSet.ElementCount;
+
+					const int32 NumVertexBuffers = VertexBufferSet.Buffers.Num();
+					VertexMeshFormat.Buffers.SetNum(NumVertexBuffers);
+
+					for (int32 BufferIndex = 0; BufferIndex < NumVertexBuffers; ++BufferIndex)
+					{
+						VertexMeshFormat.Buffers[BufferIndex].Channels = VertexBufferSet.Buffers[BufferIndex].Channels;	
+						VertexMeshFormat.Buffers[BufferIndex].ElementSize = VertexBufferSet.Buffers[BufferIndex].ElementSize;	
+					}
+				
+					MeshComponentsData->VertexBuffers = MoveTemp(VertexMeshFormat);
+				}
+
+				{
+					FMeshBufferSet IndexMeshFormat;
+					const FMeshBufferSet& IndexBufferSet = MeshData->IndexBuffers;
+					
+					IndexMeshFormat.ElementCount = IndexBufferSet.ElementCount;
+
+					const int32 NumIndexBuffers = IndexBufferSet.Buffers.Num();
+					IndexMeshFormat.Buffers.SetNum(NumIndexBuffers);
+
+					for (int32 BufferIndex = 0; BufferIndex < NumIndexBuffers; ++BufferIndex)
+					{
+						IndexMeshFormat.Buffers[BufferIndex].Channels = IndexBufferSet.Buffers[BufferIndex].Channels;	
+						IndexMeshFormat.Buffers[BufferIndex].ElementSize = IndexBufferSet.Buffers[BufferIndex].ElementSize;	
+					}
+					
+					MeshComponentsData->IndexBuffers = MoveTemp(IndexMeshFormat);
+				}
+
+				// Copy components related additional buffers.
+				for (const TPair<EMeshBufferType, FMeshBufferSet>& AdditionalBuffer : MeshData->AdditionalBuffers)
+				{
+					const bool bIsComponentsBufferType = 
+							AdditionalBuffer.Key == EMeshBufferType::SkeletonDeformBinding      ||
+							AdditionalBuffer.Key == EMeshBufferType::PhysicsBodyDeformBinding   ||
+							AdditionalBuffer.Key == EMeshBufferType::PhysicsBodyDeformSelection ||
+							AdditionalBuffer.Key == EMeshBufferType::PhysicsBodyDeformOffsets;
+
+					if (bIsComponentsBufferType)
+					{
+						MeshComponentsData->AdditionalBuffers.Emplace(AdditionalBuffer);
+					}
+				}
+
+				int32 ComponentsIndex = -1;
+				mu::Ptr<const mu::Mesh> ComponentsKey = MeshComponentsData;
+				const int32* ComponentsIndexPtr = Options.MeshConstantMap.Find(ComponentsKey);
+				if (!ComponentsIndexPtr)
+				{
+					ComponentsIndex = Program.ConstantMeshData.Add(TPair<int32, Ptr<const Mesh>>(-1, MeshComponentsData));
+					Options.MeshConstantMap.Add(MeshComponentsData.get(), ComponentsIndex);
+				}
+				else
+				{
+					ComponentsIndex = *ComponentsIndexPtr;
+				}
+
+				Program.ConstantMeshIndices.Add(ComponentsIndex);
+			}
+
+			check(Program.ConstantMeshIndices.Num() - FirstIndex == 3);
+
+			// Using a FMeshIndexRange is not strictly needed at the moment as we have a regular layout,
+			// all meshes have 3 pices of data. This may not be the case if linker options are added or 
+			// if some data can be discarded. Notice that in the later case we will need to specify the 
+			// content per Index.
+			return Program.ConstantMeshes.Add(FMeshRange{FirstIndex, 3});
+		}
+		
+
 		/** Adds a constant image data to a program and returns its constant index. */
 		int32 AddConstantImage(FProgram& Program, const Ptr<const Image>& pImage, FLinkerOptions& Options)
 		{
@@ -400,8 +569,7 @@ namespace mu
 		}
 	}
 
-	//-------------------------------------------------------------------------------------------------
-	void ASTOpConstantResource::Link(FProgram& program, FLinkerOptions* Options)
+	void ASTOpConstantResource::Link(FProgram& Program, FLinkerOptions* Options)
 	{
 		MUTABLE_CPUPROFILER_SCOPE(ASTOpConstantResource_Link);
 
@@ -409,43 +577,30 @@ namespace mu
 		{
 			if (Type == OP_TYPE::ME_CONSTANT)
 			{
-				OP::MeshConstantArgs args;
-				FMemory::Memset(&args, 0, sizeof(args));
-
-				Ptr<Mesh> MeshData = static_cast<const Mesh*>(GetValue().get())->Clone();
-				check(MeshData);
+				OP::MeshConstantArgs Args;
+				FMemory::Memzero(Args);
 
-				args.skeleton = -1;
-				if (Ptr<const Skeleton> pSkeleton = MeshData->GetSkeleton())
-				{
-					args.skeleton = program.AddConstant(pSkeleton.get());
-					MeshData->SetSkeleton(nullptr);
-				}
+				Ptr<const Mesh> TypedValue = static_cast<const Mesh*>(GetValue().get());
+				check(TypedValue);
+				
+				Args.Value = AddConstantMesh(Program, TypedValue, *Options);
 
-				args.physicsBody = -1;
-				if (Ptr<const PhysicsBody> pPhysicsBody = MeshData->GetPhysicsBody())
+				Args.Skeleton = -1;
+				if (Ptr<const Skeleton> MeshSkeleton = TypedValue->GetSkeleton())
 				{
-					args.physicsBody = program.AddConstant(pPhysicsBody.get());
-					MeshData->SetPhysicsBody(nullptr);
+					Args.Skeleton = Program.AddConstant(MeshSkeleton.get());
 				}
 
-				// Use a map-based deduplication
-				mu::Ptr<const mu::Mesh> Key = MeshData;
-				const int32* IndexPtr = Options->MeshConstantMap.Find(Key);
-				if (!IndexPtr)
-				{
-					args.value = program.AddConstant(MeshData.get());
-					Options->MeshConstantMap.Add(MeshData, int32(args.value));
-				}
-				else
+				Args.PhysicsBody = -1;
+				if (Ptr<const PhysicsBody> MeshPhysicsBody = TypedValue->GetPhysicsBody())
 				{
-					args.value = *IndexPtr;
+					Args.PhysicsBody = Program.AddConstant(MeshPhysicsBody.get());
 				}
 
-				linkedAddress = (OP::ADDRESS)program.m_opAddress.Num();
-				program.m_opAddress.Add((uint32_t)program.m_byteCode.Num());
-				AppendCode(program.m_byteCode, Type);
-				AppendCode(program.m_byteCode, args);
+				linkedAddress = (OP::ADDRESS)Program.m_opAddress.Num();
+				Program.m_opAddress.Add((uint32)Program.m_byteCode.Num());
+				AppendCode(Program.m_byteCode, Type);
+				AppendCode(Program.m_byteCode, Args);
 			}
 			else
 			{
@@ -468,7 +623,7 @@ namespace mu
 					}
 					else
 					{
-						args.value = AddConstantImage( program, pTyped, *Options);
+						args.value = AddConstantImage(Program, pTyped, *Options);
 
 						int32 DataDescIndex = Options->AdditionalData.SourceImagePerConstant.Add(SourceDataDescriptor);
 						check(DataDescIndex == args.value);
@@ -480,7 +635,7 @@ namespace mu
 				{
 					Ptr<const Layout> pTyped = static_cast<const Layout*>(GetValue().get());
 					check(pTyped);
-					args.value = program.AddConstant(pTyped);
+					args.value = Program.AddConstant(pTyped);
 					break;
 				}
 				default:
@@ -489,10 +644,10 @@ namespace mu
 
 				if (bValidData)
 				{
-					linkedAddress = (OP::ADDRESS)program.m_opAddress.Num();
-					program.m_opAddress.Add((uint32)program.m_byteCode.Num());
-					AppendCode(program.m_byteCode, Type);
-					AppendCode(program.m_byteCode, args);
+					linkedAddress = (OP::ADDRESS)Program.m_opAddress.Num();
+					Program.m_opAddress.Add((uint32)Program.m_byteCode.Num());
+					AppendCode(Program.m_byteCode, Type);
+					AppendCode(Program.m_byteCode, args);
 				}
 				else
 				{
diff --git a/Engine/Plugins/Experimental/Mutable/Source/MutableTools/Private/MuT/CodeOptimiser.cpp b/Engine/Plugins/Experimental/Mutable/Source/MutableTools/Private/MuT/CodeOptimiser.cpp
--- a/Engine/Plugins/Experimental/Mutable/Source/MutableTools/Private/MuT/CodeOptimiser.cpp
+++ b/Engine/Plugins/Experimental/Mutable/Source/MutableTools/Private/MuT/CodeOptimiser.cpp
@@ -503,21 +503,23 @@ namespace mu
 			pSystem->GetPrivate()->BeginBuild( model );
 
 			// Calculate the value and replace this op by a constant
-			switch( DataType )
+			switch(DataType)
 			{
 			case DT_MESH:
 			{
 				MUTABLE_CPUPROFILER_SCOPE(ConstantMesh);
 
-				mu::Ptr<const Mesh> pMesh = pSystem->GetPrivate()->BuildMesh( model, LocalParams.get(), at );
+				// Build mesh with ExecutionOptions indicating a full mesh load is needed for mesh constant resolution.
+				const uint8 ExecutionOptions = static_cast<uint8>(EMeshExecutionOptions::AllFlags);
+				mu::Ptr<const Mesh> BuildMesh = pSystem->GetPrivate()->BuildMesh(model, LocalParams.get(), at, ExecutionOptions);
 
-				if (pMesh)
+				if (BuildMesh)
 				{
-					mu::Ptr<ASTOpConstantResource> constantOp = new ASTOpConstantResource();
-					constantOp->Type = OP_TYPE::ME_CONSTANT;
-					constantOp->SetValue( pMesh, DiskCacheContext );
-					Result = constantOp;
-				  }
+					mu::Ptr<ASTOpConstantResource> ConstantOp = new ASTOpConstantResource();
+					ConstantOp->Type = OP_TYPE::ME_CONSTANT;
+					ConstantOp->SetValue(BuildMesh, DiskCacheContext);
+					Result = ConstantOp;
+				}
 				break;
 			}
 
diff --git a/Engine/Plugins/Experimental/Mutable/Source/MutableTools/Private/MuT/Compiler.cpp b/Engine/Plugins/Experimental/Mutable/Source/MutableTools/Private/MuT/Compiler.cpp
--- a/Engine/Plugins/Experimental/Mutable/Source/MutableTools/Private/MuT/Compiler.cpp
+++ b/Engine/Plugins/Experimental/Mutable/Source/MutableTools/Private/MuT/Compiler.cpp
@@ -463,14 +463,14 @@ namespace mu
 		uint64 NumHighResBytes = 0;
 
 		// Maximum number of roms
-		int32 MaxRomCount = Program.ConstantImageLODs.Num() + Program.ConstantMeshes.Num();
+		int32 MaxRomCount = Program.ConstantImageLODs.Num() + Program.ConstantMeshData.Num();
 		Program.m_roms.Reserve(MaxRomCount);
 
 		TSet<uint32> UsedIds;
 		UsedIds.Reserve(MaxRomCount);
 
 		TArray<FRomData> RomDatas;
-		RomDatas.SetNumZeroed(FMath::Max(Program.ConstantImageLODs.Num(), Program.ConstantMeshes.Num()));
+		RomDatas.SetNumZeroed(FMath::Max(Program.ConstantImageLODs.Num(), Program.ConstantMeshData.Num()));
 
 		// Images
 		{
@@ -579,10 +579,10 @@ namespace mu
 		{
 			MUTABLE_CPUPROFILER_SCOPE(GenerateRoms_MeshIds);
 
-			ParallelFor(Program.ConstantMeshes.Num(),
+			ParallelFor(Program.ConstantMeshData.Num(),
 				[EmbeddedDataBytesLimit, &Program, &RomDatas](uint32 ResourceIndex)
 				{
-					TPair<int32, Ptr<const Mesh>>& ResData = Program.ConstantMeshes[ResourceIndex];
+					TPair<int32, Ptr<const Mesh>>& ResData = Program.ConstantMeshData[ResourceIndex];
 
 					// This shouldn't have been serialised with rom support before.
 					check(ResData.Key < 0);
@@ -616,7 +616,7 @@ namespace mu
 		{
 			MUTABLE_CPUPROFILER_SCOPE(GenerateRoms_MeshIdsUnique);
 
-			for (int32 ResourceIndex = 0; ResourceIndex < Program.ConstantMeshes.Num(); ++ResourceIndex)
+			for (int32 ResourceIndex = 0; ResourceIndex < Program.ConstantMeshData.Num(); ++ResourceIndex)
 			{
 				FRomData& RomData = RomDatas[ResourceIndex];
 
@@ -635,7 +635,7 @@ namespace mu
 
 				int32 RomIndex = Program.m_roms.Add(RomData);
 
-				TPair<int32, Ptr<const Mesh>>& ResData = Program.ConstantMeshes[ResourceIndex];
+				TPair<int32, Ptr<const Mesh>>& ResData = Program.ConstantMeshData[ResourceIndex];
 				ResData.Key = RomIndex;
 			}
 		}
diff --git a/Engine/Plugins/Experimental/Mutable/Source/MutableTools/Private/MuT/DataPacker.cpp b/Engine/Plugins/Experimental/Mutable/Source/MutableTools/Private/MuT/DataPacker.cpp
--- a/Engine/Plugins/Experimental/Mutable/Source/MutableTools/Private/MuT/DataPacker.cpp
+++ b/Engine/Plugins/Experimental/Mutable/Source/MutableTools/Private/MuT/DataPacker.cpp
@@ -230,12 +230,10 @@ namespace mu
                 default:
                 {
                     // If we hit the constant we are analysing
-                    if ( m_opType == OP_TYPE::ME_CONSTANT
-                         &&
-                         program.GetOpType(at) == m_opType )
+                    if (m_opType == OP_TYPE::ME_CONSTANT && program.GetOpType(at) == m_opType)
                     {
-						OP::MeshConstantArgs args = program.GetOpArgs<OP::MeshConstantArgs>(at);
-                        if ( args.value == m_constant )
+						OP::MeshConstantArgs Args = program.GetOpArgs<OP::MeshConstantArgs>(at);
+                        if (Args.Value == m_constant)
                         {
                             // Accumulate the currently relevant parameters
                             const TArray<int>& currentParams = GetCurrentState();
