[Mutable] Fix crash when resizing bone format of empty mesh buffers.

#rnx

diff --git a/Engine/Plugins/Experimental/Mutable/Source/MutableRuntime/Private/MuR/OpMeshFormat.cpp b/Engine/Plugins/Experimental/Mutable/Source/MutableRuntime/Private/MuR/OpMeshFormat.cpp
--- a/Engine/Plugins/Experimental/Mutable/Source/MutableRuntime/Private/MuR/OpMeshFormat.cpp
+++ b/Engine/Plugins/Experimental/Mutable/Source/MutableRuntime/Private/MuR/OpMeshFormat.cpp
@@ -534,20 +534,20 @@ namespace mu
 		// we will find in the mesh. Otherwise, find it from the vertex buffers data.
 		if (VertexBuffers.IsDescriptor())
 		{
-			return InMesh->BoneMap.Num() - 1;
+			return FMath::Max(0, InMesh->BoneMap.Num() - 1);
 		}
 
 		const UntypedMeshBufferIteratorConst BoneIndicesBegin(VertexBuffers, MBS_BONEINDICES, Channel.SemanticIndex);
 
 		if (!BoneIndicesBegin.ptr())
 		{
-			return -1;
+			return 0;
 		}
 
 		const int32 NumVertices = VertexBuffers.GetElementCount();
 		const int32 NumInfluences = BoneIndicesBegin.GetComponents();
 		
-		int32 HighestBoneIndex = -1;
+		int32 HighestBoneIndex = 0;
 		for (int32 VertexIndex = 0; VertexIndex < NumVertices; ++VertexIndex)
 		{
 			// If MAX_TOTAL_INFLUENCES ever changed, the next line would no longer work or compile and 
