[mutable] Fix crash creating a new layout node.
#rnx

diff --git a/Engine/Plugins/Experimental/Mutable/Source/CustomizableObjectEditor/Private/MuCOE/SCustomizableObjectNodeLayoutBlocksEditor.cpp b/Engine/Plugins/Experimental/Mutable/Source/CustomizableObjectEditor/Private/MuCOE/SCustomizableObjectNodeLayoutBlocksEditor.cpp
--- a/Engine/Plugins/Experimental/Mutable/Source/CustomizableObjectEditor/Private/MuCOE/SCustomizableObjectNodeLayoutBlocksEditor.cpp
+++ b/Engine/Plugins/Experimental/Mutable/Source/CustomizableObjectEditor/Private/MuCOE/SCustomizableObjectNodeLayoutBlocksEditor.cpp
@@ -154,7 +154,7 @@ TSharedRef<SWidget> SCustomizableObjectNodeLayoutBlocksEditor::BuildLayoutToolBa
 	const ISlateStyle* const StyleSet = &FCoreStyle::Get();
 	const FName& StyleName = "ToolBar";
 
-	if (CurrentLayout->PackingStrategy!=ECustomizableObjectTextureLayoutPackingStrategy::Overlay)
+	if (CurrentLayout && CurrentLayout->PackingStrategy!=ECustomizableObjectTextureLayoutPackingStrategy::Overlay)
 	{
 		LayoutToolbarBuilder.BeginSection("Blocks");
 		{
