[mutable] Rename some material node outputs to "modifier".
- Add the "modifier" pin type.
- Change the output of modifier nodes to "modifier"
- Add a "modifier" pin input to objects.
- Reduced log warning level for some data upgrade issues.
- Fix bug with details for Clip Mesh node showing up correctly only the second time.
#rnx

diff --git a/Engine/Plugins/Experimental/Mutable/Source/CustomizableObject/Internal/MuCO/CustomizableObjectCustomVersion.h b/Engine/Plugins/Experimental/Mutable/Source/CustomizableObject/Internal/MuCO/CustomizableObjectCustomVersion.h
--- a/Engine/Plugins/Experimental/Mutable/Source/CustomizableObject/Internal/MuCO/CustomizableObjectCustomVersion.h
+++ b/Engine/Plugins/Experimental/Mutable/Source/CustomizableObject/Internal/MuCO/CustomizableObjectCustomVersion.h
@@ -130,6 +130,8 @@ struct CUSTOMIZABLEOBJECT_API FCustomizableObjectCustomVersion
 
 		UseUVRects, // Instead of pointing at layout blocks in the edited parent, define custom absolute UV layout rects.
 
+		AddModifierPin, // Move the modifiers output pin from Material to Modifier
+
 		// -----<new versions can be added above this line>-------------------------------------------------
 		VersionPlusOne,
 		LatestVersion = VersionPlusOne - 1
diff --git a/Engine/Plugins/Experimental/Mutable/Source/CustomizableObjectEditor/Private/MuCOE/EdGraphSchema_CustomizableObject.cpp b/Engine/Plugins/Experimental/Mutable/Source/CustomizableObjectEditor/Private/MuCOE/EdGraphSchema_CustomizableObject.cpp
--- a/Engine/Plugins/Experimental/Mutable/Source/CustomizableObjectEditor/Private/MuCOE/EdGraphSchema_CustomizableObject.cpp
+++ b/Engine/Plugins/Experimental/Mutable/Source/CustomizableObjectEditor/Private/MuCOE/EdGraphSchema_CustomizableObject.cpp
@@ -226,6 +226,7 @@ UEdGraphNode* FCustomizableObjectSchemaAction_Paste::PerformAction(class UEdGrap
 const FName UEdGraphSchema_CustomizableObject::PC_Object("object");
 const FName UEdGraphSchema_CustomizableObject::PC_Component("component");
 const FName UEdGraphSchema_CustomizableObject::PC_Material("material");
+const FName UEdGraphSchema_CustomizableObject::PC_Modifier("modifier");
 const FName UEdGraphSchema_CustomizableObject::PC_Mesh("mesh");
 const FName UEdGraphSchema_CustomizableObject::PC_PassThroughMesh("passThroughMesh");
 const FName UEdGraphSchema_CustomizableObject::PC_Layout("layout");
@@ -533,7 +534,7 @@ void UEdGraphSchema_CustomizableObject::GetGraphContextActions(FGraphContextMenu
 
 	{
 		// External Pin Nodes
-		TArray<FName> PinTypes({ PC_Material, PC_Mesh, PC_Image, PC_Projector, PC_GroupProjector, PC_Color, PC_Float, PC_Bool, PC_Enum, PC_Stack, PC_PassThroughImage, PC_MaterialAsset, PC_PoseAsset });
+		TArray<FName> PinTypes({ PC_Material, PC_Modifier, PC_Mesh, PC_Image, PC_Projector, PC_GroupProjector, PC_Color, PC_Float, PC_Bool, PC_Enum, PC_Stack, PC_PassThroughImage, PC_MaterialAsset, PC_PoseAsset });
 
 		// Add pin types from extensions
 		for (const FRegisteredCustomizableObjectPinType& PinType : ICustomizableObjectModule::Get().GetExtendedPinTypes())
@@ -720,6 +721,10 @@ FLinearColor UEdGraphSchema_CustomizableObject::GetPinTypeColor(const FName& Typ
 	{
 		return FLinearColor(0.000000f, 0.100000f, 0.600000f, 1.000000f); // Blue
 	}
+	else if (TypeString == PC_Modifier)
+	{
+		return FColor::Silver;
+	}
 	else if (TypeString == PC_Object)
 	{
 		return FLinearColor(0.000000f, 0.400000f, 0.910000f, 1.000000f); // Light blue
@@ -1335,6 +1340,10 @@ FText UEdGraphSchema_CustomizableObject::GetPinCategoryName(const FName& PinCate
 	{
 		return LOCTEXT("Material_Pin_Category", "Material");
 	}
+	else if (PinCategory == UEdGraphSchema_CustomizableObject::PC_Modifier)
+	{
+		return LOCTEXT("Modifier_Pin_Category", "Modifier");
+	}
 	else if (PinCategory == UEdGraphSchema_CustomizableObject::PC_Mesh)
 	{
 		return LOCTEXT("Mesh_Pin_Category", "Mesh");
diff --git a/Engine/Plugins/Experimental/Mutable/Source/CustomizableObjectEditor/Private/MuCOE/GenerateMutableSource/GenerateMutableSource.cpp b/Engine/Plugins/Experimental/Mutable/Source/CustomizableObjectEditor/Private/MuCOE/GenerateMutableSource/GenerateMutableSource.cpp
--- a/Engine/Plugins/Experimental/Mutable/Source/CustomizableObjectEditor/Private/MuCOE/GenerateMutableSource/GenerateMutableSource.cpp
+++ b/Engine/Plugins/Experimental/Mutable/Source/CustomizableObjectEditor/Private/MuCOE/GenerateMutableSource/GenerateMutableSource.cpp
@@ -1105,11 +1105,11 @@ mu::Ptr<mu::NodeObject> GenerateMutableSource(const UEdGraphPin * Pin, FMutableG
 							continue;
 						}
 
-						mu::NodeSurfacePtr SurfaceNode = GenerateMutableSourceSurface(ChildNodePin, GenerationContext);
+						mu::Ptr<mu::NodeSurface> SurfaceNode = GenerateMutableSourceSurface(ChildNodePin, GenerationContext);
 						LODNode->Surfaces.Add(SurfaceNode);
 					}
 
-					// Process modifiers.
+					// Process legacy modifiers.
 					for (UEdGraphPin* const ChildNodePin : ConnectedLODPins)
 					{
 						if (!Cast<UCustomizableObjectNodeModifierBase>(ChildNodePin->GetOwningNode()))
@@ -1117,15 +1117,42 @@ mu::Ptr<mu::NodeObject> GenerateMutableSource(const UEdGraphPin * Pin, FMutableG
 							continue;
 						}
 
+						// Warn about legacy modifier connection
+						FString Msg = FString::Printf(TEXT("The object has legacy modifier connections (to material pins?) that should be updated."));
+						GenerationContext.Compiler->CompilerLog(FText::FromString(Msg), Node, EMessageSeverity::Warning);
+
+						// Set it to "None" to indicate we don't care about component id.
+						GenerationContext.CurrentMeshComponent = FName();
+
+						mu::Ptr<mu::NodeModifier> ModifierNode = GenerateMutableSourceModifier(ChildNodePin, GenerationContext);
+						LODNode->Modifiers.Add(ModifierNode);
+
+						GenerationContext.CurrentMeshComponent = ComponentName;
+					}
+				}
+
+				// Process modifiers.
+				const UEdGraphPin* ModifierPin = TypedNodeObj->ModifiersPin();
+				if (ModifierPin)
+				{
+					TArray<UEdGraphPin*> ConnectedModifierPins = FollowInputPinArray(*ModifierPin);
+					for (UEdGraphPin* const ChildNodePin : ConnectedModifierPins)
+					{
+						if (!Cast<UCustomizableObjectNodeModifierBase>(ChildNodePin->GetOwningNode()))
+						{
+							continue;
+						}
+
 						// Set it to "None" to indicate we don't care about component id.
 						GenerationContext.CurrentMeshComponent = FName();
 
-						mu::NodeModifierPtr ModifierNode = GenerateMutableSourceModifier(ChildNodePin, GenerationContext);
+						mu::Ptr<mu::NodeModifier> ModifierNode = GenerateMutableSourceModifier(ChildNodePin, GenerationContext);
 						LODNode->Modifiers.Add(ModifierNode);
 
 						GenerationContext.CurrentMeshComponent = ComponentName;
 					}
 				}
+
 			}
 		}
 
diff --git a/Engine/Plugins/Experimental/Mutable/Source/CustomizableObjectEditor/Private/MuCOE/GenerateMutableSource/GenerateMutableSourceModifier.cpp b/Engine/Plugins/Experimental/Mutable/Source/CustomizableObjectEditor/Private/MuCOE/GenerateMutableSource/GenerateMutableSourceModifier.cpp
--- a/Engine/Plugins/Experimental/Mutable/Source/CustomizableObjectEditor/Private/MuCOE/GenerateMutableSource/GenerateMutableSourceModifier.cpp
+++ b/Engine/Plugins/Experimental/Mutable/Source/CustomizableObjectEditor/Private/MuCOE/GenerateMutableSource/GenerateMutableSourceModifier.cpp
@@ -22,7 +22,7 @@
 #define LOCTEXT_NAMESPACE "CustomizableObjectEditor"
 
 
-mu::NodeModifierPtr GenerateMutableSourceModifier(const UEdGraphPin * Pin, FMutableGraphGenerationContext & GenerationContext)
+mu::Ptr<mu::NodeModifier> GenerateMutableSourceModifier(const UEdGraphPin * Pin, FMutableGraphGenerationContext & GenerationContext)
 {
 	check(Pin)
 	RETURN_ON_CYCLE(*Pin, GenerationContext)
diff --git a/Engine/Plugins/Experimental/Mutable/Source/CustomizableObjectEditor/Private/MuCOE/GenerateMutableSource/GenerateMutableSourceModifier.h b/Engine/Plugins/Experimental/Mutable/Source/CustomizableObjectEditor/Private/MuCOE/GenerateMutableSource/GenerateMutableSourceModifier.h
--- a/Engine/Plugins/Experimental/Mutable/Source/CustomizableObjectEditor/Private/MuCOE/GenerateMutableSource/GenerateMutableSourceModifier.h
+++ b/Engine/Plugins/Experimental/Mutable/Source/CustomizableObjectEditor/Private/MuCOE/GenerateMutableSource/GenerateMutableSourceModifier.h
@@ -7,4 +7,4 @@
 class UEdGraphPin;
 struct FMutableGraphGenerationContext;
 
-mu::NodeModifierPtr GenerateMutableSourceModifier(const UEdGraphPin* Pin, FMutableGraphGenerationContext& GenerationContext);
+mu::Ptr<mu::NodeModifier> GenerateMutableSourceModifier(const UEdGraphPin* Pin, FMutableGraphGenerationContext& GenerationContext);
diff --git a/Engine/Plugins/Experimental/Mutable/Source/CustomizableObjectEditor/Private/MuCOE/Nodes/CustomizableObjectNodeEditMaterial.cpp b/Engine/Plugins/Experimental/Mutable/Source/CustomizableObjectEditor/Private/MuCOE/Nodes/CustomizableObjectNodeEditMaterial.cpp
--- a/Engine/Plugins/Experimental/Mutable/Source/CustomizableObjectEditor/Private/MuCOE/Nodes/CustomizableObjectNodeEditMaterial.cpp
+++ b/Engine/Plugins/Experimental/Mutable/Source/CustomizableObjectEditor/Private/MuCOE/Nodes/CustomizableObjectNodeEditMaterial.cpp
@@ -180,7 +180,7 @@ void UCustomizableObjectNodeEditMaterial::BackwardsCompatibleFixup()
 						if (ExistingBlock.Id == BlockId)
 						{
 							bSkipBlock = true;
-							UE_LOG(LogMutable, Warning, TEXT("[%s] UCustomizableObjectNodeEditMaterial has a duplicated layout block id. One has been ignored during version upgrade."), *GetOutermost()->GetName());
+							UE_LOG(LogMutable, Log, TEXT("[%s] UCustomizableObjectNodeEditMaterial has a duplicated layout block id. One has been ignored during version upgrade."), *GetOutermost()->GetName());
 							break;
 						}
 					}
diff --git a/Engine/Plugins/Experimental/Mutable/Source/CustomizableObjectEditor/Private/MuCOE/Nodes/CustomizableObjectNodeMeshClipDeform.cpp b/Engine/Plugins/Experimental/Mutable/Source/CustomizableObjectEditor/Private/MuCOE/Nodes/CustomizableObjectNodeMeshClipDeform.cpp
--- a/Engine/Plugins/Experimental/Mutable/Source/CustomizableObjectEditor/Private/MuCOE/Nodes/CustomizableObjectNodeMeshClipDeform.cpp
+++ b/Engine/Plugins/Experimental/Mutable/Source/CustomizableObjectEditor/Private/MuCOE/Nodes/CustomizableObjectNodeMeshClipDeform.cpp
@@ -16,7 +16,7 @@ void UCustomizableObjectNodeMeshClipDeform::AllocateDefaultPins(UCustomizableObj
 	UEdGraphPin* ClipMeshPin = CustomCreatePin(EGPD_Input, Schema->PC_Mesh, FName("Clip Shape"));
 	ClipMeshPin->bDefaultValueIsIgnored = true;
 	
-	UEdGraphPin* OutputPin = CustomCreatePin(EGPD_Output, Schema->PC_Material, FName("Material"));
+	UEdGraphPin* OutputPin = CustomCreatePin(EGPD_Output, Schema->PC_Modifier, FName("Modifier"));
 	ClipMeshPin->bDefaultValueIsIgnored = true;
 }
 
@@ -29,7 +29,7 @@ UEdGraphPin* UCustomizableObjectNodeMeshClipDeform::ClipShapePin() const
 
 UEdGraphPin* UCustomizableObjectNodeMeshClipDeform::OutputPin() const
 {
-	return FindPin(TEXT("Material"));
+	return FindPin(TEXT("Modifier"));
 }
 
 
@@ -38,11 +38,6 @@ FText UCustomizableObjectNodeMeshClipDeform::GetNodeTitle(ENodeTitleType::Type T
 	return LOCTEXT("Clip_Deform_Mesh", "Clip Deform Mesh");
 }
 
-FLinearColor UCustomizableObjectNodeMeshClipDeform::GetNodeTitleColor() const
-{
-	const UEdGraphSchema_CustomizableObject* Schema = GetDefault<UEdGraphSchema_CustomizableObject>();
-	return Schema->GetPinTypeColor(Schema->PC_Material);
-}
 
 FText UCustomizableObjectNodeMeshClipDeform::GetTooltipText() const
 {
diff --git a/Engine/Plugins/Experimental/Mutable/Source/CustomizableObjectEditor/Private/MuCOE/Nodes/CustomizableObjectNodeMeshClipMorph.cpp b/Engine/Plugins/Experimental/Mutable/Source/CustomizableObjectEditor/Private/MuCOE/Nodes/CustomizableObjectNodeMeshClipMorph.cpp
--- a/Engine/Plugins/Experimental/Mutable/Source/CustomizableObjectEditor/Private/MuCOE/Nodes/CustomizableObjectNodeMeshClipMorph.cpp
+++ b/Engine/Plugins/Experimental/Mutable/Source/CustomizableObjectEditor/Private/MuCOE/Nodes/CustomizableObjectNodeMeshClipMorph.cpp
@@ -6,6 +6,7 @@
 #include "MuCOE/EdGraphSchema_CustomizableObject.h"
 #include "MuCOE/ICustomizableObjectEditor.h"
 #include "MuCOE/UnrealEditorPortabilityHelpers.h"
+#include "MuCOE/Nodes/CustomizableObjectNodeObject.h"
 
 class UCustomizableObjectNodeRemapPins;
 class UEdGraphPin;
@@ -88,7 +89,7 @@ void UCustomizableObjectNodeMeshClipMorph::ChangeStartOffsetTransform()
 
 UEdGraphPin* UCustomizableObjectNodeMeshClipMorph::OutputPin() const
 {
-	return FindPin(TEXT("Material"));
+	return FindPin(TEXT("Modifier"));
 }
 
 
@@ -106,6 +107,7 @@ void UCustomizableObjectNodeMeshClipMorph::PostEditChangeProperty(FPropertyChang
 	Super::PostEditChangeProperty(PropertyChangedEvent);
 }
 
+
 void UCustomizableObjectNodeMeshClipMorph::Serialize(FArchive& Ar)
 {
 	Super::Serialize(Ar);
@@ -135,7 +137,7 @@ void UCustomizableObjectNodeMeshClipMorph::AllocateDefaultPins(UCustomizableObje
 {
 	const UEdGraphSchema_CustomizableObject* Schema = GetDefault<UEdGraphSchema_CustomizableObject>();
 
-	UEdGraphPin* OutputPin = CustomCreatePin(EGPD_Output, Schema->PC_Material, FName("Material"));
+	UEdGraphPin* OutputPin = CustomCreatePin(EGPD_Output, Schema->PC_Modifier, FName("Modifier"));
 }
 
 
@@ -145,13 +147,6 @@ FText UCustomizableObjectNodeMeshClipMorph::GetNodeTitle(ENodeTitleType::Type Ti
 }
 
 
-FLinearColor UCustomizableObjectNodeMeshClipMorph::GetNodeTitleColor() const
-{
-	const UEdGraphSchema_CustomizableObject* Schema = GetDefault<UEdGraphSchema_CustomizableObject>();
-	return Schema->GetPinTypeColor(Schema->PC_Material);
-}
-
-
 void UCustomizableObjectNodeMeshClipMorph::PinConnectionListChanged(UEdGraphPin* Pin)
 {
 	Super::PinConnectionListChanged(Pin);
@@ -168,7 +163,6 @@ void UCustomizableObjectNodeMeshClipMorph::PinConnectionListChanged(UEdGraphPin*
 }
 
 
-
 FText UCustomizableObjectNodeMeshClipMorph::GetTooltipText() const
 {
 	return LOCTEXT("Clip_Mesh_Morph_Tooltip", "Defines a cutting plane on a bone to cut tagged Materials that go past it, while morphing the mesh after the cut to blend in more naturally.\nIt only cuts and morphs mesh that receives some influence of that bone or other descendant bones.");
diff --git a/Engine/Plugins/Experimental/Mutable/Source/CustomizableObjectEditor/Private/MuCOE/Nodes/CustomizableObjectNodeMeshClipWithMesh.cpp b/Engine/Plugins/Experimental/Mutable/Source/CustomizableObjectEditor/Private/MuCOE/Nodes/CustomizableObjectNodeMeshClipWithMesh.cpp
--- a/Engine/Plugins/Experimental/Mutable/Source/CustomizableObjectEditor/Private/MuCOE/Nodes/CustomizableObjectNodeMeshClipWithMesh.cpp
+++ b/Engine/Plugins/Experimental/Mutable/Source/CustomizableObjectEditor/Private/MuCOE/Nodes/CustomizableObjectNodeMeshClipWithMesh.cpp
@@ -52,7 +52,7 @@ void UCustomizableObjectNodeMeshClipWithMesh::AllocateDefaultPins(UCustomizableO
 	UEdGraphPin* ClipMeshPin = CustomCreatePin(EGPD_Input, Schema->PC_Mesh, FName("Clip Mesh"));
 	ClipMeshPin->bDefaultValueIsIgnored = true;
 
-	UEdGraphPin* OutputPin_p = CustomCreatePin(EGPD_Output, Schema->PC_Material, FName("Material"));
+	UEdGraphPin* OutputPin_p = CustomCreatePin(EGPD_Output, Schema->PC_Modifier, FName("Modifier"));
 }
 
 
@@ -73,13 +73,6 @@ FText UCustomizableObjectNodeMeshClipWithMesh::GetNodeTitle(ENodeTitleType::Type
 }
 
 
-FLinearColor UCustomizableObjectNodeMeshClipWithMesh::GetNodeTitleColor() const
-{
-	const UEdGraphSchema_CustomizableObject* Schema = GetDefault<UEdGraphSchema_CustomizableObject>();
-	return Schema->GetPinTypeColor(Schema->PC_Material);
-}
-
-
 void UCustomizableObjectNodeMeshClipWithMesh::PinConnectionListChanged(UEdGraphPin* Pin)
 {
 	Super::PinConnectionListChanged(Pin);
@@ -95,7 +88,7 @@ void UCustomizableObjectNodeMeshClipWithMesh::PinConnectionListChanged(UEdGraphP
 
 UEdGraphPin* UCustomizableObjectNodeMeshClipWithMesh::OutputPin() const
 {
-	return FindPin(TEXT("Material"));
+	return FindPin(TEXT("Modifier"));
 }
 
 
diff --git a/Engine/Plugins/Experimental/Mutable/Source/CustomizableObjectEditor/Private/MuCOE/Nodes/CustomizableObjectNodeMeshClipWithMeshDetails.cpp b/Engine/Plugins/Experimental/Mutable/Source/CustomizableObjectEditor/Private/MuCOE/Nodes/CustomizableObjectNodeMeshClipWithMeshDetails.cpp
--- a/Engine/Plugins/Experimental/Mutable/Source/CustomizableObjectEditor/Private/MuCOE/Nodes/CustomizableObjectNodeMeshClipWithMeshDetails.cpp
+++ b/Engine/Plugins/Experimental/Mutable/Source/CustomizableObjectEditor/Private/MuCOE/Nodes/CustomizableObjectNodeMeshClipWithMeshDetails.cpp
@@ -44,7 +44,6 @@ void FCustomizableObjectNodeMeshClipWithMeshDetails::CustomizeDetails(IDetailLay
 
 	IDetailCategoryBuilder& CustomizableObjectToClipCategory = DetailBuilder.EditCategory("CustomizableObjectToClip");
 
-
 	if(Node != nullptr)
 	{
 		SelectedCO = Node->CustomizableObjectToClipWith;
@@ -54,10 +53,6 @@ void FCustomizableObjectNodeMeshClipWithMeshDetails::CustomizeDetails(IDetailLay
 		// Hidding unnecessary properties
 		DetailBuilder.HideProperty("CustomizableObjectToClipWith");
 		DetailBuilder.HideProperty("ArrayMaterialNodeToClipWithID");
-		if (!Node->bUseTags)
-		{
-			DetailBuilder.HideProperty("Tags");			
-		}
 
 		if (const UEdGraphPin* ConnectedPin = FollowInputPin(*Node->ClipMeshPin());
 			ConnectedPin && !Cast<UCustomizableObjectNodeStaticMesh>(ConnectedPin->GetOwningNode()))
@@ -103,6 +98,12 @@ void FCustomizableObjectNodeMeshClipWithMeshDetails::CustomizeDetails(IDetailLay
 		// Initializing the inital value of the ComboBox
 		SelectedMethod = SetInitialClippingMethod();
 		
+		// This flag is now up to date and can be used to hide ui.
+		if (!Node->bUseTags)
+		{
+			DetailBuilder.HideProperty("Tags");
+		}
+
 		// ComboBox to select the Clipping method (material or/and tags)
 		CustomizableObjectToClipCategory.AddCustomRow(LOCTEXT("FCustomizableObjectNodeMeshClipWithMeshDetails", "Blocks"))
 		[
diff --git a/Engine/Plugins/Experimental/Mutable/Source/CustomizableObjectEditor/Private/MuCOE/Nodes/CustomizableObjectNodeModifierBase.cpp b/Engine/Plugins/Experimental/Mutable/Source/CustomizableObjectEditor/Private/MuCOE/Nodes/CustomizableObjectNodeModifierBase.cpp
--- a/Engine/Plugins/Experimental/Mutable/Source/CustomizableObjectEditor/Private/MuCOE/Nodes/CustomizableObjectNodeModifierBase.cpp
+++ b/Engine/Plugins/Experimental/Mutable/Source/CustomizableObjectEditor/Private/MuCOE/Nodes/CustomizableObjectNodeModifierBase.cpp
@@ -2,8 +2,122 @@
 
 #include "MuCOE/Nodes/CustomizableObjectNodeModifierBase.h"
 
+#include "MuCOE/EdGraphSchema_CustomizableObject.h"
+#include "MuCOE/CustomizableObjectEditorLogger.h"
+#include "MuCO/CustomizableObjectCustomVersion.h"
+#include "MuCOE/Nodes/CustomizableObjectNodeObject.h"
 
-UCustomizableObjectNodeModifierBase::UCustomizableObjectNodeModifierBase()
-	: Super()
+
+FLinearColor UCustomizableObjectNodeModifierBase::GetNodeTitleColor() const
+{
+	const UEdGraphSchema_CustomizableObject* Schema = GetDefault<UEdGraphSchema_CustomizableObject>();
+	return Schema->GetPinTypeColor(Schema->PC_Modifier);
+}
+
+
+void UCustomizableObjectNodeModifierBase::PostBackwardsCompatibleFixup()
 {
+	Super::PostBackwardsCompatibleFixup();
+
+	int32 CustomizableObjectCustomVersion = GetLinkerCustomVersion(FCustomizableObjectCustomVersion::GUID);
+
+	const UEdGraphSchema_CustomizableObject* Schema = GetDefault<UEdGraphSchema_CustomizableObject>();
+
+	// Remove "Material" pin and add "Modifier"
+	if (CustomizableObjectCustomVersion < FCustomizableObjectCustomVersion::AddModifierPin)
+	{
+		// Replace the output pin
+		UEdGraphPin* OldPin = FindPin(TEXT("Material"));
+		UEdGraphPin* NewPin = CustomCreatePin(EGPD_Output, Schema->PC_Modifier, FName("Modifier"));
+		if (OldPin && NewPin)
+		{
+			FGuid PinId = NewPin->PinId;
+			NewPin->CopyPersistentDataFromOldPin(*OldPin);
+			NewPin->PinId = PinId;
+			NewPin->bHidden = OldPin->bHidden;
+
+			CustomRemovePin(*OldPin);
+
+			// Reconnect it to the correct output of the target node
+			TArray<UEdGraphPin*> LinksToRemove;
+			for (int32 LinkIndex = 0; LinkIndex < NewPin->LinkedTo.Num(); ++LinkIndex)
+			{
+				UEdGraphPin* LinkedToPin = NewPin->LinkedTo[LinkIndex];
+
+				if (!LinkedToPin)
+				{
+					continue;
+				}
+
+				UEdGraphNode* ToNode = LinkedToPin->GetOwningNode();
+				if (!ToNode)
+				{
+					continue;
+				}
+
+				if (UCustomizableObjectNodeObject* ToObjectNode = Cast<UCustomizableObjectNodeObject>(ToNode))
+				{
+					UEdGraphPin* ToModifierPin = ToObjectNode->ModifiersPin();
+					if (ToModifierPin == LinkedToPin)
+					{
+						// It is already correctly connected
+						continue;
+					}
+
+					if (ToModifierPin)
+					{
+						LinksToRemove.Add(LinkedToPin);
+						NewPin->MakeLinkTo(ToModifierPin);
+					}
+				}
+				else
+				{
+					// The modifier is connected to a node for which automatic upgrade support is not implemented. This will be warned below
+				}
+			}
+
+			// Remove reconnected links
+			for (UEdGraphPin* LinkToRemove : LinksToRemove)
+			{
+				NewPin->BreakLinkTo(LinkToRemove);
+			}
+		}
+	}
+
+	// Check for old legacy connections that need manual update
+	{
+		UEdGraphPin* NewPin = FindPin(TEXT("Modifier"));
+		if (NewPin)
+		{
+			for (int32 LinkIndex = 0; LinkIndex < NewPin->LinkedTo.Num(); ++LinkIndex)
+			{
+				UEdGraphPin* LinkedToPin = NewPin->LinkedTo[LinkIndex];
+
+				if (!LinkedToPin)
+				{
+					continue;
+				}
+
+				UEdGraphNode* ToNode = LinkedToPin->GetOwningNode();
+				if (!ToNode)
+				{
+					continue;
+				}
+
+				if (LinkedToPin->PinType.PinCategory != Schema->PC_Modifier)
+				{
+					// The modifier is connected to a node for which automatic upgrade support is not implemented
+					FString Msg = FString::Printf(TEXT("A modifier node has a legacy connection to a node [%s] without automatic upgrade support. Manual update is probably needed."), *ToNode->GetName());
+					FCustomizableObjectEditorLogger::CreateLog(FText::FromString(Msg))
+						.Severity(EMessageSeverity::Warning)
+						.Context(*this)
+						.BaseObject(true)
+						.Log();
+
+				}
+			}
+		}
+	}
+
 }
+
diff --git a/Engine/Plugins/Experimental/Mutable/Source/CustomizableObjectEditor/Private/MuCOE/Nodes/CustomizableObjectNodeModifierClipWithUVMask.cpp b/Engine/Plugins/Experimental/Mutable/Source/CustomizableObjectEditor/Private/MuCOE/Nodes/CustomizableObjectNodeModifierClipWithUVMask.cpp
--- a/Engine/Plugins/Experimental/Mutable/Source/CustomizableObjectEditor/Private/MuCOE/Nodes/CustomizableObjectNodeModifierClipWithUVMask.cpp
+++ b/Engine/Plugins/Experimental/Mutable/Source/CustomizableObjectEditor/Private/MuCOE/Nodes/CustomizableObjectNodeModifierClipWithUVMask.cpp
@@ -19,7 +19,7 @@ void UCustomizableObjectNodeModifierClipWithUVMask::AllocateDefaultPins(UCustomi
 	UEdGraphPin* ClipMaskPin = CustomCreatePin(EGPD_Input, Schema->PC_Image, FName("Clip Mask"));
 	ClipMaskPin->bDefaultValueIsIgnored = true;
 
-	UEdGraphPin* OutputPin_p = CustomCreatePin(EGPD_Output, Schema->PC_Material, FName("Modifier"));
+	UEdGraphPin* OutputPin_p = CustomCreatePin(EGPD_Output, Schema->PC_Modifier, FName("Modifier"));
 }
 
 
@@ -29,13 +29,6 @@ FText UCustomizableObjectNodeModifierClipWithUVMask::GetNodeTitle(ENodeTitleType
 }
 
 
-FLinearColor UCustomizableObjectNodeModifierClipWithUVMask::GetNodeTitleColor() const
-{
-	const UEdGraphSchema_CustomizableObject* Schema = GetDefault<UEdGraphSchema_CustomizableObject>();
-	return Schema->GetPinTypeColor(Schema->PC_Material);
-}
-
-
 void UCustomizableObjectNodeModifierClipWithUVMask::PinConnectionListChanged(UEdGraphPin* Pin)
 {
 	Super::PinConnectionListChanged(Pin);
diff --git a/Engine/Plugins/Experimental/Mutable/Source/CustomizableObjectEditor/Private/MuCOE/Nodes/CustomizableObjectNodeObject.cpp b/Engine/Plugins/Experimental/Mutable/Source/CustomizableObjectEditor/Private/MuCOE/Nodes/CustomizableObjectNodeObject.cpp
--- a/Engine/Plugins/Experimental/Mutable/Source/CustomizableObjectEditor/Private/MuCOE/Nodes/CustomizableObjectNodeObject.cpp
+++ b/Engine/Plugins/Experimental/Mutable/Source/CustomizableObjectEditor/Private/MuCOE/Nodes/CustomizableObjectNodeObject.cpp
@@ -24,6 +24,7 @@ class UCustomizableObjectNodeRemapPins;
 
 const FName UCustomizableObjectNodeObject::ChildrenPinName(TEXT("Children"));
 const FName UCustomizableObjectNodeObject::ComponentsPinName(TEXT("Components"));
+const FName UCustomizableObjectNodeObject::ModifiersPinName(TEXT("Modifiers"));
 const FName UCustomizableObjectNodeObject::OutputPinName(TEXT("Object"));
 const TCHAR* UCustomizableObjectNodeObject::LODPinNamePrefix = TEXT("LOD ");
 
@@ -154,6 +155,9 @@ void UCustomizableObjectNodeObject::AllocateDefaultPins(UCustomizableObjectNodeR
 	UEdGraphPin* ComponentsPin = CustomCreatePin(EGPD_Input, Schema->PC_Component, ComponentsPinName, true);
 	ComponentsPin->bDefaultValueIsIgnored = true;
 
+	UEdGraphPin* ModifiersPin = CustomCreatePin(EGPD_Input, Schema->PC_Modifier, ModifiersPinName, true);
+	ModifiersPin->bDefaultValueIsIgnored = true;
+
 	UEdGraphPin* ChildrenPin = CustomCreatePin(EGPD_Input, Schema->PC_Object, ChildrenPinName, true);
 	ChildrenPin->bDefaultValueIsIgnored = true;
 
diff --git a/Engine/Plugins/Experimental/Mutable/Source/CustomizableObjectEditor/Private/MuCOE/Nodes/CustomizableObjectNodeRemoveMeshBlocks.cpp b/Engine/Plugins/Experimental/Mutable/Source/CustomizableObjectEditor/Private/MuCOE/Nodes/CustomizableObjectNodeRemoveMeshBlocks.cpp
--- a/Engine/Plugins/Experimental/Mutable/Source/CustomizableObjectEditor/Private/MuCOE/Nodes/CustomizableObjectNodeRemoveMeshBlocks.cpp
+++ b/Engine/Plugins/Experimental/Mutable/Source/CustomizableObjectEditor/Private/MuCOE/Nodes/CustomizableObjectNodeRemoveMeshBlocks.cpp
@@ -95,7 +95,7 @@ void UCustomizableObjectNodeRemoveMeshBlocks::BackwardsCompatibleFixup()
 						if (ExistingBlock.Id == BlockId)
 						{
 							bSkipBlock = true;
-							UE_LOG(LogMutable, Warning, TEXT("[%s] UCustomizableObjectNodeRemoveMeshBlocks has a duplicated layout block id. One has been ignored during version upgrade."), *GetOutermost()->GetName());
+							UE_LOG(LogMutable, Log, TEXT("[%s] UCustomizableObjectNodeRemoveMeshBlocks has a duplicated layout block id. One has been ignored during version upgrade."), *GetOutermost()->GetName());
 							break;
 						}
 					}
@@ -161,6 +161,7 @@ void UCustomizableObjectNodeRemoveMeshBlocks::PostEditChangeProperty(FPropertyCh
 	Super::PostEditChangeProperty(PropertyChangedEvent);
 }
 
+
 void UCustomizableObjectNodeRemoveMeshBlocks::AllocateDefaultPins(UCustomizableObjectNodeRemapPins* RemapPins)
 {
 	const UEdGraphSchema_CustomizableObject* Schema = GetDefault<UEdGraphSchema_CustomizableObject>();
diff --git a/Engine/Plugins/Experimental/Mutable/Source/CustomizableObjectEditor/Public/MuCOE/EdGraphSchema_CustomizableObject.h b/Engine/Plugins/Experimental/Mutable/Source/CustomizableObjectEditor/Public/MuCOE/EdGraphSchema_CustomizableObject.h
--- a/Engine/Plugins/Experimental/Mutable/Source/CustomizableObjectEditor/Public/MuCOE/EdGraphSchema_CustomizableObject.h
+++ b/Engine/Plugins/Experimental/Mutable/Source/CustomizableObjectEditor/Public/MuCOE/EdGraphSchema_CustomizableObject.h
@@ -91,6 +91,7 @@ public:
 	static const FName PC_Object;
 	static const FName PC_Component;
 	static const FName PC_Material;
+	static const FName PC_Modifier;
 	static const FName PC_Mesh;
 	static const FName PC_PassThroughMesh;
 	static const FName PC_Layout;
diff --git a/Engine/Plugins/Experimental/Mutable/Source/CustomizableObjectEditor/Public/MuCOE/Nodes/CustomizableObjectNodeMeshClipDeform.h b/Engine/Plugins/Experimental/Mutable/Source/CustomizableObjectEditor/Public/MuCOE/Nodes/CustomizableObjectNodeMeshClipDeform.h
--- a/Engine/Plugins/Experimental/Mutable/Source/CustomizableObjectEditor/Public/MuCOE/Nodes/CustomizableObjectNodeMeshClipDeform.h
+++ b/Engine/Plugins/Experimental/Mutable/Source/CustomizableObjectEditor/Public/MuCOE/Nodes/CustomizableObjectNodeMeshClipDeform.h
@@ -37,7 +37,6 @@ public:
 	
 	// EdGraphNode interface
 	virtual FText GetNodeTitle(ENodeTitleType::Type TitleType) const override;
-	virtual FLinearColor GetNodeTitleColor() const override;
 	virtual FText GetTooltipText() const override;
 
 	// UCustomizableObjectNode interface
diff --git a/Engine/Plugins/Experimental/Mutable/Source/CustomizableObjectEditor/Public/MuCOE/Nodes/CustomizableObjectNodeMeshClipMorph.h b/Engine/Plugins/Experimental/Mutable/Source/CustomizableObjectEditor/Public/MuCOE/Nodes/CustomizableObjectNodeMeshClipMorph.h
--- a/Engine/Plugins/Experimental/Mutable/Source/CustomizableObjectEditor/Public/MuCOE/Nodes/CustomizableObjectNodeMeshClipMorph.h
+++ b/Engine/Plugins/Experimental/Mutable/Source/CustomizableObjectEditor/Public/MuCOE/Nodes/CustomizableObjectNodeMeshClipMorph.h
@@ -82,13 +82,12 @@ public:
 
 	// EdGraphNode interface
 	virtual FText GetNodeTitle(ENodeTitleType::Type TitleType) const override;
-	virtual FLinearColor GetNodeTitleColor() const override;
 	virtual FText GetTooltipText() const override;
 
 	// UCustomizableObjectNode interface
 	virtual void AllocateDefaultPins(UCustomizableObjectNodeRemapPins* RemapPins) override;
 	virtual void PinConnectionListChanged(UEdGraphPin* Pin) override;
-	
+
 	// Own interface
 	UEdGraphPin* OutputPin() const;
 
diff --git a/Engine/Plugins/Experimental/Mutable/Source/CustomizableObjectEditor/Public/MuCOE/Nodes/CustomizableObjectNodeMeshClipWithMesh.h b/Engine/Plugins/Experimental/Mutable/Source/CustomizableObjectEditor/Public/MuCOE/Nodes/CustomizableObjectNodeMeshClipWithMesh.h
--- a/Engine/Plugins/Experimental/Mutable/Source/CustomizableObjectEditor/Public/MuCOE/Nodes/CustomizableObjectNodeMeshClipWithMesh.h
+++ b/Engine/Plugins/Experimental/Mutable/Source/CustomizableObjectEditor/Public/MuCOE/Nodes/CustomizableObjectNodeMeshClipWithMesh.h
@@ -51,7 +51,6 @@ public:
 
 	// EdGraphNode interface
 	virtual FText GetNodeTitle(ENodeTitleType::Type TitleType) const override;
-	virtual FLinearColor GetNodeTitleColor() const override;
 	virtual FText GetTooltipText() const override;
 
 	// UCustomizableObjectNode interface
diff --git a/Engine/Plugins/Experimental/Mutable/Source/CustomizableObjectEditor/Public/MuCOE/Nodes/CustomizableObjectNodeModifierBase.h b/Engine/Plugins/Experimental/Mutable/Source/CustomizableObjectEditor/Public/MuCOE/Nodes/CustomizableObjectNodeModifierBase.h
--- a/Engine/Plugins/Experimental/Mutable/Source/CustomizableObjectEditor/Public/MuCOE/Nodes/CustomizableObjectNodeModifierBase.h
+++ b/Engine/Plugins/Experimental/Mutable/Source/CustomizableObjectEditor/Public/MuCOE/Nodes/CustomizableObjectNodeModifierBase.h
@@ -15,6 +15,11 @@ class CUSTOMIZABLEOBJECTEDITOR_API UCustomizableObjectNodeModifierBase : public
 public:
 	GENERATED_BODY()
 
-	UCustomizableObjectNodeModifierBase();
+	// EdGraphNode interface
+	virtual FLinearColor GetNodeTitleColor() const override;
+
+	// UCustomizableObjectNode interface
+	virtual void PostBackwardsCompatibleFixup() override;
+
 };
 
diff --git a/Engine/Plugins/Experimental/Mutable/Source/CustomizableObjectEditor/Public/MuCOE/Nodes/CustomizableObjectNodeModifierClipWithUVMask.h b/Engine/Plugins/Experimental/Mutable/Source/CustomizableObjectEditor/Public/MuCOE/Nodes/CustomizableObjectNodeModifierClipWithUVMask.h
--- a/Engine/Plugins/Experimental/Mutable/Source/CustomizableObjectEditor/Public/MuCOE/Nodes/CustomizableObjectNodeModifierClipWithUVMask.h
+++ b/Engine/Plugins/Experimental/Mutable/Source/CustomizableObjectEditor/Public/MuCOE/Nodes/CustomizableObjectNodeModifierClipWithUVMask.h
@@ -33,7 +33,6 @@ public:
 
 	// EdGraphNode interface
 	virtual FText GetNodeTitle(ENodeTitleType::Type TitleType) const override;
-	virtual FLinearColor GetNodeTitleColor() const override;
 	virtual FText GetTooltipText() const override;
 
 	// UCustomizableObjectNode interface
diff --git a/Engine/Plugins/Experimental/Mutable/Source/CustomizableObjectEditor/Public/MuCOE/Nodes/CustomizableObjectNodeObject.h b/Engine/Plugins/Experimental/Mutable/Source/CustomizableObjectEditor/Public/MuCOE/Nodes/CustomizableObjectNodeObject.h
--- a/Engine/Plugins/Experimental/Mutable/Source/CustomizableObjectEditor/Public/MuCOE/Nodes/CustomizableObjectNodeObject.h
+++ b/Engine/Plugins/Experimental/Mutable/Source/CustomizableObjectEditor/Public/MuCOE/Nodes/CustomizableObjectNodeObject.h
@@ -258,6 +258,11 @@ public:
 		return FindPin(ComponentsPinName);
 	}
 
+	UEdGraphPin* ModifiersPin() const
+	{
+		return FindPin(ModifiersPinName);
+	}
+
 	UEdGraphPin* ChildrenPin() const
 	{
 		return FindPin(ChildrenPinName);
@@ -292,6 +297,7 @@ public:
 private:
 	static const FName ChildrenPinName;
 	static const FName ComponentsPinName;
+	static const FName ModifiersPinName;
 	static const FName OutputPinName;
 	static const TCHAR* LODPinNamePrefix;
 
