[mutable] Fix UI crash with missing layouts.
#rnx

diff --git a/Engine/Plugins/Experimental/Mutable/Source/CustomizableObjectEditor/Private/MuCOE/SCustomizableObjectNodeLayoutBlocksEditor.cpp b/Engine/Plugins/Experimental/Mutable/Source/CustomizableObjectEditor/Private/MuCOE/SCustomizableObjectNodeLayoutBlocksEditor.cpp
--- a/Engine/Plugins/Experimental/Mutable/Source/CustomizableObjectEditor/Private/MuCOE/SCustomizableObjectNodeLayoutBlocksEditor.cpp
+++ b/Engine/Plugins/Experimental/Mutable/Source/CustomizableObjectEditor/Private/MuCOE/SCustomizableObjectNodeLayoutBlocksEditor.cpp
@@ -489,7 +489,7 @@ TSharedPtr<IToolTip> SCustomizableObjectNodeLayoutBlocksEditor::GenerateInfoTool
 	};
 
 	// Duplicate command
-	if (CurrentLayout->PackingStrategy != ECustomizableObjectTextureLayoutPackingStrategy::Overlay)
+	if (CurrentLayout && CurrentLayout->PackingStrategy != ECustomizableObjectTextureLayoutPackingStrategy::Overlay)
 	{
 		BuildShortcutAndTooltip(LOCTEXT("ShortCut_DuplicateBlocks", "CTRL + D"), LOCTEXT("Tooltip_DuplicateBlocks", "Duplicate selected block/s"));
 		BuildShortcutAndTooltip(LOCTEXT("ShortCut_CreateNewBlock", "CTRL + N"), LOCTEXT("Tooltip_CreateNewBlock", "Create new block"));
