[mutable] Fix another UI crash due to null layout.
#rnx

diff --git a/Engine/Plugins/Experimental/Mutable/Source/CustomizableObjectEditor/Private/MuCOE/SCustomizableObjectNodeLayoutBlocksEditor.cpp b/Engine/Plugins/Experimental/Mutable/Source/CustomizableObjectEditor/Private/MuCOE/SCustomizableObjectNodeLayoutBlocksEditor.cpp
--- a/Engine/Plugins/Experimental/Mutable/Source/CustomizableObjectEditor/Private/MuCOE/SCustomizableObjectNodeLayoutBlocksEditor.cpp
+++ b/Engine/Plugins/Experimental/Mutable/Source/CustomizableObjectEditor/Private/MuCOE/SCustomizableObjectNodeLayoutBlocksEditor.cpp
@@ -60,7 +60,7 @@ SCustomizableObjectNodeLayoutBlocksEditor::SCustomizableObjectNodeLayoutBlocksEd
 
 void SCustomizableObjectNodeLayoutBlocksEditor::Construct(const FArguments& InArgs)
 {
-	CurrentLayout = 0;
+	CurrentLayout = nullptr;
 	
 	BindCommands();
 }
@@ -428,10 +428,9 @@ TArray<FCustomizableObjectLayoutBlock> SCustomizableObjectNodeLayoutBlocksEditor
 	if (CurrentLayout)
 	{
 		Blocks = CurrentLayout->Blocks;
+		Blocks.Append(CurrentLayout->AutomaticBlocks);
 	}
 
-	Blocks.Append(CurrentLayout->AutomaticBlocks );
-
 	return Blocks;
 }
 
