[Mutable] Prevent crash when generating realtime morphs if vertex name mapping data is not correct.

#rnx

diff --git a/Engine/Plugins/Experimental/Mutable/Source/CustomizableObject/Private/MuCO/CustomizableObjectSystem.cpp b/Engine/Plugins/Experimental/Mutable/Source/CustomizableObject/Private/MuCO/CustomizableObjectSystem.cpp
--- a/Engine/Plugins/Experimental/Mutable/Source/CustomizableObject/Private/MuCO/CustomizableObjectSystem.cpp
+++ b/Engine/Plugins/Experimental/Mutable/Source/CustomizableObject/Private/MuCO/CustomizableObjectSystem.cpp
@@ -1904,6 +1904,8 @@ namespace impl
 			TArray<int32> SectionMorphTargetVerticesCount;
 			SectionMorphTargetVerticesCount.SetNumZeroed(ComponentMorphTargetsData.RealTimeMorphTargetNames.Num());
 
+			int32 NumInvalidVertexMorphNamesFound = 0;
+			
 			const int32 FirstGeneratedLOD = FMath::Max((int32)OperationData->GetRequestedLODs()[ComponentIndex], OperationData->GetMinLOD());
 			for (int32 LODIndex = FirstGeneratedLOD; LODIndex < OperationData->NumLODsAvailablePerComponent[ComponentIndex]; ++LODIndex)
 			{
@@ -1930,7 +1932,7 @@ namespace impl
 
 				const uint32* const VertexMorphsResourceIdBuffer = reinterpret_cast<const uint32*>(MeshSet.GetBufferData(VertexMorphsResourceIdBufferIndex));
 				TArrayView<const uint32> VertexMorphsResourceIdView(VertexMorphsResourceIdBuffer, MeshSet.GetElementCount());
-
+					
 				const int32 SurfaceCount = LOD.Mesh->GetSurfaceCount();
 				for (int32 Section = 0; Section < SurfaceCount; ++Section)
 				{
@@ -1983,6 +1985,7 @@ namespace impl
 						}
 
 						TArrayView<const FMorphTargetVertexData> SpanMorphData = MorphTargetReconstructionData->DataView;
+						const int32 NumNamesInResolutionMap = MorphTargetReconstructionData->NameResolutionMap.Num();
 
 						for (int32 SpanVertexIdx = SpanStart; SpanVertexIdx < SpanEnd; ++SpanVertexIdx)
 						{
@@ -2001,6 +2004,12 @@ namespace impl
 
 							for (const FMorphTargetVertexData& SourceVertex : MorphsVertexDataView)
 							{
+								if (SourceVertex.MorphNameIndex >= (uint32)NumNamesInResolutionMap)
+								{
+									++NumInvalidVertexMorphNamesFound;
+									continue;
+								}
+
 								const uint32 ResolvedNameIndex =
 										MorphTargetReconstructionData->NameResolutionMap[SourceVertex.MorphNameIndex];
 
@@ -2032,7 +2041,12 @@ namespace impl
 					}
 				}
 			}
-		
+	
+			if (NumInvalidVertexMorphNamesFound > 0)
+			{
+				UE_LOG(LogMutable, Warning, TEXT("Invalid real-time morphs names found in instance vertices. Some morph may not work as expected."));
+			}
+	
 			// Remove empty morph targets;
 			for (int32 MorphIndex = 0; MorphIndex < NumMorphs; ++MorphIndex)
 			{
@@ -2064,7 +2078,6 @@ namespace impl
 		}
 	}
 
-
 	/** End of the GetMeshes tasks. */
 	void Task_Mutable_GetMeshes_End(
 		const TSharedRef<FUpdateContextPrivate>& OperationData,
