[Mutable] Fix COI Editor red grooms.

#rnx

diff --git a/Engine/Plugins/Experimental/Mutable/Source/CustomizableObjectEditor/Private/MuCOE/CustomizableObjectInstanceEditor.cpp b/Engine/Plugins/Experimental/Mutable/Source/CustomizableObjectEditor/Private/MuCOE/CustomizableObjectInstanceEditor.cpp
--- a/Engine/Plugins/Experimental/Mutable/Source/CustomizableObjectEditor/Private/MuCOE/CustomizableObjectInstanceEditor.cpp
+++ b/Engine/Plugins/Experimental/Mutable/Source/CustomizableObjectEditor/Private/MuCOE/CustomizableObjectInstanceEditor.cpp
@@ -598,9 +598,12 @@ void FCustomizableObjectInstanceEditor::OnInstancePropertySelectionChanged(FProp
 void FCustomizableObjectInstanceEditor::OnUpdatePreviewInstance(UCustomizableObjectInstance* Instance)
 {
 	check(CustomizableObjectInstance)
+	
+	if (Instance->GetPrivate()->SkeletalMeshStatus != ESkeletalMeshStatus::Success)
+	{
+		Actor->GetRootComponent()->SetVisibility(false, true);
+	}
 
-	const bool bVisible = CustomizableObjectInstance->GetPrivate()->SkeletalMeshStatus == ESkeletalMeshStatus::Success;
-	Actor->GetRootComponent()->SetVisibility(bVisible, true);
 	Actor->GetRootComponent()->UpdateBounds();
 	
 	// Sets which UV material to draw
