[Mutable] Fix additional data file handle resolution when multiple data blocks point to the same file id.

#rnx

diff --git a/Engine/Plugins/Experimental/Mutable/Source/CustomizableObject/Private/MuCO/CustomizableObjectInstance.cpp b/Engine/Plugins/Experimental/Mutable/Source/CustomizableObject/Private/MuCO/CustomizableObjectInstance.cpp
--- a/Engine/Plugins/Experimental/Mutable/Source/CustomizableObject/Private/MuCO/CustomizableObjectInstance.cpp
+++ b/Engine/Plugins/Experimental/Mutable/Source/CustomizableObject/Private/MuCO/CustomizableObjectInstance.cpp
@@ -5614,6 +5614,29 @@ UE::Tasks::FTask UCustomizableInstancePrivate::LoadAdditionalAssetsAndData(
 		}
 
 		TArray<uint32> OpenFilesIds;
+
+		auto OpenOrGetFileHandleForBlock = [&OpenFilesIds, &OpenFileHandles, BulkData](const FMutableStreamableBlock& Block) 
+			-> IAsyncReadFileHandle*
+		{
+			int32 FileHandleIndex = OpenFilesIds.Find(Block.FileId);
+			if (FileHandleIndex == INDEX_NONE && BulkData)
+			{
+				TUniquePtr<IAsyncReadFileHandle> ReadFileHandle = BulkData->OpenFileAsyncRead(Block.FileId, Block.Flags);
+
+				OpenFileHandles.Emplace(MoveTemp(ReadFileHandle));
+				FileHandleIndex = OpenFilesIds.Add(Block.FileId);
+				
+				check(OpenFileHandles.Num() == OpenFilesIds.Num());
+			}
+
+			if (OpenFileHandles.IsValidIndex(FileHandleIndex))
+			{
+				return OpenFileHandles[FileHandleIndex].Get();
+			}
+
+			return nullptr;
+		};
+
 		const int32 NumMorphBlocks = RealTimeMorphStreamableBlocksToStream.Num();
 		for (int32 I = 0; I < NumMorphBlocks; ++I)
 		{
@@ -5640,19 +5663,9 @@ UE::Tasks::FTask UCustomizableInstancePrivate::LoadAdditionalAssetsAndData(
 			ReadDestData.Data.SetNumUninitialized(NumElems);
 
 			IAsyncReadFileHandle* FileHandle = nullptr;
-
 			if (!bUseFBulkData)
-			{
-				int32 FileHandleIndex = OpenFilesIds.Find(Block.FileId);
-				if (FileHandleIndex == INDEX_NONE && BulkData)
-				{
-					TUniquePtr<IAsyncReadFileHandle> ReadFileHandle = BulkData->OpenFileAsyncRead(Block.FileId, Block.Flags);
-
-					FileHandle = ReadFileHandle.Get();
-
-					OpenFileHandles.Emplace(MoveTemp(ReadFileHandle));
-					FileHandleIndex = OpenFilesIds.Add(Block.FileId);
-				}
+			{	
+				FileHandle = OpenOrGetFileHandleForBlock(Block);
 			}
 
 			BlockReadInfos.Emplace(FBlockReadInfo
@@ -5692,19 +5705,9 @@ UE::Tasks::FTask UCustomizableInstancePrivate::LoadAdditionalAssetsAndData(
 			ReadDestData.Data.SetNumUninitialized(NumElems);
 
 			IAsyncReadFileHandle* FileHandle = nullptr;
-
 			if (!bUseFBulkData)
 			{
-				int32 FileHandleIndex = OpenFilesIds.Find(Block.FileId);
-				if (FileHandleIndex == INDEX_NONE && BulkData)
-				{
-					TUniquePtr<IAsyncReadFileHandle> ReadFileHandle = BulkData->OpenFileAsyncRead(Block.FileId, Block.Flags);
-
-					FileHandle = ReadFileHandle.Get();
-
-					OpenFileHandles.Emplace(MoveTemp(ReadFileHandle));
-					FileHandleIndex = OpenFilesIds.Add(Block.FileId);
-				}
+				FileHandle = OpenOrGetFileHandleForBlock(Block);
 			}
 
 			BlockReadInfos.Emplace(FBlockReadInfo
