[mutable] Avoid unnecessary instance generation time for parameter relevancy in PIE
#rnx

diff --git a/Engine/Plugins/Experimental/Mutable/Source/CustomizableObject/Private/MuCO/CustomizableObjectInstanceDescriptor.cpp b/Engine/Plugins/Experimental/Mutable/Source/CustomizableObject/Private/MuCO/CustomizableObjectInstanceDescriptor.cpp
--- a/Engine/Plugins/Experimental/Mutable/Source/CustomizableObject/Private/MuCO/CustomizableObjectInstanceDescriptor.cpp
+++ b/Engine/Plugins/Experimental/Mutable/Source/CustomizableObject/Private/MuCO/CustomizableObjectInstanceDescriptor.cpp
@@ -4,7 +4,7 @@
 
 #include "MuCO/CustomizableObjectInstancePrivate.h"
 #include "MuCO/CustomizableObjectSystemPrivate.h"
-#include "UnrealMutableImageProvider.h"
+#include "MuCO/UnrealMutableImageProvider.h"
 #include "MuCO/CustomizableObject.h"
 #include "MuCO/CustomizableObjectPrivate.h"
 #include "MuCO/DefaultImageProvider.h"
@@ -14,6 +14,10 @@
 #include "MuR/Parameters.h"
 #include "MuR/Ptr.h"
 
+#if WITH_EDITOR
+#include "Editor.h"
+#endif
+
 #include UE_INLINE_GENERATED_CPP_BY_NAME(CustomizableObjectInstanceDescriptor)
 
 
@@ -447,7 +451,18 @@ void FCustomizableObjectInstanceDescriptor::SetCustomizableObject(UCustomizableO
 bool FCustomizableObjectInstanceDescriptor::GetBuildParameterRelevancy() const
 {
 #if WITH_EDITOR
-	return true;
+	// In editor, calculate the parameter relevancy by default.
+	bool bResultBuildRelevancy = true;
+	// However if we are in a PIE session, do it only if requested to simulate a more game-like performance.
+	if (GIsEditor)
+	{
+		FWorldContext* PIEWorldContext = GEditor->GetPIEWorldContext();
+		if (PIEWorldContext)
+		{
+			bResultBuildRelevancy = bBuildParameterRelevancy;
+		}
+	}
+	return bResultBuildRelevancy;
 #else
 	return bBuildParameterRelevancy;
 #endif
