[Mutable] Added animation timeline to CO and COI Editors.

- Added Viewport timeline.
- Added Preview Settings animation selector.
- Added Viewport speed selector.

#rnx

diff --git a/Engine/Plugins/Experimental/Mutable/Source/CustomizableObjectEditor/CustomizableObjectEditor.Build.cs b/Engine/Plugins/Experimental/Mutable/Source/CustomizableObjectEditor/CustomizableObjectEditor.Build.cs
--- a/Engine/Plugins/Experimental/Mutable/Source/CustomizableObjectEditor/CustomizableObjectEditor.Build.cs
+++ b/Engine/Plugins/Experimental/Mutable/Source/CustomizableObjectEditor/CustomizableObjectEditor.Build.cs
@@ -20,7 +20,6 @@ public class CustomizableObjectEditor : ModuleRules
 				"SkeletonEditor",
 				"WorkspaceMenuStructure",
 				"MessageLog",
-				"KismetWidgets",
 			}
 		);
 
@@ -70,6 +69,7 @@ public class CustomizableObjectEditor : ModuleRules
 				"TextureCompressor",
 				"TextureBuildUtilities",
 				"ImageCore",
+				"KismetWidgets"
 			}
 		);
 
diff --git a/Engine/Plugins/Experimental/Mutable/Source/CustomizableObjectEditor/Private/MuCOE/CustomizableObjectCustomSettingsDetails.cpp b/Engine/Plugins/Experimental/Mutable/Source/CustomizableObjectEditor/Private/MuCOE/CustomizableObjectCustomSettingsDetails.cpp
--- a/Engine/Plugins/Experimental/Mutable/Source/CustomizableObjectEditor/Private/MuCOE/CustomizableObjectCustomSettingsDetails.cpp
+++ b/Engine/Plugins/Experimental/Mutable/Source/CustomizableObjectEditor/Private/MuCOE/CustomizableObjectCustomSettingsDetails.cpp
@@ -2,27 +2,95 @@
 
 #include "MuCOE/CustomizableObjectCustomSettingsDetails.h"
 
+#include "Animation/Skeleton.h"
 #include "DetailLayoutBuilder.h"
 #include "DetailWidgetRow.h"
 #include "IDetailsView.h"
+#include "PropertyCustomizationHelpers.h"
+#include "SCustomizableObjectEditorViewport.h"
+#include "SMutableScrubPanel.h"
+#include "Animation/DebugSkelMeshComponent.h"
 #include "MuCO/CustomizableObject.h"
 #include "MuCOE/CustomizableObjectCustomSettings.h"
+#include "MuCOE/CustomizableObjectEditorViewportClient.h"
+#include "Toolkits/ToolkitManager.h"
+
 
 TSharedRef<IDetailCustomization> FCustomizableObjectCustomSettingsDetails::MakeInstance()
 {
 	return MakeShareable(new FCustomizableObjectCustomSettingsDetails);
 }
 
+
 void FCustomizableObjectCustomSettingsDetails::CustomizeDetails(IDetailLayoutBuilder& DetailBuilder)
 {
 	const IDetailsView* DetailsView = DetailBuilder.GetDetailsView();
 	if (DetailsView->GetSelectedObjects().Num() > 0 && DetailsView->GetSelectedObjects()[0].IsValid())
 	{
-		IDetailCategoryBuilder& MainCategory = DetailBuilder.EditCategory("Custom Settings");
-		MainCategory.AddCustomRow(FText::FromString("Custom Settings"))
+		UCustomSettings* CustomSettings = Cast<UCustomSettings>(DetailsView->GetSelectedObjects()[0]);
+
+		TWeakPtr<ICustomizableObjectInstanceEditor> WeakEditor = CustomSettings->GetEditor();
+		TSharedPtr<ICustomizableObjectInstanceEditor> Editor = WeakEditor.Pin();
+
+		if (Editor->ShowLightingSettings())
+		{
+			IDetailCategoryBuilder& MainCategory = DetailBuilder.EditCategory("Custom Settings");
+			MainCategory.AddCustomRow(FText::FromString("Custom Settings"))
+			[
+				SNew(SCustomizableObjectCustomSettings)
+					.PreviewSettings(CustomSettings)
+			];
+		}
+		
+		TSharedPtr<FCustomizableObjectEditorViewportClient> ViewportClient = Editor->GetViewport()->GetViewportClient();
+		
+		TSharedPtr<IPropertyHandle> PropertyHandle = DetailBuilder.GetProperty(GET_MEMBER_NAME_STRING_CHECKED(UCustomSettings, Animation));
+		IDetailPropertyRow* DetailPropertyRow = DetailBuilder.EditDefaultProperty(PropertyHandle);
+		
+		DetailPropertyRow->CustomWidget()
+		.NameContent()
+		[
+			DetailPropertyRow->GetPropertyHandle()->CreatePropertyNameWidget()
+		]
+		.ValueContent()
+		.MaxDesiredWidth(250.0f)
+		.MinDesiredWidth(250.0f)
 		[
-			SNew(SCustomizableObjectCustomSettings)
-				.PreviewSettings(Cast<UCustomSettings>(DetailsView->GetSelectedObjects()[0]))
+			SNew(SObjectPropertyEntryBox)
+			.AllowedClass(UAnimationAsset::StaticClass())
+			.PropertyHandle(DetailPropertyRow->GetPropertyHandle())
+			.OnShouldFilterAsset(FOnShouldFilterAsset::CreateLambda([WeakEditor](const FAssetData& InAssetData)
+			{
+				TSharedPtr<ICustomizableObjectInstanceEditor> Editor = WeakEditor.Pin();
+				if (!Editor)
+				{
+					return true;
+				}
+
+				for (TWeakObjectPtr<UDebugSkelMeshComponent>& WeakPreviewMeshComponent : Editor->GetViewport()->GetViewportClient()->GetPreviewMeshComponents())
+				{
+					UDebugSkelMeshComponent* PreviewMeshComponent = WeakPreviewMeshComponent.Get();
+					if (!PreviewMeshComponent)
+					{
+						continue;
+					}
+
+					USkeletalMesh* SkeletalMeshAsset = PreviewMeshComponent->GetSkeletalMeshAsset();
+					if (!SkeletalMeshAsset)
+					{
+						continue;
+					}
+
+					if (USkeleton* Skeleton = SkeletalMeshAsset->GetSkeleton())
+					{
+						return !Skeleton->IsCompatibleForEditor(InAssetData);
+					}
+				}
+
+				return true;
+			}))
+			.ThumbnailPool(DetailBuilder.GetThumbnailPool())
 		];
 	}
 }
+
diff --git a/Engine/Plugins/Experimental/Mutable/Source/CustomizableObjectEditor/Private/MuCOE/CustomizableObjectEditor.cpp b/Engine/Plugins/Experimental/Mutable/Source/CustomizableObjectEditor/Private/MuCOE/CustomizableObjectEditor.cpp
--- a/Engine/Plugins/Experimental/Mutable/Source/CustomizableObjectEditor/Private/MuCOE/CustomizableObjectEditor.cpp
+++ b/Engine/Plugins/Experimental/Mutable/Source/CustomizableObjectEditor/Private/MuCOE/CustomizableObjectEditor.cpp
@@ -2,6 +2,8 @@
 
 #include "MuCOE/CustomizableObjectEditor.h"
 
+#include "AssetViewerSettings.h"
+#include "ContentBrowserModule.h"
 #include "CustomizableObjectEditorPerformanceAnalyzer.h"
 #include "Animation/DebugSkelMeshComponent.h"
 #include "AssetRegistry/ARFilter.h"
@@ -14,6 +16,7 @@
 #include "Framework/MultiBox/MultiBoxBuilder.h"
 #include "Framework/Notifications/NotificationManager.h"
 #include "GraphEditorActions.h"
+#include "IContentBrowserSingleton.h"
 #include "HAL/PlatformApplicationMisc.h"
 #include "IDetailsView.h"
 #include "MaterialGraph/MaterialGraphNode.h"
@@ -236,18 +239,15 @@ void FCustomizableObjectEditor::InitCustomizableObjectEditor(const EToolkitMode:
 	// \TODO: Create only when needed?
 	TagExplorer = SNew(SCustomizableObjectEditorTagExplorer).CustomizableObjectEditor(this);
 	
-	FAdvancedPreviewSceneModule& AdvancedPreviewSceneModule = FModuleManager::LoadModuleChecked<FAdvancedPreviewSceneModule>("AdvancedPreviewScene");
-
 	TSharedPtr<FAdvancedPreviewScene> AdvancedPreviewScene = StaticCastSharedPtr<FAdvancedPreviewScene>(Viewport->GetPreviewScene());
 
 	CustomizableObjectEditorAdvancedPreviewSettings =
 		SNew(SCustomizableObjectEditorAdvancedPreviewSettings, AdvancedPreviewScene.ToSharedRef())
 		.CustomSettings(CustomSettings)
-		.CustomizableObjectEditor(this);
+		.CustomizableObjectEditor(SharedThis(this).ToWeakPtr());
 	CustomizableObjectEditorAdvancedPreviewSettings->LoadProfileEnvironment();
 	AdvancedPreviewSettingsWidget = CustomizableObjectEditorAdvancedPreviewSettings;
 
-
 	const TSharedRef<FTabManager::FLayout> StandaloneDefaultLayout = FTabManager::NewLayout( "Standalone_CustomizableObjectEditor_Layout_v1.4" )
 	->AddArea
 	(
@@ -1999,12 +1999,24 @@ UEdGraphNode* FCustomizableObjectEditor::CreateCommentBox(const FVector2D& InTar
 }
 
 
-TSharedPtr<SCustomizableObjectEditorAdvancedPreviewSettings> FCustomizableObjectEditor::GetCustomizableObjectEditorAdvancedPreviewSettings()
+TSharedPtr<SCustomizableObjectEditorAdvancedPreviewSettings> FCustomizableObjectEditor::GetAdvancedPreviewSettings()
 {
 	return CustomizableObjectEditorAdvancedPreviewSettings;
 }
 
 
+bool FCustomizableObjectEditor::ShowLightingSettings()
+{
+	return true;
+}
+
+
+bool FCustomizableObjectEditor::ShowProfileManagementOptions()
+{
+	return true;
+}
+
+
 void FCustomizableObjectEditor::FindProperty(const FProperty* Property, const void* InContainer, const FString& FindString, const UObject& Context, bool& bFound)
 {
 	if (!Property || !InContainer)
diff --git a/Engine/Plugins/Experimental/Mutable/Source/CustomizableObjectEditor/Private/MuCOE/CustomizableObjectEditor.h b/Engine/Plugins/Experimental/Mutable/Source/CustomizableObjectEditor/Private/MuCOE/CustomizableObjectEditor.h
--- a/Engine/Plugins/Experimental/Mutable/Source/CustomizableObjectEditor/Private/MuCOE/CustomizableObjectEditor.h
+++ b/Engine/Plugins/Experimental/Mutable/Source/CustomizableObjectEditor/Private/MuCOE/CustomizableObjectEditor.h
@@ -213,6 +213,9 @@ public:
 	virtual void ShowGizmoLight(ULightComponent& SelectedLight) override;
 	virtual void HideGizmoLight() override;
 	virtual UCustomizableObjectEditorProperties* GetEditorProperties() override;
+	virtual TSharedPtr< SCustomizableObjectEditorAdvancedPreviewSettings> GetAdvancedPreviewSettings() override;
+	virtual bool ShowLightingSettings() override;
+	virtual bool ShowProfileManagementOptions() override;
 
 	/** Select only this node only. Do nothing if already was only selected. */
 	void SelectSingleNode(UCustomizableObjectNode& Node);
@@ -258,9 +261,6 @@ public:
 	/** Called when the Object Properties needs to be updated */
 	virtual void UpdateObjectProperties() override;
 
-	/** Getter of CustomizableObjectEditorAdvancedPreviewSettings */
-	TSharedPtr<class SCustomizableObjectEditorAdvancedPreviewSettings> GetCustomizableObjectEditorAdvancedPreviewSettings();
-
 	/** Debug the object as a raw mutable data in the internal tools. */
 	void DebugObject() const;
 
diff --git a/Engine/Plugins/Experimental/Mutable/Source/CustomizableObjectEditor/Private/MuCOE/CustomizableObjectEditorActions.cpp b/Engine/Plugins/Experimental/Mutable/Source/CustomizableObjectEditor/Private/MuCOE/CustomizableObjectEditorActions.cpp
--- a/Engine/Plugins/Experimental/Mutable/Source/CustomizableObjectEditor/Private/MuCOE/CustomizableObjectEditorActions.cpp
+++ b/Engine/Plugins/Experimental/Mutable/Source/CustomizableObjectEditor/Private/MuCOE/CustomizableObjectEditorActions.cpp
@@ -3,6 +3,7 @@
 #include "MuCOE/CustomizableObjectEditorActions.h"
 
 #include "MuCOE/CustomizableObjectEditorStyle.h"
+#include "MuCOE/SCustomizableObjectEditorViewport.h"
 
 #define LOCTEXT_NAMESPACE "CustomizableObjectEditor"
 
@@ -45,6 +46,7 @@ FCustomizableObjectEditorViewportCommands::FCustomizableObjectEditorViewportComm
 	FCustomizableObjectEditorStyle::GetStyleSetName()
 	)
 {
+	PlaybackSpeedCommands.AddZeroed(EMutableAnimationPlaybackSpeeds::NumPlaybackSpeeds);
 }
 
 
@@ -61,6 +63,16 @@ void FCustomizableObjectEditorViewportCommands::RegisterCommands()
 	UI_COMMAND(BakeInstance, "Bake Instance", "Create baked unreal resources for the current preview instance.", EUserInterfaceActionType::Button, FInputChord());
 	UI_COMMAND(StateChangeShowData, "Show or hide test results", "Show or hide test results", EUserInterfaceActionType::ToggleButton, FInputChord());
 	UI_COMMAND(StateChangeShowGeometryData, "Show instance geometry data", "Show instance geometry data", EUserInterfaceActionType::ToggleButton, FInputChord());
+	
+	UI_COMMAND( PlaybackSpeedCommands[EMutableAnimationPlaybackSpeeds::OneTenth],	"x0.1", "Set the animation playback speed to a tenth of normal", EUserInterfaceActionType::RadioButton, FInputChord() );
+	UI_COMMAND( PlaybackSpeedCommands[EMutableAnimationPlaybackSpeeds::Quarter],		"x0.25", "Set the animation playback speed to a quarter of normal", EUserInterfaceActionType::RadioButton, FInputChord() );
+	UI_COMMAND( PlaybackSpeedCommands[EMutableAnimationPlaybackSpeeds::Half],		"x0.5", "Set the animation playback speed to a half of normal", EUserInterfaceActionType::RadioButton, FInputChord() );
+	UI_COMMAND( PlaybackSpeedCommands[EMutableAnimationPlaybackSpeeds::ThreeQuarters],		"x0.75", "Set the animation playback speed to three quarters of normal", EUserInterfaceActionType::RadioButton, FInputChord() );
+	UI_COMMAND( PlaybackSpeedCommands[EMutableAnimationPlaybackSpeeds::Normal],		"x1.0", "Set the animation playback speed to normal", EUserInterfaceActionType::RadioButton, FInputChord() );
+	UI_COMMAND( PlaybackSpeedCommands[EMutableAnimationPlaybackSpeeds::Double],		"x2.0", "Set the animation playback speed to double the speed of normal", EUserInterfaceActionType::RadioButton, FInputChord() );
+	UI_COMMAND( PlaybackSpeedCommands[EMutableAnimationPlaybackSpeeds::FiveTimes],	"x5.0", "Set the animation playback speed to five times the normal speed", EUserInterfaceActionType::RadioButton, FInputChord() );
+	UI_COMMAND( PlaybackSpeedCommands[EMutableAnimationPlaybackSpeeds::TenTimes],	"x10.0", "Set the animation playback speed to ten times the normal speed", EUserInterfaceActionType::RadioButton, FInputChord() );
+	UI_COMMAND( PlaybackSpeedCommands[EMutableAnimationPlaybackSpeeds::Custom],	"xCustom", "Set the animation playback speed to assigned custom speed", EUserInterfaceActionType::RadioButton, FInputChord() );
 
 	UI_COMMAND(ShowDisplayInfo, "Mesh Info", "Display mesh info in the viewport", EUserInterfaceActionType::ToggleButton, FInputChord());
 	UI_COMMAND(EnableClothSimulation, "Enable Cloth Simulation", "Show simulated cloth mesh", EUserInterfaceActionType::ToggleButton, FInputChord());
diff --git a/Engine/Plugins/Experimental/Mutable/Source/CustomizableObjectEditor/Private/MuCOE/CustomizableObjectEditorViewportClient.cpp b/Engine/Plugins/Experimental/Mutable/Source/CustomizableObjectEditor/Private/MuCOE/CustomizableObjectEditorViewportClient.cpp
--- a/Engine/Plugins/Experimental/Mutable/Source/CustomizableObjectEditor/Private/MuCOE/CustomizableObjectEditorViewportClient.cpp
+++ b/Engine/Plugins/Experimental/Mutable/Source/CustomizableObjectEditor/Private/MuCOE/CustomizableObjectEditorViewportClient.cpp
@@ -2,6 +2,7 @@
 
 #include "MuCOE/CustomizableObjectEditorViewportClient.h"
 
+#include "AnimPreviewInstance.h"
 #include "MuCOE/CustomizableObjectInstanceBakingUtils.h"
 #include "Animation/DebugSkelMeshComponent.h"
 #include "Animation/PoseAsset.h"
@@ -18,6 +19,7 @@
 #include "Editor/EditorPerProjectUserSettings.h"
 #include "Editor/UnrealEdTypes.h"
 #include "EditorModeManager.h"
+#include "GameFramework/WorldSettings.h"
 #include "FileHelpers.h"
 #include "IContentBrowserSingleton.h"
 #include "InputKeyEventArgs.h"
@@ -59,6 +61,13 @@ class UTextureMipDataProviderFactory;
 #define LOCTEXT_NAMESPACE "CustomizableObjectEditor" 
 
 
+namespace EMutableAnimationPlaybackSpeeds
+{
+	// Speed scales for animation playback, must match EMutableAnimationPlaybackSpeeds::Type
+	float Values[EMutableAnimationPlaybackSpeeds::NumPlaybackSpeeds] = { 0.1f, 0.25f, 0.5f, 0.75f, 1.0f, 2.0f, 5.0f, 10.0f, 0.f };
+}
+
+
 FCustomizableObjectEditorViewportClient::FCustomizableObjectEditorViewportClient(TWeakPtr<ICustomizableObjectInstanceEditor> InCustomizableObjectEditor, FPreviewScene* InPreviewScene, const TSharedPtr<SEditorViewport>& EditorViewportWidget)
 	: FEditorViewportClient(&GLevelEditorModeTools(), InPreviewScene, EditorViewportWidget)
 	, CustomizableObjectEditorPtr(InCustomizableObjectEditor)
@@ -161,10 +170,7 @@ FCustomizableObjectEditorViewportClient::FCustomizableObjectEditorViewportClient
 	{
 		AddRealtimeOverride(false, LOCTEXT("RealtimeOverrideMessage_InstanceViewport", "Instance Viewport")); // We are PIE, don't start in realtime mode
 	}
-
-	IsPlayingAnimation = false;
-	AnimationBeingPlayed = nullptr;
-
+	
 	// Lighting 
 	SelectedLightComponent = nullptr;
 	
@@ -704,6 +710,8 @@ void FCustomizableObjectEditorViewportClient::SetPreviewActor(const TWeakObjectP
 			SkeletalMeshComponent->bDrawBinormals = bDrawBinormals;
 		}
 	}
+
+	bUpdated = false;
 	
 	InInstance->UpdatedNativeDelegate.AddSP(SharedThis(this), &FCustomizableObjectEditorViewportClient::OnInstanceUpdate);
 	
@@ -711,10 +719,54 @@ void FCustomizableObjectEditorViewportClient::SetPreviewActor(const TWeakObjectP
 }
 
 
+TArray<TWeakObjectPtr<UDebugSkelMeshComponent>>& FCustomizableObjectEditorViewportClient::GetPreviewMeshComponents()
+{
+	return SkeletalMeshComponents;
+}
+
+
+void FCustomizableObjectEditorViewportClient::SetPreviewAnimationAsset(UAnimationAsset* AnimAsset)
+{
+	for (TWeakObjectPtr<UDebugSkelMeshComponent>& WeakSkeletalMeshComponent : SkeletalMeshComponents)
+	{
+		UDebugSkelMeshComponent* SkeletalMeshComponent = WeakSkeletalMeshComponent.Get();
+		if (!SkeletalMeshComponent)
+		{
+			continue;
+		}
+			
+		if (AnimAsset)
+		{
+			// Early out if the new preview asset is the same as the current one, to avoid replaying from the beginning, etc...
+			if (SkeletalMeshComponent->PreviewInstance &&
+				AnimAsset == SkeletalMeshComponent->PreviewInstance->GetCurrentAsset() &&
+				SkeletalMeshComponent->IsPreviewOn())
+			{
+				return;
+			}
+
+			// Treat it as invalid if it's got a bogus skeleton pointer
+			if (AnimAsset->GetSkeleton() == nullptr)
+			{
+				return;
+			}
+		}
+
+		SkeletalMeshComponent->EnablePreview(true, AnimAsset);
+	}
+}
+
+
 void FCustomizableObjectEditorViewportClient::OnInstanceUpdate(UCustomizableObjectInstance* Instance)
 {
+	if (!bUpdated)
+	{
+		bUpdated = true;
+		
+		SetAnimation(CustomizableObjectEditorPtr.Pin()->GetCustomSettings()->Animation);
+	}
+	
 	Invalidate();
-	ReSetAnimation();
 	
 	// Configure the initial orbital position of the camera
 	if (!bIsCameraSetup)
@@ -1308,68 +1360,17 @@ bool FCustomizableObjectEditorViewportClient::CanSetWidgetMode(UE::Widget::EWidg
 }
 
 
-void FCustomizableObjectEditorViewportClient::SetAnimation(UAnimationAsset* Animation, EAnimationMode::Type AnimationType)
+void FCustomizableObjectEditorViewportClient::SetAnimation(UAnimationAsset* Animation)
 {
-	bool bFoundComponent = false;
-
-	for (TWeakObjectPtr<USkeletalMeshComponent> SkeletalMeshComponent : SkeletalMeshComponents)
+	for (TWeakObjectPtr<UDebugSkelMeshComponent>& WeakPreviewMeshComponent : SkeletalMeshComponents)
 	{
-		if (SkeletalMeshComponent.IsValid() && Animation != nullptr
-			&& UE_MUTABLE_GETSKINNEDASSET(SkeletalMeshComponent) != nullptr
-			&& UE_MUTABLE_GETSKINNEDASSET(SkeletalMeshComponent)->GetSkeleton() == Animation->GetSkeleton()
-			)
+		UDebugSkelMeshComponent* PreviewMeshComponent = WeakPreviewMeshComponent.Get();
+		if (!PreviewMeshComponent)
 		{
-			SetRealtime(true);
-			IsPlayingAnimation = true;
-			AnimationBeingPlayed = Animation;
-
-			if (UPoseAsset* PoseAsset = Cast<UPoseAsset>(Animation))
-			{
-				SkeletalMeshComponent->SetAnimationMode(EAnimationMode::AnimationBlueprint);
-				SkeletalMeshComponent->InitAnim(false);
-				SkeletalMeshComponent->SetAnimation(PoseAsset);
-
-				UAnimSingleNodeInstance* SingleNodeInstance = Cast<UAnimSingleNodeInstance>(SkeletalMeshComponent->GetAnimInstance());
-				if (SingleNodeInstance)
-				{
-					TArray<FName> ArrayPoseNames = PoseAsset->GetPoseFNames();
-					for (int32 i = 0; i < ArrayPoseNames.Num(); ++i)
-					{
-						SingleNodeInstance->SetPreviewCurveOverride(ArrayPoseNames[i], 1.0f, false);
-					}
-				}
-			}
-			else
-			{
-				SkeletalMeshComponent->SetAnimationMode(AnimationType);
-				SkeletalMeshComponent->PlayAnimation(Animation, true);
-				SkeletalMeshComponent->SetPlayRate(1.f);
-			}
-
-			bFoundComponent = true;
+			continue;
 		}
-	}
-	
-	if(!bFoundComponent)
-	{
-		IsPlayingAnimation = false;
-		AnimationBeingPlayed = nullptr;
-	}
-}
 
-
-void FCustomizableObjectEditorViewportClient::ReSetAnimation()
-{
-	if ((IsPlayingAnimation == true) && (AnimationBeingPlayed != nullptr))
-	{
-		if (Cast<UPoseAsset>(AnimationBeingPlayed))
-		{
-			SetAnimation(AnimationBeingPlayed, EAnimationMode::AnimationBlueprint);
-		}
-		else
-		{
-			SetAnimation(AnimationBeingPlayed, EAnimationMode::AnimationSingleNode);
-		}
+		PreviewMeshComponent->EnablePreview(true, Animation);
 	}
 }
 
@@ -1610,7 +1611,6 @@ void FCustomizableObjectEditorViewportClient::AddReferencedObjects(FReferenceCol
 	Collector.AddReferencedObject(ClipMorphMaterial);
 	Collector.AddReferencedObject(ClipMeshMaterial);
 	Collector.AddReferencedObject(TransparentPlaneMaterialXY);
-	Collector.AddReferencedObject(AnimationBeingPlayed);
 
 	if (BakeTempInstance)
 	{
@@ -1834,6 +1834,37 @@ const TArray<ULightComponent*>& FCustomizableObjectEditorViewportClient::GetLigh
 }
 
 
+void FCustomizableObjectEditorViewportClient::SetPlaybackSpeedMode(EMutableAnimationPlaybackSpeeds::Type InMode)
+{
+	AnimationPlaybackSpeedMode = InMode;
+
+	if (const UWorld* World = GetWorld())
+	{
+		const float AnimationSpeed = (InMode == EMutableAnimationPlaybackSpeeds::Custom) ? GetCustomAnimationSpeed() : EMutableAnimationPlaybackSpeeds::Values[AnimationPlaybackSpeedMode];
+		World->GetWorldSettings()->TimeDilation = AnimationSpeed;
+	}
+}
+
+
+void FCustomizableObjectEditorViewportClient::SetCustomAnimationSpeed(float Speed)
+{
+	CustomAnimationSpeed = Speed;
+	SetPlaybackSpeedMode(EMutableAnimationPlaybackSpeeds::Custom);
+}
+
+
+float FCustomizableObjectEditorViewportClient::GetCustomAnimationSpeed() const
+{
+	return CustomAnimationSpeed;
+}
+
+
+EMutableAnimationPlaybackSpeeds::Type FCustomizableObjectEditorViewportClient::GetPlaybackSpeedMode() const
+{
+	return AnimationPlaybackSpeedMode;
+}
+
+
 void FCustomizableObjectEditorViewportClient::OnShowDisplayInfo()
 {
 	bShowDisplayInfo = !bShowDisplayInfo;
diff --git a/Engine/Plugins/Experimental/Mutable/Source/CustomizableObjectEditor/Private/MuCOE/CustomizableObjectInstanceEditor.cpp b/Engine/Plugins/Experimental/Mutable/Source/CustomizableObjectEditor/Private/MuCOE/CustomizableObjectInstanceEditor.cpp
--- a/Engine/Plugins/Experimental/Mutable/Source/CustomizableObjectEditor/Private/MuCOE/CustomizableObjectInstanceEditor.cpp
+++ b/Engine/Plugins/Experimental/Mutable/Source/CustomizableObjectEditor/Private/MuCOE/CustomizableObjectInstanceEditor.cpp
@@ -27,6 +27,7 @@
 #include "MuCOE/SCustomizableObjectEditorTextureAnalyzer.h"
 #include "MuCOE/UnrealEditorPortabilityHelpers.h"
 #include "PropertyEditorModule.h"
+#include "SCustomizableObjectEditorAdvancedPreviewSettings.h"
 #include "SCustomizableObjectEditorViewport.h"
 #include "Selection.h"
 #include "Widgets/Docking/SDockTab.h"
@@ -158,6 +159,28 @@ void FCustomizableObjectInstanceEditor::UnregisterTabSpawners(const TSharedRef<c
 }
 
 
+void UCustomSettings::PostEditChangeProperty(FPropertyChangedEvent& PropertyChangedEvent)
+{
+	UObject::PostEditChangeProperty(PropertyChangedEvent);
+	
+	if (!PropertyChangedEvent.MemberProperty)
+	{
+		return;
+	}
+	
+	if (PropertyChangedEvent.MemberProperty->GetFName() == GET_MEMBER_NAME_CHECKED(UCustomSettings, Animation))
+	{
+		TSharedPtr<ICustomizableObjectInstanceEditor> Editor = WeakEditor.Pin();
+		if (!Editor)
+		{
+			return;
+		}
+		
+		Editor->GetViewport()->GetViewportClient()->SetAnimation(Animation);
+	}
+}
+
+
 ULightComponent* UCustomSettings::GetSelectedLight() const
 {
 	return SelectedLight;
@@ -184,13 +207,13 @@ void UCustomSettings::SetLightsPreset(UCustomizableObjectEditorViewportLights& I
 
 TWeakPtr<ICustomizableObjectInstanceEditor> UCustomSettings::GetEditor() const
 {
-	return Editor;
+	return WeakEditor;
 }
 
 
 void UCustomSettings::SetEditor(TSharedPtr<ICustomizableObjectInstanceEditor> InEditor)
 {
-	Editor = InEditor;
+	WeakEditor = InEditor;
 }
 
 
@@ -263,7 +286,12 @@ void FCustomizableObjectInstanceEditor::InitCustomizableObjectInstanceEditor( co
 	FAdvancedPreviewSceneModule& AdvancedPreviewSceneModule = FModuleManager::LoadModuleChecked<FAdvancedPreviewSceneModule>("AdvancedPreviewScene");
 
 	TSharedPtr<FAdvancedPreviewScene> AdvancedPreviewScene = StaticCastSharedPtr<FAdvancedPreviewScene>(Viewport->GetPreviewScene());
-	AdvancedPreviewSettingsWidget = AdvancedPreviewSceneModule.CreateAdvancedPreviewSceneSettingsWidget(AdvancedPreviewScene.ToSharedRef());
+
+	CustomizableObjectEditorAdvancedPreviewSettings = SNew(SCustomizableObjectEditorAdvancedPreviewSettings, AdvancedPreviewScene.ToSharedRef())
+		.CustomSettings(CustomSettings)
+		.CustomizableObjectEditor(SharedThis(this).ToWeakPtr());
+	
+	AdvancedPreviewSettingsWidget = CustomizableObjectEditorAdvancedPreviewSettings;
 
 	const TSharedRef<FTabManager::FLayout> StandaloneDefaultLayout = FTabManager::NewLayout( "Standalone_CustomizableObjectInstanceEditor_Layout_v2.1" )
 	->AddArea
@@ -523,34 +551,6 @@ TSharedPtr<SCustomizableObjectEditorViewportTabBody> FCustomizableObjectInstance
 }
 
 
-void FCustomizableObjectInstanceEditor::SetPoseAsset(class UPoseAsset* PoseAssetParameter)
-{
-	PoseAsset = PoseAssetParameter;
-
-	if (PoseAsset != nullptr)
-	{
-		Viewport->SetAnimation(nullptr, EAnimationMode::AnimationBlueprint);
-
-		for (TWeakObjectPtr<UDebugSkelMeshComponent> PreviewSkeletalMeshComponent : PreviewSkeletalMeshComponents)
-		{
-			PreviewSkeletalMeshComponent->SetAnimationMode(EAnimationMode::AnimationBlueprint);
-			PreviewSkeletalMeshComponent->InitAnim(false);
-			PreviewSkeletalMeshComponent->SetAnimation(PoseAsset);
-
-			UAnimSingleNodeInstance* SingleNodeInstance = Cast<UAnimSingleNodeInstance>(PreviewSkeletalMeshComponent->GetAnimInstance());
-			if (SingleNodeInstance)
-			{
-				const TArray<FName>& ArrayPoseSmartNames = PoseAsset->GetPoseFNames();
-				for (int32 i = 0; i < ArrayPoseSmartNames.Num(); ++i)
-				{
-					SingleNodeInstance->SetPreviewCurveOverride(ArrayPoseSmartNames[i], 1.0f, false);
-				}
-			}
-		}
-	}
-}
-
-
 UProjectorParameter* FCustomizableObjectInstanceEditor::GetProjectorParameter()
 {
 	return ProjectorParameter;
@@ -589,6 +589,24 @@ UCustomizableObjectEditorProperties* FCustomizableObjectInstanceEditor::GetEdito
 }
 
 
+TSharedPtr<SCustomizableObjectEditorAdvancedPreviewSettings> FCustomizableObjectInstanceEditor::GetAdvancedPreviewSettings()
+{
+	return CustomizableObjectEditorAdvancedPreviewSettings;
+}
+
+
+bool FCustomizableObjectInstanceEditor::ShowLightingSettings()
+{
+	return false;
+}
+
+
+bool FCustomizableObjectInstanceEditor::ShowProfileManagementOptions()
+{
+	return false;
+}
+
+
 void FCustomizableObjectInstanceEditor::OnInstancePropertySelectionChanged(FProperty* InProperty)
 {
 	Viewport->GetViewportClient()->Invalidate();
diff --git a/Engine/Plugins/Experimental/Mutable/Source/CustomizableObjectEditor/Private/MuCOE/CustomizableObjectInstanceEditor.h b/Engine/Plugins/Experimental/Mutable/Source/CustomizableObjectEditor/Private/MuCOE/CustomizableObjectInstanceEditor.h
--- a/Engine/Plugins/Experimental/Mutable/Source/CustomizableObjectEditor/Private/MuCOE/CustomizableObjectInstanceEditor.h
+++ b/Engine/Plugins/Experimental/Mutable/Source/CustomizableObjectEditor/Private/MuCOE/CustomizableObjectInstanceEditor.h
@@ -2,6 +2,7 @@
 
 #pragma once
 
+#include "DetailCategoryBuilder.h"
 #include "MuCO/CustomizableObjectPrivate.h"
 #include "MuCOE/ICustomizableObjectInstanceEditor.h"
 #include "MuCOE/CustomizableObjectEditorViewportLights.h"
@@ -18,6 +19,7 @@ class SWidget;
 class UCustomizableObject;
 class UCustomizableObjectInstance;
 class UPoseAsset;
+class UAnimationAsset;
 class SLevelOfDetailSettings;
 class SCustomizableObjectEditorViewportTabBody;
 class SCustomizableObjecEditorTextureAnalyzer;
@@ -107,6 +109,9 @@ class UCustomSettings : public UObject
 	GENERATED_BODY()
 
 public:
+	// UObject interface
+	virtual void PostEditChangeProperty(FPropertyChangedEvent& PropertyChangedEvent) override;
+
 	ULightComponent* GetSelectedLight() const;
 
 	void SetSelectedLight(ULightComponent* Light);
@@ -118,7 +123,10 @@ public:
 	TWeakPtr<ICustomizableObjectInstanceEditor> GetEditor() const;
 
 	void SetEditor(TSharedPtr<ICustomizableObjectInstanceEditor> Editor);
-
+	
+	UPROPERTY(Category = Animation, EditAnywhere)
+	TObjectPtr<UAnimationAsset> Animation;
+	
 private:
 	UPROPERTY()
 	TObjectPtr<ULightComponent> SelectedLight = nullptr;
@@ -126,7 +134,7 @@ private:
 	UPROPERTY()
 	TObjectPtr<UCustomizableObjectEditorViewportLights> LightsPreset;
 
-	TWeakPtr<ICustomizableObjectInstanceEditor> Editor;
+	TWeakPtr<ICustomizableObjectInstanceEditor> WeakEditor;
 };
 
 USTRUCT()
@@ -199,17 +207,19 @@ public:
 	virtual FLinearColor GetWorldCentricTabColorScale() const override;
 
 	/** ICustomizableObjectInstanceEditor interface */
-	UCustomizableObjectInstance* GetPreviewInstance() override;
+	virtual UCustomizableObjectInstance* GetPreviewInstance() override;
 	virtual void RefreshTool() override;
 	virtual TSharedPtr<SCustomizableObjectEditorViewportTabBody> GetViewport() override;
-	virtual void SetPoseAsset(UPoseAsset* PoseAssetParameter) override;
 	virtual UProjectorParameter* GetProjectorParameter() override;
 	virtual UCustomSettings* GetCustomSettings() override;
 	virtual void HideGizmo() override;
 	virtual void ShowGizmoProjectorParameter(const FString& ParamName, int32 RangeIndex) override;
 	virtual void HideGizmoProjectorParameter() override;
 	virtual UCustomizableObjectEditorProperties* GetEditorProperties() override;
-	
+	virtual TSharedPtr<SCustomizableObjectEditorAdvancedPreviewSettings> GetAdvancedPreviewSettings() override;
+	virtual bool ShowLightingSettings() override;
+	virtual bool ShowProfileManagementOptions() override;
+
 	/** Callback to notify the editor when the PreviewInstance has been updated */
 	void OnUpdatePreviewInstance(UCustomizableObjectInstance* Instance);
 
@@ -285,6 +295,8 @@ private:
 	TSharedPtr<SCustomizableObjectEditorViewportTabBody> Viewport;
 	TSharedPtr<IDetailsView> CustomizableInstanceDetailsView;
 
+	TSharedPtr<SCustomizableObjectEditorAdvancedPreviewSettings> CustomizableObjectEditorAdvancedPreviewSettings;
+	
 	/** Level of Details Settings widget. */
 	TSharedPtr<SLevelOfDetailSettings> LevelOfDetailSettings;
 
diff --git a/Engine/Plugins/Experimental/Mutable/Source/CustomizableObjectEditor/Private/MuCOE/SCustomizableObjectEditorAdvancedPreviewSettings.cpp b/Engine/Plugins/Experimental/Mutable/Source/CustomizableObjectEditor/Private/MuCOE/SCustomizableObjectEditorAdvancedPreviewSettings.cpp
--- a/Engine/Plugins/Experimental/Mutable/Source/CustomizableObjectEditor/Private/MuCOE/SCustomizableObjectEditorAdvancedPreviewSettings.cpp
+++ b/Engine/Plugins/Experimental/Mutable/Source/CustomizableObjectEditor/Private/MuCOE/SCustomizableObjectEditorAdvancedPreviewSettings.cpp
@@ -386,7 +386,7 @@ FReply SUObjectSelecWindow::OnButtonClick(EAppReturnType::Type ButtonID)
 																 SkyLightNameValue,
 																 PreviewSceneProfile);
 			// Make new entry in profile
-			CustomizableObjectEditor->GetCustomizableObjectEditorAdvancedPreviewSettings()->AddProfileToEditor(PreviewSceneProfile);
+			CustomizableObjectEditor.Pin()->GetAdvancedPreviewSettings()->AddProfileToEditor(PreviewSceneProfile);
 
 			if (LevelSelectWindow != nullptr)
 			{
@@ -417,50 +417,49 @@ void SCustomizableObjectEditorAdvancedPreviewSettings::Construct(const FArgument
 	SAdvancedPreviewDetailsTab::FArguments ParentArgs;
 	ParentArgs.AdditionalSettings( InArgs._CustomSettings );
 	SAdvancedPreviewDetailsTab::Construct(ParentArgs, InPreviewScene);
-
-	//SkyMaterial = nullptr;
-	//InstancedSkyMaterial = nullptr;
-	//DefaultTexture = nullptr;
-
+	
 	CustomizableObjectEditor = InArgs._CustomizableObjectEditor.Get();
 
-	TSharedRef< SWidget > ActualContent = ChildSlot.GetWidget();
-	ChildSlot
-	[
-		SNew(SVerticalBox)
-		+ SVerticalBox::Slot()
-		.FillHeight(1.0f)
-		[
-			ActualContent
-		]
-		+ SVerticalBox::Slot()
-		.AutoHeight()
+	if (CustomizableObjectEditor.Pin()->ShowProfileManagementOptions())
+	{
+		TSharedRef< SWidget > ActualContent = ChildSlot.GetWidget();
+		ChildSlot
 		[
-			SNew(SBorder)
-			.HAlign(HAlign_Fill)
+			SNew(SVerticalBox)
+			+ SVerticalBox::Slot()
+			.FillHeight(1.0f)
 			[
-				SNew(SUniformGridPanel)
-				.SlotPadding(UE_MUTABLE_GET_MARGIN("StandardDialog.SlotPadding"))
-				.MinDesiredSlotWidth(UE_MUTABLE_GET_FLOAT("StandardDialog.MinDesiredSlotWidth"))
-				.MinDesiredSlotHeight(UE_MUTABLE_GET_FLOAT("StandardDialog.MinDesiredSlotHeight"))
-				+ SUniformGridPanel::Slot(0, 0)
-				[
-					SNew(STextBlock)
-					.Font(FSlateFontInfo(FPaths::EngineContentDir() / TEXT("Slate/Fonts/Roboto-Bold.ttf"), 12))
-					.Text(LOCTEXT("LevelPreviewSettings", "Profile from Level Asset"))
-				]
-				+ SUniformGridPanel::Slot(1, 0)
+				ActualContent
+			]
+			+ SVerticalBox::Slot()
+			.AutoHeight()
+			[
+				SNew(SBorder)
+				.HAlign(HAlign_Fill)
 				[
-					SNew(SButton)
-					.VAlign(VAlign_Center)
-					.HAlign(HAlign_Center)
-					.OnClicked(this, &SCustomizableObjectEditorAdvancedPreviewSettings::ShowAddProfileWindow)
-					.Text(LOCTEXT("AddProfile", "Add new Profile"))
-					.ToolTipText(LOCTEXT("AddProfileToolTip", "Add new Profile reading information from an existing level"))
+					SNew(SUniformGridPanel)
+					.SlotPadding(UE_MUTABLE_GET_MARGIN("StandardDialog.SlotPadding"))
+					.MinDesiredSlotWidth(UE_MUTABLE_GET_FLOAT("StandardDialog.MinDesiredSlotWidth"))
+					.MinDesiredSlotHeight(UE_MUTABLE_GET_FLOAT("StandardDialog.MinDesiredSlotHeight"))
+					+ SUniformGridPanel::Slot(0, 0)
+					[
+						SNew(STextBlock)
+						.Font(FSlateFontInfo(FPaths::EngineContentDir() / TEXT("Slate/Fonts/Roboto-Bold.ttf"), 12))
+						.Text(LOCTEXT("LevelPreviewSettings", "Profile from Level Asset"))
+					]
+					+ SUniformGridPanel::Slot(1, 0)
+					[
+						SNew(SButton)
+						.VAlign(VAlign_Center)
+						.HAlign(HAlign_Center)
+						.OnClicked(this, &SCustomizableObjectEditorAdvancedPreviewSettings::ShowAddProfileWindow)
+						.Text(LOCTEXT("AddProfile", "Add new Profile"))
+						.ToolTipText(LOCTEXT("AddProfileToolTip", "Add new Profile reading information from an existing level"))
+					]
 				]
 			]
-		]
-	];
+		];	
+	}
 }
 
 
@@ -470,11 +469,7 @@ FReply SCustomizableObjectEditorAdvancedPreviewSettings::ShowAddProfileWindow()
 		SNew(SLevelSelectWindow)
 		.CustomizableObjectEditor(CustomizableObjectEditor)
 		.DefaultSettings(DefaultSettings);
-
-	if (LevelSelectWindow->ShowModal() != EAppReturnType::Cancel)
-	{
-	}
-
+	
 	return FReply::Handled();
 }
 
diff --git a/Engine/Plugins/Experimental/Mutable/Source/CustomizableObjectEditor/Private/MuCOE/SCustomizableObjectEditorAdvancedPreviewSettings.h b/Engine/Plugins/Experimental/Mutable/Source/CustomizableObjectEditor/Private/MuCOE/SCustomizableObjectEditorAdvancedPreviewSettings.h
--- a/Engine/Plugins/Experimental/Mutable/Source/CustomizableObjectEditor/Private/MuCOE/SCustomizableObjectEditorAdvancedPreviewSettings.h
+++ b/Engine/Plugins/Experimental/Mutable/Source/CustomizableObjectEditor/Private/MuCOE/SCustomizableObjectEditorAdvancedPreviewSettings.h
@@ -5,6 +5,7 @@
 #include "CustomizableObjectInstanceEditor.h"
 #include "AssetRegistry/AssetData.h"
 #include "SAdvancedPreviewDetailsTab.h"
+#include "MuCOE/ICustomizableObjectEditor.h"
 #include "Widgets/SWindow.h"
 
 class FAdvancedPreviewScene;
@@ -22,7 +23,7 @@ class SLevelSelectWindow : public SWindow
 {
 public:
 	SLATE_BEGIN_ARGS(SLevelSelectWindow) {}
-		SLATE_ATTRIBUTE(FCustomizableObjectEditor*, CustomizableObjectEditor)
+		SLATE_ATTRIBUTE(TWeakPtr<ICustomizableObjectInstanceEditor>, CustomizableObjectEditor)
 		SLATE_ATTRIBUTE(UAssetViewerSettings*, DefaultSettings)
 	SLATE_END_ARGS()
 
@@ -52,7 +53,7 @@ private:
 	TArray<FString> ArraySkyLightName;
 	FString LevelAssetPath;
 	bool ProfileAddedSuccessfully;
-	FCustomizableObjectEditor* CustomizableObjectEditor; // Pointer to editor instance
+	TWeakPtr<ICustomizableObjectInstanceEditor> CustomizableObjectEditor;
 	UAssetViewerSettings* DefaultSettings;
 };
 
@@ -65,7 +66,7 @@ public:
 	SLATE_BEGIN_ARGS(SUObjectSelecWindow) {}
 		SLATE_ATTRIBUTE(FString, LevelAssetPath)
 		SLATE_ATTRIBUTE(SLevelSelectWindow*, LevelSelectWindow)
-		SLATE_ATTRIBUTE(FCustomizableObjectEditor*, CustomizableObjectEditor)
+		SLATE_ATTRIBUTE(TWeakPtr<ICustomizableObjectInstanceEditor>, CustomizableObjectEditor)
 		SLATE_ATTRIBUTE(UAssetViewerSettings*, DefaultSettings)
 	SLATE_END_ARGS()
 
@@ -84,7 +85,7 @@ private:
 	TSharedPtr<STextComboBox> SkyLight;
 	FString LevelAssetPath;
 	SLevelSelectWindow* LevelSelectWindow; // Pointer to parent window, needed to close it if new profile added successfully
-	FCustomizableObjectEditor* CustomizableObjectEditor; // Pointer to editor instance
+	TWeakPtr<ICustomizableObjectInstanceEditor> CustomizableObjectEditor;
 	TArray<TSharedPtr<FString>> ArrayDirectionalLightName;
 	TArray<TSharedPtr<FString>> ArrayPostProcessName;
 	TArray<TSharedPtr<FString>> ArraySkyLightName;
@@ -100,7 +101,7 @@ class SCustomizableObjectEditorAdvancedPreviewSettings : public SAdvancedPreview
 public:
 	SLATE_BEGIN_ARGS(SCustomizableObjectEditorAdvancedPreviewSettings) {}
 		SLATE_ARGUMENT(UCustomSettings*, CustomSettings) // From SAdvancedPreviewDetailsTab
-		SLATE_ATTRIBUTE(FCustomizableObjectEditor*, CustomizableObjectEditor)
+		SLATE_ATTRIBUTE(TWeakPtr<ICustomizableObjectInstanceEditor>, CustomizableObjectEditor)
 	SLATE_END_ARGS()
 
 	SCustomizableObjectEditorAdvancedPreviewSettings();
@@ -115,7 +116,6 @@ public:
 	/** Load the environment map of each prodile in DefaultSettings::Profiles since default constructor loads a predetermined environment */
 	void LoadProfileEnvironment();
 
-	FCustomizableObjectEditor* CustomizableObjectEditor; // Pointer to editor instance
-
+	TWeakPtr<ICustomizableObjectInstanceEditor> CustomizableObjectEditor; // Pointer to editor instance
 };
 
diff --git a/Engine/Plugins/Experimental/Mutable/Source/CustomizableObjectEditor/Private/MuCOE/SCustomizableObjectEditorViewport.cpp b/Engine/Plugins/Experimental/Mutable/Source/CustomizableObjectEditor/Private/MuCOE/SCustomizableObjectEditorViewport.cpp
--- a/Engine/Plugins/Experimental/Mutable/Source/CustomizableObjectEditor/Private/MuCOE/SCustomizableObjectEditorViewport.cpp
+++ b/Engine/Plugins/Experimental/Mutable/Source/CustomizableObjectEditor/Private/MuCOE/SCustomizableObjectEditorViewport.cpp
@@ -2,10 +2,13 @@
 
 #include "MuCOE/SCustomizableObjectEditorViewport.h"
 
+#include "CustomizableObjectInstanceEditor.h"
 #include "Animation/AnimationAsset.h"
 #include "Animation/DebugSkelMeshComponent.h"
 #include "Editor/UnrealEdEngine.h"
 #include "EditorViewportCommands.h"
+#include "ObjectEditorUtils.h"
+#include "SMutableScrubPanel.h"
 #include "Framework/Commands/UICommandList.h"
 #include "Framework/MultiBox/MultiBoxBuilder.h"
 #include "MuCO/CustomizableObjectInstance.h"
@@ -272,6 +275,8 @@ void SCustomizableObjectEditorViewportTabBody::Construct(const FArguments& InArg
 	
 	ViewportWidget = SNew(SCustomizableObjectEditorViewport, PreviewScenePtr.ToSharedRef(), SharedThis(this), WeakEditor);
 
+	LevelViewportClient = StaticCastSharedPtr<FCustomizableObjectEditorViewportClient>(ViewportWidget->GetViewportClient());
+
 	this->ChildSlot
 		[
 			SNew(SVerticalBox)
@@ -285,9 +290,12 @@ void SCustomizableObjectEditorViewportTabBody::Construct(const FArguments& InArg
 			[
 				ViewportWidget.ToSharedRef()
 			]
+			+SVerticalBox::Slot()
+			.AutoHeight()
+			[
+				SAssignNew(ScrubPanel, SMutableScrubPanel, LevelViewportClient.ToSharedRef())
+			]
 		];
-
-	LevelViewportClient = StaticCastSharedPtr<FCustomizableObjectEditorViewportClient>(ViewportWidget->GetViewportClient());
 	
 	BindCommands();
 }
@@ -449,6 +457,18 @@ bool SCustomizableObjectEditorViewportTabBody::IsVisible() const
 }
 
 
+void SCustomizableObjectEditorViewportTabBody::OnSetPlaybackSpeed(int32 PlaybackSpeedMode)
+{
+	LevelViewportClient->SetPlaybackSpeedMode(static_cast<EMutableAnimationPlaybackSpeeds::Type>(PlaybackSpeedMode));
+}
+
+
+bool SCustomizableObjectEditorViewportTabBody::IsPlaybackSpeedSelected(int32 PlaybackSpeedMode)
+{
+	return PlaybackSpeedMode == LevelViewportClient->GetPlaybackSpeedMode();
+}
+
+
 void SCustomizableObjectEditorViewportTabBody::BindCommands()
 {
 	FUICommandList& CommandList = *UICommandList;
@@ -500,7 +520,17 @@ void SCustomizableObjectEditorViewportTabBody::BindCommands()
 		FExecuteAction::CreateSP(EditorViewportClientRef, &FCustomizableObjectEditorViewportClient::BakeInstance),
 		FCanExecuteAction(),
 		FIsActionChecked());
-
+	
+	//Create a menu item for each playback speed in EMutableAnimationPlaybackSpeeds
+	for (int32 Index = 0; Index < static_cast<int32>(EMutableAnimationPlaybackSpeeds::NumPlaybackSpeeds); ++Index)
+	{
+		CommandList.MapAction( 
+			Commands.PlaybackSpeedCommands[Index],
+			FExecuteAction::CreateSP(this, &SCustomizableObjectEditorViewportTabBody::OnSetPlaybackSpeed, Index),
+			FCanExecuteAction(),
+			FIsActionChecked::CreateSP(this, &SCustomizableObjectEditorViewportTabBody::IsPlaybackSpeedSelected, Index));
+	}
+	
 	//Bind LOD preview menu commands
 	const FCustomizableObjectEditorViewportLODCommands& ViewportLODMenuCommands = FCustomizableObjectEditorViewportLODCommands::Get();
 
@@ -617,7 +647,6 @@ void SCustomizableObjectEditorViewportTabBody::BindCommands()
 		FExecuteAction::CreateSP(EditorViewportClientRef, &FCustomizableObjectEditorViewportClient::SetViewportType, LVT_OrthoNegativeXY),
 		FCanExecuteAction(),
 		FIsActionChecked::CreateSP(EditorViewportClientRef, &FCustomizableObjectEditorViewportClient::IsActiveViewportType, LVT_OrthoNegativeXY));
-
 	CommandList.MapAction( 
 		Commands.ShowDisplayInfo,
 		FExecuteAction::CreateSP(EditorViewportClientRef, &FCustomizableObjectEditorViewportClient::OnShowDisplayInfo),
@@ -981,12 +1010,6 @@ void SCustomizableObjectEditorViewportTabBody::OnUVChannelChanged(TSharedPtr<FSt
 }
 
 
-void SCustomizableObjectEditorViewportTabBody::SetAnimation(UAnimationAsset* Animation, EAnimationMode::Type AnimationType)
-{
-	LevelViewportClient->SetAnimation(Animation, AnimationType);
-}
-
-
 FReply SCustomizableObjectEditorViewportTabBody::OnDrop(const FGeometry& MyGeometry, const FDragDropEvent& DragDropEvent)
 {
 	if (FAssetDragDropOp* DragDropOp = DragDropEvent.GetOperationAs<FAssetDragDropOp>().Get())
@@ -997,7 +1020,7 @@ FReply SCustomizableObjectEditorViewportTabBody::OnDrop(const FGeometry& MyGeome
 			UAnimationAsset* AnimationAsset = Cast<UAnimationAsset>(DragDropOp->GetAssets()[0].GetAsset());
 			if (AnimationAsset)
 			{
-				LevelViewportClient->SetAnimation(AnimationAsset, EAnimationMode::AnimationSingleNode);
+				FObjectEditorUtils::SetPropertyValue(WeakEditor.Pin()->GetCustomSettings(), GET_MEMBER_NAME_CHECKED(UCustomSettings, Animation), AnimationAsset);
 
 				return FReply::Handled();
 			}
diff --git a/Engine/Plugins/Experimental/Mutable/Source/CustomizableObjectEditor/Private/MuCOE/SCustomizableObjectEditorViewport.h b/Engine/Plugins/Experimental/Mutable/Source/CustomizableObjectEditor/Private/MuCOE/SCustomizableObjectEditorViewport.h
--- a/Engine/Plugins/Experimental/Mutable/Source/CustomizableObjectEditor/Private/MuCOE/SCustomizableObjectEditorViewport.h
+++ b/Engine/Plugins/Experimental/Mutable/Source/CustomizableObjectEditor/Private/MuCOE/SCustomizableObjectEditorViewport.h
@@ -6,6 +6,7 @@
 #include "MuCOE/CustomizableObjectEditorViewportClient.h"
 #include "UObject/GCObject.h"
 
+class SMutableScrubPanel;
 class FCustomizableObjectPreviewScene;
 class SCustomizableObjectEditorViewportTabBody;
 enum ERotationGridMode : int;
@@ -97,9 +98,7 @@ public:
 	SLATE_END_ARGS()
 
 	void Construct(const FArguments& InArgs);
-
-	void SetAnimation(class UAnimationAsset* Animation, EAnimationMode::Type AnimationType);
-
+	
 	virtual FReply OnDrop(const FGeometry& MyGeometry, const FDragDropEvent& DragDropEvent) override;
 	
 	// SWidget interface
@@ -223,7 +222,11 @@ public:
 private:
 	/** Determines the visibility of the viewport. */
 	bool IsVisible() const;
+
+	void OnSetPlaybackSpeed(int32 PlaybackSpeedMode);
 	
+	bool IsPlaybackSpeedSelected(int32 PlaybackSpeedMode);
+
 	// Components for the preview mesh.
 	TArray<TWeakObjectPtr<UDebugSkelMeshComponent>> PreviewSkeletalMeshComponents;
 
@@ -264,6 +267,8 @@ private:
 	// Combo box for UV channel selection
 	TSharedPtr<STextComboBox> UVChannelOptionCombo;
 
+	TSharedPtr<SMutableScrubPanel> ScrubPanel;
+
 	// Generates the Combobox options for each UV Channel
 	void GenerateUVSectionOptions();
 
diff --git a/Engine/Plugins/Experimental/Mutable/Source/CustomizableObjectEditor/Private/MuCOE/SCustomizableObjectEditorViewportToolBar.cpp b/Engine/Plugins/Experimental/Mutable/Source/CustomizableObjectEditor/Private/MuCOE/SCustomizableObjectEditorViewportToolBar.cpp
--- a/Engine/Plugins/Experimental/Mutable/Source/CustomizableObjectEditor/Private/MuCOE/SCustomizableObjectEditorViewportToolBar.cpp
+++ b/Engine/Plugins/Experimental/Mutable/Source/CustomizableObjectEditor/Private/MuCOE/SCustomizableObjectEditorViewportToolBar.cpp
@@ -34,6 +34,7 @@
 #include "Widgets/Layout/SScaleBox.h"
 #include "Widgets/Layout/SUniformGridPanel.h"
 #include "Widgets/Input/SSlider.h"
+#include "MuCOE/CustomizableObjectEditorActions.h"
 
 class SWidget;
 struct FSlateBrush;
@@ -117,6 +118,19 @@ void SCustomizableObjectEditorViewportToolBar::Construct(const FArguments& InArg
 			.ParentToolBar(SharedThis(this))
 			.Label(LOCTEXT("CharacterMenuLabel","Character"))
 			.OnGetMenuContent(this, &SCustomizableObjectEditorViewportToolBar::GenerateCharacterMenu)
+		]
+
+	// Playback Menu
+	+ SHorizontalBox::Slot()
+		.AutoWidth()
+		.Padding(2.0f, 2.0f)
+		[
+			SNew(SEditorViewportToolbarMenu)
+			.ToolTipText(LOCTEXT("PlaybackSpeedMenuTooltip", "Playback Speed Options. Control the time dilation of the scene's update."))
+			.ParentToolBar(SharedThis(this))
+			.Label(this, &SCustomizableObjectEditorViewportToolBar::GetPlaybackMenuLabel)
+			.LabelIcon(FAppStyle::GetBrush("AnimViewportMenu.PlayBackSpeed"))
+			.OnGetMenuContent(this, &SCustomizableObjectEditorViewportToolBar::GeneratePlaybackMenu)
 		];
 	
 	TSharedRef<SWidget> RTSButtons = GenerateRTSButtons();
@@ -331,14 +345,77 @@ TSharedRef<SWidget> SCustomizableObjectEditorViewportToolBar::GenerateViewportTy
 	return CameraMenuBuilder.MakeWidget();
 }
 
+
+class SCusttomAnimationSpeedSetting : public SCompoundWidget
+{
+public:
+	DECLARE_DELEGATE_OneParam(FOnCustomSpeedChanged, float);
+
+	SLATE_BEGIN_ARGS(SCusttomAnimationSpeedSetting) {}
+		SLATE_ATTRIBUTE(float, CustomSpeed)
+		SLATE_EVENT(FOnCustomSpeedChanged, OnCustomSpeedChanged)
+	SLATE_END_ARGS()
+
+	void Construct(const FArguments& InArgs )
+	{
+		CustomSpeed = InArgs._CustomSpeed;
+		OnCustomSpeedChanged = InArgs._OnCustomSpeedChanged;
+
+		this->ChildSlot
+		[
+			SNew(SBox)
+			.HAlign(HAlign_Right)
+			[
+				SNew(SBox)
+				.Padding(FMargin(4.0f, 0.0f, 0.0f, 0.0f))
+				.WidthOverride(100.0f)
+				[
+					SNew(SSpinBox<float>)
+					.Font(FAppStyle::GetFontStyle(TEXT("MenuItem.Font")))
+					.ToolTipText(LOCTEXT("AnimationCustomSpeed", "Set Custom Speed."))
+					.MinValue(0.f)
+					.MaxSliderValue(10.f)
+					.SupportDynamicSliderMaxValue(true)
+					.Value(CustomSpeed)
+					.OnValueChanged(OnCustomSpeedChanged)
+				]
+			]
+		];
+	}
+
+private:
+    TAttribute<float> CustomSpeed = 1.0f;
+	FOnCustomSpeedChanged OnCustomSpeedChanged;
+};
+
+
 TSharedRef<SWidget> SCustomizableObjectEditorViewportToolBar::GeneratePlaybackMenu() const
 {
-	const bool bInShouldCloseWindowAfterMenuSelection = true;
+	const FCustomizableObjectEditorViewportCommands& Actions = FCustomizableObjectEditorViewportCommands::Get();
+	
+	FMenuBuilder InMenuBuilder(true, Viewport.Pin()->GetCommandList());
+
+	InMenuBuilder.BeginSection("AnimViewportPlaybackSpeed", LOCTEXT("PlaybackMenu_SpeedLabel", "Playback Speed") );
 
-	FMenuBuilder PlaybackMenuBuilder(bInShouldCloseWindowAfterMenuSelection, Viewport.Pin()->GetCommandList());
+	for(int32 PlaybackSpeedIndex = 0; PlaybackSpeedIndex < EMutableAnimationPlaybackSpeeds::NumPlaybackSpeeds; ++PlaybackSpeedIndex)
+	{
+		InMenuBuilder.AddMenuEntry( Actions.PlaybackSpeedCommands[PlaybackSpeedIndex] );
+	}
 	
-	return PlaybackMenuBuilder.MakeWidget();
+	TSharedPtr<SWidget> AnimSpeedWidget = SNew(SCusttomAnimationSpeedSetting)
+		.CustomSpeed_Lambda([Viewport = Viewport]()
+			{
+				return Viewport.Pin()->GetViewportClient()->GetCustomAnimationSpeed();
+			})
+		.OnCustomSpeedChanged_Lambda([Viewport = Viewport](float CustomSpeed)
+			{
+				return Viewport.Pin()->GetViewportClient()->SetCustomAnimationSpeed(CustomSpeed);
+			});
+	
+	InMenuBuilder.AddWidget(AnimSpeedWidget.ToSharedRef(), LOCTEXT("PlaybackMenu_Speed_Custom", "Custom Speed:"));
+	InMenuBuilder.EndSection();
 
+	return InMenuBuilder.MakeWidget();
 }
 
 
@@ -395,9 +472,28 @@ TSharedRef<SWidget> SCustomizableObjectEditorViewportToolBar::GenerateViewportOp
 
 FText SCustomizableObjectEditorViewportToolBar::GetPlaybackMenuLabel() const
 {
-	return LOCTEXT("PlaybackError", "Error");
+	FText Label = LOCTEXT("PlaybackError", "Error");
+	if (Viewport.IsValid())
+	{
+		for(int Index = 0; Index < EMutableAnimationPlaybackSpeeds::NumPlaybackSpeeds; ++Index)
+		{
+			if (Viewport.Pin()->GetViewportClient()->GetPlaybackSpeedMode() == Index)
+			{
+				const int32 NumFractionalDigits = (Index == EMutableAnimationPlaybackSpeeds::Quarter || Index == EMutableAnimationPlaybackSpeeds::ThreeQuarters) ? 2 : 1;
+
+				const FNumberFormattingOptions FormatOptions = FNumberFormattingOptions()
+					.SetMinimumFractionalDigits(NumFractionalDigits)
+					.SetMaximumFractionalDigits(NumFractionalDigits);
+
+				Label = FText::Format(LOCTEXT("AnimViewportPlaybackMenuLabel", "x{0}"), FText::AsNumber(EMutableAnimationPlaybackSpeeds::Values[Index], &FormatOptions));
+			}
+		}
+	}
+	
+	return Label;
 }
 
+
 FText SCustomizableObjectEditorViewportToolBar::GetCameraMenuLabel() const
 {
 	FText Label = LOCTEXT("Viewport_Default", "Camera");
diff --git a/Engine/Plugins/Experimental/Mutable/Source/CustomizableObjectEditor/Private/MuCOE/SMutableScrubPanel.cpp b/Engine/Plugins/Experimental/Mutable/Source/CustomizableObjectEditor/Private/MuCOE/SMutableScrubPanel.cpp
new file mode 100644
--- /dev/null
+++ b/Engine/Plugins/Experimental/Mutable/Source/CustomizableObjectEditor/Private/MuCOE/SMutableScrubPanel.cpp
@@ -0,0 +1,715 @@
+﻿// Copyright Epic Games, Inc. All Rights Reserved.
+
+#include "MuCOE/SMutableScrubPanel.h"
+
+#include "Widgets/SBoxPanel.h"
+#include "Animation/AnimSequenceBase.h"
+#include "Animation/AnimSequence.h"
+#include "Animation/AnimBlueprintGeneratedClass.h"
+#include "Animation/AnimBlueprint.h"
+#include "AnimPreviewInstance.h"
+#include "SScrubControlPanel.h"
+#include "ScopedTransaction.h"
+#include "Animation/BlendSpace.h"
+#include "Animation/AnimData/IAnimationDataModel.h"
+#include "Animation/DebugSkelMeshComponent.h"
+#include "MuCOE/CustomizableObjectEditorViewportClient.h"
+#include "UObject/WeakObjectPtrTemplates.h"
+
+#define LOCTEXT_NAMESPACE "CustomizableObjectEditor"
+
+
+void SMutableScrubPanel::Construct(const FArguments& InArgs, const TSharedRef<FCustomizableObjectEditorViewportClient>& InPreviewScene)
+{
+	bSliderBeingDragged = false;
+	OnSetInputViewRange = InArgs._OnSetInputViewRange;
+
+	PreviewScenePtr = InPreviewScene;
+
+	this->ChildSlot
+	[
+		SNew(SHorizontalBox)
+		+SHorizontalBox::Slot()
+		.HAlign(HAlign_Fill) 
+		.VAlign(VAlign_Center)
+		.FillWidth(1)
+		.Padding(0.0f)
+		[
+			SAssignNew(ScrubControlPanel, SScrubControlPanel)
+			.IsEnabled(true)
+			.Value(this, &SMutableScrubPanel::GetScrubValue)
+			.NumOfKeys(this, &SMutableScrubPanel::GetNumberOfKeys)
+			.SequenceLength(this, &SMutableScrubPanel::GetSequenceLength)
+			.DisplayDrag(this, &SMutableScrubPanel::GetDisplayDrag)
+			.OnValueChanged(this, &SMutableScrubPanel::OnValueChanged)
+			.OnBeginSliderMovement(this, &SMutableScrubPanel::OnBeginSliderMovement)
+			.OnEndSliderMovement(this, &SMutableScrubPanel::OnEndSliderMovement)
+			.OnClickedForwardPlay(this, &SMutableScrubPanel::OnClick_Forward)
+			.OnClickedForwardStep(this, &SMutableScrubPanel::OnClick_Forward_Step)
+			.OnClickedForwardEnd(this, &SMutableScrubPanel::OnClick_Forward_End)
+			.OnClickedBackwardPlay(this, &SMutableScrubPanel::OnClick_Backward)
+			.OnClickedBackwardStep(this, &SMutableScrubPanel::OnClick_Backward_Step)
+			.OnClickedBackwardEnd(this, &SMutableScrubPanel::OnClick_Backward_End)
+			.OnClickedToggleLoop(this, &SMutableScrubPanel::OnClick_ToggleLoop)
+			.OnGetLooping(this, &SMutableScrubPanel::IsLoopStatusOn)
+			.OnGetPlaybackMode(this, &SMutableScrubPanel::GetPlaybackMode)
+			.ViewInputMin(this, &SMutableScrubPanel::GetViewMinInput)
+			.ViewInputMax(this, &SMutableScrubPanel::GetViewMaxInput)
+			.bDisplayAnimScrubBarEditing(InArgs._bDisplayAnimScrubBarEditing)
+			.OnSetInputViewRange(InArgs._OnSetInputViewRange)
+			.OnReZeroAnimSequence( this, &SMutableScrubPanel::OnReZeroAnimSequence )
+			.bAllowZoom(true)
+			.IsRealtimeStreamingMode(this, &SMutableScrubPanel::IsRealtimeStreamingMode)
+		]
+	];
+}
+
+
+float SMutableScrubPanel::GetViewMinInput() const
+{
+	for (TWeakObjectPtr<UDebugSkelMeshComponent>& WeakPreviewMeshComponent : GetPreviewScene()->GetPreviewMeshComponents())
+	{
+		UDebugSkelMeshComponent* PreviewMeshComponent = WeakPreviewMeshComponent.Get();
+		if (!PreviewMeshComponent)
+		{
+			continue;
+		}
+
+		if (PreviewMeshComponent->PreviewInstance != NULL)
+		{
+			return 0.0f;
+		}
+		if (PreviewMeshComponent->GetAnimInstance() != NULL)
+		{
+			return FMath::Max<float>((float)(PreviewMeshComponent->GetAnimInstance()->LifeTimer - 30.0), 0.0f);
+		}
+	}
+	
+	return 0.f; 
+}
+
+
+float SMutableScrubPanel::GetViewMaxInput() const
+{ 
+	for (TWeakObjectPtr<UDebugSkelMeshComponent>& WeakPreviewMeshComponent : GetPreviewScene()->GetPreviewMeshComponents())
+	{
+		UDebugSkelMeshComponent* PreviewMeshComponent = WeakPreviewMeshComponent.Get();
+		if (!PreviewMeshComponent)
+		{
+			continue;
+		}
+		
+		if (PreviewMeshComponent->PreviewInstance != NULL)
+		{
+			return PreviewMeshComponent->PreviewInstance->GetLength();
+		}
+		else if (PreviewMeshComponent->GetAnimInstance() != NULL)
+		{
+			return static_cast<float>(PreviewMeshComponent->GetAnimInstance()->LifeTimer);
+		}
+	}
+
+	return 0.f;
+}
+
+
+FReply SMutableScrubPanel::OnClick_Forward_Step()
+{
+	for (TWeakObjectPtr<UDebugSkelMeshComponent>& WeakPreviewMeshComponent : GetPreviewScene()->GetPreviewMeshComponents())
+	{
+		UDebugSkelMeshComponent* PreviewMeshComponent = WeakPreviewMeshComponent.Get();
+		if (!PreviewMeshComponent)
+		{
+			continue;
+		}
+		
+		if (UAnimSingleNodeInstance* PreviewInstance = GetPreviewInstance(PreviewMeshComponent))
+		{
+			bool bShouldStepCloth = FMath::Abs(PreviewInstance->GetLength() - PreviewInstance->GetCurrentTime()) > SMALL_NUMBER;
+
+			PreviewInstance->SetPlaying(false);
+			PreviewInstance->StepForward();
+
+			if(PreviewMeshComponent && bShouldStepCloth)
+			{
+				PreviewMeshComponent->bPerformSingleClothingTick = true;
+			}
+		}
+		else if (PreviewMeshComponent)
+		{
+			// BlendSpaces and Animation Blueprints combine animations so there's no such thing as a frame. However, 1/30 is a sensible/common rate.
+			const float FixedFrameRate = 30.0f;
+
+			// Advance a single frame, leaving it paused afterwards
+			PreviewMeshComponent->GlobalAnimRateScale = 1.0f;
+			GetPreviewScene()->Tick(1.0f / FixedFrameRate); 
+			PreviewMeshComponent->GlobalAnimRateScale = 0.0f;
+		}
+	}
+
+	return FReply::Handled();
+}
+
+
+FReply SMutableScrubPanel::OnClick_Forward_End()
+{
+	for (TWeakObjectPtr<UDebugSkelMeshComponent>& WeakPreviewMeshComponent : GetPreviewScene()->GetPreviewMeshComponents())
+	{
+		UDebugSkelMeshComponent* PreviewMeshComponent = WeakPreviewMeshComponent.Get();
+		if (!PreviewMeshComponent)
+        {
+        	continue;
+        }
+		
+		UAnimSingleNodeInstance* PreviewInstance = GetPreviewInstance(PreviewMeshComponent);
+        if (PreviewInstance)
+        {
+        	PreviewInstance->SetPlaying(false);
+        	PreviewInstance->SetPosition(PreviewInstance->GetLength(), false);
+        }
+	}
+
+	return FReply::Handled();
+}
+
+
+FReply SMutableScrubPanel::OnClick_Backward_Step()
+{
+	for (TWeakObjectPtr<UDebugSkelMeshComponent>& WeakPreviewMeshComponent : GetPreviewScene()->GetPreviewMeshComponents())
+	{
+		UDebugSkelMeshComponent* PreviewMeshComponent = WeakPreviewMeshComponent.Get();
+		if (!PreviewMeshComponent)
+		{
+			continue;
+		}
+
+		UAnimSingleNodeInstance* PreviewInstance = GetPreviewInstance(PreviewMeshComponent);
+		if (PreviewInstance)
+		{
+			bool bShouldStepCloth = PreviewInstance->GetCurrentTime() > SMALL_NUMBER;
+
+			PreviewInstance->SetPlaying(false);
+			PreviewInstance->StepBackward();
+
+			if(PreviewMeshComponent && bShouldStepCloth)
+			{
+				PreviewMeshComponent->bPerformSingleClothingTick = true;
+			}
+		}
+	}
+	
+	return FReply::Handled();
+}
+
+
+FReply SMutableScrubPanel::OnClick_Backward_End()
+{
+	for (TWeakObjectPtr<UDebugSkelMeshComponent>& WeakPreviewMeshComponent : GetPreviewScene()->GetPreviewMeshComponents())
+	{
+		UDebugSkelMeshComponent* PreviewMeshComponent = WeakPreviewMeshComponent.Get();
+		if (!PreviewMeshComponent)
+		{
+			continue;
+		}
+
+		UAnimSingleNodeInstance* PreviewInstance = GetPreviewInstance(PreviewMeshComponent);
+		if (PreviewInstance)
+		{
+			PreviewInstance->SetPlaying(false);
+			PreviewInstance->SetPosition(0.f, false);
+		}
+	}
+
+	return FReply::Handled();
+}
+
+
+FReply SMutableScrubPanel::OnClick_Forward()
+{
+	for (TWeakObjectPtr<UDebugSkelMeshComponent>& WeakPreviewMeshComponent : GetPreviewScene()->GetPreviewMeshComponents())
+	{
+		UDebugSkelMeshComponent* PreviewMeshComponent = WeakPreviewMeshComponent.Get();
+		if (!PreviewMeshComponent)
+		{
+			continue;
+		}
+
+		UAnimSingleNodeInstance* PreviewInstance = GetPreviewInstance(PreviewMeshComponent);
+
+		if (PreviewInstance)
+		{
+			bool bIsReverse = PreviewInstance->IsReverse();
+			bool bIsPlaying = PreviewInstance->IsPlaying();
+			// if current bIsReverse and bIsPlaying, we'd like to just turn off reverse
+			if (bIsReverse && bIsPlaying)
+			{
+				PreviewInstance->SetReverse(false);
+			}
+			// already playing, simply pause
+			else if (bIsPlaying) 
+			{
+				PreviewInstance->SetPlaying(false);
+			
+				if(PreviewMeshComponent && PreviewMeshComponent->bPauseClothingSimulationWithAnim)
+				{
+					PreviewMeshComponent->SuspendClothingSimulation();
+				}
+			}
+			// if not playing, play forward
+			else 
+			{
+				//if we're at the end of the animation, jump back to the beginning before playing
+				if ( GetScrubValue() >= GetSequenceLength() )
+				{
+					PreviewInstance->SetPosition(0.0f, false);
+				}
+
+				PreviewInstance->SetReverse(false);
+				PreviewInstance->SetPlaying(true);
+
+				if(PreviewMeshComponent && PreviewMeshComponent->bPauseClothingSimulationWithAnim)
+				{
+					PreviewMeshComponent->ResumeClothingSimulation();
+				}
+			}
+		}
+		else if(PreviewMeshComponent)
+		{
+			PreviewMeshComponent->GlobalAnimRateScale = (PreviewMeshComponent->GlobalAnimRateScale > 0.0f) ? 0.0f : 1.0f;
+		}
+	}
+	
+	return FReply::Handled();
+}
+
+
+FReply SMutableScrubPanel::OnClick_Backward()
+{
+	for (TWeakObjectPtr<UDebugSkelMeshComponent>& WeakPreviewMeshComponent : GetPreviewScene()->GetPreviewMeshComponents())
+	{
+		UDebugSkelMeshComponent* PreviewMeshComponent = WeakPreviewMeshComponent.Get();
+		if (!PreviewMeshComponent)
+		{
+			continue;
+		}
+
+		UAnimSingleNodeInstance* PreviewInstance = GetPreviewInstance(PreviewMeshComponent);
+		if (PreviewInstance)
+		{
+			bool bIsReverse = PreviewInstance->IsReverse();
+			bool bIsPlaying = PreviewInstance->IsPlaying();
+			// if currently playing forward, just simply turn on reverse
+			if (!bIsReverse && bIsPlaying)
+			{
+				PreviewInstance->SetReverse(true);
+			}
+			else if (bIsPlaying)
+			{
+				PreviewInstance->SetPlaying(false);
+			}
+			else
+			{
+				//if we're at the beginning of the animation, jump back to the end before playing
+				if ( GetScrubValue() <= 0.0f )
+				{
+					PreviewInstance->SetPosition(GetSequenceLength(), false);
+				}
+
+				PreviewInstance->SetPlaying(true);
+				PreviewInstance->SetReverse(true);
+			}
+		}
+	}
+	
+	return FReply::Handled();
+}
+
+
+FReply SMutableScrubPanel::OnClick_ToggleLoop()
+{
+	for (TWeakObjectPtr<UDebugSkelMeshComponent>& WeakPreviewMeshComponent : GetPreviewScene()->GetPreviewMeshComponents())
+	{
+		UDebugSkelMeshComponent* PreviewMeshComponent = WeakPreviewMeshComponent.Get();
+		if (!PreviewMeshComponent)
+		{
+			continue;
+		}
+
+		UAnimSingleNodeInstance* PreviewInstance = GetPreviewInstance(PreviewMeshComponent);
+		if (PreviewInstance)
+		{
+			bool bIsLooping = PreviewInstance->IsLooping();
+			PreviewInstance->SetLooping(!bIsLooping);
+		}
+	}
+	
+	return FReply::Handled();
+}
+
+
+bool SMutableScrubPanel::IsLoopStatusOn() const
+{
+	for (TWeakObjectPtr<UDebugSkelMeshComponent>& WeakPreviewMeshComponent : GetPreviewScene()->GetPreviewMeshComponents())
+	{
+		UDebugSkelMeshComponent* PreviewMeshComponent = WeakPreviewMeshComponent.Get();
+		if (!PreviewMeshComponent)
+		{
+			continue;
+		}
+
+		if (UAnimSingleNodeInstance* PreviewInstance = GetPreviewInstance(PreviewMeshComponent))
+		{
+			return PreviewInstance->IsLooping();
+		}
+	}
+	
+	return false;
+}
+
+
+EPlaybackMode::Type SMutableScrubPanel::GetPlaybackMode() const
+{
+	for (TWeakObjectPtr<UDebugSkelMeshComponent>& WeakPreviewMeshComponent : GetPreviewScene()->GetPreviewMeshComponents())
+	{
+		UDebugSkelMeshComponent* PreviewMeshComponent = WeakPreviewMeshComponent.Get();
+		if (!PreviewMeshComponent)
+		{
+			continue;
+		}
+
+		if (UAnimSingleNodeInstance* PreviewInstance = GetPreviewInstance(PreviewMeshComponent))
+		{
+			if (PreviewInstance->IsPlaying())
+			{
+				return PreviewInstance->IsReverse() ? EPlaybackMode::PlayingReverse : EPlaybackMode::PlayingForward;
+			}
+			return EPlaybackMode::Stopped;
+		}
+		else if (PreviewMeshComponent)
+		{
+			return PreviewMeshComponent->GlobalAnimRateScale > 0.0f ? EPlaybackMode::PlayingForward : EPlaybackMode::Stopped;
+		}
+	}
+	
+	return EPlaybackMode::Stopped;
+}
+
+
+bool SMutableScrubPanel::IsRealtimeStreamingMode() const
+{
+	for (TWeakObjectPtr<UDebugSkelMeshComponent>& WeakPreviewMeshComponent : GetPreviewScene()->GetPreviewMeshComponents())
+	{
+		UDebugSkelMeshComponent* PreviewMeshComponent = WeakPreviewMeshComponent.Get();
+		if (!PreviewMeshComponent)
+		{
+			continue;
+		}
+
+		if (UAnimSingleNodeInstance* PreviewInstance = GetPreviewInstance(PreviewMeshComponent))
+		{
+			return !PreviewInstance->GetCurrentAsset();
+		}
+	}
+	
+	return true;
+}
+
+
+void SMutableScrubPanel::OnValueChanged(float NewValue)
+{
+	for (TWeakObjectPtr<UDebugSkelMeshComponent>& WeakPreviewMeshComponent : GetPreviewScene()->GetPreviewMeshComponents())
+	{
+		UDebugSkelMeshComponent* PreviewMeshComponent = WeakPreviewMeshComponent.Get();
+		if (!PreviewMeshComponent)
+		{
+			continue;
+		}
+
+		if (UAnimSingleNodeInstance* PreviewInstance = GetPreviewInstance(PreviewMeshComponent))
+		{
+			PreviewInstance->SetPosition(NewValue);
+		}
+		else
+		{
+			UAnimInstance* Instance;
+			FAnimBlueprintDebugData* DebugData;
+			if (GetAnimBlueprintDebugData(/*out*/ Instance, /*out*/ DebugData))
+			{
+				DebugData->SetSnapshotIndexByTime(Instance, NewValue);
+			}
+		}
+	}
+}
+
+
+void SMutableScrubPanel::OnBeginSliderMovement()
+{
+	bSliderBeingDragged = true;
+
+	for (TWeakObjectPtr<UDebugSkelMeshComponent>& WeakPreviewMeshComponent : GetPreviewScene()->GetPreviewMeshComponents())
+	{
+		UDebugSkelMeshComponent* PreviewMeshComponent = WeakPreviewMeshComponent.Get();
+		if (!PreviewMeshComponent)
+		{
+			continue;
+		}
+
+		if (UAnimSingleNodeInstance* PreviewInstance = GetPreviewInstance(PreviewMeshComponent))
+		{
+			PreviewInstance->SetPlaying(false);
+		}
+	}	
+}
+
+
+void SMutableScrubPanel::OnEndSliderMovement(float NewValue)
+{
+	bSliderBeingDragged = false;
+}
+
+
+uint32 SMutableScrubPanel::GetNumberOfKeys() const
+{
+	for (TWeakObjectPtr<UDebugSkelMeshComponent>& WeakPreviewMeshComponent : GetPreviewScene()->GetPreviewMeshComponents())
+	{
+		UDebugSkelMeshComponent* PreviewMeshComponent = WeakPreviewMeshComponent.Get();
+		if (!PreviewMeshComponent)
+		{
+			continue;
+		}
+
+		UAnimSingleNodeInstance* PreviewInstance = GetPreviewInstance(PreviewMeshComponent);
+		if (!PreviewInstance)
+		{
+			continue;
+		}
+			
+		float Length = PreviewInstance->GetLength();
+		// if anim sequence, use correct num frames
+		int32 NumKeys = (int32) (Length/0.0333f); 
+		if (PreviewInstance->GetCurrentAsset())
+		{
+			if (PreviewInstance->GetCurrentAsset()->IsA(UAnimSequenceBase::StaticClass()))
+			{
+				NumKeys = CastChecked<UAnimSequenceBase>(PreviewInstance->GetCurrentAsset())->GetNumberOfSampledKeys();
+			}
+			else if(PreviewInstance->GetCurrentAsset()->IsA(UBlendSpace::StaticClass()))
+			{
+				// Blendspaces dont display frame notches, so just return 0 here
+				NumKeys = 0;
+			}
+		}
+		return NumKeys;
+	}
+
+	return 1;
+}
+
+
+float SMutableScrubPanel::GetSequenceLength() const
+{
+	for (TWeakObjectPtr<UDebugSkelMeshComponent>& WeakPreviewMeshComponent : GetPreviewScene()->GetPreviewMeshComponents())
+	{
+		UDebugSkelMeshComponent* PreviewMeshComponent = WeakPreviewMeshComponent.Get();
+		if (!PreviewMeshComponent)
+		{
+			continue;
+		}
+			
+		UAnimSingleNodeInstance* PreviewInstance = GetPreviewInstance(PreviewMeshComponent);
+		if (!PreviewInstance)
+		{
+			continue;
+		}
+			
+		return PreviewInstance->GetLength();
+	}
+
+	return 0.f;
+}
+
+
+void SMutableScrubPanel::Tick( const FGeometry& AllottedGeometry, const double InCurrentTime, const float InDeltaTime )
+{
+	if (bSliderBeingDragged)
+	{
+		GetPreviewScene()->Invalidate();
+	}
+}
+
+
+UAnimSingleNodeInstance* SMutableScrubPanel::GetPreviewInstance(UDebugSkelMeshComponent* PreviewMeshComponent) const
+{
+	return PreviewMeshComponent && PreviewMeshComponent->IsPreviewOn() ? PreviewMeshComponent->PreviewInstance : nullptr;
+}
+
+
+float SMutableScrubPanel::GetScrubValue() const
+{
+	for (TWeakObjectPtr<UDebugSkelMeshComponent>& WeakPreviewMeshComponent : GetPreviewScene()->GetPreviewMeshComponents())
+	{
+		UDebugSkelMeshComponent* PreviewMeshComponent = WeakPreviewMeshComponent.Get();
+		if (!PreviewMeshComponent)
+		{
+			continue;
+		}
+			
+		UAnimSingleNodeInstance* PreviewInstance = GetPreviewInstance(PreviewMeshComponent);
+		if (!PreviewInstance)
+		{
+			continue;
+		}
+			
+		return PreviewInstance->GetCurrentTime(); 
+	}
+	
+	return 0.f;
+}
+
+
+bool SMutableScrubPanel::GetAnimBlueprintDebugData(UAnimInstance*& Instance, FAnimBlueprintDebugData*& DebugInfo) const
+{
+	Instance = nullptr;
+
+	for (TWeakObjectPtr<UDebugSkelMeshComponent>& WeakPreviewMeshComponent : GetPreviewScene()->GetPreviewMeshComponents())
+	{
+		UDebugSkelMeshComponent* PreviewMeshComponent = WeakPreviewMeshComponent.Get();
+		if (!PreviewMeshComponent)
+		{
+			continue;
+		}
+		
+		UAnimInstance* AnimInstance = PreviewMeshComponent->GetAnimInstance();
+		if (AnimInstance && AnimInstance->GetClass()->ClassGeneratedBy)
+		{
+			Instance = AnimInstance;
+			break;
+		}
+	}
+
+	if (Instance)
+	{
+		// Avoid updating the instance if we're replaying the past
+		if (UAnimBlueprintGeneratedClass* AnimBlueprintClass = Cast<UAnimBlueprintGeneratedClass>(Instance->GetClass()))
+		{
+			if (UAnimBlueprint* Blueprint = Cast<UAnimBlueprint>(AnimBlueprintClass->ClassGeneratedBy))
+			{
+				if (Blueprint->GetObjectBeingDebugged() == Instance)
+				{
+					DebugInfo = &(AnimBlueprintClass->GetAnimBlueprintDebugData());
+					return true;
+				}
+			}
+		}
+	}
+
+	return false;
+}
+
+
+TSharedRef<FCustomizableObjectEditorViewportClient> SMutableScrubPanel::GetPreviewScene() const
+{
+	return PreviewScenePtr.Pin().ToSharedRef();
+}
+
+
+void SMutableScrubPanel::OnReZeroAnimSequence(int32 FrameIndex)
+{
+	for (TWeakObjectPtr<UDebugSkelMeshComponent>& WeakPreviewMeshComponent : GetPreviewScene()->GetPreviewMeshComponents())
+	{
+		UDebugSkelMeshComponent* PreviewSkelComp = WeakPreviewMeshComponent.Get();
+		if (!PreviewSkelComp)
+		{
+			continue;
+		}
+		
+		UAnimSingleNodeInstance* PreviewInstance = GetPreviewInstance(PreviewSkelComp);
+		if (!PreviewInstance)
+		{
+			continue;
+		}
+
+		if (PreviewInstance->GetCurrentAsset())
+		{
+			if (UAnimSequence* AnimSequence = Cast<UAnimSequence>( PreviewInstance->GetCurrentAsset()))
+			{
+				if (const USkeleton* Skeleton = AnimSequence->GetSkeleton())
+				{
+					const FName RootBoneName = Skeleton->GetReferenceSkeleton().GetBoneName(0);
+
+					if (AnimSequence->GetDataModel()->IsValidBoneTrackName(RootBoneName))
+					{
+						TArray<FVector3f> PosKeys;
+						TArray<FQuat4f> RotKeys;
+						TArray<FVector3f> ScaleKeys;
+
+						TArray<FTransform> BoneTransforms;
+						AnimSequence->GetDataModel()->GetBoneTrackTransforms(RootBoneName, BoneTransforms);
+
+						PosKeys.SetNum(BoneTransforms.Num());
+						RotKeys.SetNum(BoneTransforms.Num());
+						ScaleKeys.SetNum(BoneTransforms.Num());
+
+						// Find vector that would translate current root bone location onto origin.
+						FVector FrameTransform = FVector::ZeroVector;
+						if (FrameIndex == INDEX_NONE)
+						{
+							// Use current transform
+							FrameTransform = PreviewSkelComp->GetComponentSpaceTransforms()[0].GetLocation();
+						}
+						else if(BoneTransforms.IsValidIndex(FrameIndex))
+						{
+							// Use transform at frame
+							FrameTransform = BoneTransforms[FrameIndex].GetLocation();
+						}
+
+						FVector ApplyTranslation = -1.f * FrameTransform;
+
+						// Convert into world space
+						const FVector WorldApplyTranslation = PreviewSkelComp->GetComponentTransform().TransformVector(ApplyTranslation);
+						ApplyTranslation = PreviewSkelComp->GetComponentTransform().InverseTransformVector(WorldApplyTranslation);
+
+						for(int32 KeyIndex = 0; KeyIndex < BoneTransforms.Num(); KeyIndex++)
+						{
+							PosKeys[KeyIndex] = FVector3f(BoneTransforms[KeyIndex].GetLocation() + ApplyTranslation);
+							RotKeys[KeyIndex] = FQuat4f(BoneTransforms[KeyIndex].GetRotation());
+							ScaleKeys[KeyIndex] = FVector3f(BoneTransforms[KeyIndex].GetScale3D());
+						}
+
+						IAnimationDataController& Controller = AnimSequence->GetController();
+						Controller.SetBoneTrackKeys(RootBoneName, PosKeys, RotKeys, ScaleKeys);
+					}
+				}
+			}
+		}
+	}
+}
+
+
+bool SMutableScrubPanel::GetDisplayDrag() const
+{
+	for (TWeakObjectPtr<UDebugSkelMeshComponent>& WeakPreviewMeshComponent : GetPreviewScene()->GetPreviewMeshComponents())
+	{
+		UDebugSkelMeshComponent* PreviewMeshComponent = WeakPreviewMeshComponent.Get();
+		if (!PreviewMeshComponent)
+		{
+			continue;
+		}
+		
+		UAnimSingleNodeInstance* PreviewInstance = GetPreviewInstance(PreviewMeshComponent);
+		if (!PreviewInstance)
+		{
+			continue;
+		}
+
+		if (PreviewInstance->GetCurrentAsset())
+		{
+			return true;
+		}
+	}
+
+	return false;
+}
+
+
+#undef LOCTEXT_NAMESPACE
diff --git a/Engine/Plugins/Experimental/Mutable/Source/CustomizableObjectEditor/Private/MuCOE/SMutableScrubPanel.h b/Engine/Plugins/Experimental/Mutable/Source/CustomizableObjectEditor/Private/MuCOE/SMutableScrubPanel.h
new file mode 100644
--- /dev/null
+++ b/Engine/Plugins/Experimental/Mutable/Source/CustomizableObjectEditor/Private/MuCOE/SMutableScrubPanel.h
@@ -0,0 +1,92 @@
+﻿// Copyright Epic Games, Inc. All Rights Reserved.
+
+#pragma once
+
+#include "CoreMinimal.h"
+#include "Widgets/DeclarativeSyntaxSupport.h"
+#include "IPersonaPreviewScene.h"
+#include "Input/Reply.h"
+#include "Widgets/SCompoundWidget.h"
+#include "SScrubWidget.h"
+#include "ITransportControl.h"
+
+class FCustomizableObjectEditorViewportClient;
+class SScrubControlPanel;
+class UAnimationAsset;
+class UAnimInstance;
+class UAnimSequenceBase;
+class UDebugSkelMeshComponent;
+class UAnimSingleNodeInstance;
+struct FAnimBlueprintDebugData;
+
+
+class SMutableScrubPanel : public SCompoundWidget
+{
+public:
+	SLATE_BEGIN_ARGS(SMutableScrubPanel) : _LockedSequence() {}
+		SLATE_ARGUMENT(UAnimSequenceBase*, LockedSequence)
+		SLATE_ARGUMENT(bool, bDisplayAnimScrubBarEditing)
+		SLATE_EVENT(FOnSetInputViewRange, OnSetInputViewRange)
+		SLATE_ARGUMENT(bool, bAllowZoom)
+	SLATE_END_ARGS()
+
+	void Construct(const FArguments& InArgs, const TSharedRef<FCustomizableObjectEditorViewportClient>& InPreviewScene);
+
+	// SWidget interface
+	virtual void Tick(const FGeometry& AllottedGeometry, const double InCurrentTime, const float InDeltaTime) override;
+
+private:
+	float GetViewMinInput() const;
+
+	float GetViewMaxInput() const;
+	
+	FReply OnClick_Forward_Step();
+
+	FReply OnClick_Forward_End();
+
+	FReply OnClick_Backward_Step();
+	
+	FReply OnClick_Backward_End();
+
+	FReply OnClick_Forward();
+
+	FReply OnClick_Backward();
+
+	FReply OnClick_ToggleLoop();
+
+	void OnValueChanged(float NewValue);
+
+	void OnReZeroAnimSequence(int32 FrameIndex);
+
+	void OnBeginSliderMovement();
+	
+	void OnEndSliderMovement(float NewValue);
+
+	EPlaybackMode::Type GetPlaybackMode() const;
+
+	bool IsLoopStatusOn() const;
+	
+	bool IsRealtimeStreamingMode() const;
+	
+	float GetScrubValue() const;
+
+	UAnimSingleNodeInstance* GetPreviewInstance(UDebugSkelMeshComponent* PreviewMeshComponent) const;
+
+	TSharedPtr<SScrubControlPanel> ScrubControlPanel;
+
+	uint32 GetNumberOfKeys() const;
+
+	float GetSequenceLength() const;
+	
+	bool GetAnimBlueprintDebugData(UAnimInstance*& Instance, FAnimBlueprintDebugData*& DebugInfo) const;
+
+	TSharedRef<FCustomizableObjectEditorViewportClient> GetPreviewScene() const;
+
+	bool GetDisplayDrag() const;
+	
+	TWeakPtr<FCustomizableObjectEditorViewportClient> PreviewScenePtr;
+
+	FOnSetInputViewRange OnSetInputViewRange;
+
+	bool bSliderBeingDragged = false;
+};
diff --git a/Engine/Plugins/Experimental/Mutable/Source/CustomizableObjectEditor/Public/MuCOE/CustomizableObjectEditorActions.h b/Engine/Plugins/Experimental/Mutable/Source/CustomizableObjectEditor/Public/MuCOE/CustomizableObjectEditorActions.h
--- a/Engine/Plugins/Experimental/Mutable/Source/CustomizableObjectEditor/Public/MuCOE/CustomizableObjectEditorActions.h
+++ b/Engine/Plugins/Experimental/Mutable/Source/CustomizableObjectEditor/Public/MuCOE/CustomizableObjectEditorActions.h
@@ -60,10 +60,12 @@ public:
 	TSharedPtr<FUICommandInfo> SetShowNormals;
 	TSharedPtr<FUICommandInfo> SetShowTangents;
 	TSharedPtr<FUICommandInfo> SetShowBinormals;
-
+	
+	/** Command list for playback speed, indexed by EPlaybackSpeeds*/
+	TArray<TSharedPtr< FUICommandInfo >> PlaybackSpeedCommands;
+	
 	/**
 	 * Initialize commands
 	 */
 	virtual void RegisterCommands() override;
-
 };
diff --git a/Engine/Plugins/Experimental/Mutable/Source/CustomizableObjectEditor/Public/MuCOE/CustomizableObjectEditorViewportClient.h b/Engine/Plugins/Experimental/Mutable/Source/CustomizableObjectEditor/Public/MuCOE/CustomizableObjectEditorViewportClient.h
--- a/Engine/Plugins/Experimental/Mutable/Source/CustomizableObjectEditor/Public/MuCOE/CustomizableObjectEditorViewportClient.h
+++ b/Engine/Plugins/Experimental/Mutable/Source/CustomizableObjectEditor/Public/MuCOE/CustomizableObjectEditorViewportClient.h
@@ -64,6 +64,26 @@ enum class EWidgetType : uint8
 };
 
 
+namespace EMutableAnimationPlaybackSpeeds
+{
+	enum Type
+	{
+		OneTenth = 0,
+		Quarter,
+		Half,
+		ThreeQuarters,
+		Normal,
+		Double,
+		FiveTimes,
+		TenTimes,
+		Custom,
+		NumPlaybackSpeeds
+	};
+
+	extern float Values[NumPlaybackSpeeds];
+}
+
+
 /** Viewport which shows the scene with the generated Instance. */
 class FCustomizableObjectEditorViewportClient : public FEditorViewportClient, public TSharedFromThis<FCustomizableObjectEditorViewportClient>
 {
@@ -93,7 +113,11 @@ public:
 	// End of FEditorViewportClient
 	
 	void SetPreviewActor(const TWeakObjectPtr<AActor>& InActor, const TWeakObjectPtr<UCustomizableObjectInstance>& InInstance, const TArray<TWeakObjectPtr<UDebugSkelMeshComponent>>& InSkeletalMeshComponents);
-	
+
+	TArray<TWeakObjectPtr<UDebugSkelMeshComponent>>& GetPreviewMeshComponents();
+
+	void SetPreviewAnimationAsset(UAnimationAsset* AnimAsset);
+
 	/**
 	 *	Draws the UV overlay for the current LOD.
 	 *
@@ -188,11 +212,9 @@ public:
 
 	/** Do not call directly. Use ICustomizableObjectEditor functions instead. */
 	void HideGizmoLight();
-	
-	void SetAnimation(UAnimationAsset* Animation, EAnimationMode::Type AnimationType);
 
-	/** Set again the animation saved in AnimationBeingPlayed (if any) */
-	void ReSetAnimation();
+	/** Play the animation. */
+	void SetAnimation(UAnimationAsset* Animation);
 
 	/** Add Light Component to the scene. */
 	void AddLightToScene(ULightComponent* AddedLight);
@@ -284,6 +306,14 @@ public:
 
 	bool IsSetShowBinormalsChecked() const;
 	
+	void SetPlaybackSpeedMode(EMutableAnimationPlaybackSpeeds::Type InMode);
+
+	void SetCustomAnimationSpeed(float Speed);
+	
+	float GetCustomAnimationSpeed() const;
+
+	EMutableAnimationPlaybackSpeeds::Type GetPlaybackSpeedMode() const;
+
 private:
 	/** Draws Mesh Bones in foreground (From: FAnimationViewportClient) */
 	void DrawMeshBones(UDebugSkelMeshComponent* MeshComponent, FPrimitiveDrawInterface* PDI);
@@ -339,15 +369,12 @@ private:
 	/** Spawned light components. */
 	TArray<ULightComponent*> LightComponents;
 	
-	/** To know if an animation ia being played by the Customizable Object and restre it after compilation */
-	bool IsPlayingAnimation;
-
 	/** true if the camera has already being setup. */
 	bool bIsCameraSetup = false;
 
-	/** Animation being played by the Customizable Object, if any */
-	TObjectPtr<UAnimationAsset> AnimationBeingPlayed;
-
+	/** true if if has been updated after changing the actor. */
+	bool bUpdated = false;
+	
 	/** Customizable object being used */
 	UCustomizableObject* CustomizableObject;
 
@@ -411,6 +438,12 @@ private:
 	FWidgetTrackingStartedDelegate WidgetTrackingStartedDelegate;
 	
 	EWidgetType WidgetType = EWidgetType::Hidden;
+
+	/** Selected playback speed mode, used for deciding scale */
+	EMutableAnimationPlaybackSpeeds::Type AnimationPlaybackSpeedMode = EMutableAnimationPlaybackSpeeds::Normal;
+
+	/** Custom Animation speed in the viewport. Transient setting. */
+	float CustomAnimationSpeed = 1.0f;
 };
 
 #if UE_ENABLE_INCLUDE_ORDER_DEPRECATED_IN_5_2
diff --git a/Engine/Plugins/Experimental/Mutable/Source/CustomizableObjectEditor/Public/MuCOE/ICustomizableObjectInstanceEditor.h b/Engine/Plugins/Experimental/Mutable/Source/CustomizableObjectEditor/Public/MuCOE/ICustomizableObjectInstanceEditor.h
--- a/Engine/Plugins/Experimental/Mutable/Source/CustomizableObjectEditor/Public/MuCOE/ICustomizableObjectInstanceEditor.h
+++ b/Engine/Plugins/Experimental/Mutable/Source/CustomizableObjectEditor/Public/MuCOE/ICustomizableObjectInstanceEditor.h
@@ -18,6 +18,8 @@ class UCustomizableObjectNodeMeshClipWithMesh;
 class ULightComponent;
 class UPoseAsset;
 class UCustomizableObjectEditorProperties;
+class SCustomizableObjectEditorAdvancedPreviewSettings;
+
 
 /**
  * Public interface to Customizable Object Instance Editor
@@ -33,8 +35,6 @@ public:
 	/** Refreshes everything in the Customizable Object Instance Editor. */
 	virtual void RefreshTool() = 0;
 	
-	virtual void SetPoseAsset(UPoseAsset* PoseAssetParameter) {}
-
 	/** Return the selected projector in the viewport/editor. */
 	virtual UProjectorParameter* GetProjectorParameter() = 0;
 
@@ -94,6 +94,12 @@ public:
 
 	/** @return Editor properties. */
 	virtual UCustomizableObjectEditorProperties* GetEditorProperties() = 0;
+	
+	virtual TSharedPtr<SCustomizableObjectEditorAdvancedPreviewSettings> GetAdvancedPreviewSettings() = 0;
+
+	virtual bool ShowLightingSettings() = 0;
+
+	virtual bool ShowProfileManagementOptions() = 0;
 };
 
 
