[mutable] Remove arbitrary limit of LODs in COs.
#rnx

diff --git a/Engine/Plugins/Experimental/Mutable/Source/CustomizableObject/Internal/MuCO/CustomizableObjectPrivate.h b/Engine/Plugins/Experimental/Mutable/Source/CustomizableObject/Internal/MuCO/CustomizableObjectPrivate.h
--- a/Engine/Plugins/Experimental/Mutable/Source/CustomizableObject/Internal/MuCO/CustomizableObjectPrivate.h
+++ b/Engine/Plugins/Experimental/Mutable/Source/CustomizableObject/Internal/MuCO/CustomizableObjectPrivate.h
@@ -1315,6 +1315,8 @@ public:
 
 		LayoutBlocksAsInt32,
 
+		RemoveLODCountLimit,
+
 		// -----<new versions can be added above this line>--------
 		LastCustomizableObjectVersion
 	};
diff --git a/Engine/Plugins/Experimental/Mutable/Source/MutableRuntime/Internal/MuR/CodeVisitor.h b/Engine/Plugins/Experimental/Mutable/Source/MutableRuntime/Internal/MuR/CodeVisitor.h
--- a/Engine/Plugins/Experimental/Mutable/Source/MutableRuntime/Internal/MuR/CodeVisitor.h
+++ b/Engine/Plugins/Experimental/Mutable/Source/MutableRuntime/Internal/MuR/CodeVisitor.h
@@ -530,20 +530,27 @@ namespace mu
 
             case OP_TYPE::IN_ADDLOD:
             {
-				OP::InstanceAddLODArgs args = program.GetOpArgs<OP::InstanceAddLODArgs>(at);
-
                 recurse = false;
 
-                STATE newState = PARENT::GetCurrentState();
-                for (int32 t=0;t<MUTABLE_OP_MAX_ADD_COUNT;++t)
+				const uint8* Data = program.GetOpArgsPointer(at);
+
+				uint8 LODCount;
+				FMemory::Memcpy(&LODCount, Data, sizeof(uint8));
+				Data += sizeof(uint8);
+
+                STATE NewState = PARENT::GetCurrentState();
+                for (int8 LODIndex=0; LODIndex < LODCount;++LODIndex)
                 {
-                    OP::ADDRESS lodAt = args.lod[t];
-                    if (lodAt)
+					OP::ADDRESS LODAddress;
+					FMemory::Memcpy(&LODAddress, Data, sizeof(OP::ADDRESS));
+					Data += sizeof(OP::ADDRESS);
+					
+					if (LODAddress)
                     {
-                        bool selected = ( (1<<t) & m_lodMask ) != 0;
-                        if ( selected )
+                        bool bSelected = ( (1<< LODIndex) & m_lodMask ) != 0;
+                        if (bSelected)
                         {
-                            PARENT::RecurseWithState( lodAt, newState );
+                            PARENT::RecurseWithState(LODAddress, NewState);
                         }
                     }
                 }
diff --git a/Engine/Plugins/Experimental/Mutable/Source/MutableRuntime/Internal/MuR/Operations.h b/Engine/Plugins/Experimental/Mutable/Source/MutableRuntime/Internal/MuR/Operations.h
--- a/Engine/Plugins/Experimental/Mutable/Source/MutableRuntime/Internal/MuR/Operations.h
+++ b/Engine/Plugins/Experimental/Mutable/Source/MutableRuntime/Internal/MuR/Operations.h
@@ -11,7 +11,6 @@
 #include "Templates/Function.h"
 
 #define MUTABLE_OP_MAX_INTERPOLATE_COUNT	6
-#define MUTABLE_OP_MAX_ADD_COUNT			7
 #define MUTABLE_OP_MAX_SWIZZLE_CHANNELS		4
 
 namespace mu
@@ -993,8 +992,6 @@ namespace mu
 			ADDRESS source;
 		};
 
-
-        //-------------------------------------------------------------------------------------
         struct InstanceAddArgs
         {
             ADDRESS instance;
@@ -1013,11 +1010,6 @@ namespace mu
             ADDRESS relevantParametersListIndex;
         };
 
-        struct InstanceAddLODArgs
-        {
-            ADDRESS lod[ MUTABLE_OP_MAX_ADD_COUNT ];
-        };
-
 		struct InstanceAddExtensionDataArgs
 		{
 			// This is a reference to an op that produces the Instance that the ExtensionData will
@@ -1031,7 +1023,6 @@ namespace mu
 			ADDRESS ExtensionDataName;
 		};
 
-        //-------------------------------------------------------------------------------------
         struct LayoutPackArgs
         {
             ADDRESS Source;
diff --git a/Engine/Plugins/Experimental/Mutable/Source/MutableRuntime/Private/MuR/CodeRunner.cpp b/Engine/Plugins/Experimental/Mutable/Source/MutableRuntime/Private/MuR/CodeRunner.cpp
--- a/Engine/Plugins/Experimental/Mutable/Source/MutableRuntime/Private/MuR/CodeRunner.cpp
+++ b/Engine/Plugins/Experimental/Mutable/Source/MutableRuntime/Private/MuR/CodeRunner.cpp
@@ -700,26 +700,35 @@ namespace mu
 
         case OP_TYPE::IN_ADDLOD:
         {
-			OP::InstanceAddLODArgs args = Program.GetOpArgs<OP::InstanceAddLODArgs>(item.At);
-            switch (item.Stage)
+			const uint8* Data = Program.GetOpArgsPointer(item.At);
+
+			uint8 LODCount;
+			FMemory::Memcpy(&LODCount, Data, sizeof(uint8));
+			Data += sizeof(uint8);
+
+			switch (item.Stage)
             {
             case 0:
             {                
                 TArray<FScheduledOp> deps;
-                for ( int i=0; i<MUTABLE_OP_MAX_ADD_COUNT; ++i )
+                for (uint8 LODIndex=0; LODIndex <LODCount; ++LODIndex)
                 {
-                    if ( args.lod[i] )
+					OP::ADDRESS LODAddress;
+					FMemory::Memcpy(&LODAddress, Data, sizeof(OP::ADDRESS));
+					Data += sizeof(OP::ADDRESS);
+
+                    if (LODAddress)
                     {
-                        bool selectedLod = ( (1<<i) & lodMask ) != 0;
+                        bool bSelectedLod = ( (1<< LODIndex) & lodMask ) != 0;
 
-                        if ( selectedLod )
+                        if (bSelectedLod)
                         {
-                            deps.Emplace(args.lod[i], item);
+                            deps.Emplace(LODAddress, item);
                         }
                     }
                 }
 
-                AddOp( FScheduledOp( item.At,item, 1), deps );
+                AddOp( FScheduledOp( item.At, item, 1), deps );
 
                 break;
             }
@@ -730,25 +739,29 @@ namespace mu
 				Ptr<Instance> pResult = new Instance();
 				int32 ComponentIndex = pResult->GetPrivate()->AddComponent();
 
-                for ( int32 i=0; i<MUTABLE_OP_MAX_ADD_COUNT; ++i )
+                for (uint8 LODIndex = 0; LODIndex < LODCount; ++LODIndex)
                 {
-                    if ( args.lod[i] )
-                    {
-                        bool bIsSelectedLod = ( (1<<i) & lodMask ) != 0;
+					OP::ADDRESS LODAddress;
+					FMemory::Memcpy(&LODAddress, Data, sizeof(OP::ADDRESS));
+					Data += sizeof(OP::ADDRESS);
+					
+					if ( LODAddress )
+                    {						
+						bool bIsSelectedLod = ( (1<<LODIndex) & lodMask ) != 0;
 
 						// Add an empty LOD even if not selected.
-						int32 LODIndex = pResult->GetPrivate()->AddLOD(ComponentIndex);
+						int32 InstanceLODIndex = pResult->GetPrivate()->AddLOD(ComponentIndex);
 						
 						if (bIsSelectedLod)
                         {
-							Ptr<const Instance> pLOD = LoadInstance( FCacheAddress(args.lod[i],item) );
+							Ptr<const Instance> pLOD = LoadInstance( FCacheAddress(LODAddress,item) );
 
                             // In a degenerated case, the returned pLOD may not have an LOD inside
  							if (!pLOD->GetPrivate()->Components.IsEmpty()
 								&&
 								!pLOD->GetPrivate()->Components[0].LODs.IsEmpty())
 							{
-								pResult->GetPrivate()->Components[ComponentIndex].LODs[LODIndex] = pLOD->GetPrivate()->Components[0].LODs[0];
+								pResult->GetPrivate()->Components[ComponentIndex].LODs[InstanceLODIndex] = pLOD->GetPrivate()->Components[0].LODs[0];
 							}
 						}
                     }
diff --git a/Engine/Plugins/Experimental/Mutable/Source/MutableRuntime/Private/MuR/Operations.cpp b/Engine/Plugins/Experimental/Mutable/Source/MutableRuntime/Private/MuR/Operations.cpp
--- a/Engine/Plugins/Experimental/Mutable/Source/MutableRuntime/Private/MuR/Operations.cpp
+++ b/Engine/Plugins/Experimental/Mutable/Source/MutableRuntime/Private/MuR/Operations.cpp
@@ -428,18 +428,18 @@ namespace mu
 			FMemory::Memcpy( &DefAddress, data, sizeof(OP::ADDRESS));
 			data += sizeof(OP::ADDRESS);
 
-			uint32_t CaseCount;
-			FMemory::Memcpy( &CaseCount, data, sizeof(uint32_t));
-			data += sizeof(uint32_t);
+			uint32 CaseCount;
+			FMemory::Memcpy( &CaseCount, data, sizeof(uint32));
+			data += sizeof(uint32);
 
 			f(VarAddress);
 			f(DefAddress);
 
-			for ( uint32_t C = 0; C < CaseCount; ++C )
+			for ( uint32 C = 0; C < CaseCount; ++C )
 			{	
-				//int32_t Condition;
+				//int32 Condition;
 				//FMemory::Memcpy( &Condition, data, sizeof(int32_t));
-				data += sizeof(int32_t);
+				data += sizeof(int32);
 				
 				OP::ADDRESS At;
 				FMemory::Memcpy( &At, data, sizeof(OP::ADDRESS) );
@@ -1022,10 +1022,19 @@ namespace mu
 
         case OP_TYPE::IN_ADDLOD:
         {
-			OP::InstanceAddLODArgs args = program.GetOpArgs<OP::InstanceAddLODArgs>(at);
-            for (int t=0;t<MUTABLE_OP_MAX_ADD_COUNT;++t)
-            {
-                f(args.lod[t] );
+			const uint8* Data = program.GetOpArgsPointer(at);
+
+			uint8 LODCount;
+			FMemory::Memcpy(&LODCount, Data, sizeof(uint8));
+			Data += sizeof(uint8);
+
+			for (int8 LODIndex = 0; LODIndex < LODCount; ++LODIndex)
+			{
+				OP::ADDRESS LODAddress;
+				FMemory::Memcpy(&LODAddress, Data, sizeof(OP::ADDRESS));
+				Data += sizeof(OP::ADDRESS);
+
+				f(LODAddress);
             }
             break;
         }
diff --git a/Engine/Plugins/Experimental/Mutable/Source/MutableTools/Private/MuT/ASTOpAddLOD.cpp b/Engine/Plugins/Experimental/Mutable/Source/MutableTools/Private/MuT/ASTOpAddLOD.cpp
--- a/Engine/Plugins/Experimental/Mutable/Source/MutableTools/Private/MuT/ASTOpAddLOD.cpp
+++ b/Engine/Plugins/Experimental/Mutable/Source/MutableTools/Private/MuT/ASTOpAddLOD.cpp
@@ -73,27 +73,41 @@ namespace mu
 		// Already linked?
 		if (!linkedAddress)
 		{
-			OP::InstanceAddLODArgs args;
-			memset(&args, 0, sizeof(args));
+			linkedAddress = (OP::ADDRESS)program.m_opAddress.Num();
+			program.m_opAddress.Add((uint32)program.m_byteCode.Num());
+			AppendCode(program.m_byteCode, OP_TYPE::IN_ADDLOD);
 
-			int i = 0;
-			for (auto& l : lods)
+			// Calculate LODCount mimicking previous behaviour.
+			// It may be better to leave null LODs if they are null in the AST.
+			uint8 LODCount = 0;
+			for (const ASTChild& LOD : lods)
 			{
-				if (l)
+				OP::ADDRESS LODAddress = 0;
+				if (LOD)
 				{
-					args.lod[i] = l->linkedAddress;
-					++i;
-					if (i >= MUTABLE_OP_MAX_ADD_COUNT)
+					++LODCount;
+					check(LODCount < 255);
+					if (LODCount == 255)
 					{
 						break;
 					}
 				}
 			}
 
-			linkedAddress = (OP::ADDRESS)program.m_opAddress.Num();
-			program.m_opAddress.Add((uint32_t)program.m_byteCode.Num());
-			AppendCode(program.m_byteCode, OP_TYPE::IN_ADDLOD);
-			AppendCode(program.m_byteCode, args);
+			AppendCode(program.m_byteCode, LODCount);
+			for (const ASTChild& LOD : lods)
+			{
+				OP::ADDRESS LODAddress = 0;
+				if (LOD)
+				{
+					LODAddress = LOD->linkedAddress;
+					AppendCode(program.m_byteCode, LODAddress);
+					if (LODCount == 255)
+					{
+						break;
+					}
+				}
+			}
 		}
 
 	}
diff --git a/Engine/Plugins/Experimental/Mutable/Source/MutableTools/Private/MuT/ASTOpAddLOD.h b/Engine/Plugins/Experimental/Mutable/Source/MutableTools/Private/MuT/ASTOpAddLOD.h
--- a/Engine/Plugins/Experimental/Mutable/Source/MutableTools/Private/MuT/ASTOpAddLOD.h
+++ b/Engine/Plugins/Experimental/Mutable/Source/MutableTools/Private/MuT/ASTOpAddLOD.h
@@ -13,9 +13,7 @@ namespace mu
 struct FProgram;
 
 
-	//---------------------------------------------------------------------------------------------
-	//! Operations to add elements to a LOD
-	//---------------------------------------------------------------------------------------------
+	/** Operations to add elements to a LOD. */
 	class ASTOpAddLOD final : public ASTOp
 	{
 	public:
