[mutable] Cleanup modifier nodes and remove unnecessary PatchImage and PatchMesh nodes.
#rnx

diff --git a/Engine/Plugins/Experimental/Mutable/Source/CustomizableObjectEditor/Private/MuCOE/GenerateMutableSource/GenerateMutableSource.cpp b/Engine/Plugins/Experimental/Mutable/Source/CustomizableObjectEditor/Private/MuCOE/GenerateMutableSource/GenerateMutableSource.cpp
--- a/Engine/Plugins/Experimental/Mutable/Source/CustomizableObjectEditor/Private/MuCOE/GenerateMutableSource/GenerateMutableSource.cpp
+++ b/Engine/Plugins/Experimental/Mutable/Source/CustomizableObjectEditor/Private/MuCOE/GenerateMutableSource/GenerateMutableSource.cpp
@@ -291,7 +291,7 @@ void FMutableGraphGenerationContext::GenerateClippingCOInternalTags()
 
 		for (j = 0; j < It->Value.Num(); ++j)
 		{
-			It->Value[j]->AddTag(TagName);
+			It->Value[j]->RequiredTags.Add(TagName);
 		}
 
 		UCustomizableObjectNodeMeshClipWithMesh* CustomizableObjectNodeMeshClipWithMesh = It->Key;
diff --git a/Engine/Plugins/Experimental/Mutable/Source/CustomizableObjectEditor/Private/MuCOE/GenerateMutableSource/GenerateMutableSourceModifier.cpp b/Engine/Plugins/Experimental/Mutable/Source/CustomizableObjectEditor/Private/MuCOE/GenerateMutableSource/GenerateMutableSourceModifier.cpp
--- a/Engine/Plugins/Experimental/Mutable/Source/CustomizableObjectEditor/Private/MuCOE/GenerateMutableSource/GenerateMutableSourceModifier.cpp
+++ b/Engine/Plugins/Experimental/Mutable/Source/CustomizableObjectEditor/Private/MuCOE/GenerateMutableSource/GenerateMutableSourceModifier.cpp
@@ -51,22 +51,22 @@ mu::Ptr<mu::NodeModifier> GenerateMutableSourceModifier(const UEdGraphPin * Pin,
 
 	if (const UCustomizableObjectNodeMeshClipMorph* TypedNodeClip = Cast<UCustomizableObjectNodeMeshClipMorph>(Node))
 	{
-		mu::NodeModifierMeshClipMorphPlanePtr ClipNode = new mu::NodeModifierMeshClipMorphPlane();
+		mu::Ptr<mu::NodeModifierMeshClipMorphPlane> ClipNode = new mu::NodeModifierMeshClipMorphPlane();
 		Result = ClipNode;
 
 		const FVector Origin = TypedNodeClip->GetOriginWithOffset();
 		const FVector& Normal = TypedNodeClip->Normal;
 
-		ClipNode->SetPlane(Origin.X, Origin.Y, Origin.Z, Normal.X, Normal.Y, Normal.Z);
+		ClipNode->SetPlane(FVector3f(Origin), FVector3f(Normal));
 		ClipNode->SetParams(TypedNodeClip->B, TypedNodeClip->Exponent);
 		ClipNode->SetMorphEllipse(TypedNodeClip->Radius, TypedNodeClip->Radius2, TypedNodeClip->RotationAngle);
 
 		ClipNode->SetVertexSelectionBone(GenerationContext.GetBoneUnique(TypedNodeClip->BoneName), TypedNodeClip->MaxEffectRadius);
 
-		ClipNode->SetMultipleTagPolicy( TypedNodeClip->MultipleTagPolicy );
+		ClipNode->MultipleTagsPolicy = TypedNodeClip->MultipleTagPolicy;
 		for (const FString& Tag : TypedNodeClip->Tags)
 		{
-			 ClipNode->AddTag(Tag);
+			 ClipNode->RequiredTags.Add(Tag);
 		}
 	}
 
@@ -78,9 +78,9 @@ mu::Ptr<mu::NodeModifier> GenerateMutableSourceModifier(const UEdGraphPin * Pin,
 		if (const UEdGraphPin* ConnectedPin = FollowInputPin(*TypedNodeClipDeform->ClipShapePin()))
 		{
 			FMutableGraphMeshGenerationData DummyMeshData;
-			mu::NodeMeshPtr ClipMesh = GenerateMutableSourceMesh(ConnectedPin, GenerationContext, DummyMeshData, false, true);
+			mu::Ptr<mu::NodeMesh> ClipMesh = GenerateMutableSourceMesh(ConnectedPin, GenerationContext, DummyMeshData, false, true);
 
-			ClipNode->SetClipMesh(ClipMesh.get());
+			ClipNode->ClipMesh = ClipMesh;
 
 			mu::EShapeBindingMethod BindingMethod = mu::EShapeBindingMethod::ClipDeformClosestProject;
 			switch(TypedNodeClipDeform->BindingMethod)
@@ -99,7 +99,7 @@ mu::Ptr<mu::NodeModifier> GenerateMutableSourceModifier(const UEdGraphPin * Pin,
 					break;
 			}
 
-			ClipNode->SetBindingMethod(BindingMethod);
+			ClipNode->BindingMethod = BindingMethod;
 		}
 		else
 		{
@@ -108,10 +108,10 @@ mu::Ptr<mu::NodeModifier> GenerateMutableSourceModifier(const UEdGraphPin * Pin,
 			Result = nullptr;
 		}
 	
-		ClipNode->SetMultipleTagPolicy(TypedNodeClipDeform->MultipleTagPolicy);
+		ClipNode->MultipleTagsPolicy = TypedNodeClipDeform->MultipleTagPolicy;
 		for (const FString& Tag : TypedNodeClipDeform->Tags)
 		{
-			ClipNode->AddTag(Tag);
+			ClipNode->RequiredTags.Add(Tag);
 		}		
 	}
 
@@ -165,7 +165,7 @@ mu::Ptr<mu::NodeModifier> GenerateMutableSourceModifier(const UEdGraphPin * Pin,
 				ClipMesh = TransformMesh;
 			}
 
-			ClipNode->SetClipMesh(ClipMesh.get());
+			ClipNode->ClipMesh = ClipMesh;
 		}
 		else
 		{
@@ -174,10 +174,10 @@ mu::Ptr<mu::NodeModifier> GenerateMutableSourceModifier(const UEdGraphPin * Pin,
 			Result = nullptr;
 		}
 
-		ClipNode->SetMultipleTagPolicy(TypedNodeClipMesh->MultipleTagPolicy);
+		ClipNode->MultipleTagsPolicy = TypedNodeClipMesh->MultipleTagPolicy;
 		for (const FString& Tag : TypedNodeClipMesh->Tags)
 		{
-			 ClipNode->AddTag(Tag);
+			 ClipNode->RequiredTags.Add(Tag);
 		}
 
 		if (TypedNodeClipMesh->CustomizableObjectToClipWith != nullptr)
@@ -213,7 +213,7 @@ mu::Ptr<mu::NodeModifier> GenerateMutableSourceModifier(const UEdGraphPin * Pin,
 
 			mu::Ptr<mu::NodeImage> ClipMask = GenerateMutableSourceImage(ConnectedPin, GenerationContext, 0);
 
-			ClipNode->SetClipMask(ClipMask.get());
+			ClipNode->ClipMask = ClipMask;
 		}
 		else
 		{
@@ -222,12 +222,12 @@ mu::Ptr<mu::NodeModifier> GenerateMutableSourceModifier(const UEdGraphPin * Pin,
 			Result = nullptr;
 		}
 
-		ClipNode->SetLayoutIndex(TypedNodeClipUVMask->UVChannelForMask);
+		ClipNode->LayoutIndex = TypedNodeClipUVMask->UVChannelForMask;
 
-		ClipNode->SetMultipleTagPolicy(TypedNodeClipUVMask->MultipleTagPolicy);
+		ClipNode->MultipleTagsPolicy = TypedNodeClipUVMask->MultipleTagPolicy;
 		for (const FString& Tag : TypedNodeClipUVMask->Tags)
 		{
-			ClipNode->AddTag(Tag);
+			ClipNode->RequiredTags.Add(Tag);
 		}
 	}
 
diff --git a/Engine/Plugins/Experimental/Mutable/Source/CustomizableObjectEditor/Private/MuCOE/GenerateMutableSource/GenerateMutableSourceSurface.cpp b/Engine/Plugins/Experimental/Mutable/Source/CustomizableObjectEditor/Private/MuCOE/GenerateMutableSource/GenerateMutableSourceSurface.cpp
--- a/Engine/Plugins/Experimental/Mutable/Source/CustomizableObjectEditor/Private/MuCOE/GenerateMutableSource/GenerateMutableSourceSurface.cpp
+++ b/Engine/Plugins/Experimental/Mutable/Source/CustomizableObjectEditor/Private/MuCOE/GenerateMutableSource/GenerateMutableSourceSurface.cpp
@@ -49,8 +49,6 @@
 #include "MuT/NodeMeshConstant.h"
 #include "MuT/NodeMeshFormat.h"
 #include "MuT/NodeMeshFragment.h"
-#include "MuT/NodePatchImage.h"
-#include "MuT/NodePatchMesh.h"
 #include "MuT/NodeScalarConstant.h"
 #include "MuT/NodeSurfaceEdit.h"
 #include "MuT/NodeSurfaceSwitch.h"
@@ -1146,11 +1144,7 @@ mu::Ptr<mu::NodeSurface> GenerateMutableSourceSurface(const UEdGraphPin * Pin, F
 
 				MeshFormat->SetSource(MeshPtr.get());
 
-				mu::NodePatchMeshPtr MeshPatch = new mu::NodePatchMesh();
-				MeshPatch->SetAdd(MeshFormat.get());
-
-				SurfNode->Mesh = MeshPatch;
-				MeshPatch->SetMessageContext(Node);
+				SurfNode->MeshAdd = MeshFormat;
 			}
 			
 			const int32 NumImages = ParentMaterialNode->GetNumParameters(EMaterialParameterType::Texture);
@@ -1216,7 +1210,7 @@ mu::Ptr<mu::NodeSurface> GenerateMutableSourceSurface(const UEdGraphPin * Pin, F
 		
 			for (const FString& Tag : TypedNodeExt->Tags)
 			{
-				SurfNode->Tags.Add(Tag);
+				SurfNode->EnableTags.Add(Tag);
 			}
 		}();
 	}
@@ -1234,10 +1228,10 @@ mu::Ptr<mu::NodeSurface> GenerateMutableSourceSurface(const UEdGraphPin * Pin, F
 		else
 		{
 			// Parent, probably generated, will be retrieved from the cache
-			mu::NodeSurfacePtr ParentNode = GenerateMutableSourceSurface(ParentMaterialNode->OutputPin(), GenerationContext);
+			mu::Ptr<mu::NodeSurface> ParentNode = GenerateMutableSourceSurface(ParentMaterialNode->OutputPin(), GenerationContext);
 			SurfNode->Parent = ParentNode;
 
-			mu::NodeMeshPtr RemoveMeshNode;
+			mu::Ptr<mu::NodeMesh> RemoveMeshNode;
 
 			if (const UEdGraphPin* ConnectedPin = FollowInputPin(*TypedNodeRem->RemoveMeshPin()))
 			{
@@ -1245,13 +1239,7 @@ mu::Ptr<mu::NodeSurface> GenerateMutableSourceSurface(const UEdGraphPin * Pin, F
 				RemoveMeshNode = GenerateMutableSourceMesh(ConnectedPin, GenerationContext, DummyMeshData, false, true);
 			}
 
-			if (RemoveMeshNode)
-			{
-				mu::NodePatchMeshPtr MeshPatch = new mu::NodePatchMesh();
-				MeshPatch->SetRemove(RemoveMeshNode.get());
-				SurfNode->Mesh = MeshPatch;
-				MeshPatch->SetMessageContext(Node);
-			}
+			SurfNode->MeshRemove = RemoveMeshNode;
 
 			AddModifierToSharedSurface(GenerationContext, ParentMaterialNode, *TypedNodeRem);
 		}
@@ -1290,10 +1278,7 @@ mu::Ptr<mu::NodeSurface> GenerateMutableSourceSurface(const UEdGraphPin * Pin, F
 				MeshFrag->LayoutIndex = TypedNodeRemBlocks->ParentLayoutIndex;
 				MeshFrag->SetMessageContext(TypedNodeRemBlocks);
 
-				mu::Ptr<mu::NodePatchMesh> MeshPatch = new mu::NodePatchMesh();
-				MeshPatch->SetRemove(MeshFrag.get());
-				MeshPatch->SetMessageContext(Node);
-				SurfNode->Mesh = MeshPatch;
+				SurfNode->MeshRemove = MeshFrag;
 			}
 
 			AddModifierToSharedSurface(GenerationContext, ParentMaterialNode, *TypedNodeRemBlocks);
@@ -1328,39 +1313,36 @@ mu::Ptr<mu::NodeSurface> GenerateMutableSourceSurface(const UEdGraphPin * Pin, F
 					
 					const UEdGraphPin* ConnectedImagePin = FollowInputPin(*TypedNodeEdit->GetUsedImagePin(ImageId));
 					
-					mu::Ptr<mu::NodePatchImage> ImagePatchNode = new mu::NodePatchImage;
-					ImagePatchNode->SetMessageContext(Node);
+					mu::NodeSurfaceEdit::FTexture& ImagePatch = SurfNode->Textures[ImageIndex];
 
 					// \todo: expose these two options?
-					ImagePatchNode->BlendType = mu::EBlendType::BT_BLEND;
-					ImagePatchNode->bApplyToAlpha=true;
+					ImagePatch.PatchBlendType = mu::EBlendType::BT_BLEND;
+					ImagePatch.bPatchApplyToAlpha=true;
 
 					// ReferenceTextureSize is used to limit the size of textures contributing to the final image.
 					const int32 ReferenceTextureSize = GetBaseTextureSize(GenerationContext, ParentMaterialNode, ImageIndex);
 
-					ImagePatchNode->Image = GenerateMutableSourceImage(ConnectedImagePin, GenerationContext, ReferenceTextureSize);
+					ImagePatch.PatchImage = GenerateMutableSourceImage(ConnectedImagePin, GenerationContext, ReferenceTextureSize);
 
 					const UEdGraphPin* ImageMaskPin = TypedNodeEdit->GetUsedImageMaskPin(ImageId);
 					check(ImageMaskPin); // Ensured when reconstructing EditMaterial nodes. If it fails, something is wrong.
 					
 					if (const UEdGraphPin* ConnectedMaskPin = FollowInputPin(*ImageMaskPin))
 					{
-						ImagePatchNode->Mask = GenerateMutableSourceImage(ConnectedMaskPin, GenerationContext, ReferenceTextureSize);
+						ImagePatch.PatchMask = GenerateMutableSourceImage(ConnectedMaskPin, GenerationContext, ReferenceTextureSize);
 					}
 
 					// Add the blocks to patch
 					FIntPoint GridSize = TypedNodeEdit->Layout->GetGridSize();
 					FVector2f GridSizeF = FVector2f(GridSize);
-					ImagePatchNode->Blocks.Reserve(TypedNodeEdit->Layout->Blocks.Num());
+					ImagePatch.PatchBlocks.Reserve(TypedNodeEdit->Layout->Blocks.Num());
 					for (const FCustomizableObjectLayoutBlock& LayoutBlock: TypedNodeEdit->Layout->Blocks)
 					{
 						FBox2f Rect;
 						Rect.Min = FVector2f(LayoutBlock.Min) / GridSizeF;
 						Rect.Max = FVector2f(LayoutBlock.Max) / GridSizeF;
-						ImagePatchNode->Blocks.Add(Rect);
+						ImagePatch.PatchBlocks.Add(Rect);
 					}
-
-					SurfNode->Textures[ImageIndex].Patch = ImagePatchNode;
 				}
 
 				AddModifierToSharedSurface(GenerationContext, ParentMaterialNode, *TypedNodeEdit);
@@ -1381,7 +1363,7 @@ mu::Ptr<mu::NodeSurface> GenerateMutableSourceSurface(const UEdGraphPin * Pin, F
 		else
 		{
 			// Parent, probably generated, will be retrieved from the cache
-			mu::NodeSurfacePtr ParentNode = GenerateMutableSourceSurface(ParentMaterialNode->OutputPin(), GenerationContext);
+			mu::Ptr<mu::NodeSurface> ParentNode = GenerateMutableSourceSurface(ParentMaterialNode->OutputPin(), GenerationContext);
 			SurfNode->Parent = ParentNode;
 
 			const UEdGraphPin* BaseSourcePin = FindMeshBaseSource(*ParentMaterialNode->OutputPin(), false);
@@ -1391,7 +1373,7 @@ mu::Ptr<mu::NodeSurface> GenerateMutableSourceSurface(const UEdGraphPin * Pin, F
 			}
 			else
 			{
-				//Morph Mesh
+				// Morph Mesh
 				// The Mesh Base Source from the Parent Material will always have all LODs since it is not an auxiliary Mesh (clipping, reshape...).
 				mu::MeshPtr MorphedSourceMesh = BuildMorphedMutableMesh(BaseSourcePin, TypedNodeMorph->MorphTargetName, GenerationContext, false); 
 
@@ -1399,7 +1381,7 @@ mu::Ptr<mu::NodeSurface> GenerateMutableSourceSurface(const UEdGraphPin * Pin, F
 				MorphedSourceMeshNode->SetMessageContext(Node);					
 				MorphedSourceMeshNode->SetValue(MorphedSourceMesh);
 				
-				SurfNode->Morph = MorphedSourceMeshNode;
+				SurfNode->MeshMorph = MorphedSourceMeshNode;
 				
 				if (const UEdGraphPin* ConnectedPin = FollowInputPin(*TypedNodeMorph->FactorPin()))
 				{
diff --git a/Engine/Plugins/Experimental/Mutable/Source/CustomizableObjectEditor/Private/MuCOE/SMutableGraphViewer.cpp b/Engine/Plugins/Experimental/Mutable/Source/CustomizableObjectEditor/Private/MuCOE/SMutableGraphViewer.cpp
--- a/Engine/Plugins/Experimental/Mutable/Source/CustomizableObjectEditor/Private/MuCOE/SMutableGraphViewer.cpp
+++ b/Engine/Plugins/Experimental/Mutable/Source/CustomizableObjectEditor/Private/MuCOE/SMutableGraphViewer.cpp
@@ -52,6 +52,7 @@
 #include "MuT/NodeModifierMeshClipDeform.h"
 #include "MuT/NodeModifierMeshClipMorphPlane.h"
 #include "MuT/NodeModifierMeshClipWithUVMask.h"
+#include "MuT/NodeModifierMeshClipMorphPlane.h"
 #include "MuT/NodeScalarConstant.h"
 #include "MuT/NodeScalarCurve.h"
 #include "MuT/NodeScalarSwitch.h"
@@ -59,21 +60,14 @@
 #include "MuT/NodeObjectNew.h"
 
 #include "MuT/NodeObjectGroupPrivate.h"
-#include "MuT/NodeModifierPrivate.h"
 #include "MuT/NodeImageFormatPrivate.h"
 #include "MuT/NodeMeshFormatPrivate.h"
-#include "MuT/NodePatchImage.h"
 #include "MuT/NodeMeshConstantPrivate.h"
-#include "MuT/NodeModifierMeshClipMorphPlanePrivate.h"
-#include "MuT/NodePatchMeshPrivate.h"
 #include "MuT/NodeImageSwitchPrivate.h"
 #include "MuT/NodeImageLayerColourPrivate.h"
 #include "MuT/NodeImageLayerPrivate.h"
 #include "MuT/NodeImageMipmapPrivate.h"
 #include "MuT/NodeImageResizePrivate.h"
-#include "MuT/NodeModifierMeshClipDeformPrivate.h"
-#include "MuT/NodeModifierMeshClipWithMeshPrivate.h"
-#include "MuT/NodeModifierMeshClipWithUVMaskPrivate.h"
 #include "MuT/NodeImageInterpolatePrivate.h"
 #include "MuT/NodeImagePlainColourPrivate.h"
 #include "MuT/NodeImageProjectPrivate.h"
@@ -369,14 +363,17 @@ void SMutableGraphViewer::GetChildrenForInfo(TSharedPtr<FMutableGraphTreeElement
 	else if (ParentNode->GetType() == mu::NodeSurfaceEdit::GetStaticType())
 	{
 		mu::NodeSurfaceEdit* SurfaceEdit = StaticCast<mu::NodeSurfaceEdit*>(ParentNode);
-		AddChildFunc(SurfaceEdit->Mesh.get(), TEXT("MESH"));
-		AddChildFunc(SurfaceEdit->Morph.get(), TEXT("MORPH"));
+		AddChildFunc(SurfaceEdit->MeshAdd.get(), TEXT("MESH_ADD"));
+		AddChildFunc(SurfaceEdit->MeshRemove.get(), TEXT("MESH_REMOVE"));
+		AddChildFunc(SurfaceEdit->MeshMorph.get(), TEXT("MORPH"));
 		AddChildFunc(SurfaceEdit->MorphFactor.get(), TEXT("MORPH_FACTOR"));
 
 		for (int32 l = 0; l < SurfaceEdit->Textures.Num(); ++l)
 		{
 			AddChildFunc(SurfaceEdit->Textures[l].Extend.get(), FString::Printf(TEXT("EXTEND [%d]"), l));
-			AddChildFunc(SurfaceEdit->Textures[l].Patch.get(), FString::Printf(TEXT("PATCH [%d]"), l));
+			AddChildFunc(SurfaceEdit->Textures[l].PatchImage.get(), FString::Printf(TEXT("PATCH IMAGE [%d]"), l));
+			AddChildFunc(SurfaceEdit->Textures[l].PatchMask.get(), FString::Printf(TEXT("PATCH MASK [%d]"), l));
+
 		}
 	}
 
@@ -472,13 +469,6 @@ void SMutableGraphViewer::GetChildrenForInfo(TSharedPtr<FMutableGraphTreeElement
 		AddChildFunc(Private->Source.get(), FString::Printf(TEXT("SOURCE MESH")));
 	}
 
-	else if (ParentNode->GetType() == mu::NodePatchImage::GetStaticType())
-	{
-		mu::NodePatchImage* PatchImageVar = StaticCast<mu::NodePatchImage*>(ParentNode);
-		AddChildFunc(PatchImageVar->Image.get(), FString::Printf(TEXT("IMAGE")));
-		AddChildFunc(PatchImageVar->Mask.get(), FString::Printf(TEXT("MASK")));
-	}
-
 	else if (ParentNode->GetType() == mu::NodeModifierMeshClipMorphPlane::GetStaticType())
 	{
 		// Nothing to show
@@ -487,30 +477,19 @@ void SMutableGraphViewer::GetChildrenForInfo(TSharedPtr<FMutableGraphTreeElement
 	else if (ParentNode->GetType() == mu::NodeModifierMeshClipWithMesh::GetStaticType())
 	{
 		mu::NodeModifierMeshClipWithMesh* ModifierMeshClipWithMeshVar = StaticCast<mu::NodeModifierMeshClipWithMesh*>(ParentNode);
-		mu::NodeModifierMeshClipWithMesh::Private* Private = ModifierMeshClipWithMeshVar->GetPrivate();
-		AddChildFunc(Private->ClipMesh.get(), FString::Printf(TEXT("CLIP MESH")));
+		AddChildFunc(ModifierMeshClipWithMeshVar->ClipMesh.get(), FString::Printf(TEXT("CLIP MESH")));
 	}
 
 	else if (ParentNode->GetType() == mu::NodeModifierMeshClipDeform::GetStaticType())
 	{
 		mu::NodeModifierMeshClipDeform* ModifierMeshClipDeformVar = StaticCast<mu::NodeModifierMeshClipDeform*>(ParentNode);
-		mu::NodeModifierMeshClipDeform::Private* Private = ModifierMeshClipDeformVar->GetPrivate();
-		AddChildFunc(Private->ClipMesh.get(), FString::Printf(TEXT("CLIP MESH")));
+		AddChildFunc(ModifierMeshClipDeformVar->ClipMesh.get(), FString::Printf(TEXT("CLIP MESH")));
 	}
 
 	else if (ParentNode->GetType() == mu::NodeModifierMeshClipWithUVMask::GetStaticType())
 	{
 		mu::NodeModifierMeshClipWithUVMask* ModifierMeshClipWithUVMaskVar = StaticCast<mu::NodeModifierMeshClipWithUVMask*>(ParentNode);
-		mu::NodeModifierMeshClipWithUVMask::Private* Private = ModifierMeshClipWithUVMaskVar->GetPrivate();
-		AddChildFunc(Private->ClipMask.get(), FString::Printf(TEXT("CLIP MASK")));
-	}
-	
-	else if (ParentNode->GetType() == mu::NodePatchMesh::GetStaticType())
-	{
-		mu::NodePatchMesh* PatchMeshVar = StaticCast<mu::NodePatchMesh*>(ParentNode);
-		mu::NodePatchMesh::Private* Private = PatchMeshVar->GetPrivate();
-		AddChildFunc(Private->m_pAdd.get(), FString::Printf(TEXT("ADD")));
-		AddChildFunc(Private->m_pRemove.get(), FString::Printf(TEXT("REMOVE")));
+		AddChildFunc(ModifierMeshClipWithUVMaskVar->ClipMask.get(), FString::Printf(TEXT("CLIP MASK")));
 	}
 
 	else if (ParentNode->GetType() == mu::NodeImageSwitch::GetStaticType())
diff --git a/Engine/Plugins/Experimental/Mutable/Source/MutableTools/Internal/MuT/NodeModifierMeshClipDeformPrivate.h b/Engine/Plugins/Experimental/Mutable/Source/MutableTools/Internal/MuT/NodeModifierMeshClipDeformPrivate.h
deleted file mode 100644
diff --git a/Engine/Plugins/Experimental/Mutable/Source/MutableTools/Internal/MuT/NodeModifierMeshClipMorphPlanePrivate.h b/Engine/Plugins/Experimental/Mutable/Source/MutableTools/Internal/MuT/NodeModifierMeshClipMorphPlanePrivate.h
deleted file mode 100644
diff --git a/Engine/Plugins/Experimental/Mutable/Source/MutableTools/Internal/MuT/NodeModifierMeshClipWithMeshPrivate.h b/Engine/Plugins/Experimental/Mutable/Source/MutableTools/Internal/MuT/NodeModifierMeshClipWithMeshPrivate.h
deleted file mode 100644
diff --git a/Engine/Plugins/Experimental/Mutable/Source/MutableTools/Internal/MuT/NodeModifierMeshClipWithUVMaskPrivate.h b/Engine/Plugins/Experimental/Mutable/Source/MutableTools/Internal/MuT/NodeModifierMeshClipWithUVMaskPrivate.h
deleted file mode 100644
diff --git a/Engine/Plugins/Experimental/Mutable/Source/MutableTools/Internal/MuT/NodeModifierPrivate.h b/Engine/Plugins/Experimental/Mutable/Source/MutableTools/Internal/MuT/NodeModifierPrivate.h
deleted file mode 100644
diff --git a/Engine/Plugins/Experimental/Mutable/Source/MutableTools/Internal/MuT/NodePatchImage.h b/Engine/Plugins/Experimental/Mutable/Source/MutableTools/Internal/MuT/NodePatchImage.h
deleted file mode 100644
diff --git a/Engine/Plugins/Experimental/Mutable/Source/MutableTools/Internal/MuT/NodePatchMeshPrivate.h b/Engine/Plugins/Experimental/Mutable/Source/MutableTools/Internal/MuT/NodePatchMeshPrivate.h
deleted file mode 100644
diff --git a/Engine/Plugins/Experimental/Mutable/Source/MutableTools/Private/MuT/CodeGenerator.cpp b/Engine/Plugins/Experimental/Mutable/Source/MutableTools/Private/MuT/CodeGenerator.cpp
--- a/Engine/Plugins/Experimental/Mutable/Source/MutableTools/Private/MuT/CodeGenerator.cpp
+++ b/Engine/Plugins/Experimental/Mutable/Source/MutableTools/Private/MuT/CodeGenerator.cpp
@@ -67,16 +67,13 @@
 #include "MuT/NodeMeshMorph.h"
 #include "MuT/NodeMeshReshape.h"
 #include "MuT/NodeModifier.h"
-#include "MuT/NodeModifierMeshClipDeformPrivate.h"
-#include "MuT/NodeModifierMeshClipMorphPlanePrivate.h"
-#include "MuT/NodeModifierMeshClipWithMeshPrivate.h"
-#include "MuT/NodeModifierMeshClipWithUVMaskPrivate.h"
-#include "MuT/NodeModifierPrivate.h"
+#include "MuT/NodeModifierMeshClipDeform.h"
+#include "MuT/NodeModifierMeshClipMorphPlane.h"
+#include "MuT/NodeModifierMeshClipWithMesh.h"
+#include "MuT/NodeModifierMeshClipWithUVMask.h"
 #include "MuT/NodeObject.h"
 #include "MuT/NodeObjectGroupPrivate.h"
 #include "MuT/NodeObjectNew.h"
-#include "MuT/NodePatchImage.h"
-#include "MuT/NodePatchMesh.h"
 #include "MuT/NodePrivate.h"
 #include "MuT/NodeRange.h"
 #include "MuT/NodeRangeFromScalar.h"
@@ -352,8 +349,8 @@ namespace mu
 			}
 
 			// Add every row
-			int32 rows = CacheKey.Table->GetPrivate()->Rows.Num();
-			for (size_t i = 0; i < rows; ++i)
+			int32 RowCount = CacheKey.Table->GetPrivate()->Rows.Num();
+			for (int32 i = 0; i < RowCount; ++i)
 			{
 				FParameterDesc::FIntValueDesc value;
 				value.m_value = (int16)CacheKey.Table->GetPrivate()->Rows[i].Id;
@@ -429,7 +426,7 @@ namespace mu
 
 
 	Ptr<ASTOp> CodeGenerator::GenerateImageBlockPatch(Ptr<ASTOp> InBlockOp,
-		const NodePatchImage* pPatch,
+		const NodeSurfaceEdit::FTexture& Patch,
 		Ptr<Image> PatchMask,
 		Ptr<ASTOp> conditionAd,
 		const FImageGenerationOptions& ImageOptions )
@@ -440,14 +437,14 @@ namespace mu
 			MUTABLE_CPUPROFILER_SCOPE(PatchBlend);
 
 			Ptr<ASTOpImageLayer> LayerOp = new ASTOpImageLayer();
-			LayerOp->blendType = pPatch->BlendType;
+			LayerOp->blendType = Patch.PatchBlendType;
 			LayerOp->base = InBlockOp;
 
 			// When we patch from edit nodes, we want to apply it to all the channels.
 			// \todo: since we can choose the patch function, maybe we want to be able to select this as well.
-			LayerOp->Flags = pPatch->bApplyToAlpha ? OP::ImageLayerArgs::F_APPLY_TO_ALPHA : 0;
+			LayerOp->Flags = Patch.bPatchApplyToAlpha ? OP::ImageLayerArgs::F_APPLY_TO_ALPHA : 0;
 
-			NodeImage* ImageNode = pPatch->Image.get();
+			NodeImage* ImageNode = Patch.PatchImage.get();
 			Ptr<ASTOp> BlendOp;
 			if (ImageNode)
 			{
@@ -457,7 +454,7 @@ namespace mu
 			}
 			else
 			{
-				BlendOp = GenerateMissingImageCode(TEXT("Patch top image"), EImageFormat::IF_RGB_UBYTE, pPatch->GetMessageContext(), ImageOptions);
+				BlendOp = GenerateMissingImageCode(TEXT("Patch top image"), EImageFormat::IF_RGB_UBYTE, nullptr, ImageOptions);
 			}
 			BlendOp = GenerateImageFormat(BlendOp, InBlockOp->GetImageDesc().m_format);
 			BlendOp = GenerateImageSize(BlendOp, ImageOptions.RectSize);
@@ -475,7 +472,7 @@ namespace mu
 				RectConstantOp = ConstantResult.op;
 			}
 
-			NodeImage* MaskNode = pPatch->Mask.get();
+			NodeImage* MaskNode = Patch.PatchMask.get();
 			Ptr<ASTOp> MaskOp;
 			if (MaskNode)
 			{
@@ -794,7 +791,7 @@ namespace mu
 	}
 
 
-	Ptr<Image> CodeGenerator::GenerateImageBlockPatchMask(const NodePatchImage* Patch, FIntPoint GridSize, int32 BlockPixelsX, int32 BlockPixelsY, box<FIntVector2> RectInCells )
+	Ptr<Image> CodeGenerator::GenerateImageBlockPatchMask(const NodeSurfaceEdit::FTexture& Patch, FIntPoint GridSize, int32 BlockPixelsX, int32 BlockPixelsY, box<FIntVector2> RectInCells )
 	{
 		// Create a patching mask for the block
 		Ptr<Image> PatchMask;
@@ -805,7 +802,7 @@ namespace mu
 		BlockRectInPixels.Min = { RectInCells.min[0] * BlockPixelsX, RectInCells.min[1] * BlockPixelsY };
 		BlockRectInPixels.Max = { (RectInCells.min[0] + RectInCells.size[0]) * BlockPixelsX, (RectInCells.min[1] + RectInCells.size[1]) * BlockPixelsY };
 
-		for (const FBox2f& PatchRect : Patch->Blocks)
+		for (const FBox2f& PatchRect : Patch.PatchBlocks)
 		{
 			// Does the patch rect intersects the current block at all?
 			FInt32Rect PatchRectInPixels;
@@ -952,66 +949,63 @@ namespace mu
             {
                 const FirstPassGenerator::FSurface::FEdit& e = edits[editIndex];
 
-                if ( e.Node->Mesh)
+                if ( Ptr<NodeMesh> pAdd = e.Node->MeshAdd)
                 {
-                    if ( NodeMeshPtr pAdd = e.Node->Mesh->GetAdd() )
-                    {
-						// Store the data necessary to apply modifiers for the pre-normal operations stage.
-						ActiveTags.Add(e.Node->Tags);
+					// Store the data necessary to apply modifiers for the pre-normal operations stage.
+					ActiveTags.Add(e.Node->EnableTags);
 						
-						FMeshGenerationOptions MergedMeshOptions;
-						MergedMeshOptions.bLayouts = true;
-						MergedMeshOptions.bClampUVIslands = bShareSurface && bNormalizeUVs;
-						MergedMeshOptions.bNormalizeUVs = bNormalizeUVs;
-						MergedMeshOptions.State = Options.State;
-						MergedMeshOptions.ActiveTags = e.Node->Tags;
-
-						if (SharedMeshResults)
-						{
-							check(SharedMeshResults->ExtraMeshLayouts.Num()>editIndex);
-							MergedMeshOptions.OverrideLayouts = SharedMeshResults->ExtraMeshLayouts[editIndex].GeneratedLayouts;
-						}
+					FMeshGenerationOptions MergedMeshOptions;
+					MergedMeshOptions.bLayouts = true;
+					MergedMeshOptions.bClampUVIslands = bShareSurface && bNormalizeUVs;
+					MergedMeshOptions.bNormalizeUVs = bNormalizeUVs;
+					MergedMeshOptions.State = Options.State;
+					MergedMeshOptions.ActiveTags = e.Node->EnableTags;
 
-						FMeshGenerationResult addResults;
-                        GenerateMesh(MergedMeshOptions, addResults, pAdd);
+					if (SharedMeshResults)
+					{
+						check(SharedMeshResults->ExtraMeshLayouts.Num()>editIndex);
+						MergedMeshOptions.OverrideLayouts = SharedMeshResults->ExtraMeshLayouts[editIndex].GeneratedLayouts;
+					}
 
-						ActiveTags.Pop();
+					FMeshGenerationResult addResults;
+                    GenerateMesh(MergedMeshOptions, addResults, pAdd);
 
-                        baseMeshesForEachAddedMesh[e.Node] = addResults.BaseMeshOp;
+					ActiveTags.Pop();
 
-						// Apply the modifier for the post-normal operations stage to the added mesh
-						bool bModifiersForBeforeOperations = false;
-						FGenericGenerationOptions ModifierOptions(Options);
-						ModifierOptions.ActiveTags = e.Node->Tags;
-						lastMeshOp = ApplyMeshModifiers(Options, lastMeshOp, bModifiersForBeforeOperations, surfaceNode->GetMessageContext());
+                    baseMeshesForEachAddedMesh[e.Node] = addResults.BaseMeshOp;
 
-                        FMeshGenerationResult::FExtraLayouts data;
-						data.GeneratedLayouts = addResults.GeneratedLayouts;
-						data.Condition = e.Condition;
-                        data.MeshFragment = addResults.MeshOp;
-                        meshResults.ExtraMeshLayouts[editIndex] = data;
+					// Apply the modifier for the post-normal operations stage to the added mesh
+					bool bModifiersForBeforeOperations = false;
+					FGenericGenerationOptions ModifierOptions(Options);
+					ModifierOptions.ActiveTags = e.Node->EnableTags;
+					lastMeshOp = ApplyMeshModifiers(Options, lastMeshOp, bModifiersForBeforeOperations, surfaceNode->GetMessageContext());
 
-                        Ptr<ASTOpFixed> mop = new ASTOpFixed();
-                        mop->op.type = OP_TYPE::ME_MERGE;
-                        mop->SetChild(mop->op.args.MeshMerge.base, lastMeshOp );
-                        mop->SetChild(mop->op.args.MeshMerge.added, addResults.MeshOp );
-                        // will merge the meshes under the same surface
-                        mop->op.args.MeshMerge.newSurfaceID = 0;
+                    FMeshGenerationResult::FExtraLayouts data;
+					data.GeneratedLayouts = addResults.GeneratedLayouts;
+					data.Condition = e.Condition;
+                    data.MeshFragment = addResults.MeshOp;
+                    meshResults.ExtraMeshLayouts[editIndex] = data;
 
-                        // Condition to apply
-                        if (e.Condition)
-                        {
-                            Ptr<ASTOpConditional> conditionalAd = new ASTOpConditional();
-                            conditionalAd->type = OP_TYPE::ME_CONDITIONAL;
-                            conditionalAd->no = lastMeshOp;
-                            conditionalAd->yes = mop;
-                            conditionalAd->condition = e.Condition;
-                            lastMeshOp = conditionalAd;
-                        }
-                        else
-                        {
-                            lastMeshOp = mop;
-                        }
+                    Ptr<ASTOpFixed> mop = new ASTOpFixed();
+                    mop->op.type = OP_TYPE::ME_MERGE;
+                    mop->SetChild(mop->op.args.MeshMerge.base, lastMeshOp );
+                    mop->SetChild(mop->op.args.MeshMerge.added, addResults.MeshOp );
+                    // will merge the meshes under the same surface
+                    mop->op.args.MeshMerge.newSurfaceID = 0;
+
+                    // Condition to apply
+                    if (e.Condition)
+                    {
+                        Ptr<ASTOpConditional> conditionalAd = new ASTOpConditional();
+                        conditionalAd->type = OP_TYPE::ME_CONDITIONAL;
+                        conditionalAd->no = lastMeshOp;
+                        conditionalAd->yes = mop;
+                        conditionalAd->condition = e.Condition;
+                        lastMeshOp = conditionalAd;
+                    }
+                    else
+                    {
+                        lastMeshOp = mop;
                     }
                 }
             }
@@ -1024,44 +1018,41 @@ namespace mu
             Ptr<ASTOpMeshRemoveMask> rop;
             for ( const FirstPassGenerator::FSurface::FEdit& e: edits )
             {
-                if ( e.Node->Mesh )
+                if ( Ptr<NodeMesh> pRemove = e.Node->MeshRemove )
                 {
-                    if ( NodeMeshPtr pRemove = e.Node->Mesh->GetRemove() )
-                    {
-                        FMeshGenerationResult removeResults;
-						FMeshGenerationOptions RemoveMeshOptions;
-						RemoveMeshOptions.bLayouts = false;
-						RemoveMeshOptions.State = Options.State;
-						RemoveMeshOptions.ActiveTags = e.Node->Tags;
+                    FMeshGenerationResult removeResults;
+					FMeshGenerationOptions RemoveMeshOptions;
+					RemoveMeshOptions.bLayouts = false;
+					RemoveMeshOptions.State = Options.State;
+					RemoveMeshOptions.ActiveTags = e.Node->EnableTags;
 
-                        GenerateMesh(RemoveMeshOptions, removeResults, pRemove );
+                    GenerateMesh(RemoveMeshOptions, removeResults, pRemove );
 
-                        Ptr<ASTOpFixed> maskOp = new ASTOpFixed();
-                        maskOp->op.type = OP_TYPE::ME_MASKDIFF;
+                    Ptr<ASTOpFixed> maskOp = new ASTOpFixed();
+                    maskOp->op.type = OP_TYPE::ME_MASKDIFF;
 
-                        // By default, remove from the base
-                        Ptr<ASTOp> removeFrom = meshResults.BaseMeshOp;
-                        // See if we want to remove from an added mesh instead.
-                        if ( e.Node->Parent )
+                    // By default, remove from the base
+                    Ptr<ASTOp> removeFrom = meshResults.BaseMeshOp;
+                    // See if we want to remove from an added mesh instead.
+                    if ( e.Node->Parent )
+                    {
+                        auto addedBaseMeshIt = baseMeshesForEachAddedMesh.find(e.Node->Parent.get());
+                        if (addedBaseMeshIt!=baseMeshesForEachAddedMesh.end())
                         {
-                            auto addedBaseMeshIt = baseMeshesForEachAddedMesh.find(e.Node->Parent.get());
-                            if (addedBaseMeshIt!=baseMeshesForEachAddedMesh.end())
-                            {
-                                removeFrom = addedBaseMeshIt->second;
-                            }
+                            removeFrom = addedBaseMeshIt->second;
                         }
+                    }
 
-                        maskOp->SetChild(maskOp->op.args.MeshMaskDiff.source, removeFrom );
-                        maskOp->SetChild(maskOp->op.args.MeshMaskDiff.fragment, removeResults.MeshOp );
-
-                        if (!rop)
-                        {
-                            rop = new ASTOpMeshRemoveMask();
-                            rop->source = lastMeshOp;
-                        }
+                    maskOp->SetChild(maskOp->op.args.MeshMaskDiff.source, removeFrom );
+                    maskOp->SetChild(maskOp->op.args.MeshMaskDiff.fragment, removeResults.MeshOp );
 
-                        rop->AddRemove( e.Condition, maskOp );
+                    if (!rop)
+                    {
+                        rop = new ASTOpMeshRemoveMask();
+                        rop->source = lastMeshOp;
                     }
+
+                    rop->AddRemove( e.Condition, maskOp );
                 }
             }
 
@@ -1073,7 +1064,7 @@ namespace mu
             // Apply mesh morphs from child objects "edit surface" nodes
             for ( const FirstPassGenerator::FSurface::FEdit& e: edits )
             {
-                if ( NodeMeshPtr pMorph = e.Node->Morph )
+                if ( NodeMeshPtr pMorph = e.Node->MeshMorph )
                 {
 					// Not needed because it has been generated before already.
 					// Base mesh
@@ -1362,7 +1353,8 @@ namespace mu
 								continue;
 							}
 
-							if (const NodePatchImage* Patch = e.Node->Textures[t].Patch.get())
+							const NodeSurfaceEdit::FTexture& Patch = e.Node->Textures[t];
+							if (Patch.PatchImage.get())
 							{
 								// Does the current block need to be patched?
 
@@ -1565,7 +1557,8 @@ namespace mu
 										continue;
 									}
 
-									if (const NodePatchImage* Patch = Edit.Node->Textures[t].Patch.get())
+									const NodeSurfaceEdit::FTexture& Patch = Edit.Node->Textures[t];
+									if (Patch.PatchImage)
 									{
 										// Is the current block to be patched?
 										Ptr<Image> PatchMask = GenerateImageBlockPatchMask( Patch, GridSize, BlockPixelsX, BlockPixelsY, RectInCells );
@@ -2370,7 +2363,7 @@ namespace mu
 				ClipOptions.State = Options.State;
 
 				FMeshGenerationResult clipResult;
-				GenerateMesh(ClipOptions, clipResult, TypedClipNode->GetPrivate()->ClipMesh);
+				GenerateMesh(ClipOptions, clipResult, TypedClipNode->ClipMesh);
 				op->clip = clipResult.MeshOp;
 
 				if (!op->clip)
@@ -2406,7 +2399,7 @@ namespace mu
 				const NodeModifierMeshClipWithUVMask* TypedClipNode = static_cast<const NodeModifierMeshClipWithUVMask*>(m.Node);
 				Ptr<ASTOpMeshMaskClipUVMask> op = new ASTOpMeshMaskClipUVMask();
 				op->Source = preModifiersMesh;
-				op->LayoutIndex = TypedClipNode->GetPrivate()->LayoutIndex;
+				op->LayoutIndex = TypedClipNode->LayoutIndex;
 
 				// Parameters
 				FImageGenerationOptions ClipOptions;
@@ -2415,7 +2408,7 @@ namespace mu
 				ClipOptions.State = Options.State;
 
 				FImageGenerationResult ClipMaskResult;
-				GenerateImage(ClipOptions, ClipMaskResult, TypedClipNode->GetPrivate()->ClipMask);
+				GenerateImage(ClipOptions, ClipMaskResult, TypedClipNode->ClipMask);
 
 				// It could be IF_L_UBIT, but since this should be optimized out at compile time, leave the most cpu efficient.
 				op->Mask = GenerateImageFormat(ClipMaskResult.op, mu::EImageFormat::IF_L_UBYTE);
@@ -2452,7 +2445,7 @@ namespace mu
 
 			if (m.Node->GetType() == NodeModifierMeshClipMorphPlane::GetStaticType())
 			{
-				const NodeModifierMeshClipMorphPlane::Private* TypedNode = static_cast<const NodeModifierMeshClipMorphPlane*>(m.Node)->GetPrivate();
+				const NodeModifierMeshClipMorphPlane* TypedNode = static_cast<const NodeModifierMeshClipMorphPlane*>(m.Node);
 				Ptr<ASTOpMeshClipMorphPlane> op = new ASTOpMeshClipMorphPlane();
 				op->source = lastMeshOp;
 
@@ -2460,10 +2453,10 @@ namespace mu
 				{
 					FShape morphShape;
 					morphShape.type = (uint8_t)FShape::Type::Ellipse;
-					morphShape.position = TypedNode->m_origin;
-					morphShape.up = TypedNode->m_normal;
+					morphShape.position = TypedNode->Parameters.Origin;
+					morphShape.up = TypedNode->Parameters.Normal;
 					// TODO: Move rotation to ellipse rotation reference base instead of passing it directly
-					morphShape.size = FVector3f(TypedNode->m_radius1, TypedNode->m_radius2, TypedNode->m_rotation);
+					morphShape.size = FVector3f(TypedNode->Parameters.Radius1, TypedNode->Parameters.Radius2, TypedNode->Parameters.Rotation);
 
 					// Generate a "side" vector.
 					// \todo: make generic and move to the vector class
@@ -2471,31 +2464,31 @@ namespace mu
 						// Generate vector perpendicular to normal for ellipse rotation reference base
 						FVector3f aux_base(0.f, 1.f, 0.f);
 
-						if (FMath::Abs(FVector3f::DotProduct(TypedNode->m_normal, aux_base)) > 0.95f)
+						if (FMath::Abs(FVector3f::DotProduct(TypedNode->Parameters.Normal, aux_base)) > 0.95f)
 						{
 							aux_base = FVector3f(0.f, 0.f, 1.f);
 						}
 
-						morphShape.side = FVector3f::CrossProduct(TypedNode->m_normal, aux_base);
+						morphShape.side = FVector3f::CrossProduct(TypedNode->Parameters.Normal, aux_base);
 					}
 					op->morphShape = morphShape;
 				}
 
 				// Selection box
-				if (TypedNode->m_vertexSelectionType == NodeModifierMeshClipMorphPlane::Private::VS_SHAPE)
+				if (TypedNode->Parameters.VertexSelectionType == FClipMorphPlaneParameters::VS_SHAPE)
 				{
 					op->vertexSelectionType = OP::MeshClipMorphPlaneArgs::VS_SHAPE;
 					FShape selectionShape;
-					selectionShape.type = (uint8_t)FShape::Type::AABox;
-					selectionShape.position = TypedNode->m_selectionBoxOrigin;
-					selectionShape.size = TypedNode->m_selectionBoxRadius;
+					selectionShape.type = (uint8)FShape::Type::AABox;
+					selectionShape.position = TypedNode->Parameters.SelectionBoxOrigin;
+					selectionShape.size = TypedNode->Parameters.SelectionBoxRadius;
 					op->selectionShape = selectionShape;
 				}
-				else if (TypedNode->m_vertexSelectionType == NodeModifierMeshClipMorphPlane::Private::VS_BONE_HIERARCHY)
+				else if (TypedNode->Parameters.VertexSelectionType == FClipMorphPlaneParameters::VS_BONE_HIERARCHY)
 				{
 					op->vertexSelectionType = OP::MeshClipMorphPlaneArgs::VS_BONE_HIERARCHY;
-					op->vertexSelectionBone = TypedNode->m_vertexSelectionBone;
-					op->vertexSelectionBoneMaxRadius = TypedNode->m_maxEffectRadius;
+					op->vertexSelectionBone = TypedNode->Parameters.VertexSelectionBone;
+					op->vertexSelectionBoneMaxRadius = TypedNode->Parameters.MaxEffectRadius;
 				}
 				else
 				{
@@ -2503,8 +2496,8 @@ namespace mu
 				}
 
 				// Parameters
-				op->dist = TypedNode->m_dist;
-				op->factor = TypedNode->m_factor;
+				op->dist = TypedNode->Parameters.DistanceToPlane;
+				op->factor = TypedNode->Parameters.LinearityFactor;
 
 				modifiedMeshOp = op;
 			}
@@ -2533,7 +2526,7 @@ namespace mu
 
 			if (M.Node->GetType()==NodeModifierMeshClipDeform::GetStaticType())
 			{
-				const NodeModifierMeshClipDeform::Private* TypedClipNode = static_cast<const NodeModifierMeshClipDeform*>(M.Node)->GetPrivate();
+				const NodeModifierMeshClipDeform* TypedClipNode = static_cast<const NodeModifierMeshClipDeform*>(M.Node);
 				Ptr<ASTOpMeshBindShape>  BindOp = new ASTOpMeshBindShape();
 				Ptr<ASTOpMeshClipDeform> ClipOp = new ASTOpMeshClipDeform();
 
diff --git a/Engine/Plugins/Experimental/Mutable/Source/MutableTools/Private/MuT/CodeGenerator.h b/Engine/Plugins/Experimental/Mutable/Source/MutableTools/Private/MuT/CodeGenerator.h
--- a/Engine/Plugins/Experimental/Mutable/Source/MutableTools/Private/MuT/CodeGenerator.h
+++ b/Engine/Plugins/Experimental/Mutable/Source/MutableTools/Private/MuT/CodeGenerator.h
@@ -33,7 +33,6 @@
 #include "MuT/NodeMeshTable.h"
 #include "MuT/NodeObjectGroup.h"
 #include "MuT/NodeObjectNew.h"
-#include "MuT/NodePatchImage.h"
 #include "MuT/NodeProjector.h"
 #include "MuT/NodeString.h"
 #include "MuT/NodeSurfaceNew.h"
@@ -505,10 +504,10 @@ namespace mu
 		Ptr<ASTOp> ApplyTiling(Ptr<ASTOp> Source, UE::Math::TIntVector2<int32> Size, EImageFormat Format);
 
 		/** Generate a layout block-sized image with a mask including all pixels in the blocks defined in the patch node. */
-		Ptr<Image> GenerateImageBlockPatchMask(const NodePatchImage* Patch, FIntPoint GridSize, int32 BlockPixelsX, int32 BlockPixelsY, box<FIntVector2> RectInCells);
+		Ptr<Image> GenerateImageBlockPatchMask(const NodeSurfaceEdit::FTexture&, FIntPoint GridSize, int32 BlockPixelsX, int32 BlockPixelsY, box<FIntVector2> RectInCells);
 
 		/** Generate all the operations to apply the block patching on top of the BlockOp, and masking with PatchMask. */
-		Ptr<ASTOp> GenerateImageBlockPatch(Ptr<ASTOp> BlockOp, const NodePatchImage*, Ptr<Image> PatchMask, Ptr<ASTOp> ConditionOp, const FImageGenerationOptions&);
+		Ptr<ASTOp> GenerateImageBlockPatch(Ptr<ASTOp> BlockOp, const NodeSurfaceEdit::FTexture&, Ptr<Image> PatchMask, Ptr<ASTOp> ConditionOp, const FImageGenerationOptions&);
 
         //-----------------------------------------------------------------------------------------
         // Meshes
diff --git a/Engine/Plugins/Experimental/Mutable/Source/MutableTools/Private/MuT/CodeGenerator_FirstPass.cpp b/Engine/Plugins/Experimental/Mutable/Source/MutableTools/Private/MuT/CodeGenerator_FirstPass.cpp
--- a/Engine/Plugins/Experimental/Mutable/Source/MutableTools/Private/MuT/CodeGenerator_FirstPass.cpp
+++ b/Engine/Plugins/Experimental/Mutable/Source/MutableTools/Private/MuT/CodeGenerator_FirstPass.cpp
@@ -17,10 +17,10 @@
 #include "MuT/ErrorLog.h"
 #include "MuT/ErrorLogPrivate.h"
 #include "MuT/NodeComponent.h"
-#include "MuT/NodeModifierMeshClipDeformPrivate.h"
-#include "MuT/NodeModifierMeshClipMorphPlanePrivate.h"
-#include "MuT/NodeModifierMeshClipWithMeshPrivate.h"
-#include "MuT/NodeModifierMeshClipWithUVMaskPrivate.h"
+#include "MuT/NodeModifierMeshClipDeform.h"
+#include "MuT/NodeModifierMeshClipMorphPlane.h"
+#include "MuT/NodeModifierMeshClipWithMesh.h"
+#include "MuT/NodeModifierMeshClipWithUVMask.h"
 #include "MuT/NodeObject.h"
 #include "MuT/NodeObjectGroupPrivate.h"
 #include "MuT/NodeObjectNew.h"
@@ -97,10 +97,10 @@ namespace mu
             for (int32 e=0; e< Surfaces[s].Edits.Num(); ++e)
             {
                 const FSurface::FEdit& edit = Surfaces[s].Edits[e];
-                for (int32 t=0; t<edit.Node->Tags.Num(); ++t)
+                for (int32 t=0; t<edit.Node->EnableTags.Num(); ++t)
                 {
                     int32 tag = -1;
-					const FString& tagStr = edit.Node->Tags[t];
+					const FString& tagStr = edit.Node->EnableTags[t];
 
                     for (int32 i = 0; i<Tags.Num() && tag<0; ++i)
                     {
diff --git a/Engine/Plugins/Experimental/Mutable/Source/MutableTools/Private/MuT/CodeGenerator_FirstPass.h b/Engine/Plugins/Experimental/Mutable/Source/MutableTools/Private/MuT/CodeGenerator_FirstPass.h
--- a/Engine/Plugins/Experimental/Mutable/Source/MutableTools/Private/MuT/CodeGenerator_FirstPass.h
+++ b/Engine/Plugins/Experimental/Mutable/Source/MutableTools/Private/MuT/CodeGenerator_FirstPass.h
@@ -20,7 +20,6 @@
 #include "MuT/NodeModifierMeshClipWithUVMask.h"
 #include "MuT/NodeObjectGroup.h"
 #include "MuT/NodeObjectNew.h"
-#include "MuT/NodePatchMesh.h"
 #include "MuT/NodeSurfaceEdit.h"
 #include "MuT/NodeSurfaceNew.h"
 #include "MuT/NodeSurfaceVariation.h"
diff --git a/Engine/Plugins/Experimental/Mutable/Source/MutableTools/Private/MuT/CodeGenerator_Mesh.cpp b/Engine/Plugins/Experimental/Mutable/Source/MutableTools/Private/MuT/CodeGenerator_Mesh.cpp
--- a/Engine/Plugins/Experimental/Mutable/Source/MutableTools/Private/MuT/CodeGenerator_Mesh.cpp
+++ b/Engine/Plugins/Experimental/Mutable/Source/MutableTools/Private/MuT/CodeGenerator_Mesh.cpp
@@ -44,11 +44,8 @@
 #include "MuT/NodeMeshApplyPose.h"
 #include "MuT/NodeMeshApplyPosePrivate.h"
 #include "MuT/NodeMeshClipDeform.h"
-#include "MuT/NodeMeshClipDeformPrivate.h"
 #include "MuT/NodeMeshClipMorphPlane.h"
-#include "MuT/NodeMeshClipMorphPlanePrivate.h"
 #include "MuT/NodeMeshClipWithMesh.h"
-#include "MuT/NodeMeshClipWithMeshPrivate.h"
 #include "MuT/NodeMeshConstant.h"
 #include "MuT/NodeMeshConstantPrivate.h"
 #include "MuT/NodeMeshFormat.h"
@@ -1403,15 +1400,13 @@ namespace mu
     void CodeGenerator::GenerateMesh_ClipMorphPlane(const FMeshGenerationOptions& InOptions, FMeshGenerationResult& OutResult,
                                                      const NodeMeshClipMorphPlane* ClipNode )
     {
-        const auto& node = *ClipNode->GetPrivate();
-
         Ptr<ASTOpMeshClipMorphPlane> op = new ASTOpMeshClipMorphPlane();
 
         // Base
-        if (node.m_pSource)
+        if (ClipNode->Source)
         {
 			FMeshGenerationOptions BaseOptions = InOptions;
-            GenerateMesh(BaseOptions, OutResult, node.m_pSource);
+            GenerateMesh(BaseOptions, OutResult, ClipNode->Source);
             op->source = OutResult.MeshOp;
         }
         else
@@ -1424,9 +1419,9 @@ namespace mu
         // Morph to an ellipse
         {
             op->morphShape.type = (uint8_t)FShape::Type::Ellipse;
-            op->morphShape.position = node.m_origin;
-            op->morphShape.up = node.m_normal;
-            op->morphShape.size = FVector3f(node.m_radius1, node.m_radius2, node.m_rotation); // TODO: Move rotation to ellipse rotation reference base instead of passing it directly
+            op->morphShape.position = ClipNode->Parameters.Origin;
+            op->morphShape.up = ClipNode->Parameters.Normal;
+            op->morphShape.size = FVector3f(ClipNode->Parameters.Radius1, ClipNode->Parameters.Radius2, ClipNode->Parameters.Rotation); // TODO: Move rotation to ellipse rotation reference base instead of passing it directly
 
                                                                                       // Generate a "side" vector.
                                                                                       // \todo: make generic and move to the vector class
@@ -1434,29 +1429,29 @@ namespace mu
                 // Generate vector perpendicular to normal for ellipse rotation reference base
 				FVector3f aux_base(0.f, 1.f, 0.f);
 
-                if (fabs(FVector3f::DotProduct(node.m_normal, aux_base)) > 0.95f)
+                if (fabs(FVector3f::DotProduct(ClipNode->Parameters.Normal, aux_base)) > 0.95f)
                 {
                     aux_base = FVector3f(0.f, 0.f, 1.f);
                 }
 
-                op->morphShape.side = FVector3f::CrossProduct(node.m_normal, aux_base);
+                op->morphShape.side = FVector3f::CrossProduct(ClipNode->Parameters.Normal, aux_base);
             }
         }
 
         // Selection by shape
-        if (node.m_vertexSelectionType== NodeMeshClipMorphPlane::Private::VS_SHAPE)
+        if (ClipNode->Parameters.VertexSelectionType== FClipMorphPlaneParameters::VS_SHAPE)
         {
             op->vertexSelectionType = OP::MeshClipMorphPlaneArgs::VS_SHAPE;
             op->selectionShape.type = (uint8_t)FShape::Type::AABox;
-            op->selectionShape.position = node.m_selectionBoxOrigin;
-            op->selectionShape.size = node.m_selectionBoxRadius;
+            op->selectionShape.position = ClipNode->Parameters.SelectionBoxOrigin;
+            op->selectionShape.size = ClipNode->Parameters.SelectionBoxRadius;
         }
-        else if (node.m_vertexSelectionType == NodeMeshClipMorphPlane::Private::VS_BONE_HIERARCHY)
+        else if (ClipNode->Parameters.VertexSelectionType == FClipMorphPlaneParameters::VS_BONE_HIERARCHY)
         {
             // Selection by bone hierarchy?
             op->vertexSelectionType = OP::MeshClipMorphPlaneArgs::VS_BONE_HIERARCHY;
-            op->vertexSelectionBone = node.m_vertexSelectionBone;
-			op->vertexSelectionBoneMaxRadius = node.m_maxEffectRadius;
+            op->vertexSelectionBone = ClipNode->Parameters.VertexSelectionBone;
+			op->vertexSelectionBoneMaxRadius = ClipNode->Parameters.MaxEffectRadius;
         }
         else
         {
@@ -1464,8 +1459,8 @@ namespace mu
         }
 
         // Parameters
-        op->dist = node.m_dist;
-        op->factor = node.m_factor;
+        op->dist = ClipNode->Parameters.DistanceToPlane;
+        op->factor = ClipNode->Parameters.LinearityFactor;
 
         OutResult.MeshOp = op;
     }
@@ -1475,15 +1470,13 @@ namespace mu
     void CodeGenerator::GenerateMesh_ClipWithMesh(const FMeshGenerationOptions& InOptions, FMeshGenerationResult& OutResult,
                                                    const NodeMeshClipWithMesh* ClipNode)
     {
-        const auto& node = *ClipNode->GetPrivate();
-
         Ptr<ASTOpFixed> op = new ASTOpFixed();
         op->op.type = OP_TYPE::ME_CLIPWITHMESH;
 
         // Base
-        if (node.m_pSource)
+        if (ClipNode->Source)
         {
-            GenerateMesh(InOptions, OutResult, node.m_pSource );
+            GenerateMesh(InOptions, OutResult, ClipNode->Source );
             op->SetChild( op->op.args.MeshClipWithMesh.source, OutResult.MeshOp );
         }
         else
@@ -1494,7 +1487,7 @@ namespace mu
         }
 
         // Clipping mesh
-        if (node.m_pClipMesh)
+        if (ClipNode->ClipMesh)
         {
 			FMeshGenerationOptions ClipOptions = InOptions;
 			ClipOptions.bLayouts = false;
@@ -1502,7 +1495,7 @@ namespace mu
 			ClipOptions.ActiveTags.Empty();
 
             FMeshGenerationResult clipResult;
-            GenerateMesh(ClipOptions, clipResult, node.m_pClipMesh);
+            GenerateMesh(ClipOptions, clipResult, ClipNode->ClipMesh);
             op->SetChild( op->op.args.MeshClipWithMesh.clipMesh, clipResult.MeshOp );
 		}
         else
@@ -1518,15 +1511,13 @@ namespace mu
 	//---------------------------------------------------------------------------------------------
 	void CodeGenerator::GenerateMesh_ClipDeform(const FMeshGenerationOptions& InOptions, FMeshGenerationResult& Result, const NodeMeshClipDeform* ClipDeform)
 	{
-		const auto& Node = *ClipDeform->GetPrivate();
-
 		const Ptr<ASTOpMeshBindShape> OpBind = new ASTOpMeshBindShape();
 		const Ptr<ASTOpMeshClipDeform> OpClipDeform = new ASTOpMeshClipDeform();
 
 		// Base Mesh
-		if (Node.m_pBaseMesh)
+		if (ClipDeform->BaseMesh)
 		{
-			GenerateMesh(InOptions, Result, Node.m_pBaseMesh);
+			GenerateMesh(InOptions, Result, ClipDeform->BaseMesh);
 			OpBind->Mesh = Result.MeshOp;
 		}
 		else
@@ -1536,7 +1527,7 @@ namespace mu
 		}
 
 		// Base Shape
-		if (Node.m_pClipShape)
+		if (ClipDeform->ClipShape)
 		{
 			FMeshGenerationOptions ClipOptions = InOptions;
 			ClipOptions.bLayouts = false;
@@ -1544,7 +1535,7 @@ namespace mu
 			ClipOptions.ActiveTags.Empty();
 
 			FMeshGenerationResult baseResult;
-			GenerateMesh(ClipOptions, baseResult, Node.m_pClipShape);
+			GenerateMesh(ClipOptions, baseResult, ClipDeform->ClipShape);
 			OpBind->Shape = baseResult.MeshOp;
 			OpClipDeform->ClipShape = baseResult.MeshOp;
 		}
diff --git a/Engine/Plugins/Experimental/Mutable/Source/MutableTools/Private/MuT/Node.cpp b/Engine/Plugins/Experimental/Mutable/Source/MutableTools/Private/MuT/Node.cpp
--- a/Engine/Plugins/Experimental/Mutable/Source/MutableTools/Private/MuT/Node.cpp
+++ b/Engine/Plugins/Experimental/Mutable/Source/MutableTools/Private/MuT/Node.cpp
@@ -29,7 +29,14 @@
 #include "MuT/NodeColourTable.h"
 #include "MuT/NodeScalarSwitch.h"
 #include "MuT/NodeMeshFragment.h"
-#include "MuT/NodePatchImage.h"
+#include "MuT/NodeMeshClipMorphPlane.h"
+#include "MuT/NodeMeshClipDeform.h"
+#include "MuT/NodeMeshClipWithMesh.h"
+#include "MuT/NodeModifier.h"
+#include "MuT/NodeModifierMeshClipDeform.h"
+#include "MuT/NodeModifierMeshClipMorphPlane.h"
+#include "MuT/NodeModifierMeshClipWithMesh.h"
+#include "MuT/NodeModifierMeshClipWithUVMask.h"
 
 namespace mu
 {
@@ -66,8 +73,15 @@ namespace mu
 	FNodeType NodeColourFromScalars::StaticType = FNodeType(Node::EType::ColorFromScalars, NodeColour::GetStaticType());
 
 	FNodeType NodeMeshFragment::StaticType = FNodeType(Node::EType::MeshFragment, NodeMesh::GetStaticType());
-
-	FNodeType NodePatchImage::StaticType = FNodeType(Node::EType::PatchImage, Node::GetStaticType());
+	FNodeType NodeMeshClipMorphPlane::StaticType = FNodeType(Node::EType::MeshClipMorphPlane, NodeMesh::GetStaticType());
+	FNodeType NodeMeshClipDeform::StaticType = FNodeType(Node::EType::MeshClipDeform, NodeMesh::GetStaticType());
+	FNodeType NodeMeshClipWithMesh::StaticType = FNodeType(Node::EType::MeshClipWithMesh, NodeMesh::GetStaticType());
+
+	FNodeType NodeModifier::StaticType = FNodeType(Node::EType::Modifier, Node::GetStaticType());
+	FNodeType NodeModifierMeshClipDeform::StaticType = FNodeType(Node::EType::ModifierMeshClipDeform, NodeModifier::GetStaticType());
+	FNodeType NodeModifierMeshClipMorphPlane::StaticType = FNodeType(Node::EType::ModifierMeshClipMorphPlane, NodeModifier::GetStaticType());
+	FNodeType NodeModifierMeshClipWithMesh::StaticType = FNodeType(Node::EType::ModifierMeshClipWithMesh, NodeModifier::GetStaticType());
+	FNodeType NodeModifierMeshClipWithUVMask::StaticType = FNodeType(Node::EType::ModifierMeshClipWithUVMask, NodeModifier::GetStaticType());
 
 
 	FNodeType::FNodeType()
diff --git a/Engine/Plugins/Experimental/Mutable/Source/MutableTools/Private/MuT/NodeMeshClipDeform.cpp b/Engine/Plugins/Experimental/Mutable/Source/MutableTools/Private/MuT/NodeMeshClipDeform.cpp
deleted file mode 100644
diff --git a/Engine/Plugins/Experimental/Mutable/Source/MutableTools/Private/MuT/NodeMeshClipDeformPrivate.h b/Engine/Plugins/Experimental/Mutable/Source/MutableTools/Private/MuT/NodeMeshClipDeformPrivate.h
deleted file mode 100644
diff --git a/Engine/Plugins/Experimental/Mutable/Source/MutableTools/Private/MuT/NodeMeshClipMorphPlane.cpp b/Engine/Plugins/Experimental/Mutable/Source/MutableTools/Private/MuT/NodeMeshClipMorphPlane.cpp
--- a/Engine/Plugins/Experimental/Mutable/Source/MutableTools/Private/MuT/NodeMeshClipMorphPlane.cpp
+++ b/Engine/Plugins/Experimental/Mutable/Source/MutableTools/Private/MuT/NodeMeshClipMorphPlane.cpp
@@ -5,7 +5,6 @@
 #include "Misc/AssertionMacros.h"
 #include "MuR/MutableMath.h"
 #include "MuT/NodeLayout.h"
-#include "MuT/NodeMeshClipMorphPlanePrivate.h"
 #include "MuT/NodeMeshPrivate.h"
 #include "MuT/NodePrivate.h"
 
@@ -13,77 +12,56 @@
 namespace mu
 {
 
-	//---------------------------------------------------------------------------------------------
-	// Static initialisation
-	//---------------------------------------------------------------------------------------------
-    FNodeType NodeMeshClipMorphPlane::Private::s_type = FNodeType(Node::EType::MeshClipMorphPlane, NodeMesh::GetStaticType() );
-
-
-	//---------------------------------------------------------------------------------------------
-	//!
-	//---------------------------------------------------------------------------------------------
-
-    MUTABLE_IMPLEMENT_NODE( NodeMeshClipMorphPlane )
-
-
-	//---------------------------------------------------------------------------------------------
-	// Own Interface
-	//---------------------------------------------------------------------------------------------
-    NodeMeshPtr NodeMeshClipMorphPlane::GetSource() const
+	NodeMeshPtr NodeMeshClipMorphPlane::GetSource() const
 	{
-		return m_pD->m_pSource;
+		return Source;
 	}
 
 
     //---------------------------------------------------------------------------------------------
     void NodeMeshClipMorphPlane::SetSource( NodeMesh* p )
     {
-        m_pD->m_pSource = p;
+        Source = p;
     }
 
 
 	//---------------------------------------------------------------------------------------------
-	void NodeMeshClipMorphPlane::SetPlane(float centerX, float centerY, float centerZ, float normalX, float normalY, float normalZ)
+	void NodeMeshClipMorphPlane::SetPlane(FVector3f Center, FVector3f Normal)
 	{
-		m_pD->m_origin = FVector3f(centerX, centerY, centerZ);
-		m_pD->m_normal = FVector3f(normalX, normalY, normalZ);
+		Parameters.Origin = Center;
+		Parameters.Normal = Normal;
 	}
 
 
 	//---------------------------------------------------------------------------------------------
 	void NodeMeshClipMorphPlane::SetParams(float dist, float factor)
 	{
-		m_pD->m_dist = dist;
-		m_pD->m_factor = factor;
+		Parameters.DistanceToPlane = dist;
+		Parameters.LinearityFactor = factor;
 	}
 
 	//---------------------------------------------------------------------------------------------
 	void NodeMeshClipMorphPlane::SetMorphEllipse(float radius1, float radius2, float rotation)
 	{
-		m_pD->m_radius1 = radius1;
-		m_pD->m_radius2 = radius2;
-		m_pD->m_rotation = rotation;
+		Parameters.Radius1 = radius1;
+		Parameters.Radius2 = radius2;
+		Parameters.Rotation = rotation;
 	}
 
 	//---------------------------------------------------------------------------------------------
-	void NodeMeshClipMorphPlane::SetVertexSelectionBox(float centerX, float centerY, float centerZ, float radiusX, float radiusY, float radiusZ)
+	void NodeMeshClipMorphPlane::SetVertexSelectionBox(FVector3f Center, FVector3f Radius)
 	{
-		m_pD->m_vertexSelectionType = Private::VS_SHAPE;
-		m_pD->m_selectionBoxOrigin = FVector3f(centerX, centerY, centerZ);
-		m_pD->m_selectionBoxRadius = FVector3f(radiusX, radiusY, radiusZ);
+		Parameters.VertexSelectionType = FClipMorphPlaneParameters::VS_SHAPE;
+		Parameters.SelectionBoxOrigin = Center;
+		Parameters.SelectionBoxRadius = Radius;
 	}
 
 	//---------------------------------------------------------------------------------------------
 	void NodeMeshClipMorphPlane::SetVertexSelectionBone(const FBoneName& BoneId, float maxEffectRadius)
 	{
-		m_pD->m_vertexSelectionType = Private::VS_BONE_HIERARCHY;
-		m_pD->m_vertexSelectionBone = BoneId;
-		m_pD->m_maxEffectRadius = maxEffectRadius;
+		Parameters.VertexSelectionType = FClipMorphPlaneParameters::VS_BONE_HIERARCHY;
+		Parameters.VertexSelectionBone = BoneId;
+		Parameters.MaxEffectRadius = maxEffectRadius;
 	}
 
-	//---------------------------------------------------------------------------------------------
-	void NodeMeshClipMorphPlane::AddTag(const char* tagName)
-	{
-		m_pD->Tags.Add(tagName);
-	}
 }
diff --git a/Engine/Plugins/Experimental/Mutable/Source/MutableTools/Private/MuT/NodeMeshClipMorphPlanePrivate.h b/Engine/Plugins/Experimental/Mutable/Source/MutableTools/Private/MuT/NodeMeshClipMorphPlanePrivate.h
deleted file mode 100644
diff --git a/Engine/Plugins/Experimental/Mutable/Source/MutableTools/Private/MuT/NodeMeshClipWithMesh.cpp b/Engine/Plugins/Experimental/Mutable/Source/MutableTools/Private/MuT/NodeMeshClipWithMesh.cpp
deleted file mode 100644
diff --git a/Engine/Plugins/Experimental/Mutable/Source/MutableTools/Private/MuT/NodeMeshClipWithMeshPrivate.h b/Engine/Plugins/Experimental/Mutable/Source/MutableTools/Private/MuT/NodeMeshClipWithMeshPrivate.h
deleted file mode 100644
diff --git a/Engine/Plugins/Experimental/Mutable/Source/MutableTools/Private/MuT/NodeModifier.cpp b/Engine/Plugins/Experimental/Mutable/Source/MutableTools/Private/MuT/NodeModifier.cpp
deleted file mode 100644
diff --git a/Engine/Plugins/Experimental/Mutable/Source/MutableTools/Private/MuT/NodeModifierMeshClipDeform.cpp b/Engine/Plugins/Experimental/Mutable/Source/MutableTools/Private/MuT/NodeModifierMeshClipDeform.cpp
deleted file mode 100644
diff --git a/Engine/Plugins/Experimental/Mutable/Source/MutableTools/Private/MuT/NodeModifierMeshClipMorphPlane.cpp b/Engine/Plugins/Experimental/Mutable/Source/MutableTools/Private/MuT/NodeModifierMeshClipMorphPlane.cpp
--- a/Engine/Plugins/Experimental/Mutable/Source/MutableTools/Private/MuT/NodeModifierMeshClipMorphPlane.cpp
+++ b/Engine/Plugins/Experimental/Mutable/Source/MutableTools/Private/MuT/NodeModifierMeshClipMorphPlane.cpp
@@ -5,66 +5,51 @@
 
 #include "Misc/AssertionMacros.h"
 #include "MuR/MutableMath.h"
-#include "MuT/NodeModifierMeshClipMorphPlanePrivate.h"
 #include "MuT/NodePrivate.h"
 
 
 namespace mu
 {
 
-	//---------------------------------------------------------------------------------------------
-	// Static initialisation
-	//---------------------------------------------------------------------------------------------
-    FNodeType NodeModifierMeshClipMorphPlane::Private::s_type = FNodeType(Node::EType::ModifierMeshClipMorphPlane, NodeModifier::GetStaticType() );
-
-
-	//---------------------------------------------------------------------------------------------
-	//!
-	//---------------------------------------------------------------------------------------------
-
-    MUTABLE_IMPLEMENT_NODE( NodeModifierMeshClipMorphPlane )
-
-
 	//---------------------------------------------------------------------------------------------
 	// Own Interface
 	//---------------------------------------------------------------------------------------------
-    void NodeModifierMeshClipMorphPlane::SetPlane( float centerX, float centerY, float centerZ,
-                                                   float normalX, float normalY, float normalZ)
+    void NodeModifierMeshClipMorphPlane::SetPlane(FVector3f Center, FVector3f Normal)
 	{
-		m_pD->m_origin = FVector3f(centerX, centerY, centerZ);
-		m_pD->m_normal = FVector3f(normalX, normalY, normalZ);
+		Parameters.Origin = Center;
+		Parameters.Normal = Normal;
 	}
 
 
 	//---------------------------------------------------------------------------------------------
 	void NodeModifierMeshClipMorphPlane::SetParams(float dist, float factor)
 	{
-		m_pD->m_dist = dist;
-		m_pD->m_factor = factor;
+		Parameters.DistanceToPlane = dist;
+		Parameters.LinearityFactor = factor;
 	}
 
 	//---------------------------------------------------------------------------------------------
 	void NodeModifierMeshClipMorphPlane::SetMorphEllipse(float radius1, float radius2, float rotation)
 	{
-		m_pD->m_radius1 = radius1;
-		m_pD->m_radius2 = radius2;
-		m_pD->m_rotation = rotation;
+		Parameters.Radius1 = radius1;
+		Parameters.Radius2 = radius2;
+		Parameters.Rotation = rotation;
 	}
 
 	//---------------------------------------------------------------------------------------------
 	void NodeModifierMeshClipMorphPlane::SetVertexSelectionBox(float centerX, float centerY, float centerZ, float radiusX, float radiusY, float radiusZ)
 	{
-		m_pD->m_vertexSelectionType = Private::VS_SHAPE;
-		m_pD->m_selectionBoxOrigin = FVector3f(centerX, centerY, centerZ);
-		m_pD->m_selectionBoxRadius = FVector3f(radiusX, radiusY, radiusZ);
+		Parameters.VertexSelectionType = FClipMorphPlaneParameters::VS_SHAPE;
+		Parameters.SelectionBoxOrigin = FVector3f(centerX, centerY, centerZ);
+		Parameters.SelectionBoxRadius = FVector3f(radiusX, radiusY, radiusZ);
 	}
 
 	//---------------------------------------------------------------------------------------------
 	void NodeModifierMeshClipMorphPlane::SetVertexSelectionBone(const FBoneName& BoneId, float maxEffectRadius)
 	{
-		m_pD->m_vertexSelectionType = Private::VS_BONE_HIERARCHY;
-		m_pD->m_vertexSelectionBone = BoneId;
-		m_pD->m_maxEffectRadius = maxEffectRadius;
+		Parameters.VertexSelectionType = FClipMorphPlaneParameters::VS_BONE_HIERARCHY;
+		Parameters.VertexSelectionBone = BoneId;
+		Parameters.MaxEffectRadius = maxEffectRadius;
 	}
 
 }
diff --git a/Engine/Plugins/Experimental/Mutable/Source/MutableTools/Private/MuT/NodeModifierMeshClipWithMesh.cpp b/Engine/Plugins/Experimental/Mutable/Source/MutableTools/Private/MuT/NodeModifierMeshClipWithMesh.cpp
deleted file mode 100644
diff --git a/Engine/Plugins/Experimental/Mutable/Source/MutableTools/Private/MuT/NodeModifierMeshClipWithUVMask.cpp b/Engine/Plugins/Experimental/Mutable/Source/MutableTools/Private/MuT/NodeModifierMeshClipWithUVMask.cpp
deleted file mode 100644
diff --git a/Engine/Plugins/Experimental/Mutable/Source/MutableTools/Private/MuT/NodePatchMesh.cpp b/Engine/Plugins/Experimental/Mutable/Source/MutableTools/Private/MuT/NodePatchMesh.cpp
deleted file mode 100644
diff --git a/Engine/Plugins/Experimental/Mutable/Source/MutableTools/Public/MuT/NodeMeshClipDeform.h b/Engine/Plugins/Experimental/Mutable/Source/MutableTools/Public/MuT/NodeMeshClipDeform.h
--- a/Engine/Plugins/Experimental/Mutable/Source/MutableTools/Public/MuT/NodeMeshClipDeform.h
+++ b/Engine/Plugins/Experimental/Mutable/Source/MutableTools/Public/MuT/NodeMeshClipDeform.h
@@ -16,43 +16,24 @@ namespace mu
 	{
 	public:
 
-		NodeMeshClipDeform();
-
-		//-----------------------------------------------------------------------------------------
-		// Node Interface
-		//-----------------------------------------------------------------------------------------
-
-		const FNodeType* GetType() const override;
-		static const FNodeType* GetStaticType();
-
-		//-----------------------------------------------------------------------------------------
-		// Own Interface
-		//-----------------------------------------------------------------------------------------
-
-		const NodeMeshPtr& GetBaseMesh() const;
-		void SetBaseMesh( const NodeMeshPtr& );
-
-		const NodeMeshPtr& GetClipShape() const;
-		void SetClipShape(const NodeMeshPtr&);
-
-		const NodeImagePtr& GetShapeWeights() const;
-		void SetShapeWeights(const NodeImagePtr&);
+		Ptr<NodeMesh> BaseMesh;
+		Ptr<NodeMesh> ClipShape;
+		Ptr<NodeImage> ShapeWeights;
 
+	public:
 
-        //-----------------------------------------------------------------------------------------
-		// Interface pattern
-		//-----------------------------------------------------------------------------------------
-		class Private;
-		Private* GetPrivate() const;
+		// Node interface
+		virtual const FNodeType* GetType() const override { return &StaticType; }
+		static const FNodeType* GetStaticType() { return &StaticType; }
 
 	protected:
 
 		//! Forbidden. Manage with the Ptr<> template.
-		~NodeMeshClipDeform();
+		inline ~NodeMeshClipDeform() {}
 
 	private:
 
-		Private* m_pD;
+		static FNodeType StaticType;
 
 	};
 
diff --git a/Engine/Plugins/Experimental/Mutable/Source/MutableTools/Public/MuT/NodeMeshClipMorphPlane.h b/Engine/Plugins/Experimental/Mutable/Source/MutableTools/Public/MuT/NodeMeshClipMorphPlane.h
--- a/Engine/Plugins/Experimental/Mutable/Source/MutableTools/Public/MuT/NodeMeshClipMorphPlane.h
+++ b/Engine/Plugins/Experimental/Mutable/Source/MutableTools/Public/MuT/NodeMeshClipMorphPlane.h
@@ -6,31 +6,26 @@
 #include "MuR/RefCounted.h"
 #include "MuT/Node.h"
 #include "MuT/NodeMesh.h"
+#include "MuT/NodeModifierMeshClipMorphPlane.h"
 
 
 namespace mu
 {
 
-	// Forward definitions
-    class NodeMeshClipMorphPlane;
-    typedef Ptr<NodeMeshClipMorphPlane> NodeMeshClipMorphPlanePtr;
-    typedef Ptr<const NodeMeshClipMorphPlane> NodeMeshClipMorphPlaneConst;
-
-	struct FBoneName;
-
     //! This node applies a geometric transform represented by a 4x4 matrix to a mesh
     class MUTABLETOOLS_API NodeMeshClipMorphPlane : public NodeMesh
 	{
 	public:
 
-		NodeMeshClipMorphPlane();
+		Ptr<NodeMesh> Source;
 
-		//-----------------------------------------------------------------------------------------
-		// Node Interface
-		//-----------------------------------------------------------------------------------------
+		FClipMorphPlaneParameters Parameters;
 
-		virtual const FNodeType* GetType() const;
-		static const FNodeType* GetStaticType();
+	public:
+
+		// Node interface
+		virtual const FNodeType* GetType() const override { return &StaticType; }
+		static const FNodeType* GetStaticType() { return &StaticType; }
 
 		//-----------------------------------------------------------------------------------------
 		// Own Interface
@@ -40,36 +35,28 @@ namespace mu
         NodeMeshPtr GetSource() const;
 		void SetSource( NodeMesh* );
 
-		void SetPlane(float centerX, float centerY, float centerZ, float normalX, float normalY, float normalZ);
+		void SetPlane(FVector3f Center, FVector3f Normal);
 		void SetParams(float dist, float factor);
 		void SetMorphEllipse(float radius1, float radius2, float rotation);
 
 		//! Define an axis-aligned box that will select the vertices to be morphed.
 		//! If this is not set, or the size of the box is 0, all vertices will be affected.
 		//! Only one of Box or Bone Hierarchy can be used (the last one set)
-		void SetVertexSelectionBox(float centerX, float centerY, float centerZ, float radiusX, float radiusY, float radiusZ);
+		void SetVertexSelectionBox(FVector3f Center, FVector3f Radius);
 
 		//! Define the root bone of the subhierarchy of the mesh that will be affected.
 		//! Only one of Box or Bone Hierarchy can be used (the last one set)
 		void SetVertexSelectionBone(const FBoneName& BoneId, float maxEffectRadius);
 
-		//! Add a tag to the clip morph operation, which will only affect surfaces with the same tag
-		void AddTag(const char* tagName);
-
-		//-----------------------------------------------------------------------------------------
-		// Interface pattern
-		//-----------------------------------------------------------------------------------------
-		class Private;
-		Private* GetPrivate() const;
 
 	protected:
 
 		//! Forbidden. Manage with the Ptr<> template.
-        ~NodeMeshClipMorphPlane();
+		~NodeMeshClipMorphPlane() {}
 
 	private:
 
-		Private* m_pD;
+		static FNodeType StaticType;
 
 	};
 
diff --git a/Engine/Plugins/Experimental/Mutable/Source/MutableTools/Public/MuT/NodeMeshClipWithMesh.h b/Engine/Plugins/Experimental/Mutable/Source/MutableTools/Public/MuT/NodeMeshClipWithMesh.h
--- a/Engine/Plugins/Experimental/Mutable/Source/MutableTools/Public/MuT/NodeMeshClipWithMesh.h
+++ b/Engine/Plugins/Experimental/Mutable/Source/MutableTools/Public/MuT/NodeMeshClipWithMesh.h
@@ -11,54 +11,29 @@
 namespace mu
 {
 
-	// Forward definitions
-    class NodeMeshClipWithMesh;
-    typedef Ptr<NodeMeshClipWithMesh> NodeMeshClipWithMeshPtr;
-    typedef Ptr<const NodeMeshClipWithMesh> NodeMeshClipWithMeshConst;
-
-
     //! This node applies a geometric transform represented by a 4x4 matrix to a mesh
 	//! \ingroup model
     class MUTABLETOOLS_API NodeMeshClipWithMesh : public NodeMesh
 	{
 	public:
 
-		NodeMeshClipWithMesh();
-
-		//-----------------------------------------------------------------------------------------
-		// Node Interface
-		//-----------------------------------------------------------------------------------------
-
-		virtual const FNodeType* GetType() const;
-		static const FNodeType* GetStaticType();
+		Ptr<NodeMesh> Source;
+		Ptr<NodeMesh> ClipMesh;
 
-		//-----------------------------------------------------------------------------------------
-		// Own Interface
-		//-----------------------------------------------------------------------------------------
-
-		//! Source mesh to be clipped
-        NodeMeshPtr GetSource() const;
-		void SetSource(NodeMesh*);
-
-		void SetClipMesh(NodeMesh*);
-
-		//! Add a tag to the clip morph operation, which will only affect surfaces with the same tag
-		void AddTag(const char* tagName);
+	public:
 
-		//-----------------------------------------------------------------------------------------
-		// Interface pattern
-		//-----------------------------------------------------------------------------------------
-		class Private;
-		Private* GetPrivate() const;
+		// Node interface
+		virtual const FNodeType* GetType() const override { return &StaticType; }
+		static const FNodeType* GetStaticType() { return &StaticType; }
 
 	protected:
 
 		//! Forbidden. Manage with the Ptr<> template.
-        ~NodeMeshClipWithMesh();
+		inline ~NodeMeshClipWithMesh() {}
 
 	private:
 
-		Private* m_pD;
+		static FNodeType StaticType;
 
 	};
 
diff --git a/Engine/Plugins/Experimental/Mutable/Source/MutableTools/Public/MuT/NodeModifier.h b/Engine/Plugins/Experimental/Mutable/Source/MutableTools/Public/MuT/NodeModifier.h
--- a/Engine/Plugins/Experimental/Mutable/Source/MutableTools/Public/MuT/NodeModifier.h
+++ b/Engine/Plugins/Experimental/Mutable/Source/MutableTools/Public/MuT/NodeModifier.h
@@ -5,6 +5,7 @@
 #include "HAL/Platform.h"
 #include "MuR/Ptr.h"
 #include "MuR/RefCounted.h"
+#include "MuR/Serialisation.h"
 #include "MuT/Node.h"
 #include "Containers/UnrealString.h"
 
@@ -30,32 +31,11 @@ namespace mu
 	typedef Ptr<NodeModifier> NodeModifierPtr;
 	typedef Ptr<const NodeModifier> NodeModifierConst;
 
-	//! This class is the parent of all nodes that output a component.
-	//! \ingroup model
+	/** Parent of all node classes that apply modifiers to surfaces. */
 	class MUTABLETOOLS_API NodeModifier : public Node
 	{
 	public:
 
-		//-----------------------------------------------------------------------------------------
-        // Node interface
-		//-----------------------------------------------------------------------------------------
-
-		const FNodeType* GetType() const override;
-		static const FNodeType* GetStaticType();
-
-        //-----------------------------------------------------------------------------------------
-        // Own interface
-        //-----------------------------------------------------------------------------------------
-
-        /** Add a tag to the surface, which will be affected by modifier nodes with the same tag. */
-        void AddTag(const FString& TagName);
-
-		/** Set the policy to interprete the tags when there is more than one. */
-		void SetMultipleTagPolicy(EMutableMultipleTagPolicy);
-
-		/** Set the stage to apply this modifier in.Default is before normal operations. */
-		void SetStage( bool bBeforeNormalOperation );
-
 		/** Tags that target surface need to have enabled to receive this modifier. */
 		TArray<FString> RequiredTags;
 
@@ -65,16 +45,22 @@ namespace mu
 		// Wether the modifier has to be applied after the normal node operations or before
 		bool bApplyBeforeNormalOperations = true;
 
-		//-----------------------------------------------------------------------------------------
-		// Interface pattern
-		//-----------------------------------------------------------------------------------------
-		class Private;
+
+	public:
+
+		// Node interface
+		virtual const FNodeType* GetType() const override { return &StaticType; }
+		static const FNodeType* GetStaticType() { return &StaticType; }
 
 	protected:
 
 		//! Forbidden. Manage with the Ptr<> template.
 		inline ~NodeModifier() {}
 
+	private:
+
+		static FNodeType StaticType;
+
 	};
 
 
diff --git a/Engine/Plugins/Experimental/Mutable/Source/MutableTools/Public/MuT/NodeModifierMeshClipDeform.h b/Engine/Plugins/Experimental/Mutable/Source/MutableTools/Public/MuT/NodeModifierMeshClipDeform.h
--- a/Engine/Plugins/Experimental/Mutable/Source/MutableTools/Public/MuT/NodeModifierMeshClipDeform.h
+++ b/Engine/Plugins/Experimental/Mutable/Source/MutableTools/Public/MuT/NodeModifierMeshClipDeform.h
@@ -5,59 +5,37 @@
 #include "HAL/Platform.h"
 #include "MuR/Ptr.h"
 #include "MuR/RefCounted.h"
+#include "MuR/Mesh.h"
 #include "MuT/Node.h"
 #include "MuT/NodeModifier.h"
+#include "MuT/NodeMesh.h"
 
 namespace mu
 {
-
-	// Forward definitions
-	class NodeModifierMeshClipDeform;
-	typedef Ptr<NodeModifierMeshClipDeform> NodeModifierMeshClipDeformPtr;
-	typedef Ptr<const NodeModifierMeshClipDeform> NodeModifierMeshClipDeformPtrConst;
-
-	class NodeMesh;
-	typedef Ptr<NodeMesh> NodeMeshPtr;
-	typedef Ptr<const NodeMesh> NodeMeshPtrConst;
-
-	enum class EShapeBindingMethod : uint32;
 	
-	//! This node makes a new component from several meshes and images.
-	//! \ingroup model
+	/** */
 	class MUTABLETOOLS_API NodeModifierMeshClipDeform : public NodeModifier
 	{
 	public:
 
-		NodeModifierMeshClipDeform();
+		Ptr<NodeMesh> ClipMesh;
 
-		//-----------------------------------------------------------------------------------------
-        // Node interface
-		//-----------------------------------------------------------------------------------------
+		EShapeBindingMethod BindingMethod = EShapeBindingMethod::ClipDeformClosestProject;
 
-        const FNodeType* GetType() const override;
-		static const FNodeType* GetStaticType();
-
-		//-----------------------------------------------------------------------------------------
-        // Own interface
-		//-----------------------------------------------------------------------------------------
+	public:
 
-		void SetClipMesh( NodeMesh* InClipMesh);
-		void SetBindingMethod(EShapeBindingMethod BindingMethod);
-	
-		//-----------------------------------------------------------------------------------------
-		// Interface pattern
-		//-----------------------------------------------------------------------------------------
-		class Private;
-		Private* GetPrivate() const;
+		// Node interface
+		virtual const FNodeType* GetType() const override { return &StaticType; }
+		static const FNodeType* GetStaticType() { return &StaticType; }
 
 	protected:
 
 		//! Forbidden. Manage with the Ptr<> template.
-		~NodeModifierMeshClipDeform();
+		inline ~NodeModifierMeshClipDeform() {}
 
 	private:
 
-		Private* m_pD;
+		static FNodeType StaticType;
 
 	};
 
diff --git a/Engine/Plugins/Experimental/Mutable/Source/MutableTools/Public/MuT/NodeModifierMeshClipMorphPlane.h b/Engine/Plugins/Experimental/Mutable/Source/MutableTools/Public/MuT/NodeModifierMeshClipMorphPlane.h
--- a/Engine/Plugins/Experimental/Mutable/Source/MutableTools/Public/MuT/NodeModifierMeshClipMorphPlane.h
+++ b/Engine/Plugins/Experimental/Mutable/Source/MutableTools/Public/MuT/NodeModifierMeshClipMorphPlane.h
@@ -4,6 +4,7 @@
 
 #include "MuR/Ptr.h"
 #include "MuR/RefCounted.h"
+#include "MuR/Skeleton.h"
 #include "MuT/Node.h"
 #include "MuT/NodeModifier.h"
 
@@ -11,33 +12,65 @@
 namespace mu
 {
 
-	// Forward definitions
-	class NodeModifierMeshClipMorphPlane;
-	typedef Ptr<NodeModifierMeshClipMorphPlane> NodeModifierMeshClipMorphPlanePtr;
-	typedef Ptr<const NodeModifierMeshClipMorphPlane> NodeModifierMeshClipMorphPlanePtrConst;
+	/** */
+	struct FClipMorphPlaneParameters
+	{
+		// Morph field parameters
+
+		 //! Distance to the plane of last affected vertex
+		float DistanceToPlane = 0;
+
+		//! "Linearity" factor of the influence.
+		float LinearityFactor = 0;
+
+		// Ellipse location
+		FVector3f Origin = FVector3f(0, 0, 0);
+		FVector3f Normal = FVector3f(0, 0, 0);
+		float Radius1 = 0, Radius2 = 0, Rotation = 0;
+
+		//! Typed of vertex selection
+		typedef enum
+		{
+			//! All vertices, so no extra info is needed
+			VS_ALL,
+
+			//! Select vertices inside a shape
+			VS_SHAPE,
+
+			//! Select all vertices affected by any bone in a sub hierarchy
+			VS_BONE_HIERARCHY,
+		} EVertexSelection;
+
+		// Vertex selection box
+		uint8 VertexSelectionType = VS_ALL;
+		FVector3f SelectionBoxOrigin = FVector3f(0, 0, 0);
+		FVector3f SelectionBoxRadius = FVector3f(0, 0, 0);
+		FBoneName VertexSelectionBone;
+
+		// Max distance a vertex can have to the bone in order to be affected. A negative value
+		// means no limit.
+		float MaxEffectRadius = -1.0f;
+	};
 
-	struct FBoneName;
 
-	//! This node makes a new component from several meshes and images.
-	//! \ingroup model
+	/** */
 	class MUTABLETOOLS_API NodeModifierMeshClipMorphPlane : public NodeModifier
 	{
 	public:
 
-		NodeModifierMeshClipMorphPlane();
+		FClipMorphPlaneParameters Parameters;
 
-		//-----------------------------------------------------------------------------------------
-        // Node interface
-		//-----------------------------------------------------------------------------------------
+	public:
 
-        const FNodeType* GetType() const override;
-		static const FNodeType* GetStaticType();
+		// Node interface
+		virtual const FNodeType* GetType() const override { return &StaticType; }
+		static const FNodeType* GetStaticType() { return &StaticType; }
 
 		//-----------------------------------------------------------------------------------------
         // Own interface
 		//-----------------------------------------------------------------------------------------
 
-		void SetPlane(float centerX, float centerY, float centerZ, float normalX, float normalY, float normalZ);
+		void SetPlane(FVector3f Center, FVector3f Normal);
 		void SetParams(float dist, float factor);
 		void SetMorphEllipse(float radius1, float radius2, float rotation);
 
@@ -49,20 +82,14 @@ namespace mu
 		//! Only one of Box or Bone Hierarchy can be used (the last one set)
 		void SetVertexSelectionBone(const FBoneName& BoneId, float maxEffectRadius);
 
-		//-----------------------------------------------------------------------------------------
-		// Interface pattern
-		//-----------------------------------------------------------------------------------------
-		class Private;
-		Private* GetPrivate() const;
-
 	protected:
 
 		//! Forbidden. Manage with the Ptr<> template.
-		~NodeModifierMeshClipMorphPlane();
+		~NodeModifierMeshClipMorphPlane() {}
 
 	private:
 
-		Private* m_pD;
+		static FNodeType StaticType;
 
 	};
 
diff --git a/Engine/Plugins/Experimental/Mutable/Source/MutableTools/Public/MuT/NodeModifierMeshClipWithMesh.h b/Engine/Plugins/Experimental/Mutable/Source/MutableTools/Public/MuT/NodeModifierMeshClipWithMesh.h
--- a/Engine/Plugins/Experimental/Mutable/Source/MutableTools/Public/MuT/NodeModifierMeshClipWithMesh.h
+++ b/Engine/Plugins/Experimental/Mutable/Source/MutableTools/Public/MuT/NodeModifierMeshClipWithMesh.h
@@ -6,52 +6,34 @@
 #include "MuR/RefCounted.h"
 #include "MuT/Node.h"
 #include "MuT/NodeModifier.h"
-
+#include "MuT/NodeMesh.h"
 
 namespace mu
 {
 
-	// Forward definitions
-	class NodeMesh;
-
-
-	//! This node makes a new component from several meshes and images.
-	//! \ingroup model
+	/** */
 	class MUTABLETOOLS_API NodeModifierMeshClipWithMesh : public NodeModifier
 	{
 	public:
 
-		NodeModifierMeshClipWithMesh();
-
-		//-----------------------------------------------------------------------------------------
-        // Node interface
-		//-----------------------------------------------------------------------------------------
+		//! 
+		Ptr<NodeMesh> ClipMesh;
 
-        const FNodeType* GetType() const override;
-		static const FNodeType* GetStaticType();
-
-		//-----------------------------------------------------------------------------------------
-        // Own interface
-		//-----------------------------------------------------------------------------------------
-
-		void SetClipMesh(NodeMesh*);
+	public:
 
-		//-----------------------------------------------------------------------------------------
-		// Interface pattern
-		//-----------------------------------------------------------------------------------------
-		class Private;
-		Private* GetPrivate() const;
+		// Node interface
+		virtual const FNodeType* GetType() const override { return &StaticType; }
+		static const FNodeType* GetStaticType() { return &StaticType; }
 
 	protected:
 
 		//! Forbidden. Manage with the Ptr<> template.
-		~NodeModifierMeshClipWithMesh();
+		inline ~NodeModifierMeshClipWithMesh() {}
 
 	private:
 
-		Private* m_pD;
+		static FNodeType StaticType;
 
 	};
 
-
 }
diff --git a/Engine/Plugins/Experimental/Mutable/Source/MutableTools/Public/MuT/NodeModifierMeshClipWithUVMask.h b/Engine/Plugins/Experimental/Mutable/Source/MutableTools/Public/MuT/NodeModifierMeshClipWithUVMask.h
--- a/Engine/Plugins/Experimental/Mutable/Source/MutableTools/Public/MuT/NodeModifierMeshClipWithUVMask.h
+++ b/Engine/Plugins/Experimental/Mutable/Source/MutableTools/Public/MuT/NodeModifierMeshClipWithUVMask.h
@@ -6,53 +6,37 @@
 #include "MuR/RefCounted.h"
 #include "MuT/Node.h"
 #include "MuT/NodeModifier.h"
+#include "MuT/NodeImage.h"
 
 
 namespace mu
 {
 
-	// Forward definitions
-	class NodeImage;
-
-	//! This node makes a new component from several meshes and images.
-	//! \ingroup model
+	/** */
 	class MUTABLETOOLS_API NodeModifierMeshClipWithUVMask : public NodeModifier
 	{
 	public:
 
-		NodeModifierMeshClipWithUVMask();
-
-		//-----------------------------------------------------------------------------------------
-        // Node interface
-		//-----------------------------------------------------------------------------------------
-
-        const FNodeType* GetType() const override;
-		static const FNodeType* GetStaticType();
+		/** Image with the regions to remove. It will be interpreted as a bitmap. */
+		Ptr<NodeImage> ClipMask;
 
-		//-----------------------------------------------------------------------------------------
-        // Own interface
-		//-----------------------------------------------------------------------------------------
+		/** Layout index of the UVs to use inthe source mesh to ben clipped with the mask. */
+		uint8 LayoutIndex = 0;
 
-		/** Set the expression generating the image witht he UV mask use to clip the mesh. */
-		void SetClipMask(NodeImage*);
-
-		/** Set the UV channel index for the UVs to check against the mask. */
-		void SetLayoutIndex(uint8 LayoutIndex);
+	public:
 
-		//-----------------------------------------------------------------------------------------
-		// Interface pattern
-		//-----------------------------------------------------------------------------------------
-		class Private;
-		Private* GetPrivate() const;
+		// Node interface
+		virtual const FNodeType* GetType() const override { return &StaticType; }
+		static const FNodeType* GetStaticType() { return &StaticType; }
 
 	protected:
 
 		//! Forbidden. Manage with the Ptr<> template.
-		~NodeModifierMeshClipWithUVMask();
+		inline ~NodeModifierMeshClipWithUVMask() {}
 
 	private:
 
-		Private* m_pD;
+		static FNodeType StaticType;
 
 	};
 
diff --git a/Engine/Plugins/Experimental/Mutable/Source/MutableTools/Public/MuT/NodePatchMesh.h b/Engine/Plugins/Experimental/Mutable/Source/MutableTools/Public/MuT/NodePatchMesh.h
deleted file mode 100644
diff --git a/Engine/Plugins/Experimental/Mutable/Source/MutableTools/Public/MuT/NodeSurfaceEdit.h b/Engine/Plugins/Experimental/Mutable/Source/MutableTools/Public/MuT/NodeSurfaceEdit.h
--- a/Engine/Plugins/Experimental/Mutable/Source/MutableTools/Public/MuT/NodeSurfaceEdit.h
+++ b/Engine/Plugins/Experimental/Mutable/Source/MutableTools/Public/MuT/NodeSurfaceEdit.h
@@ -4,13 +4,12 @@
 
 #include "MuR/Ptr.h"
 #include "MuR/RefCounted.h"
+#include "MuR/ImageTypes.h"
 #include "MuT/Node.h"
 #include "MuT/NodeSurface.h"
 #include "MuT/NodeMesh.h"
-#include "MuT/NodePatchMesh.h"
 #include "MuT/NodeImage.h"
 #include "MuT/NodeScalar.h"
-#include "MuT/NodePatchImage.h"
 
 
 namespace mu
@@ -25,20 +24,47 @@ namespace mu
 	public:
 
 		Ptr<NodeSurface> Parent;
-		Ptr<NodePatchMesh> Mesh;
-		Ptr<NodeMesh> Morph;
+
+		/** Mesh to remove from the modified surface. */
+		Ptr<NodeMesh> MeshRemove;
+
+		/** Mesh to add to the modified surface. */
+		Ptr<NodeMesh> MeshAdd;
+
+		/** Morph to apply to the modified surface. */
+		Ptr<NodeMesh> MeshMorph;
+
+		/** Factor of the morph to apply. */
 		Ptr<NodeScalar> MorphFactor;
 
+		/** Data for every modified texture. */
 		struct FTexture
 		{
+			/** Image to add if extgending. */
 			Ptr<NodeImage> Extend;
-			Ptr<NodePatchImage> Patch;
+
+			/** Image to blend if patching. */
+			Ptr<NodeImage> PatchImage;
+
+			/** Optional mask controlling the blending area. */
+			Ptr<NodeImage> PatchMask;
+
+			/** Rects in the parent layout homogeneous UV space to patch. */
+			TArray<FBox2f> PatchBlocks;
+
+			/** */
+			EBlendType PatchBlendType = EBlendType::BT_BLEND;
+
+			/** Patch alpha channel as well? */
+			bool bPatchApplyToAlpha = false;
+
 		};
 
+		/** Textures to modify. */
 		TArray<FTexture> Textures;
 
-		//! Tags in this surface edit
-		TArray<FString> Tags;
+		/** Tags enabled by this surface edit. */
+		TArray<FString> EnableTags;
 
 	public:
 
