[Mutable] Automatically store gameplay tags in each row's editor metadata.

- Added public API to retrieve a DataTable given a parameter and its value.

#rnx

diff --git a/Engine/Plugins/Experimental/Mutable/Source/CustomizableObject/Internal/MuCO/CustomizableObjectPrivate.h b/Engine/Plugins/Experimental/Mutable/Source/CustomizableObject/Internal/MuCO/CustomizableObjectPrivate.h
--- a/Engine/Plugins/Experimental/Mutable/Source/CustomizableObject/Internal/MuCO/CustomizableObjectPrivate.h
+++ b/Engine/Plugins/Experimental/Mutable/Source/CustomizableObject/Internal/MuCO/CustomizableObjectPrivate.h
@@ -949,6 +949,11 @@ struct FModelResources
 	UPROPERTY()
 	TMap<FString, FMutableStateData> StateUIDataMap;
 
+#if WITH_EDITORONLY_DATA
+	/** DataTable used by an int parameter and its value. */
+	TMap<TTuple<FString, FString>, TSet<TSoftObjectPtr<UDataTable>>> IntParameterOptionDataTable;
+#endif
+	
 	UPROPERTY()
 	TArray<FCustomizableObjectClothConfigData> ClothSharedConfigsData;	
 
@@ -1314,6 +1319,8 @@ public:
 		ModelStreamableBulkData,
 
 		LayoutBlocksAsInt32,
+		
+		IntParameterOptionDataTable,
 
 		RemoveLODCountLimit,
 
diff --git a/Engine/Plugins/Experimental/Mutable/Source/CustomizableObject/Private/MuCO/CustomizableObject.cpp b/Engine/Plugins/Experimental/Mutable/Source/CustomizableObject/Private/MuCO/CustomizableObject.cpp
--- a/Engine/Plugins/Experimental/Mutable/Source/CustomizableObject/Private/MuCO/CustomizableObject.cpp
+++ b/Engine/Plugins/Experimental/Mutable/Source/CustomizableObject/Private/MuCO/CustomizableObject.cpp
@@ -600,6 +600,11 @@ void UCustomizableObjectPrivate::SaveCompiledData(FArchive& MemoryWriter, bool b
 
 	SerializeStreamedResources(MemoryWriter, GetPublic(), GetPublic()->StreamedResourceData, bIsCooking);
 
+	for (FCustomizableObjectStreamedResourceData Extension : GetPublic()->StreamedExtensionData)
+	{
+		Extension.GetPath()->Data.Data.Serialize(MemoryWriter);
+	}
+
 	int32 NumReferencedMaterials = LocalModelResources.Materials.Num();
 	MemoryWriter << NumReferencedMaterials;
 
@@ -677,6 +682,10 @@ void UCustomizableObjectPrivate::SaveCompiledData(FArchive& MemoryWriter, bool b
 	MemoryWriter << LocalModelResources.ParameterUIDataMap;
 	MemoryWriter << LocalModelResources.StateUIDataMap;
 
+#if WITH_EDITORONLY_DATA
+	MemoryWriter << LocalModelResources.IntParameterOptionDataTable;
+#endif
+	
 	MemoryWriter << LocalModelResources.ClothingAssetsData;
 	MemoryWriter << LocalModelResources.ClothSharedConfigsData;
 
@@ -823,6 +832,10 @@ void UCustomizableObjectPrivate::LoadCompiledData(FArchive& MemoryReader, const
 		MemoryReader << LocalModelResource.ParameterUIDataMap;
 		MemoryReader << LocalModelResource.StateUIDataMap;
 
+#if WITH_EDITORONLY_DATA
+		MemoryReader << LocalModelResource.IntParameterOptionDataTable;
+#endif
+
 		MemoryReader << LocalModelResource.ClothingAssetsData; 
 		MemoryReader << LocalModelResource.ClothSharedConfigsData; 
 
@@ -1686,6 +1699,17 @@ FMutableStateUIMetadata UCustomizableObject::GetStateUIMetadata(const FString& S
 }
 
 
+#if WITH_EDITOR
+TSet<TSoftObjectPtr<UDataTable>> UCustomizableObject::GetIntParameterOptionDataTable(FString& ParamName, const FString& OptionName)
+{
+	const FModelResources& ModelResources = GetPrivate()->GetModelResources();
+	const TSet<TSoftObjectPtr<UDataTable>>* Result = ModelResources.IntParameterOptionDataTable.Find(MakeTuple(ParamName, OptionName));
+	
+	return Result ? *Result : TSet<TSoftObjectPtr<UDataTable>>();
+}
+#endif
+
+
 float UCustomizableObject::GetFloatParameterDefaultValue(const FString& InParameterName) const
 {
 	const int32 ParameterIndex = FindParameter(InParameterName);
diff --git a/Engine/Plugins/Experimental/Mutable/Source/CustomizableObject/Private/MuCO/CustomizableObjectInstanceDescriptor.cpp b/Engine/Plugins/Experimental/Mutable/Source/CustomizableObject/Private/MuCO/CustomizableObjectInstanceDescriptor.cpp
--- a/Engine/Plugins/Experimental/Mutable/Source/CustomizableObject/Private/MuCO/CustomizableObjectInstanceDescriptor.cpp
+++ b/Engine/Plugins/Experimental/Mutable/Source/CustomizableObject/Private/MuCO/CustomizableObjectInstanceDescriptor.cpp
@@ -1500,8 +1500,8 @@ FLinearColor FCustomizableObjectInstanceDescriptor::GetColorParameterSelectedOpt
 		LogParameterNotFoundWarning(ColorParamName, ParameterIndexInObject, ColorParamIndex, CustomizableObject, __FUNCTION__);
 		return FCustomizableObjectVectorParameterValue::DEFAULT_PARAMETER_VALUE;
 	}
-	
-	return VectorParameters[ColorParamIndex].ParameterValue;
+
+	return VectorParameters.IsValidIndex(ColorParamIndex) ? VectorParameters[ColorParamIndex].ParameterValue : FLinearColor();
 }
 
 
diff --git a/Engine/Plugins/Experimental/Mutable/Source/CustomizableObject/Public/MuCO/CustomizableObject.h b/Engine/Plugins/Experimental/Mutable/Source/CustomizableObject/Public/MuCO/CustomizableObject.h
--- a/Engine/Plugins/Experimental/Mutable/Source/CustomizableObject/Public/MuCO/CustomizableObject.h
+++ b/Engine/Plugins/Experimental/Mutable/Source/CustomizableObject/Public/MuCO/CustomizableObject.h
@@ -588,6 +588,11 @@ public:
 	UFUNCTION(BlueprintCallable, Category = CustomizableObject)
 	FMutableStateUIMetadata GetStateUIMetadata(const FString& StateName) const;
 
+#if WITH_EDITOR
+	/** Return the DataTables used by the given parameter and its value (if any). */
+	UFUNCTION(BlueprintCallable, Category = CustomizableObject)
+	TSet<TSoftObjectPtr<UDataTable>> GetIntParameterOptionDataTable(FString& ParamName, const FString& OptionName);
+#endif
 	
 private:
 	/** Textures marked as low priority will generate defaulted resident mips (if texture streaming is enabled).
diff --git a/Engine/Plugins/Experimental/Mutable/Source/CustomizableObjectEditor/Private/MuCOE/CustomizableObjectCompiler.cpp b/Engine/Plugins/Experimental/Mutable/Source/CustomizableObjectEditor/Private/MuCOE/CustomizableObjectCompiler.cpp
--- a/Engine/Plugins/Experimental/Mutable/Source/CustomizableObjectEditor/Private/MuCOE/CustomizableObjectCompiler.cpp
+++ b/Engine/Plugins/Experimental/Mutable/Source/CustomizableObjectEditor/Private/MuCOE/CustomizableObjectCompiler.cpp
@@ -1013,6 +1013,7 @@ void FCustomizableObjectCompiler::CompileInternal(bool bAsync)
 
 		ModelResources.ParameterUIDataMap = MoveTemp(GenerationContext.ParameterUIDataMap);
 		ModelResources.StateUIDataMap = MoveTemp(GenerationContext.StateUIDataMap);
+		ModelResources.IntParameterOptionDataTable = MoveTemp(GenerationContext.IntParameterOptionDataTable);
 
 		// Create the RealTimeMorphsTargets Blocks from the per mesh Morph data.
 		uint64 RealTimeMorphDataSize = 0;
diff --git a/Engine/Plugins/Experimental/Mutable/Source/CustomizableObjectEditor/Private/MuCOE/GenerateMutableSource/GenerateMutableSource.h b/Engine/Plugins/Experimental/Mutable/Source/CustomizableObjectEditor/Private/MuCOE/GenerateMutableSource/GenerateMutableSource.h
--- a/Engine/Plugins/Experimental/Mutable/Source/CustomizableObjectEditor/Private/MuCOE/GenerateMutableSource/GenerateMutableSource.h
+++ b/Engine/Plugins/Experimental/Mutable/Source/CustomizableObjectEditor/Private/MuCOE/GenerateMutableSource/GenerateMutableSource.h
@@ -750,6 +750,7 @@ struct FMutableGraphGenerationContext
 	TMap<FString, FString> CustomizableObjectPathMap;
 	TMap<FString, FMutableParameterData> ParameterUIDataMap;
 	TMap<FString, FMutableStateData> StateUIDataMap;
+	TMap<TTuple<FString, FString>, TSet<TSoftObjectPtr<UDataTable>>> IntParameterOptionDataTable;
 	TMultiMap<const UCustomizableObjectNodeObjectGroup*, FGroupProjectorTempData> ProjectorGroupMap;
 	//TMap<UPhysicsAsset*, uint32> DiscartedPhysicsAssetMap;
 
diff --git a/Engine/Plugins/Experimental/Mutable/Source/CustomizableObjectEditor/Private/MuCOE/GenerateMutableSource/GenerateMutableSourceTable.cpp b/Engine/Plugins/Experimental/Mutable/Source/CustomizableObjectEditor/Private/MuCOE/GenerateMutableSource/GenerateMutableSourceTable.cpp
--- a/Engine/Plugins/Experimental/Mutable/Source/CustomizableObjectEditor/Private/MuCOE/GenerateMutableSource/GenerateMutableSourceTable.cpp
+++ b/Engine/Plugins/Experimental/Mutable/Source/CustomizableObjectEditor/Private/MuCOE/GenerateMutableSource/GenerateMutableSourceTable.cpp
@@ -170,8 +170,8 @@ bool FillTableColumn(const UCustomizableObjectNodeTable* TableNode, mu::TablePtr
 								{
 									if (GameplayTagData)
 									{
-										const FGameplayTagContainer TagContainer = *(FGameplayTagContainer*)GameplayTagData;
-										TagContainer.GetGameplayTagArray(GameplayTags);
+										const FGameplayTagContainer* TagContainer = reinterpret_cast<FGameplayTagContainer*>(GameplayTagData);
+										TagContainer->GetGameplayTagArray(GameplayTags);
 									}
 								}
 							}
@@ -958,11 +958,18 @@ void GenerateTableParameterUIData(const UDataTable* DataTable, const UCustomizab
 	// Checking if the parameter name already exists
 	GenerationContext.AddParameterNameUnique(TableNode, TableNode->ParameterName);
 
+	TArray<FName> RowNames = GetRowsToCompile(*DataTable, *TableNode, GenerationContext);
+
+	for (const FName& Name : RowNames)
+	{
+		TSet<TSoftObjectPtr<UDataTable>>& DataTables = GenerationContext.IntParameterOptionDataTable.FindOrAdd(MakeTuple(TableNode->ParameterName, Name.ToString()));
+		DataTables.Add(TSoftObjectPtr<UDataTable>(DataTable));
+	}
+	
 	// Generating Parameter UI MetaData if not exists
 	if (!GenerationContext.ParameterUIDataMap.Contains(TableNode->ParameterName))
 	{
 		// Getting Table and row names to access the information
-		TArray<FName> RowNames = GetRowsToCompile(*DataTable, *TableNode, GenerationContext);
 
 		FMutableParameterData ParameterUIData(TableNode->ParamUIMetadata, EMutableParameterType::Int);
 		ParameterUIData.IntegerParameterGroupType = TableNode->bAddNoneOption ? ECustomizableObjectGroupType::COGT_ONE_OR_NONE : ECustomizableObjectGroupType::COGT_ONE;
@@ -992,37 +999,65 @@ void GenerateTableParameterUIData(const UDataTable* DataTable, const UCustomizab
 			GenerationContext.Compiler->CompilerLog(LogMessage, TableNode);
 		}
 
-		if (!bIsValidMetadataColumn && !bIsValidThumbnailColumn)
+		if (!bIsValidMetadataColumn)
 		{
-			// Do no add a meta data if it is not needed.
 			return;
 		}
 
 		for (int32 NameIndex = 0; NameIndex < RowNames.Num(); ++NameIndex)
 		{
 			FName RowName = RowNames[NameIndex];
-			FMutableParamUIMetadata MetadataValue;
 
-			if (bIsValidMetadataColumn)
+			if (uint8* MetadataCellData = GetCellData(RowName, *DataTable, *MetadataColumnProperty))
 			{
-				if (uint8* CellData = GetCellData(RowName, *DataTable, *MetadataColumnProperty))
+				FMutableParamUIMetadata MetadataValue = *reinterpret_cast<FMutableParamUIMetadata*>(MetadataCellData);
+
+				FIntegerParameterUIData IntegerMetadata = FIntegerParameterUIData(MetadataValue);
+
+				// Add thumbnail
+				if (bIsValidThumbnailColumn && MetadataValue.EditorUIThumbnailObject.IsNull())
 				{
-					MetadataValue = *(FMutableParamUIMetadata*)CellData;
+					if (uint8* ThumbnailCellData = GetCellData(RowName, *DataTable, *ThumbnailColumnProperty))
+					{
+						FSoftObjectPtr* ObjectPtr = reinterpret_cast<FSoftObjectPtr*>(ThumbnailCellData);
+						IntegerMetadata.ParamUIMetadata.EditorUIThumbnailObject = ObjectPtr->ToSoftObjectPath();
+					}
 				}
-			}
-
-			FIntegerParameterUIData IntegerMetadata = FIntegerParameterUIData(MetadataValue);
 
-			if (bIsValidThumbnailColumn && MetadataValue.EditorUIThumbnailObject.IsNull())
-			{
-				if (uint8* CellData = GetCellData(RowName, *DataTable, *ThumbnailColumnProperty))
+				// Add tags
+				if (TableNode->bGatherTags)
 				{
-					FSoftObjectPtr ObjectPtr = *(FSoftObjectPtr*)CellData;
-					IntegerMetadata.ParamUIMetadata.EditorUIThumbnailObject = ObjectPtr.ToSoftObjectPath();
+					if (const UScriptStruct* Struct = DataTable->GetRowStruct())
+					{
+						for (TFieldIterator<FProperty> It(Struct); It; ++It)
+						{
+							FProperty* ColumnProperty = *It;
+
+							if (!ColumnProperty)
+							{
+								continue;
+							}
+
+							if (const FStructProperty* StructProperty = CastField<FStructProperty>(ColumnProperty))
+							{
+								if (StructProperty->Struct == TBaseStructure<FGameplayTagContainer>::Get())
+								{
+									FName ColumnName = FName(DataTableUtils::GetPropertyExportName(ColumnProperty));
+								
+									FProperty* TagColumnProperty = DataTable->FindTableProperty(ColumnName);
+									if (uint8* TagCellData = GetCellData(RowName, *DataTable, *TagColumnProperty))
+									{
+										FGameplayTagContainer* TagContainer = reinterpret_cast<FGameplayTagContainer*>(TagCellData);
+										IntegerMetadata.ParamUIMetadata.EditorGameplayTags.AppendTags(*TagContainer);
+									}
+								}
+							}
+						}
+					}
 				}
-			}
 
-			ParameterUIDataRef.ArrayIntegerParameterOption.Add(RowName.ToString(), IntegerMetadata);
+				ParameterUIDataRef.ArrayIntegerParameterOption.Add(RowName.ToString(), IntegerMetadata);
+			}
 		}
 	}
 }
diff --git a/Engine/Plugins/Experimental/Mutable/Source/CustomizableObjectEditor/Public/MuCOE/Nodes/CustomizableObjectNodeTable.h b/Engine/Plugins/Experimental/Mutable/Source/CustomizableObjectEditor/Public/MuCOE/Nodes/CustomizableObjectNodeTable.h
--- a/Engine/Plugins/Experimental/Mutable/Source/CustomizableObjectEditor/Public/MuCOE/Nodes/CustomizableObjectNodeTable.h
+++ b/Engine/Plugins/Experimental/Mutable/Source/CustomizableObjectEditor/Public/MuCOE/Nodes/CustomizableObjectNodeTable.h
@@ -225,6 +225,9 @@ public:
 	/** Name of the column that contains the asset to use its thumbnails as option thumbnails. */
 	UPROPERTY(EditAnywhere, Category = UI)
 	FName ThumbnailColumn;
+	
+	UPROPERTY(EditAnywhere, Category = UI, meta = (Tooltip = "Given a row, add all tags found in GameplayTag columns to its Parameter UI Metadata"))
+	bool bGatherTags = true;
 
 	/** Map to relate a Structure Column with its Data */
 	UPROPERTY()
