[Mutable] Remove incorrect code from previous 35691972 submit.

Two shelves were incorrectly merged.

#rnx

diff --git a/Engine/Plugins/Experimental/Mutable/Source/CustomizableObject/Internal/MuCO/CustomizableObjectPrivate.h b/Engine/Plugins/Experimental/Mutable/Source/CustomizableObject/Internal/MuCO/CustomizableObjectPrivate.h
--- a/Engine/Plugins/Experimental/Mutable/Source/CustomizableObject/Internal/MuCO/CustomizableObjectPrivate.h
+++ b/Engine/Plugins/Experimental/Mutable/Source/CustomizableObject/Internal/MuCO/CustomizableObjectPrivate.h
@@ -1324,6 +1324,8 @@ public:
 
 		RemoveLODCountLimit,
 
+		IntParameterOptionDataTablePartialBackout,
+		
 		// -----<new versions can be added above this line>--------
 		LastCustomizableObjectVersion
 	};
diff --git a/Engine/Plugins/Experimental/Mutable/Source/CustomizableObject/Private/MuCO/CustomizableObject.cpp b/Engine/Plugins/Experimental/Mutable/Source/CustomizableObject/Private/MuCO/CustomizableObject.cpp
--- a/Engine/Plugins/Experimental/Mutable/Source/CustomizableObject/Private/MuCO/CustomizableObject.cpp
+++ b/Engine/Plugins/Experimental/Mutable/Source/CustomizableObject/Private/MuCO/CustomizableObject.cpp
@@ -599,12 +599,7 @@ void UCustomizableObjectPrivate::SaveCompiledData(FArchive& MemoryWriter, bool b
 	MemoryWriter << LocalModelResources.ReferenceSkeletalMeshesData;
 
 	SerializeStreamedResources(MemoryWriter, GetPublic(), GetPublic()->StreamedResourceData, bIsCooking);
-
-	for (FCustomizableObjectStreamedResourceData Extension : GetPublic()->StreamedExtensionData)
-	{
-		Extension.GetPath()->Data.Data.Serialize(MemoryWriter);
-	}
-
+	
 	int32 NumReferencedMaterials = LocalModelResources.Materials.Num();
 	MemoryWriter << NumReferencedMaterials;
 
