[Mutable] Fix TSoftObjectPtr constantness static analyzer warning.

#rnx

diff --git a/Engine/Plugins/Experimental/Mutable/Source/CustomizableObjectEditor/Private/MuCOE/GenerateMutableSource/GenerateMutableSourceTable.cpp b/Engine/Plugins/Experimental/Mutable/Source/CustomizableObjectEditor/Private/MuCOE/GenerateMutableSource/GenerateMutableSourceTable.cpp
--- a/Engine/Plugins/Experimental/Mutable/Source/CustomizableObjectEditor/Private/MuCOE/GenerateMutableSource/GenerateMutableSourceTable.cpp
+++ b/Engine/Plugins/Experimental/Mutable/Source/CustomizableObjectEditor/Private/MuCOE/GenerateMutableSource/GenerateMutableSourceTable.cpp
@@ -963,7 +963,7 @@ void GenerateTableParameterUIData(const UDataTable* DataTable, const UCustomizab
 	for (const FName& Name : RowNames)
 	{
 		TSet<TSoftObjectPtr<UDataTable>>& DataTables = GenerationContext.IntParameterOptionDataTable.FindOrAdd(MakeTuple(TableNode->ParameterName, Name.ToString()));
-		DataTables.Add(TSoftObjectPtr<UDataTable>(DataTable));
+		DataTables.Add(TSoftObjectPtr<UDataTable>(const_cast<UDataTable*>(DataTable)));
 	}
 	
 	// Generating Parameter UI MetaData if not exists
