[Mutable] Disable Layout block generation when packing strategy is set to Overlay.

- When the layout packing is set to Overlay uv are allowed to not be normalized or even negative. This could generate blocks with negative coords which are not supported.

#rnx

diff --git a/Engine/Plugins/Experimental/Mutable/Source/CustomizableObjectEditor/Private/MuCOE/GenerateMutableSource/GenerateMutableSourceLayout.cpp b/Engine/Plugins/Experimental/Mutable/Source/CustomizableObjectEditor/Private/MuCOE/GenerateMutableSource/GenerateMutableSourceLayout.cpp
--- a/Engine/Plugins/Experimental/Mutable/Source/CustomizableObjectEditor/Private/MuCOE/GenerateMutableSource/GenerateMutableSourceLayout.cpp
+++ b/Engine/Plugins/Experimental/Mutable/Source/CustomizableObjectEditor/Private/MuCOE/GenerateMutableSource/GenerateMutableSourceLayout.cpp
@@ -112,9 +112,11 @@ mu::Ptr<mu::NodeLayout> CreateMutableLayoutNode(FMutableGraphGenerationContext&
 
 
 	ECustomizableObjectLayoutAutomaticBlocksStrategy AutomaticBlockStrategy = UnrealLayout->AutomaticBlocksStrategy;
-
-	if (AutomaticBlockStrategy == ECustomizableObjectLayoutAutomaticBlocksStrategy::Ignore)
-	{
+	ECustomizableObjectTextureLayoutPackingStrategy PackingStrategy = UnrealLayout->PackingStrategy;
+	
+	if (AutomaticBlockStrategy == ECustomizableObjectLayoutAutomaticBlocksStrategy::Ignore || 
+		PackingStrategy == ECustomizableObjectTextureLayoutPackingStrategy::Overlay)
+ 	{
 		// Legacy behavior
 		if (!UnrealLayout->Blocks.Num())
 		{
