[Mutable] Add cvar to disable realtime morph targets from loading and generating.

- Added cvar mutable.EnableRealTimeMorphTargets with default value true.

#rnx

diff --git a/Engine/Plugins/Experimental/Mutable/Source/CustomizableObject/Private/MuCO/CustomizableObjectInstance.cpp b/Engine/Plugins/Experimental/Mutable/Source/CustomizableObject/Private/MuCO/CustomizableObjectInstance.cpp
--- a/Engine/Plugins/Experimental/Mutable/Source/CustomizableObject/Private/MuCO/CustomizableObjectInstance.cpp
+++ b/Engine/Plugins/Experimental/Mutable/Source/CustomizableObject/Private/MuCO/CustomizableObjectInstance.cpp
@@ -3952,6 +3952,12 @@ void UCustomizableInstancePrivate::BuildOrCopyMorphTargetsData(const TSharedRef<
 {
 	MUTABLE_CPUPROFILER_SCOPE(UCustomizableInstancePrivate::BuildOrCopyMorphTargetsData);
 
+	// This is a bit redundant as ComponentMorphTargets should not be generated.
+	if (!CVarEnableRealTimeMorphTargets.GetValueOnAnyThread())
+	{
+		return;
+	}
+
 	if (!SkeletalMesh)
 	{
 		return;
@@ -5297,6 +5303,8 @@ UE::Tasks::FTask UCustomizableInstancePrivate::LoadAdditionalAssetsAndData(
 		}
 	}
 
+	const bool bMorphTargetsEnabled = CVarEnableRealTimeMorphTargets.GetValueOnAnyThread();
+
 	// Load assets coming from SubMeshes of the newly generated Mesh
 	if (OperationData->InstanceUpdateData.LODs.Num())
 	{
@@ -5342,16 +5350,19 @@ UE::Tasks::FTask UCustomizableInstancePrivate::LoadAdditionalAssetsAndData(
 					}
 					else if (TypedResourceId.Type == (uint8)FCustomizableObjectStreameableResourceId::EType::RealTimeMorphTarget)
 					{
-						check(TypedResourceId.Id != 0 && TypedResourceId.Id <= TNumericLimits<uint32>::Max());
-
-						const TMap<uint32, FRealTimeMorphStreamable>& MorphsStreamables = ModelStreamableBulkData->RealTimeMorphStreamables;
-						if (MorphsStreamables.Contains((uint32)TypedResourceId.Id))
-						{
-							RealTimeMorphStreamableBlocksToStream.AddUnique(TypedResourceId.Id);
-						}
-						else
+						if (bMorphTargetsEnabled)
 						{
-							UE_LOG(LogMutable, Error, TEXT("Invalid streamed real time morph target data block [%d] found."), TypedResourceId.Id);
+							check(TypedResourceId.Id != 0 && TypedResourceId.Id <= TNumericLimits<uint32>::Max());
+
+							const TMap<uint32, FRealTimeMorphStreamable>& MorphsStreamables = ModelStreamableBulkData->RealTimeMorphStreamables;
+							if (MorphsStreamables.Contains((uint32)TypedResourceId.Id))
+							{
+								RealTimeMorphStreamableBlocksToStream.AddUnique(TypedResourceId.Id);
+							}
+							else
+							{
+								UE_LOG(LogMutable, Error, TEXT("Invalid streamed real time morph target data block [%d] found."), TypedResourceId.Id);
+							}
 						}
 					}
 					else if (TypedResourceId.Type == (uint8)FCustomizableObjectStreameableResourceId::EType::Clothing)
@@ -5544,6 +5555,7 @@ UE::Tasks::FTask UCustomizableInstancePrivate::LoadAdditionalAssetsAndData(
 	{
 #if WITH_EDITOR
 		// On editor the data is always loaded, load directly form the ModelResources.
+		if (bMorphTargetsEnabled)
 		{
 			MUTABLE_CPUPROFILER_SCOPE(RealTimeMorphStreamingEditor);
 			for (uint32 BlockId : RealTimeMorphStreamableBlocksToStream)
@@ -5637,46 +5649,48 @@ UE::Tasks::FTask UCustomizableInstancePrivate::LoadAdditionalAssetsAndData(
 			return nullptr;
 		};
 
-		const int32 NumMorphBlocks = RealTimeMorphStreamableBlocksToStream.Num();
-		for (int32 I = 0; I < NumMorphBlocks; ++I)
+		if (bMorphTargetsEnabled)
 		{
-			MUTABLE_CPUPROFILER_SCOPE(RealTimeMorphStreamingRequest_Alloc);
+			const int32 NumMorphBlocks = RealTimeMorphStreamableBlocksToStream.Num();
+			for (int32 I = 0; I < NumMorphBlocks; ++I)
+			{
+				MUTABLE_CPUPROFILER_SCOPE(RealTimeMorphStreamingRequest_Alloc);
 
-			const int32 BlockId = RealTimeMorphStreamableBlocksToStream[I];
-			const FRealTimeMorphStreamable& Streamable = ModelStreamableBulkData->RealTimeMorphStreamables[BlockId];
-			const FMutableStreamableBlock& Block = Streamable.Block; 
-		
-			FInstanceUpdateData::FMorphTargetMeshData& ReadDestData = 
-					OperationData->InstanceUpdateData.RealTimeMorphTargetMeshData.FindOrAdd(BlockId);
+				const int32 BlockId = RealTimeMorphStreamableBlocksToStream[I];
+				const FRealTimeMorphStreamable& Streamable = ModelStreamableBulkData->RealTimeMorphStreamables[BlockId];
+				const FMutableStreamableBlock& Block = Streamable.Block; 
+			
+				FInstanceUpdateData::FMorphTargetMeshData& ReadDestData = 
+						OperationData->InstanceUpdateData.RealTimeMorphTargetMeshData.FindOrAdd(BlockId);
 
-			// Only request blocks once.
-			if (ReadDestData.Data.Num())
-			{
-				continue;
-			}
+				// Only request blocks once.
+				if (ReadDestData.Data.Num())
+				{
+					continue;
+				}
 
-			ReadDestData.NameResolutionMap = ModelStreamableBulkData->RealTimeMorphStreamables[BlockId].NameResolutionMap;
+				ReadDestData.NameResolutionMap = ModelStreamableBulkData->RealTimeMorphStreamables[BlockId].NameResolutionMap;
 
-			check(Streamable.Size % sizeof(FMorphTargetVertexData) == 0);
-			uint32 NumElems = Streamable.Size / sizeof(FMorphTargetVertexData);
+				check(Streamable.Size % sizeof(FMorphTargetVertexData) == 0);
+				uint32 NumElems = Streamable.Size / sizeof(FMorphTargetVertexData);
 
-			ReadDestData.Data.SetNumUninitialized(NumElems);
+				ReadDestData.Data.SetNumUninitialized(NumElems);
 
-			IAsyncReadFileHandle* FileHandle = nullptr;
-			if (!bUseFBulkData)
-			{	
-				FileHandle = OpenOrGetFileHandleForBlock(Block);
-			}
+				IAsyncReadFileHandle* FileHandle = nullptr;
+				if (!bUseFBulkData)
+				{	
+					FileHandle = OpenOrGetFileHandleForBlock(Block);
+				}
 
-			BlockReadInfos.Emplace(FBlockReadInfo
-			{ 
-				Block.Offset,
-				FileHandle,
-				MakeArrayView(reinterpret_cast<uint8*>(ReadDestData.Data.GetData()), ReadDestData.Data.Num()*sizeof(FMorphTargetVertexData)),
-				Block.FileId
-			});
+				BlockReadInfos.Emplace(FBlockReadInfo
+				{ 
+					Block.Offset,
+					FileHandle,
+					MakeArrayView(reinterpret_cast<uint8*>(ReadDestData.Data.GetData()), ReadDestData.Data.Num()*sizeof(FMorphTargetVertexData)),
+					Block.FileId
+				});
+			}
 		}
-
 		const int32 NumClothBlocks = ClothingStreamableBlocksToStream.Num();
 		for (int32 I = 0; I < NumClothBlocks; ++I)
 		{
diff --git a/Engine/Plugins/Experimental/Mutable/Source/CustomizableObject/Private/MuCO/CustomizableObjectSystem.cpp b/Engine/Plugins/Experimental/Mutable/Source/CustomizableObject/Private/MuCO/CustomizableObjectSystem.cpp
--- a/Engine/Plugins/Experimental/Mutable/Source/CustomizableObject/Private/MuCO/CustomizableObjectSystem.cpp
+++ b/Engine/Plugins/Experimental/Mutable/Source/CustomizableObject/Private/MuCO/CustomizableObjectSystem.cpp
@@ -115,6 +115,11 @@ TAutoConsoleVariable<bool> CVarEnableUpdateOptimization(
 	false,
 	TEXT("Enable or disable update optimization when no changes are made to the parent component."));
 
+TAutoConsoleVariable<bool> CVarEnableRealTimeMorphTargets(
+	TEXT("mutable.EnableRealTimeMorphTargets"),
+	true,
+	TEXT("Enable or disable generation of realtime morph targets."));
+
 #if WITH_EDITOR
 bool bEnableLODManagmentInEditor = false;
 
@@ -2918,9 +2923,13 @@ namespace impl
 			return;
 		}
 
-		// TODO: This subtask should execute before Convert resources in a worker thread but after 
-		// Loading resources. For now keep it here.
-		Subtask_Mutable_PrepareRealTimeMorphData(OperationData);
+		
+		if (CVarEnableRealTimeMorphTargets.GetValueOnAnyThread())
+		{
+			// TODO: This subtask should execute before Convert resources in a worker thread but after 
+			// Loading resources. For now keep it here.
+			Subtask_Mutable_PrepareRealTimeMorphData(OperationData);
+		}
 
 		UCustomizableObjectInstance* CustomizableObjectInstance = OperationData->Instance.Get();
 
diff --git a/Engine/Plugins/Experimental/Mutable/Source/CustomizableObject/Public/MuCO/CustomizableObjectSystem.h b/Engine/Plugins/Experimental/Mutable/Source/CustomizableObject/Public/MuCO/CustomizableObjectSystem.h
--- a/Engine/Plugins/Experimental/Mutable/Source/CustomizableObject/Public/MuCO/CustomizableObjectSystem.h
+++ b/Engine/Plugins/Experimental/Mutable/Source/CustomizableObject/Public/MuCO/CustomizableObjectSystem.h
@@ -37,8 +37,11 @@ extern TAutoConsoleVariable<bool> CVarPreserveUserLODsOnFirstGeneration;
 
 extern TAutoConsoleVariable<bool> CVarEnableMeshCache;
 
+extern TAutoConsoleVariable<bool> CVarEnableRealTimeMorphTargets;
+
 extern TAutoConsoleVariable<bool> CVarRollbackFixModelDiskStreamerDataRace;
 
+
 #if WITH_EDITOR
 
 // Struct used to keep a copy of the EditorSettings needed to compile Customizable Objects.
