[mutable] Fixed the issue of incomplete/ill-formatted widgets being shown in the mutable debugger.
- Removed unnecessary interactions with the debugger parameters widget. The debugger can not change the model, therefore, re-generating widgets for the parameters is pointless. By removing these rendundant calls the bug is made irrelevant.

#rnx

diff --git a/Engine/Plugins/Experimental/Mutable/Source/CustomizableObjectEditor/Private/MuCOE/SMutableCodeViewer.cpp b/Engine/Plugins/Experimental/Mutable/Source/CustomizableObjectEditor/Private/MuCOE/SMutableCodeViewer.cpp
--- a/Engine/Plugins/Experimental/Mutable/Source/CustomizableObjectEditor/Private/MuCOE/SMutableCodeViewer.cpp
+++ b/Engine/Plugins/Experimental/Mutable/Source/CustomizableObjectEditor/Private/MuCOE/SMutableCodeViewer.cpp
@@ -632,6 +632,7 @@ void SMutableCodeViewer::Construct(const FArguments& InArgs, const TSharedPtr<mu
 							[
 								SAssignNew(ParametersWidget, SMutableParametersWidget)
 								.OnParametersValueChanged(this, &SMutableCodeViewer::OnPreviewParameterValueChanged)
+								.Parameters(PreviewParameters)
 							]
 						]
 					]
@@ -1742,7 +1743,6 @@ void SMutableCodeViewer::OnSelectionChanged(TSharedPtr<FMutableCodeTreeElement>
 	TreeView->GetSelectedItems(SelectedNodes);
 
 	PreviewBorder->ClearContent();
-	ParametersWidget->SetParameters(nullptr);
 
 	SelectedOperationAddress = 0;
 	bSelectedOperationIsImage = false;
@@ -1755,7 +1755,7 @@ void SMutableCodeViewer::OnSelectionChanged(TSharedPtr<FMutableCodeTreeElement>
 	// Clear all selected items in the constant resources widget
 	ConstantsWidget->ClearSelectedConstantItems();
 	
-	// Find the duplicates for the selected tree element element and highlight them
+	// Find the duplicates for the selected tree element and highlight them
 	if (InNode)
 	{
 		HighlightDuplicatesOfEntry(InNode);
@@ -1766,8 +1766,6 @@ void SMutableCodeViewer::OnSelectionChanged(TSharedPtr<FMutableCodeTreeElement>
 	SelectedOperationAddress = SelectedNodes[0]->MutableOperation;
 	const mu::OP_TYPE OperationType = MutableModel->GetPrivate()->m_program.GetOpType(SelectedOperationAddress);
 	const mu::DATATYPE OperationDataType = mu::GetOpDataType(OperationType);
-
-	ParametersWidget->SetParameters(PreviewParameters);
 	
 	switch (OperationDataType)
 	{
diff --git a/Engine/Plugins/Experimental/Mutable/Source/CustomizableObjectEditor/Private/MuCOE/SMutableParametersWidget.cpp b/Engine/Plugins/Experimental/Mutable/Source/CustomizableObjectEditor/Private/MuCOE/SMutableParametersWidget.cpp
--- a/Engine/Plugins/Experimental/Mutable/Source/CustomizableObjectEditor/Private/MuCOE/SMutableParametersWidget.cpp
+++ b/Engine/Plugins/Experimental/Mutable/Source/CustomizableObjectEditor/Private/MuCOE/SMutableParametersWidget.cpp
@@ -22,7 +22,12 @@ struct FGeometry;
 void SMutableParametersWidget::Construct(const FArguments& InArgs)
 {
 	OnParametersValueChanged = InArgs._OnParametersValueChanged;
-
+	MutableParameters = InArgs._Parameters;
+	if (MutableParameters)
+	{
+		bIsPendingUpdate = true;
+	}
+	
 	ChildSlot
 	[
 		SNew(SVerticalBox)
@@ -37,13 +42,6 @@ void SMutableParametersWidget::Construct(const FArguments& InArgs)
 }
 
 
-void SMutableParametersWidget::SetParameters(const mu::ParametersPtr& InParameters)
-{
-	MutableParameters = InParameters;
-	bIsPendingUpdate = true;
-}
-
-
 void SMutableParametersWidget::Tick(const FGeometry& AllottedGeometry, const double InCurrentTime, const float InDeltaTime)
 {
 	SCompoundWidget::Tick(AllottedGeometry, InCurrentTime, InDeltaTime);
diff --git a/Engine/Plugins/Experimental/Mutable/Source/CustomizableObjectEditor/Private/MuCOE/SMutableParametersWidget.h b/Engine/Plugins/Experimental/Mutable/Source/CustomizableObjectEditor/Private/MuCOE/SMutableParametersWidget.h
--- a/Engine/Plugins/Experimental/Mutable/Source/CustomizableObjectEditor/Private/MuCOE/SMutableParametersWidget.h
+++ b/Engine/Plugins/Experimental/Mutable/Source/CustomizableObjectEditor/Private/MuCOE/SMutableParametersWidget.h
@@ -46,11 +46,6 @@ public:
 	// SWidget interface
 	void Tick(const FGeometry& AllottedGeometry, const double InCurrentTime, const float InDeltaTime) override;
 
-	// Own interface
-
-	/** Set the image to show in the widget. */
-	void SetParameters(const mu::ParametersPtr& InParameters);
-
 private:
 
 	/** Reference to the mutable parameters. */
