[Mutable] Fix components not showing when generating components without mesh.

- Now core will return a null mesh if no geomertry is generated.
- Fixed MeshChanged accessing indices (was accessed with both InstanceComponentIndex and ObjectComponentIndex).
- Merged UpdateSkeletalMesh_PostBeginUpdate0 component loops.
- Reference meshes use the first lod available instead of the index 0.

#rnx

diff --git a/Engine/Plugins/Experimental/Mutable/Source/CustomizableObject/Private/MuCO/CustomizableObjectInstance.cpp b/Engine/Plugins/Experimental/Mutable/Source/CustomizableObject/Private/MuCO/CustomizableObjectInstance.cpp
--- a/Engine/Plugins/Experimental/Mutable/Source/CustomizableObject/Private/MuCO/CustomizableObjectInstance.cpp
+++ b/Engine/Plugins/Experimental/Mutable/Source/CustomizableObject/Private/MuCO/CustomizableObjectInstance.cpp
@@ -1922,7 +1922,7 @@ bool UCustomizableInstancePrivate::DoComponentsNeedUpdate(UCustomizableObjectIns
 	}
 
 	// Find which components need an update
-	OperationData->MeshChanged.Init(false,NumInstanceComponents);
+	OperationData->MeshChanged.Init(false, NumInstanceComponents);
 
 	for (int32 InstanceComponentIndex = 0; InstanceComponentIndex < NumInstanceComponents; ++InstanceComponentIndex)
 	{
@@ -2026,9 +2026,7 @@ bool UCustomizableInstancePrivate::UpdateSkeletalMesh_PostBeginUpdate0(UCustomiz
 	TextureReuseCache.Empty(); // Sections may have changed, so invalidate the texture reuse cache because it's indexed by section
 
 	TArray<TObjectPtr<USkeletalMesh>> OldSkeletalMeshes = SkeletalMeshes;
-
-	bool bSuccess = true;
-
+	
 	const FModelResources& ModelResources = CustomizableObject->GetPrivate()->GetModelResources();
 
 	// Collate the Extension Data on the instance into groups based on the extension that produced
@@ -2158,8 +2156,16 @@ bool UCustomizableInstancePrivate::UpdateSkeletalMesh_PostBeginUpdate0(UCustomiz
 	for (int32 InstanceComponentIndex = 0; InstanceComponentIndex < NumInstanceComponents; ++InstanceComponentIndex)
 	{
 		const FInstanceUpdateData::FComponent& Component = OperationData->InstanceUpdateData.Components[InstanceComponentIndex];
-		int32 ObjectComponentIndex = Component.Id;
+		const int32 ObjectComponentIndex = Component.Id;
 
+		if (!ComponentsData.IsValidIndex(ObjectComponentIndex))
+		{
+			ensure(false);
+			
+			InvalidateGeneratedData();
+			return false;
+		}
+		
 		// If the component doesn't need an update copy the previously generated mesh.
 		if (!OperationData->MeshChanged[InstanceComponentIndex])
 		{
@@ -2170,14 +2176,7 @@ bool UCustomizableInstancePrivate::UpdateSkeletalMesh_PostBeginUpdate0(UCustomiz
 
 			continue;
 		}
-
-		if (!ComponentsData.IsValidIndex(ObjectComponentIndex))
-		{
-			bSuccess = false;
-			ensure(false);
-			continue;
-		}
-
+		
 		if (OperationData->bUseMeshCache)
 		{
 			if (USkeletalMesh* CachedMesh = CustomizableObject->GetPrivate()->MeshCache.Get(OperationData->MeshDescriptors[ObjectComponentIndex]))
@@ -2192,38 +2191,40 @@ bool UCustomizableInstancePrivate::UpdateSkeletalMesh_PostBeginUpdate0(UCustomiz
 		// Reset last mesh IDs.
 		ComponentsData[ObjectComponentIndex].LastMeshIdPerLOD.Init(MAX_uint64, MAX_MESH_LOD_COUNT);
 
-		// If it is a referenced resource, only the first LOD is relevant.
+		// We need the first valid mesh. get it from the component, considering that some LOSs may have been skipped.
+		mu::Ptr<const mu::Mesh> ComponentMesh;
+		for (int32 FirstValidLODIndex = Component.FirstLOD; FirstValidLODIndex < OperationData->InstanceUpdateData.LODs.Num() && !ComponentMesh; ++FirstValidLODIndex)
 		{
-			mu::Ptr<const mu::Mesh> FirstMesh;
-			if (Component.LODCount > 0)
-			{
-				FirstMesh = OperationData->InstanceUpdateData.LODs[Component.FirstLOD].Mesh;
-			}
+			ComponentMesh = OperationData->InstanceUpdateData.LODs[FirstValidLODIndex].Mesh;
+		}
+		
+		if (!ComponentMesh)
+		{
+			continue;
+		}
+		
+		// If it is a referenced resource, only the first LOD is relevant.
+		if (ComponentMesh->IsReference())
+        {
+        	int32 ReferenceID = ComponentMesh->GetReferencedMesh();
+        	TSoftObjectPtr<USkeletalMesh> Ref = ModelResources.PassThroughMeshes[ReferenceID];
 
-			if (FirstMesh && FirstMesh->IsReference())
-			{
-				int32 ReferenceID = FirstMesh->GetReferencedMesh();
-				TSoftObjectPtr<USkeletalMesh> Ref = ModelResources.PassThroughMeshes[ReferenceID];
+        	if (!Ref.IsValid())
+        	{
+        		// This shouldn't happen here synchronosuly. It should have been requested as an async load.
+        		UE_LOG(LogMutable, Error, TEXT("Referenced skeletal mesh [%s] was not pre-loaded. It will be sync-loaded probably causing a hitch. CO [%s]"), *Ref.ToString(), *GetNameSafe(CustomizableObject));
+        	}
 
-				if (!Ref.IsValid())
-				{
-					// This shouldn't happen here synchronosuly. It should have been requested as an async load.
-					UE_LOG(LogMutable, Error, TEXT("Referenced skeletal mesh [%s] was not pre-loaded. It will be sync-loaded probably causing a hitch. CO [%s]"), *Ref.ToString(), *GetNameSafe(CustomizableObject));
-				}
-
-				SkeletalMeshes[ObjectComponentIndex] = Ref.LoadSynchronous();
-				continue;
-			}
-		}
+        	SkeletalMeshes[ObjectComponentIndex] = Ref.LoadSynchronous();
+        	continue;
+        }
 
 		if (!ModelResources.ReferenceSkeletalMeshesData.IsValidIndex(ObjectComponentIndex))
 		{
-			bSuccess = false;
-			break;
+			InvalidateGeneratedData();
+			return false;
 		}
 
-		const FMutableRefSkeletalMeshData& RefSkeletalMeshData = ModelResources.ReferenceSkeletalMeshesData[ObjectComponentIndex];
-
 		// Create and initialize the SkeletalMesh for this component
 		MUTABLE_CPUPROFILER_SCOPE(ConstructMesh);
 
@@ -2236,186 +2237,119 @@ bool UCustomizableInstancePrivate::UpdateSkeletalMesh_PostBeginUpdate0(UCustomiz
 			FName SkeletalMeshName = GenerateUniqueNameFromCOInstance(*Public);
 			SkeletalMeshes[ObjectComponentIndex] = NewObject<USkeletalMesh>(GetTransientPackage(), SkeletalMeshName, RF_Transient);
 		}
-
 		USkeletalMesh* SkeletalMesh = SkeletalMeshes[ObjectComponentIndex];
 		check(SkeletalMesh);
+		
+		const FMutableRefSkeletalMeshData& RefSkeletalMeshData = ModelResources.ReferenceSkeletalMeshesData[ObjectComponentIndex];
 
 		// Set up the default information any mesh from this component will have (LODArrayInfos, RenderData, Mesh settings, etc). 
 		InitSkeletalMeshData(OperationData, SkeletalMesh, RefSkeletalMeshData, *CustomizableObject, ObjectComponentIndex);
-
-		// We need the first valid mesh. get it from the component, considering that some LOSs may have been skipped.
-		mu::Ptr<const mu::Mesh> ComponentMesh;
-		for (int32 FirstValidLODIndex = Component.FirstLOD; FirstValidLODIndex < OperationData->InstanceUpdateData.LODs.Num() && !ComponentMesh; ++FirstValidLODIndex)
+		
+		// Construct a new skeleton, fix up ActiveBones and Bonemap arrays and recompute the RefInvMatrices
+		const bool bBuildSkeletonDataSuccess = BuildSkeletonData(OperationData, *SkeletalMesh, RefSkeletalMeshData, *CustomizableObject, InstanceComponentIndex);
+		if (!bBuildSkeletonDataSuccess)
 		{
-			ComponentMesh = OperationData->InstanceUpdateData.LODs[FirstValidLODIndex].Mesh;
+			InvalidateGeneratedData();
+			return false;
 		}
-		
-		if (ComponentMesh)
+
+		// Build PhysicsAsset merging physics assets coming from SubMeshes of the newly generated Mesh
+		if (mu::Ptr<const mu::PhysicsBody> MutablePhysics = ComponentMesh->GetPhysicsBody())
 		{
-			// Construct a new skeleton, fix up ActiveBones and Bonemap arrays and recompute the RefInvMatrices
-			bSuccess = BuildSkeletonData(OperationData, *SkeletalMesh, RefSkeletalMeshData, *CustomizableObject, InstanceComponentIndex);
+			constexpr bool bDisallowCollisionBetweenAssets = true;
+			UPhysicsAsset* PhysicsAssetResult = GetOrBuildMainPhysicsAsset(
+				OperationData, RefSkeletalMeshData.PhysicsAsset, MutablePhysics.get(), bDisallowCollisionBetweenAssets, InstanceComponentIndex);
 
-			if (!bSuccess)
-			{
-				break;
-			}
+			SkeletalMesh->SetPhysicsAsset(PhysicsAssetResult);
 
-			// Build PhysicsAsset merging physics assets coming from SubMeshes of the newly generated Mesh
-			if (mu::Ptr<const mu::PhysicsBody> MutablePhysics = ComponentMesh->GetPhysicsBody())
+			if (PhysicsAssetResult)
 			{
-				constexpr bool bDisallowCollisionBetweenAssets = true;
-				UPhysicsAsset* PhysicsAssetResult = GetOrBuildMainPhysicsAsset(
-					OperationData, RefSkeletalMeshData.PhysicsAsset, MutablePhysics.get(), bDisallowCollisionBetweenAssets, InstanceComponentIndex);
+				// We are setting the physics asset mesh preview to the generated skeletal mesh.
+				// this is fine if the phyiscs asset is also generated, but not sure about referenced assets.
+				// In any case, don't mark the asset as modified.
 
-				SkeletalMesh->SetPhysicsAsset(PhysicsAssetResult);
-
-				if (PhysicsAssetResult)
-				{
-					// We are setting the physics asset mesh preview to the generated skeletal mesh.
-					// this is fine if the phyiscs asset is also generated, but not sure about referenced assets.
-					// In any case, don't mark the asset as modified.
-
-					// This is only called in editor, no need to add editor guards.	
-					constexpr bool bMarkAsDirty = false;
-					PhysicsAssetResult->SetPreviewMesh(SkeletalMesh, bMarkAsDirty);
-				}
+				// This is only called in editor, no need to add editor guards.	
+				constexpr bool bMarkAsDirty = false;
+				PhysicsAssetResult->SetPreviewMesh(SkeletalMesh, bMarkAsDirty);
 			}
+		}
 
-			const int32 NumAdditionalPhysicsNum = ComponentMesh->AdditionalPhysicsBodies.Num();
-			for (int32 I = 0; I < NumAdditionalPhysicsNum; ++I)
+		const int32 NumAdditionalPhysicsNum = ComponentMesh->AdditionalPhysicsBodies.Num();
+		for (int32 I = 0; I < NumAdditionalPhysicsNum; ++I)
+		{
+			mu::Ptr<const mu::PhysicsBody> AdditionalPhysiscBody = ComponentMesh->AdditionalPhysicsBodies[I];
+			
+			check(AdditionalPhysiscBody);
+			if (!AdditionalPhysiscBody->bBodiesModified)
 			{
-				mu::Ptr<const mu::PhysicsBody> AdditionalPhysiscBody = ComponentMesh->AdditionalPhysicsBodies[I];
-				
-				check(AdditionalPhysiscBody);
-				if (!AdditionalPhysiscBody->bBodiesModified)
-				{
-					continue;
-				}
-
-				const int32 PhysicsBodyExternalId = ComponentMesh->AdditionalPhysicsBodies[I]->CustomId;
-				
-				const FAnimBpOverridePhysicsAssetsInfo& Info = ModelResources.AnimBpOverridePhysiscAssetsInfo[PhysicsBodyExternalId];
+				continue;
+			}
 
-				// Make sure the AnimInstance class is loaded. It is expected to be already loaded at this point though. 
-				UClass* AnimInstanceClassLoaded = Info.AnimInstanceClass.LoadSynchronous();
-				TSubclassOf<UAnimInstance> AnimInstanceClass = TSubclassOf<UAnimInstance>(AnimInstanceClassLoaded);
-				if (!ensureAlways(AnimInstanceClass))
-				{
-					continue;
-				}
+			const int32 PhysicsBodyExternalId = ComponentMesh->AdditionalPhysicsBodies[I]->CustomId;
+			
+			const FAnimBpOverridePhysicsAssetsInfo& Info = ModelResources.AnimBpOverridePhysiscAssetsInfo[PhysicsBodyExternalId];
 
-				FAnimBpGeneratedPhysicsAssets& PhysicsAssetsUsedByAnimBp = AnimBpPhysicsAssets.FindOrAdd(AnimInstanceClass);
+			// Make sure the AnimInstance class is loaded. It is expected to be already loaded at this point though. 
+			UClass* AnimInstanceClassLoaded = Info.AnimInstanceClass.LoadSynchronous();
+			TSubclassOf<UAnimInstance> AnimInstanceClass = TSubclassOf<UAnimInstance>(AnimInstanceClassLoaded);
+			if (!ensureAlways(AnimInstanceClass))
+			{
+				continue;
+			}
 
-				TObjectPtr<UPhysicsAsset> PhysicsAssetTemplate = TObjectPtr<UPhysicsAsset>(Info.SourceAsset.Get());
+			FAnimBpGeneratedPhysicsAssets& PhysicsAssetsUsedByAnimBp = AnimBpPhysicsAssets.FindOrAdd(AnimInstanceClass);
 
-				check(PhysicsAssetTemplate);
+			TObjectPtr<UPhysicsAsset> PhysicsAssetTemplate = TObjectPtr<UPhysicsAsset>(Info.SourceAsset.Get());
 
-				FAnimInstanceOverridePhysicsAsset& Entry =
-					PhysicsAssetsUsedByAnimBp.AnimInstancePropertyIndexAndPhysicsAssets.Emplace_GetRef();
+			check(PhysicsAssetTemplate);
 
-				Entry.PropertyIndex = Info.PropertyIndex;
-				Entry.PhysicsAsset = MakePhysicsAssetFromTemplateAndMutableBody(
-					OperationData,PhysicsAssetTemplate, AdditionalPhysiscBody.get(), InstanceComponentIndex);
-			}
+			FAnimInstanceOverridePhysicsAsset& Entry =
+				PhysicsAssetsUsedByAnimBp.AnimInstancePropertyIndexAndPhysicsAssets.Emplace_GetRef();
 
-			// Add sockets from the SkeletalMesh of reference and from the MutableMesh
-			BuildMeshSockets(OperationData, SkeletalMesh, ModelResources, RefSkeletalMeshData, ComponentMesh);
-		}
-		else
-		{
-			bSuccess = false;
-			break;
+			Entry.PropertyIndex = Info.PropertyIndex;
+			Entry.PhysicsAsset = MakePhysicsAssetFromTemplateAndMutableBody(
+				OperationData,PhysicsAssetTemplate, AdditionalPhysiscBody.get(), InstanceComponentIndex);
 		}
 
+		// Add sockets from the SkeletalMesh of reference and from the MutableMesh
+		BuildMeshSockets(OperationData, SkeletalMesh, ModelResources, RefSkeletalMeshData, ComponentMesh);
+		
 		for (const TPair<const UCustomizableObjectExtension*, TArray<FInputPinDataContainer>>& Pair : ExtensionToExtensionData)
 		{
 			Pair.Key->OnSkeletalMeshCreated(Pair.Value, ObjectComponentIndex, SkeletalMesh);
 		}
-	}
-
-
-	if (bSuccess)
-	{
-		MUTABLE_CPUPROFILER_SCOPE(BuildSkeletalMeshData);
+		
+		// Mesh to copy data from if possible. 
+		USkeletalMesh* OldSkeletalMesh = OldSkeletalMeshes.IsValidIndex(ObjectComponentIndex) ? OldSkeletalMeshes[ObjectComponentIndex] : nullptr;
 
-		for (int32 InstanceComponentIndex = 0; bSuccess && InstanceComponentIndex < OperationData->NumInstanceComponents; ++InstanceComponentIndex)
+		BuildOrCopyElementData(OperationData, SkeletalMesh, Public, InstanceComponentIndex);
+		bool const bCopyRenderDataSuccess = BuildOrCopyRenderData(OperationData, SkeletalMesh, OldSkeletalMesh, Public, InstanceComponentIndex);
+		if (!bCopyRenderDataSuccess)
 		{
-			const FInstanceUpdateData::FComponent& Component = OperationData->InstanceUpdateData.Components[InstanceComponentIndex];
-
-			int32 ObjectComponentIndex = Component.Id;
-
-			if (OperationData->bUseMeshCache)
-			{
-				const TArray<mu::FResourceID>& MeshId = OperationData->MeshDescriptors[ObjectComponentIndex];
-				if (CustomizableObject->GetPrivate()->MeshCache.Get(MeshId))
-				{
-					continue;
-				}
-			}
-			
-			if ( (!OperationData->MeshChanged.IsValidIndex(ObjectComponentIndex)) 
-				|| !OperationData->MeshChanged[ObjectComponentIndex] )
-			{
-				continue;
-			}
-
-			USkeletalMesh* SkeletalMesh = SkeletalMeshes.IsValidIndex(ObjectComponentIndex) ? SkeletalMeshes[ObjectComponentIndex] : nullptr;
-
-			if (!SkeletalMesh)
-			{
-				continue;
-			}
-
-			// Skip if it is a referenced mesh
-			{
-				mu::Ptr<const mu::Mesh> FirstMesh;
-				if (Component.LODCount > 0)
-				{
-					FirstMesh = OperationData->InstanceUpdateData.LODs[Component.FirstLOD].Mesh;
-				}
-				if (FirstMesh && FirstMesh->IsReference())
-				{
-					continue;
-				}
-			}
-
-			// Mesh to copy data from if possible. 
-			USkeletalMesh* OldSkeletalMesh = OldSkeletalMeshes.IsValidIndex(ObjectComponentIndex) ? OldSkeletalMeshes[ObjectComponentIndex] : nullptr;
-
-			{
-				BuildOrCopyElementData(OperationData, SkeletalMesh, Public, InstanceComponentIndex);
-				bSuccess = BuildOrCopyRenderData(OperationData, SkeletalMesh, OldSkeletalMesh, Public, InstanceComponentIndex);
-			}
-
-			if (bSuccess)
-			{
-				BuildOrCopyMorphTargetsData(OperationData, SkeletalMesh, OldSkeletalMesh, Public, InstanceComponentIndex);
-				BuildOrCopyClothingData(OperationData, SkeletalMesh, OldSkeletalMesh, Public, InstanceComponentIndex);
+			InvalidateGeneratedData();
+			return false;
+		}
 
-				FSkeletalMeshRenderData* RenderData = SkeletalMesh->GetResourceForRendering();
-				ensure(RenderData && RenderData->LODRenderData.Num() > 0);
-				ensure(SkeletalMesh->GetLODNum() > 0);
+		BuildOrCopyMorphTargetsData(OperationData, SkeletalMesh, OldSkeletalMesh, Public, InstanceComponentIndex);
+		BuildOrCopyClothingData(OperationData, SkeletalMesh, OldSkeletalMesh, Public, InstanceComponentIndex);
 
-				for (FSkeletalMeshLODRenderData& LODResource : RenderData->LODRenderData)
-				{
-					UnrealConversionUtils::UpdateSkeletalMeshLODRenderDataBuffersSize(LODResource);
-				}
-			}
+		FSkeletalMeshRenderData* RenderData = SkeletalMesh->GetResourceForRendering();
+		ensure(RenderData && RenderData->LODRenderData.Num() > 0);
+		ensure(SkeletalMesh->GetLODNum() > 0);
 
-			if (OperationData->bUseMeshCache)
-			{
-				const TArray<mu::FResourceID>& MeshId = OperationData->MeshDescriptors[ObjectComponentIndex];
-				CustomizableObject->GetPrivate()->MeshCache.Add(MeshId, SkeletalMesh);
-			}
+		for (FSkeletalMeshLODRenderData& LODResource : RenderData->LODRenderData)
+		{
+			UnrealConversionUtils::UpdateSkeletalMeshLODRenderDataBuffersSize(LODResource);
 		}
-	}
 
-	if (!bSuccess)
-	{
-		InvalidateGeneratedData();
+		if (OperationData->bUseMeshCache)
+		{
+			const TArray<mu::FResourceID>& MeshId = OperationData->MeshDescriptors[ObjectComponentIndex];
+			CustomizableObject->GetPrivate()->MeshCache.Add(MeshId, SkeletalMesh);
+		}
 	}
 
-	return bSuccess;
+	return true;
 }
 
 
diff --git a/Engine/Plugins/Experimental/Mutable/Source/CustomizableObject/Private/MuCO/CustomizableObjectSystem.cpp b/Engine/Plugins/Experimental/Mutable/Source/CustomizableObject/Private/MuCO/CustomizableObjectSystem.cpp
--- a/Engine/Plugins/Experimental/Mutable/Source/CustomizableObject/Private/MuCO/CustomizableObjectSystem.cpp
+++ b/Engine/Plugins/Experimental/Mutable/Source/CustomizableObject/Private/MuCO/CustomizableObjectSystem.cpp
@@ -1917,8 +1917,12 @@ namespace impl
 				const FInstanceUpdateData::FComponent& Component = UpdateData.Components[ComponentIndex];
 				const FInstanceUpdateData::FLOD& LOD = UpdateData.LODs[Component.FirstLOD+LODIndex];
 				check(LOD.bGenerated);
-				check(LOD.Mesh);
 
+				if (!LOD.Mesh)
+				{
+					continue;
+				}
+				
 				const mu::FMeshBufferSet& MeshSet = LOD.Mesh->GetVertexBuffers();
 
 				int32 VertexMorphsInfoIndexAndCountBufferIndex, VertexMorphsInfoIndexAndCountBufferChannel;
@@ -2401,7 +2405,8 @@ namespace impl
 
 		LOD.Mesh = GetMeshTask.GetResult();
 
-		if (LOD.Mesh->IsReference())
+		if (LOD.Mesh &&
+			LOD.Mesh->IsReference())
 		{
 			const UCustomizableObjectInstance* Instance = OperationData->Instance.Get();
 			const UCustomizableObject* CustomizableObject = Instance->GetCustomizableObject();
diff --git a/Engine/Plugins/Experimental/Mutable/Source/MutableRuntime/Private/MuR/System.cpp b/Engine/Plugins/Experimental/Mutable/Source/MutableRuntime/Private/MuR/System.cpp
--- a/Engine/Plugins/Experimental/Mutable/Source/MutableRuntime/Private/MuR/System.cpp
+++ b/Engine/Plugins/Experimental/Mutable/Source/MutableRuntime/Private/MuR/System.cpp
@@ -591,12 +591,7 @@ namespace mu
 					{
 						Result = SystemPrivate->WorkingMemoryManager.LoadMesh(FCacheAddress(RootAddress, 0, 0), true);
 					}
-
-					if (!Result)
-					{
-						Result = new Mesh();
-					}
-
+					
 					SystemPrivate->WorkingMemoryManager.EndRunnerThread();
 					SystemPrivate->WorkingMemoryManager.CurrentInstanceCache = nullptr;
 					
