[mutable] Fixed the "compile only selected" compile option not working between editor sessions if a full compile is not performed after restarting the editor. The TableToParamNames wasn't serialized to disk in the ModelResources.
In addition, fixed the same compile option not working when the selected option included "None" values in tables where it was already automatically included.
#tests Editor
#rnx

diff --git a/Engine/Plugins/Experimental/Mutable/Source/CustomizableObject/Internal/MuCO/CustomizableObjectCustomVersion.h b/Engine/Plugins/Experimental/Mutable/Source/CustomizableObject/Internal/MuCO/CustomizableObjectCustomVersion.h
--- a/Engine/Plugins/Experimental/Mutable/Source/CustomizableObject/Internal/MuCO/CustomizableObjectCustomVersion.h
+++ b/Engine/Plugins/Experimental/Mutable/Source/CustomizableObject/Internal/MuCO/CustomizableObjectCustomVersion.h
@@ -132,6 +132,8 @@ struct CUSTOMIZABLEOBJECT_API FCustomizableObjectCustomVersion
 
 		AddModifierPin, // Move the modifiers output pin from Material to Modifier
 
+		CorrectlySerializeTableToParamNames,
+
 		// -----<new versions can be added above this line>-------------------------------------------------
 		VersionPlusOne,
 		LatestVersion = VersionPlusOne - 1
diff --git a/Engine/Plugins/Experimental/Mutable/Source/CustomizableObject/Internal/MuCO/CustomizableObjectPrivate.h b/Engine/Plugins/Experimental/Mutable/Source/CustomizableObject/Internal/MuCO/CustomizableObjectPrivate.h
--- a/Engine/Plugins/Experimental/Mutable/Source/CustomizableObject/Internal/MuCO/CustomizableObjectPrivate.h
+++ b/Engine/Plugins/Experimental/Mutable/Source/CustomizableObject/Internal/MuCO/CustomizableObjectPrivate.h
@@ -868,6 +868,17 @@ public:
 	TSharedPtr<FModelStreamableBulkData> StreamingData;
 };
 
+USTRUCT()
+struct FMutableParamNameSet
+{
+	GENERATED_BODY()
+
+	TSet<FString> ParamNames;
+
+	friend FArchive& operator<<(FArchive& Ar, FMutableParamNameSet& MutableParamNameSet);
+};
+
+
 // Referenced materials, skeletons, passthrough textures...
 USTRUCT()
 struct FModelResources
@@ -974,7 +985,8 @@ struct FModelResources
 	TArray<FCustomizableObjectMeshToMeshVertData> EditorOnlyClothingMeshToMeshVertData;
 
 	// Stores what param names use a certain table as a table can be used from multiple table nodes, useful for partial compilations to restrict params
-	TMap<TObjectPtr<const UDataTable>, TSet<FString>> TableToParamNames;
+	UPROPERTY()
+	TMap<FString, FMutableParamNameSet> TableToParamNames;
 #endif
 
 	/** Max number of components in the compiled Model. */
diff --git a/Engine/Plugins/Experimental/Mutable/Source/CustomizableObject/Private/MuCO/CustomizableObject.cpp b/Engine/Plugins/Experimental/Mutable/Source/CustomizableObject/Private/MuCO/CustomizableObject.cpp
--- a/Engine/Plugins/Experimental/Mutable/Source/CustomizableObject/Private/MuCO/CustomizableObject.cpp
+++ b/Engine/Plugins/Experimental/Mutable/Source/CustomizableObject/Private/MuCO/CustomizableObject.cpp
@@ -697,6 +697,7 @@ void UCustomizableObjectPrivate::SaveCompiledData(FObjectAndNameAsStringProxyArc
 	MemoryWriter << CustomizableObjectPathMap;
 	MemoryWriter << GroupNodeMap;
 	MemoryWriter << ParticipatingObjects;
+	MemoryWriter << LocalModelResources.TableToParamNames;
 
 	if (!bIsCooking)
 	{
@@ -850,6 +851,7 @@ void UCustomizableObjectPrivate::LoadCompiledData(FObjectAndNameAsStringProxyArc
 			MemoryReader << CustomizableObjectPathMap;
 			MemoryReader << GroupNodeMap;
 			MemoryReader << ParticipatingObjects;
+			MemoryReader << LocalModelResource.TableToParamNames;
 
 			if (!bIsCooking)
 			{
@@ -3099,4 +3101,12 @@ bool FCompilationRequest::operator==(const FCompilationRequest& Other) const
 #endif
 
 
+FArchive& operator<<(FArchive& Ar, FMutableParamNameSet& MutableParamNameSet)
+{
+	Ar << MutableParamNameSet.ParamNames;
+
+	return Ar;
+}
+
+
 #undef LOCTEXT_NAMESPACE
diff --git a/Engine/Plugins/Experimental/Mutable/Source/CustomizableObjectEditor/Private/MuCOE/GenerateMutableSource/GenerateMutableSource.h b/Engine/Plugins/Experimental/Mutable/Source/CustomizableObjectEditor/Private/MuCOE/GenerateMutableSource/GenerateMutableSource.h
--- a/Engine/Plugins/Experimental/Mutable/Source/CustomizableObjectEditor/Private/MuCOE/GenerateMutableSource/GenerateMutableSource.h
+++ b/Engine/Plugins/Experimental/Mutable/Source/CustomizableObjectEditor/Private/MuCOE/GenerateMutableSource/GenerateMutableSource.h
@@ -850,7 +850,7 @@ struct FMutableGraphGenerationContext
 	TMap<FString, FString> ParamNamesToSelectedOptions;
 
 	// Stores what param names use a certain table as a table can be used from multiple table nodes, useful for partial compilations to restrict params
-	TMap<TObjectPtr<const UDataTable>, TSet<FString>> TableToParamNames;
+	TMap<FString, FMutableParamNameSet> TableToParamNames;
 
 	TArray<const UEdGraphNode*> LimitedParameters;
 	int32 ParameterLimitationCount = 0;
diff --git a/Engine/Plugins/Experimental/Mutable/Source/CustomizableObjectEditor/Private/MuCOE/GenerateMutableSource/GenerateMutableSourceTable.cpp b/Engine/Plugins/Experimental/Mutable/Source/CustomizableObjectEditor/Private/MuCOE/GenerateMutableSource/GenerateMutableSourceTable.cpp
--- a/Engine/Plugins/Experimental/Mutable/Source/CustomizableObjectEditor/Private/MuCOE/GenerateMutableSource/GenerateMutableSourceTable.cpp
+++ b/Engine/Plugins/Experimental/Mutable/Source/CustomizableObjectEditor/Private/MuCOE/GenerateMutableSource/GenerateMutableSourceTable.cpp
@@ -775,20 +775,23 @@ void RestrictRowNamesToSelectedOption(TArray<FName>& InOutRowNames, const UCusto
 {
 	if (!GenerationContext.ParamNamesToSelectedOptions.IsEmpty())
 	{
-		TSet<FString>* ParamNames = GenerationContext.TableToParamNames.Find(TableNode.Table);
+		FMutableParamNameSet* ParamNameSet = GenerationContext.TableToParamNames.Find(TableNode.Table->GetPathName());
 
-		if (ParamNames && !ParamNames->IsEmpty())
+		if (ParamNameSet && !ParamNameSet->ParamNames.IsEmpty())
 		{
 			TSet<FName> SelectedOptionNames;
 
-			for (const FString& ParamName : *ParamNames)
+			for (const FString& ParamName : ParamNameSet->ParamNames)
 			{
 				// If the param is in the map restrict to only the selected option
 				FString* SelectedOptionString = GenerationContext.ParamNamesToSelectedOptions.Find(ParamName);
 
 				if (SelectedOptionString)
 				{
-					SelectedOptionNames.Add(FName(*SelectedOptionString));
+					if (!(*SelectedOptionString == FString("None") && TableNode.bAddNoneOption))
+					{
+						SelectedOptionNames.Add(FName(*SelectedOptionString));
+					}
 				}
 			}
 
@@ -1069,14 +1072,14 @@ mu::TablePtr GenerateMutableSourceTable(const UDataTable* DataTable, const UCust
 
 	if (GenerationContext.ParamNamesToSelectedOptions.IsEmpty())
 	{
-		TSet<FString>* ParamNames = GenerationContext.TableToParamNames.Find(DataTable);
+		FMutableParamNameSet* ParamNameSet = GenerationContext.TableToParamNames.Find(DataTable->GetPathName());
 
-		if (!ParamNames)
+		if (!ParamNameSet)
 		{
-			ParamNames = &GenerationContext.TableToParamNames.Add(DataTable);
+			ParamNameSet = &GenerationContext.TableToParamNames.Add(DataTable->GetPathName());
 		}
 
-		ParamNames->Add(TableNode->ParameterName);
+		ParamNameSet->ParamNames.Add(TableNode->ParameterName);
 	}
 
 	// Checking if the table is in the cache
