[Mutable] Add material slot name index to surface metadata.

- Now there is no need to keep number of material slot names and materials in sync.

#rnx

diff --git a/Engine/Plugins/Experimental/Mutable/Source/CustomizableObject/Internal/MuCO/CustomizableObjectPrivate.h b/Engine/Plugins/Experimental/Mutable/Source/CustomizableObject/Internal/MuCO/CustomizableObjectPrivate.h
--- a/Engine/Plugins/Experimental/Mutable/Source/CustomizableObject/Internal/MuCO/CustomizableObjectPrivate.h
+++ b/Engine/Plugins/Experimental/Mutable/Source/CustomizableObject/Internal/MuCO/CustomizableObjectPrivate.h
@@ -540,17 +540,23 @@ struct FMutableSurfaceMetadata
 {
 	GENERATED_USTRUCT_BODY()
 	
+	UPROPERTY()
+	int32 MaterialSlotIndex = INDEX_NONE;
+	
 	UPROPERTY()
 	bool bCastShadow = false;
+	uint8 UnusedPadding[3] = {0};
 
 	friend FArchive& operator<<(FArchive& Ar, FMutableSurfaceMetadata& Elem)
 	{
+		Ar << Elem.MaterialSlotIndex;
 		Ar << Elem.bCastShadow;
 
 		return Ar;
 	}
 };
 
+static_assert(sizeof(FMutableSurfaceMetadata) == (sizeof(int32) + sizeof(bool) + sizeof(uint8)*3));
 
 USTRUCT()
 struct FClothingStreamable
@@ -1342,6 +1348,8 @@ public:
 		IntParameterOptionDataTablePartialRestore,
 
 		CorrectlySerializeTableToParamNames,
+
+		AddMaterialSlotNameIndexToSurfaceMetadata,
 		
 		// -----<new versions can be added above this line>--------
 		LastCustomizableObjectVersion
diff --git a/Engine/Plugins/Experimental/Mutable/Source/CustomizableObject/Internal/MuCO/CustomizableObjectSystemPrivate.h b/Engine/Plugins/Experimental/Mutable/Source/CustomizableObject/Internal/MuCO/CustomizableObjectSystemPrivate.h
--- a/Engine/Plugins/Experimental/Mutable/Source/CustomizableObject/Internal/MuCO/CustomizableObjectSystemPrivate.h
+++ b/Engine/Plugins/Experimental/Mutable/Source/CustomizableObject/Internal/MuCO/CustomizableObjectSystemPrivate.h
@@ -358,6 +358,9 @@ struct FInstanceUpdateData
 
 		/** Id of the surface in the mutable core instance. */
 		uint32 SurfaceId = 0;
+
+		/** Id of the metadata associated with this surface. */ 
+		uint32 SurfaceMetadataId = 0;
 	};
 
 	struct FLOD
diff --git a/Engine/Plugins/Experimental/Mutable/Source/CustomizableObject/Private/MuCO/CustomizableObjectInstance.cpp b/Engine/Plugins/Experimental/Mutable/Source/CustomizableObject/Private/MuCO/CustomizableObjectInstance.cpp
--- a/Engine/Plugins/Experimental/Mutable/Source/CustomizableObject/Private/MuCO/CustomizableObjectInstance.cpp
+++ b/Engine/Plugins/Experimental/Mutable/Source/CustomizableObject/Private/MuCO/CustomizableObjectInstance.cpp
@@ -6234,11 +6234,16 @@ void UCustomizableInstancePrivate::BuildMaterials(const TSharedRef<FUpdateContex
 				FSkeletalMaterial& MaterialSlot = Materials.AddDefaulted_GetRef();
 				MaterialSlot.MaterialInterface = MaterialTemplate;
 
-				if (ModelResources.MaterialSlotNames.IsValidIndex(Surface.MaterialIndex))
+				const FMutableSurfaceMetadata* FoundSurfaceMetadata = ModelResources.SurfaceMetadata.Find(Surface.SurfaceMetadataId);
+				
+				if (FoundSurfaceMetadata)
 				{
-					MaterialSlot.MaterialSlotName = ModelResources.MaterialSlotNames[Surface.MaterialIndex];
+					const int32 MaterialSlotNameIndex = FoundSurfaceMetadata->MaterialSlotIndex;
+					if (ModelResources.MaterialSlotNames.IsValidIndex(MaterialSlotNameIndex))
+					{
+						MaterialSlot.MaterialSlotName = ModelResources.MaterialSlotNames[MaterialSlotNameIndex];
+					}
 				}
-
 				if (RefSkeletalMeshData)
 				{
 					SetMeshUVChannelDensity(MaterialSlot.UVChannelData, RefSkeletalMeshData->Settings.DefaultUVChannelDensity);
diff --git a/Engine/Plugins/Experimental/Mutable/Source/CustomizableObject/Private/MuCO/CustomizableObjectSystem.cpp b/Engine/Plugins/Experimental/Mutable/Source/CustomizableObject/Private/MuCO/CustomizableObjectSystem.cpp
--- a/Engine/Plugins/Experimental/Mutable/Source/CustomizableObject/Private/MuCO/CustomizableObjectSystem.cpp
+++ b/Engine/Plugins/Experimental/Mutable/Source/CustomizableObject/Private/MuCO/CustomizableObjectSystem.cpp
@@ -2239,7 +2239,9 @@ namespace impl
 				
 
 				// This lambda does all the work to fill up the surface data
-				auto AddSurface = [&LOD, &SurfacesSharedId, &ModelResources, &GetImagesData, OperationData, MutableInstance, CustomizableObject, InstanceComponentIndex, MutableLODIndex](uint32 SurfaceId, int32 InstanceSurfaceIndex)
+				auto AddSurface = 
+					[&LOD, &SurfacesSharedId, &ModelResources, &GetImagesData, OperationData, MutableInstance, CustomizableObject, InstanceComponentIndex, MutableLODIndex]
+					(uint32 SurfaceId, uint32 SurfaceMetadataId, int32 InstanceSurfaceIndex)
 					{
 						int32 BaseSurfaceIndex = InstanceSurfaceIndex;
 						int32 BaseLODIndex = MutableLODIndex;
@@ -2250,6 +2252,7 @@ namespace impl
 
 						// Now Surface.MaterialIndex is decoded from a parameter at the end of this if()
 						Surface.SurfaceId = SurfaceId;
+						Surface.SurfaceMetadataId = SurfaceMetadataId;
 
 						const int32 SharedSurfaceId = MutableInstance->GetSharedSurfaceId(InstanceComponentIndex, MutableLODIndex, InstanceSurfaceIndex);
 						const int32 SharedSurfaceIndex = SurfacesSharedId.Find(SharedSurfaceId);
@@ -2268,6 +2271,7 @@ namespace impl
 							MutableInstance->FindBaseSurfaceBySharedId(InstanceComponentIndex, SharedSurfaceId, BaseSurfaceIndex, BaseLODIndex);
 
 							Surface.SurfaceId = MutableInstance->GetSurfaceId(InstanceComponentIndex, BaseLODIndex, BaseSurfaceIndex);
+							Surface.SurfaceMetadataId = MutableInstance->GetSurfaceCustomId(InstanceComponentIndex, BaseLODIndex, BaseSurfaceIndex);
 						}
 
 						// Vectors
@@ -2360,7 +2364,8 @@ namespace impl
 					for (int32 SurfaceIndex = 0; SurfaceIndex < SurfaceCount; ++SurfaceIndex)
 					{
 						uint32 SurfaceId = MutableInstance->GetSurfaceId(InstanceComponentIndex, MutableLODIndex, SurfaceIndex);
-						AddSurface(SurfaceId, SurfaceIndex);
+						uint32 SurfaceMetadataId = MutableInstance->GetSurfaceCustomId(InstanceComponentIndex, MutableLODIndex, SurfaceIndex);
+						AddSurface(SurfaceId, SurfaceMetadataId, SurfaceIndex);
 					}
 				}
 
@@ -2377,7 +2382,7 @@ namespace impl
 
 						if (InstanceSurfaceIndex >= 0)
 						{
-							AddSurface(SurfaceId, InstanceSurfaceIndex);
+							AddSurface(SurfaceId, 0, InstanceSurfaceIndex);
 						}
 					}
 				}
diff --git a/Engine/Plugins/Experimental/Mutable/Source/CustomizableObjectEditor/Private/MuCOE/GenerateMutableSource/GenerateMutableSourceSurface.cpp b/Engine/Plugins/Experimental/Mutable/Source/CustomizableObjectEditor/Private/MuCOE/GenerateMutableSource/GenerateMutableSourceSurface.cpp
--- a/Engine/Plugins/Experimental/Mutable/Source/CustomizableObjectEditor/Private/MuCOE/GenerateMutableSource/GenerateMutableSourceSurface.cpp
+++ b/Engine/Plugins/Experimental/Mutable/Source/CustomizableObjectEditor/Private/MuCOE/GenerateMutableSource/GenerateMutableSourceSurface.cpp
@@ -293,89 +293,65 @@ mu::Ptr<mu::NodeSurface> GenerateMutableSourceSurface(const UEdGraphPin * Pin, F
 		{
 			GenerationContext.AddParticipatingObject(*TypedNodeMat->GetMaterial());
 
-			const int32 lastMaterialAmount = GenerationContext.ReferencedMaterials.Num();
 			ReferencedMaterialsIndex = GenerationContext.ReferencedMaterials.AddUnique(TypedNodeMat->GetMaterial());
-			// Used ReferencedMaterialsIndex instead of TypedNodeMat->Material->GetName() to prevent material name collisions
-
-			// Take slot name from skeletal mesh if one can be found, else leave empty.
-			// Keep Referenced Materials and Materail Slot Names synchronized even if no material name can be found.
-			const bool IsNewSlotName = GenerationContext.ReferencedMaterialSlotNames.Num() == ReferencedMaterialsIndex;
-			check(IsNewSlotName == lastMaterialAmount < GenerationContext.ReferencedMaterials.Num());
-			bool SlotNameFound = false;
-			check(GenerationContext.ReferencedMaterialSlotNames.Num() >= ReferencedMaterialsIndex);
-			if (IsNewSlotName || GenerationContext.ReferencedMaterialSlotNames[ReferencedMaterialsIndex].IsNone())
+			if (ConnectedMaterialPin)
 			{
-				if (ConnectedMaterialPin)
+				if (const UEdGraphPin* SkeletalMeshPin = FindMeshBaseSource(*ConnectedMaterialPin, false))
 				{
-					if (const UEdGraphPin* SkeletalMeshPin = FindMeshBaseSource(*ConnectedMaterialPin, false))
+					FSkeletalMaterial* SkeletalMaterial = nullptr;
+					const FSkelMeshSection* ReferenceSkelMeshSection = nullptr;
+					
+					if (const UCustomizableObjectNodeSkeletalMesh* SkeletalMeshNode = Cast<UCustomizableObjectNodeSkeletalMesh>(SkeletalMeshPin->GetOwningNode()))
 					{
-						FSkeletalMaterial* SkeletalMaterial = nullptr;
-						const FSkelMeshSection* ReferenceSkelMeshSection = nullptr;
-						
-						if (const UCustomizableObjectNodeSkeletalMesh* SkeletalMeshNode = Cast<UCustomizableObjectNodeSkeletalMesh>(SkeletalMeshPin->GetOwningNode()))
-						{
-							SkeletalMaterial = SkeletalMeshNode->GetSkeletalMaterialFor(*SkeletalMeshPin);
-							ReferenceSkelMeshSection = SkeletalMeshNode->GetSkeletalMeshSectionFor(*SkeletalMeshPin);
-						}
+						SkeletalMaterial = SkeletalMeshNode->GetSkeletalMaterialFor(*SkeletalMeshPin);
+						ReferenceSkelMeshSection = SkeletalMeshNode->GetSkeletalMeshSectionFor(*SkeletalMeshPin);
+					}
 
-						else if (const UCustomizableObjectNodeTable* TableNode = Cast<UCustomizableObjectNodeTable>(SkeletalMeshPin->GetOwningNode()))
-						{
-							SkeletalMaterial = TableNode->GetDefaultSkeletalMaterialFor(*SkeletalMeshPin);
-							ReferenceSkelMeshSection = TableNode->GetDefaultSkeletalMeshSectionFor(*SkeletalMeshPin);
-						}
+					else if (const UCustomizableObjectNodeTable* TableNode = Cast<UCustomizableObjectNodeTable>(SkeletalMeshPin->GetOwningNode()))
+					{
+						SkeletalMaterial = TableNode->GetDefaultSkeletalMaterialFor(*SkeletalMeshPin);
+						ReferenceSkelMeshSection = TableNode->GetDefaultSkeletalMeshSectionFor(*SkeletalMeshPin);
+					}
 
-						if (SkeletalMaterial)
-						{
-							if (IsNewSlotName)
-							{
-								GenerationContext.ReferencedMaterialSlotNames.Add(SkeletalMaterial->MaterialSlotName);
-								SlotNameFound = true;
-							}
-							else if (GenerationContext.ReferencedMaterialSlotNames[ReferencedMaterialsIndex].IsNone())
-							{
-								GenerationContext.ReferencedMaterialSlotNames[ReferencedMaterialsIndex] = SkeletalMaterial->MaterialSlotName;
-								SlotNameFound = true;
-							}
-						}
+					// A reference to the material slot name index will be kept in the surface metadata.
+					int32 MaterialSlotNameIndex = INDEX_NONE;
+					if (SkeletalMaterial)
+					{
+						MaterialSlotNameIndex = 
+								GenerationContext.ReferencedMaterialSlotNames.AddUnique(SkeletalMaterial->MaterialSlotName);
+					}
 
-						if (ReferenceSkelMeshSection)
+					if (ReferenceSkelMeshSection)
+					{
+						auto HashSurfaceMetadataFunc = [](const FMutableSurfaceMetadata& Data) -> uint32
 						{
-							auto HashSurfaceMetadataFunc = [](const FMutableSurfaceMetadata& Data) -> uint32
-							{
-								return CityHash32(reinterpret_cast<const char*>(&Data), sizeof(FMutableSurfaceMetadata));
-							};
+							return CityHash32(reinterpret_cast<const char*>(&Data), sizeof(FMutableSurfaceMetadata));
+						};
 
-							auto CompareSurfaceMetadataFunc = [](const FMutableSurfaceMetadata& A, const FMutableSurfaceMetadata& B)
-							{
-								return FMemory::Memcmp(&A, &B, sizeof(FMutableSurfaceMetadata)) == 0;
-							};
+						auto CompareSurfaceMetadataFunc = [](const FMutableSurfaceMetadata& A, const FMutableSurfaceMetadata& B)
+						{
+							return FMemory::Memcmp(&A, &B, sizeof(FMutableSurfaceMetadata)) == 0;
+						};
 
-							FMutableSurfaceMetadata SurfaceMetadata;
-							FMemory::Memzero(SurfaceMetadata);
-							
-							SurfaceMetadata.bCastShadow = ReferenceSkelMeshSection->bCastShadow;
+						FMutableSurfaceMetadata SurfaceMetadata;
+						FMemory::Memzero(SurfaceMetadata);
+						
+						SurfaceMetadata.MaterialSlotIndex = MaterialSlotNameIndex;
+						SurfaceMetadata.bCastShadow = ReferenceSkelMeshSection->bCastShadow;
 
-							SurfaceMetadataUniqueHash = Private::GenerateUniquePersistentHash(
-									SurfaceMetadata, GenerationContext.SurfaceMetadata, HashSurfaceMetadataFunc, CompareSurfaceMetadataFunc);
+						SurfaceMetadataUniqueHash = Private::GenerateUniquePersistentHash(
+								SurfaceMetadata, GenerationContext.SurfaceMetadata, HashSurfaceMetadataFunc, CompareSurfaceMetadataFunc);
 
-							if (SurfaceMetadataUniqueHash != 0)
-							{
-								GenerationContext.SurfaceMetadata.FindOrAdd(SurfaceMetadataUniqueHash, SurfaceMetadata);
-							}
-							else
-							{
-								UE_LOG(LogMutable, Error, TEXT("Maximum number of surfaces reached."));
-							}
+						if (SurfaceMetadataUniqueHash != 0)
+						{
+							GenerationContext.SurfaceMetadata.FindOrAdd(SurfaceMetadataUniqueHash, SurfaceMetadata);
+						}
+						else
+						{
+							UE_LOG(LogMutable, Error, TEXT("Maximum number of surfaces reached."));
 						}
 					}
 				}
-
-
-				// No name was found, we need to keep index parity, so we add empty value. We may find the value later.
-				if (IsNewSlotName && !SlotNameFound)
-				{
-					GenerationContext.ReferencedMaterialSlotNames.Add(FName(NAME_None));
-				}
 			}
 		}
 
diff --git a/Engine/Plugins/Experimental/Mutable/Source/CustomizableObjectEditor/Private/MuCOE/GenerateMutableSource/GenerateMutableSourceTable.cpp b/Engine/Plugins/Experimental/Mutable/Source/CustomizableObjectEditor/Private/MuCOE/GenerateMutableSource/GenerateMutableSourceTable.cpp
--- a/Engine/Plugins/Experimental/Mutable/Source/CustomizableObjectEditor/Private/MuCOE/GenerateMutableSource/GenerateMutableSourceTable.cpp
+++ b/Engine/Plugins/Experimental/Mutable/Source/CustomizableObjectEditor/Private/MuCOE/GenerateMutableSource/GenerateMutableSourceTable.cpp
@@ -509,18 +509,7 @@ bool FillTableColumn(const UCustomizableObjectNodeTable* TableNode, mu::TablePtr
 					CurrentColumn = MutableTable->AddColumn(ColumnName, mu::ETableColumnType::Scalar);
 				}
 
-				const int32 lastMaterialAmount = GenerationContext.ReferencedMaterials.Num();
 				int32 ReferenceMaterialId = GenerationContext.ReferencedMaterials.AddUnique(MaterialInstance);
-
-				// Take slot name from skeletal mesh if one can be found, else leave empty.
-				// Keep Referenced Materials and Materail Slot Names synchronized even if no material name can be found.
-				const bool IsNewSlotName = GenerationContext.ReferencedMaterialSlotNames.Num() == ReferenceMaterialId;
-
-				if (IsNewSlotName)
-				{
-					GenerationContext.ReferencedMaterialSlotNames.Add(FName(NAME_None));
-				}
-
 				MutableTable->SetCell(CurrentColumn, RowIdx, (float)ReferenceMaterialId);
 
 				return true;
