[Mutable] Improved BackwardsComatibleFixup.

- Now all nodes do the same backwards version in sync.
- Moved fixup code from PostBackwardsComatibleFixup to PostBackwardsComatibleFixup.
- Added BackwardsComatibleFixup to PinData.

#rnx

diff --git a/Engine/Plugins/Experimental/Mutable/Source/CustomizableObject/Internal/MuCO/CustomizableObjectCustomVersion.h b/Engine/Plugins/Experimental/Mutable/Source/CustomizableObject/Internal/MuCO/CustomizableObjectCustomVersion.h
--- a/Engine/Plugins/Experimental/Mutable/Source/CustomizableObject/Internal/MuCO/CustomizableObjectCustomVersion.h
+++ b/Engine/Plugins/Experimental/Mutable/Source/CustomizableObject/Internal/MuCO/CustomizableObjectCustomVersion.h
@@ -134,6 +134,8 @@ struct CUSTOMIZABLEOBJECT_API FCustomizableObjectCustomVersion
 
 		CorrectlySerializeTableToParamNames,
 
+		MovedCompatibilityFromPostBackwardsCompatibleFixup,
+		
 		// -----<new versions can be added above this line>-------------------------------------------------
 		VersionPlusOne,
 		LatestVersion = VersionPlusOne - 1
diff --git a/Engine/Plugins/Experimental/Mutable/Source/CustomizableObjectEditor/Private/MuCOE/CustomizableObjectEditor.cpp b/Engine/Plugins/Experimental/Mutable/Source/CustomizableObjectEditor/Private/MuCOE/CustomizableObjectEditor.cpp
--- a/Engine/Plugins/Experimental/Mutable/Source/CustomizableObjectEditor/Private/MuCOE/CustomizableObjectEditor.cpp
+++ b/Engine/Plugins/Experimental/Mutable/Source/CustomizableObjectEditor/Private/MuCOE/CustomizableObjectEditor.cpp
@@ -67,6 +67,7 @@
 #include "Widgets/Layout/SScrollBox.h"
 #include "Widgets/Notifications/SNotificationList.h"
 #include "SMutableObjectViewer.h"
+#include "MuCO/CustomizableObjectCustomVersion.h"
 #include "MuCO/CustomizableSkeletalMeshActor.h"
 #include "MuCOE/CustomizableObjectInstanceFactory.h"
 
@@ -1901,14 +1902,6 @@ void FCustomizableObjectEditor::PasteNodesHere(const FVector2D& Location)
 
 	// Force new pasted Material Expressions to have same connections as graph nodes
 	//Material->MaterialGraph->LinkMaterialExpressionsFromGraph();
-
-	for (UEdGraphNode* PastedNode : PastedNodes) // Perform possible BackwardsCompatibleFixup. PostLoads have already been called at this point.
-	{
-		if (UCustomizableObjectNode* TypedNode = Cast<UCustomizableObjectNode>(PastedNode))
-		{
-			TypedNode->BackwardsCompatibleFixup();
-		}
-	}
 	
 	for (UEdGraphNode* PastedNode : PastedNodes)
 	{
diff --git a/Engine/Plugins/Experimental/Mutable/Source/CustomizableObjectEditor/Private/MuCOE/CustomizableObjectGraph.cpp b/Engine/Plugins/Experimental/Mutable/Source/CustomizableObjectEditor/Private/MuCOE/CustomizableObjectGraph.cpp
--- a/Engine/Plugins/Experimental/Mutable/Source/CustomizableObjectEditor/Private/MuCOE/CustomizableObjectGraph.cpp
+++ b/Engine/Plugins/Experimental/Mutable/Source/CustomizableObjectEditor/Private/MuCOE/CustomizableObjectGraph.cpp
@@ -2,6 +2,7 @@
 
 #include "MuCOE/CustomizableObjectGraph.h"
 
+#include "MuCO/CustomizableObjectCustomVersion.h"
 #include "MuCOE/EdGraphSchema_CustomizableObject.h"
 #include "MuCOE/Nodes/CustomizableObjectNodeObject.h"
 
@@ -19,6 +20,8 @@ void UCustomizableObjectGraph::PostLoad()
 {
 	Super::PostLoad();
 
+	// TODO UE-222779 If compatibility code has to be executed, loaded the full hierarchy and in sync the BackwardsCompatibleFixup 
+	
 	// Make sure all nodes have finished loading.
 	for (UEdGraphNode* Node : Nodes)
 	{
@@ -40,11 +43,34 @@ void UCustomizableObjectGraph::PostLoad()
 	}
 
 	// Execute backwards compatible code for all nodes. It requires all nodes to be loaded.
-	for (UEdGraphNode* Node : Nodes)
+
+	TArray<TObjectPtr<UEdGraphNode>> NodesCopy = Nodes; // Copy to be able to remove nodes inside the BackwardsCompatibleFixup.
+
+	for (int32 Version = GetLinkerCustomVersion(FCustomizableObjectCustomVersion::GUID) + 1; Version <= FCustomizableObjectCustomVersion::LatestVersion; ++Version)
 	{
-		if (UCustomizableObjectNode* CustomizableObjectNode = Cast<UCustomizableObjectNode>(Node))
+		for (UEdGraphNode* Node : NodesCopy)
 		{
-			CustomizableObjectNode->BackwardsCompatibleFixup();
+			if (UCustomizableObjectNode* CustomizableObjectNode = Cast<UCustomizableObjectNode>(Node))
+			{
+				TArray<UEdGraphPin*> PinsCopy = CustomizableObjectNode->GetAllPins(); // Copy to be able to remove pins inside the BackwardsCompatibleFixup.
+				for (UEdGraphPin* Pin : PinsCopy)
+				{
+					if (!Pin)
+					{
+						continue;
+					}
+					
+					UCustomizableObjectNodePinData* PinData = CustomizableObjectNode->GetPinData(*Pin);
+					if (!PinData)
+					{
+						continue;
+					}
+					
+					PinData->BackwardsCompatibleFixup(Version);
+				}
+				
+				CustomizableObjectNode->BackwardsCompatibleFixup(Version);
+			}
 		}
 	}
 
diff --git a/Engine/Plugins/Experimental/Mutable/Source/CustomizableObjectEditor/Private/MuCOE/EdGraphSchema_CustomizableObject.cpp b/Engine/Plugins/Experimental/Mutable/Source/CustomizableObjectEditor/Private/MuCOE/EdGraphSchema_CustomizableObject.cpp
--- a/Engine/Plugins/Experimental/Mutable/Source/CustomizableObjectEditor/Private/MuCOE/EdGraphSchema_CustomizableObject.cpp
+++ b/Engine/Plugins/Experimental/Mutable/Source/CustomizableObjectEditor/Private/MuCOE/EdGraphSchema_CustomizableObject.cpp
@@ -87,6 +87,7 @@
 #include "MuCOE/Nodes/CustomizableObjectNodeReroute.h"
 #include "Toolkits/ToolkitManager.h"
 #include "PropertyEditorModule.h"
+#include "MuCO/CustomizableObjectCustomVersion.h"
 #include "MuCOE/GraphTraversal.h"
 
 class IToolkit;
@@ -165,7 +166,6 @@ UEdGraphNode* FCustomizableObjectSchemaAction_NewNode::CreateNode(class UEdGraph
 	if (UCustomizableObjectNode* TypedResultNode = Cast<UCustomizableObjectNode>(ResultNode))
 	{
 		TypedResultNode->BeginConstruct();
-		TypedResultNode->BackwardsCompatibleFixup(); // In theory not required but added to be consistent with PostLoad (called by DuplicateObject).
 		TypedResultNode->PostBackwardsCompatibleFixup();
 	}
 	ResultNode->ReconstructNode(); // Mutable node lifecycle always starts at ReconstructNode.
diff --git a/Engine/Plugins/Experimental/Mutable/Source/CustomizableObjectEditor/Private/MuCOE/Nodes/CustomizableObjectNode.cpp b/Engine/Plugins/Experimental/Mutable/Source/CustomizableObjectEditor/Private/MuCOE/Nodes/CustomizableObjectNode.cpp
--- a/Engine/Plugins/Experimental/Mutable/Source/CustomizableObjectEditor/Private/MuCOE/Nodes/CustomizableObjectNode.cpp
+++ b/Engine/Plugins/Experimental/Mutable/Source/CustomizableObjectEditor/Private/MuCOE/Nodes/CustomizableObjectNode.cpp
@@ -357,13 +357,11 @@ void UCustomizableObjectNode::PostInitProperties()
 }
 
 
-void UCustomizableObjectNode::BackwardsCompatibleFixup()
+void UCustomizableObjectNode::BackwardsCompatibleFixup(int32 CustomizableObjectCustomVersion)
 {
-	const int32 CustomizableObjectCustomVersion = GetLinkerCustomVersion(FCustomizableObjectCustomVersion::GUID);
-	
 	// Fix UE large world coordinates automatic pin conversion. 
 	// Now all pins with PinCategory == FName("float") get automatically changed to the new double type
-	if (CustomizableObjectCustomVersion < FCustomizableObjectCustomVersion::FixBlueprintPinsUseRealNumbers)
+	if (CustomizableObjectCustomVersion == FCustomizableObjectCustomVersion::FixBlueprintPinsUseRealNumbers)
 	{
 		for (UEdGraphPin* Pin : Pins)
 		{
@@ -377,7 +375,7 @@ void UCustomizableObjectNode::BackwardsCompatibleFixup()
 		}
 	}
 	
-	if (CustomizableObjectCustomVersion < FCustomizableObjectCustomVersion::AutomaticNodeMaterialPerformance)
+	if (CustomizableObjectCustomVersion == FCustomizableObjectCustomVersion::AutomaticNodeMaterialPerformance)
 	{
 		for (TTuple<FEdGraphPinReference, UCustomizableObjectNodePinData*> Pair : PinsData_DEPRECATED)
 		{
@@ -387,7 +385,7 @@ void UCustomizableObjectNode::BackwardsCompatibleFixup()
 		PinsData_DEPRECATED.Empty();
 	}
 
-	if (CustomizableObjectCustomVersion < FCustomizableObjectCustomVersion::FixBlueprintPinsUseRealNumbersAgain)
+	if (CustomizableObjectCustomVersion == FCustomizableObjectCustomVersion::FixBlueprintPinsUseRealNumbersAgain)
 	{
 		for (UEdGraphPin* Pin : Pins)
 		{
diff --git a/Engine/Plugins/Experimental/Mutable/Source/CustomizableObjectEditor/Private/MuCOE/Nodes/CustomizableObjectNodeColorSwitch.cpp b/Engine/Plugins/Experimental/Mutable/Source/CustomizableObjectEditor/Private/MuCOE/Nodes/CustomizableObjectNodeColorSwitch.cpp
--- a/Engine/Plugins/Experimental/Mutable/Source/CustomizableObjectEditor/Private/MuCOE/Nodes/CustomizableObjectNodeColorSwitch.cpp
+++ b/Engine/Plugins/Experimental/Mutable/Source/CustomizableObjectEditor/Private/MuCOE/Nodes/CustomizableObjectNodeColorSwitch.cpp
@@ -8,13 +8,11 @@
 #define LOCTEXT_NAMESPACE "CustomizableObjectEditor"
 
 
-void UCustomizableObjectNodeColorSwitch::BackwardsCompatibleFixup()
+void UCustomizableObjectNodeColorSwitch::BackwardsCompatibleFixup(int32 CustomizableObjectCustomVersion)
 {
-	Super::BackwardsCompatibleFixup();
+	Super::BackwardsCompatibleFixup(CustomizableObjectCustomVersion);
 
-	const int32 CustomizableObjectCustomVersion = GetLinkerCustomVersion(FCustomizableObjectCustomVersion::GUID);
-
-	if (CustomizableObjectCustomVersion < FCustomizableObjectCustomVersion::BugPinsSwitch)
+	if (CustomizableObjectCustomVersion == FCustomizableObjectCustomVersion::BugPinsSwitch)
 	{
 		OutputPinReference = FindPin(TEXT("Color"));	
 	}
diff --git a/Engine/Plugins/Experimental/Mutable/Source/CustomizableObjectEditor/Private/MuCOE/Nodes/CustomizableObjectNodeColorVariation.cpp b/Engine/Plugins/Experimental/Mutable/Source/CustomizableObjectEditor/Private/MuCOE/Nodes/CustomizableObjectNodeColorVariation.cpp
--- a/Engine/Plugins/Experimental/Mutable/Source/CustomizableObjectEditor/Private/MuCOE/Nodes/CustomizableObjectNodeColorVariation.cpp
+++ b/Engine/Plugins/Experimental/Mutable/Source/CustomizableObjectEditor/Private/MuCOE/Nodes/CustomizableObjectNodeColorVariation.cpp
@@ -7,11 +7,11 @@
 #include "MuCOE/EdGraphSchema_CustomizableObject.h"
 
 
-void UCustomizableObjectNodeColorVariation::BackwardsCompatibleFixup()
+void UCustomizableObjectNodeColorVariation::BackwardsCompatibleFixup(int32 CustomizableObjectCustomVersion)
 {
-	const int32 CustomizableObjectCustomVersion = GetLinkerCustomVersion(FCustomizableObjectCustomVersion::GUID);
+	Super::BackwardsCompatibleFixup(CustomizableObjectCustomVersion);
 
-	if (CustomizableObjectCustomVersion < FCustomizableObjectCustomVersion::NodeVariationSerializationIssue)
+	if (CustomizableObjectCustomVersion == FCustomizableObjectCustomVersion::NodeVariationSerializationIssue)
 	{
 		for (const FCustomizableObjectColorVariation& OldVariation : Variations_DEPRECATED)
 		{
@@ -21,8 +21,6 @@ void UCustomizableObjectNodeColorVariation::BackwardsCompatibleFixup()
 			VariationsData.Add(Variation);
 		}
 	}
-
-	Super::BackwardsCompatibleFixup();
 }
 
 
diff --git a/Engine/Plugins/Experimental/Mutable/Source/CustomizableObjectEditor/Private/MuCOE/Nodes/CustomizableObjectNodeEditMaterial.cpp b/Engine/Plugins/Experimental/Mutable/Source/CustomizableObjectEditor/Private/MuCOE/Nodes/CustomizableObjectNodeEditMaterial.cpp
--- a/Engine/Plugins/Experimental/Mutable/Source/CustomizableObjectEditor/Private/MuCOE/Nodes/CustomizableObjectNodeEditMaterial.cpp
+++ b/Engine/Plugins/Experimental/Mutable/Source/CustomizableObjectEditor/Private/MuCOE/Nodes/CustomizableObjectNodeEditMaterial.cpp
@@ -97,14 +97,12 @@ void UCustomizableObjectNodeEditMaterial::SetParentNode(UCustomizableObject* Obj
 }
 
 
-void UCustomizableObjectNodeEditMaterial::BackwardsCompatibleFixup()
+void UCustomizableObjectNodeEditMaterial::BackwardsCompatibleFixup(int32 CustomizableObjectCustomVersion)
 {
-	Super::BackwardsCompatibleFixup();
-
-	const int32 CustomizableObjectCustomVersion = GetLinkerCustomVersion(FCustomizableObjectCustomVersion::GUID);
+	Super::BackwardsCompatibleFixup(CustomizableObjectCustomVersion);
 
 	// Convert deprecated node index list to the node id list.
-	if (CustomizableObjectCustomVersion < FCustomizableObjectCustomVersion::PostLoadToCustomVersion
+	if (CustomizableObjectCustomVersion == FCustomizableObjectCustomVersion::PostLoadToCustomVersion
 		&& BlockIds_DEPRECATED.Num() < Blocks_DEPRECATED.Num())
 	{
 		UCustomizableObjectNodeMaterialBase* ParentMaterialNode = GetParentMaterialNode();
@@ -149,7 +147,7 @@ void UCustomizableObjectNodeEditMaterial::BackwardsCompatibleFixup()
 	}
 	
 	// Convert deprecated node id list to absolute rect list.
-	if (CustomizableObjectCustomVersion < FCustomizableObjectCustomVersion::UseUVRects)
+	if (CustomizableObjectCustomVersion == FCustomizableObjectCustomVersion::UseUVRects)
 	{
 		// If we are here, it means this node was loaded from a version that didn't have it's own layout.
 		check(Layout->Blocks.IsEmpty());
@@ -219,7 +217,7 @@ void UCustomizableObjectNodeEditMaterial::BackwardsCompatibleFixup()
 		}
 	}
 
-	if (CustomizableObjectCustomVersion < FCustomizableObjectCustomVersion::AutomaticNodeMaterial)
+	if (CustomizableObjectCustomVersion == FCustomizableObjectCustomVersion::AutomaticNodeMaterial)
 	{
 		if (const UCustomizableObjectNodeMaterialBase* ParentMaterial = GetParentMaterialNode())
 		{
@@ -256,7 +254,7 @@ void UCustomizableObjectNodeEditMaterial::BackwardsCompatibleFixup()
 	}
 
 	// Fill PinsParameter.
-	if (CustomizableObjectCustomVersion < FCustomizableObjectCustomVersion::AutomaticNodeMaterialPerformanceBug) // || CustomizableObjectCustomVersion < FCustomizableObjectCustomVersion::AutomaticNodeMaterialPerformance
+	if (CustomizableObjectCustomVersion == FCustomizableObjectCustomVersion::AutomaticNodeMaterialPerformanceBug) // || CustomizableObjectCustomVersion < FCustomizableObjectCustomVersion::AutomaticNodeMaterialPerformance
 	{
 		for (const UEdGraphPin* Pin : GetAllNonOrphanPins())
 		{
@@ -267,17 +265,17 @@ void UCustomizableObjectNodeEditMaterial::BackwardsCompatibleFixup()
 		}
 	}
 	
-	if (CustomizableObjectCustomVersion < FCustomizableObjectCustomVersion::EditMaterialOnlyMutableModeParameters)
+	if (CustomizableObjectCustomVersion == FCustomizableObjectCustomVersion::EditMaterialOnlyMutableModeParameters)
 	{
 		ReconstructNode();
 	}
 
-	if (CustomizableObjectCustomVersion < FCustomizableObjectCustomVersion::EditMaterialMaskPinDesync)
+	if (CustomizableObjectCustomVersion == FCustomizableObjectCustomVersion::EditMaterialMaskPinDesync)
 	{
 		ReconstructNode();
 	}
 
-	if (CustomizableObjectCustomVersion < FCustomizableObjectCustomVersion::FixedMultilayerMaterialIds)
+	if (CustomizableObjectCustomVersion == FCustomizableObjectCustomVersion::FixedMultilayerMaterialIds)
 	{
 		ReconstructNode();
 	}
diff --git a/Engine/Plugins/Experimental/Mutable/Source/CustomizableObjectEditor/Private/MuCOE/Nodes/CustomizableObjectNodeExtendMaterial.cpp b/Engine/Plugins/Experimental/Mutable/Source/CustomizableObjectEditor/Private/MuCOE/Nodes/CustomizableObjectNodeExtendMaterial.cpp
--- a/Engine/Plugins/Experimental/Mutable/Source/CustomizableObjectEditor/Private/MuCOE/Nodes/CustomizableObjectNodeExtendMaterial.cpp
+++ b/Engine/Plugins/Experimental/Mutable/Source/CustomizableObjectEditor/Private/MuCOE/Nodes/CustomizableObjectNodeExtendMaterial.cpp
@@ -40,13 +40,11 @@ void UCustomizableObjectNodeExtendMaterial::BeginPostDuplicate(bool bDuplicateFo
 }
 
 
-void UCustomizableObjectNodeExtendMaterial::BackwardsCompatibleFixup()
+void UCustomizableObjectNodeExtendMaterial::BackwardsCompatibleFixup(int32 CustomizableObjectCustomVersion)
 {
-	Super::PostBackwardsCompatibleFixup();
-
-	const int32 CustomizableObjectCustomVersion = GetLinkerCustomVersion(FCustomizableObjectCustomVersion::GUID);
+	Super::BackwardsCompatibleFixup(CustomizableObjectCustomVersion);
 	
-	if (CustomizableObjectCustomVersion < FCustomizableObjectCustomVersion::ExtendMaterialRemoveImages)
+	if (CustomizableObjectCustomVersion == FCustomizableObjectCustomVersion::ExtendMaterialRemoveImages)
 	{
 		if (const UCustomizableObjectNodeMaterialBase* ParentMaterial = GetParentMaterialNode())
 		{
@@ -79,7 +77,7 @@ void UCustomizableObjectNodeExtendMaterial::BackwardsCompatibleFixup()
 		ReconstructNode();
 	}
 
-	if (CustomizableObjectCustomVersion < FCustomizableObjectCustomVersion::FixedMultilayerMaterialIds)
+	if (CustomizableObjectCustomVersion == FCustomizableObjectCustomVersion::FixedMultilayerMaterialIds)
 	{
 		ReconstructNode();
 	}
diff --git a/Engine/Plugins/Experimental/Mutable/Source/CustomizableObjectEditor/Private/MuCOE/Nodes/CustomizableObjectNodeExternalPin.cpp b/Engine/Plugins/Experimental/Mutable/Source/CustomizableObjectEditor/Private/MuCOE/Nodes/CustomizableObjectNodeExternalPin.cpp
--- a/Engine/Plugins/Experimental/Mutable/Source/CustomizableObjectEditor/Private/MuCOE/Nodes/CustomizableObjectNodeExternalPin.cpp
+++ b/Engine/Plugins/Experimental/Mutable/Source/CustomizableObjectEditor/Private/MuCOE/Nodes/CustomizableObjectNodeExternalPin.cpp
@@ -21,9 +21,9 @@ void UCustomizableObjectNodeExternalPin::Serialize(FArchive& Ar)
 }
 
 
-void UCustomizableObjectNodeExternalPin::BackwardsCompatibleFixup()
+void UCustomizableObjectNodeExternalPin::BackwardsCompatibleFixup(int32 CustomizableObjectCustomVersion)
 {
-	Super::BackwardsCompatibleFixup();
+	Super::BackwardsCompatibleFixup(CustomizableObjectCustomVersion);
 	
 	// Get the pin type from the actual pin.
 	if (GetLinkerCustomVersion(FCustomizableObjectCustomVersion::GUID) < FCustomizableObjectCustomVersion::BeforeCustomVersionWasAdded)
diff --git a/Engine/Plugins/Experimental/Mutable/Source/CustomizableObjectEditor/Private/MuCOE/Nodes/CustomizableObjectNodeFloatParameter.cpp b/Engine/Plugins/Experimental/Mutable/Source/CustomizableObjectEditor/Private/MuCOE/Nodes/CustomizableObjectNodeFloatParameter.cpp
--- a/Engine/Plugins/Experimental/Mutable/Source/CustomizableObjectEditor/Private/MuCOE/Nodes/CustomizableObjectNodeFloatParameter.cpp
+++ b/Engine/Plugins/Experimental/Mutable/Source/CustomizableObjectEditor/Private/MuCOE/Nodes/CustomizableObjectNodeFloatParameter.cpp
@@ -77,13 +77,11 @@ void UCustomizableObjectNodeFloatParameter::OnRenameNode(const FString& NewName)
 }
 
 
-void UCustomizableObjectNodeFloatParameter::BackwardsCompatibleFixup()
+void UCustomizableObjectNodeFloatParameter::BackwardsCompatibleFixup(int32 CustomizableObjectCustomVersion)
 {
-	Super::BackwardsCompatibleFixup();
+	Super::BackwardsCompatibleFixup(CustomizableObjectCustomVersion);
 
-	const int32 CustomizableObjectCustomVersion = GetLinkerCustomVersion(FCustomizableObjectCustomVersion::GUID);
-
-	if (CustomizableObjectCustomVersion < FCustomizableObjectCustomVersion::RemovedParameterDecorations)
+	if (CustomizableObjectCustomVersion == FCustomizableObjectCustomVersion::RemovedParameterDecorations)
 	{
 		ReconstructNode();
 	}
diff --git a/Engine/Plugins/Experimental/Mutable/Source/CustomizableObjectEditor/Private/MuCOE/Nodes/CustomizableObjectNodeFloatSwitch.cpp b/Engine/Plugins/Experimental/Mutable/Source/CustomizableObjectEditor/Private/MuCOE/Nodes/CustomizableObjectNodeFloatSwitch.cpp
--- a/Engine/Plugins/Experimental/Mutable/Source/CustomizableObjectEditor/Private/MuCOE/Nodes/CustomizableObjectNodeFloatSwitch.cpp
+++ b/Engine/Plugins/Experimental/Mutable/Source/CustomizableObjectEditor/Private/MuCOE/Nodes/CustomizableObjectNodeFloatSwitch.cpp
@@ -8,13 +8,11 @@
 #define LOCTEXT_NAMESPACE "CustomizableObjectEditor"
 
 
-void UCustomizableObjectNodeFloatSwitch::BackwardsCompatibleFixup()
+void UCustomizableObjectNodeFloatSwitch::BackwardsCompatibleFixup(int32 CustomizableObjectCustomVersion)
 {
-	Super::BackwardsCompatibleFixup();
-
-	const int32 CustomizableObjectCustomVersion = GetLinkerCustomVersion(FCustomizableObjectCustomVersion::GUID);
-
-	if (CustomizableObjectCustomVersion < FCustomizableObjectCustomVersion::BugPinsSwitch)
+	Super::BackwardsCompatibleFixup(CustomizableObjectCustomVersion);
+	
+	if (CustomizableObjectCustomVersion == FCustomizableObjectCustomVersion::BugPinsSwitch)
 	{
 		OutputPinReference = FindPin(TEXT("Float"));	
 	}
diff --git a/Engine/Plugins/Experimental/Mutable/Source/CustomizableObjectEditor/Private/MuCOE/Nodes/CustomizableObjectNodeFloatVariation.cpp b/Engine/Plugins/Experimental/Mutable/Source/CustomizableObjectEditor/Private/MuCOE/Nodes/CustomizableObjectNodeFloatVariation.cpp
--- a/Engine/Plugins/Experimental/Mutable/Source/CustomizableObjectEditor/Private/MuCOE/Nodes/CustomizableObjectNodeFloatVariation.cpp
+++ b/Engine/Plugins/Experimental/Mutable/Source/CustomizableObjectEditor/Private/MuCOE/Nodes/CustomizableObjectNodeFloatVariation.cpp
@@ -7,11 +7,11 @@
 #include "MuCOE/EdGraphSchema_CustomizableObject.h"
 
 
-void UCustomizableObjectNodeFloatVariation::BackwardsCompatibleFixup()
+void UCustomizableObjectNodeFloatVariation::BackwardsCompatibleFixup(int32 CustomizableObjectCustomVersion)
 {
-	const int32 CustomizableObjectCustomVersion = GetLinkerCustomVersion(FCustomizableObjectCustomVersion::GUID);
+	Super::BackwardsCompatibleFixup(CustomizableObjectCustomVersion);
 
-	if (CustomizableObjectCustomVersion < FCustomizableObjectCustomVersion::NodeVariationSerializationIssue)
+	if (CustomizableObjectCustomVersion == FCustomizableObjectCustomVersion::NodeVariationSerializationIssue)
 	{
 		for (const FCustomizableObjectFloatVariation& OldVariation : Variations_DEPRECATED)
 		{
@@ -21,8 +21,6 @@ void UCustomizableObjectNodeFloatVariation::BackwardsCompatibleFixup()
 			VariationsData.Add(Variation);
 		}
 	}
-
-	Super::BackwardsCompatibleFixup();
 }
 
 
diff --git a/Engine/Plugins/Experimental/Mutable/Source/CustomizableObjectEditor/Private/MuCOE/Nodes/CustomizableObjectNodeGroupProjectorParameter.cpp b/Engine/Plugins/Experimental/Mutable/Source/CustomizableObjectEditor/Private/MuCOE/Nodes/CustomizableObjectNodeGroupProjectorParameter.cpp
--- a/Engine/Plugins/Experimental/Mutable/Source/CustomizableObjectEditor/Private/MuCOE/Nodes/CustomizableObjectNodeGroupProjectorParameter.cpp
+++ b/Engine/Plugins/Experimental/Mutable/Source/CustomizableObjectEditor/Private/MuCOE/Nodes/CustomizableObjectNodeGroupProjectorParameter.cpp
@@ -114,13 +114,11 @@ FText UCustomizableObjectNodeGroupProjectorParameter::GetTooltipText() const
 }
 
 
-void UCustomizableObjectNodeGroupProjectorParameter::BackwardsCompatibleFixup()
+void UCustomizableObjectNodeGroupProjectorParameter::BackwardsCompatibleFixup(int32 CustomizableObjectCustomVersion)
 {
-	Super::BackwardsCompatibleFixup();
-
-	const int32 CustomizableObjectCustomVersion = GetLinkerCustomVersion(FCustomizableObjectCustomVersion::GUID);
+	Super::BackwardsCompatibleFixup(CustomizableObjectCustomVersion);
 	
-	if (CustomizableObjectCustomVersion < FCustomizableObjectCustomVersion::GroupProjectorPinTypeAdded)
+	if (CustomizableObjectCustomVersion == FCustomizableObjectCustomVersion::GroupProjectorPinTypeAdded)
 	{
 		if (UEdGraphPin* Pin = ProjectorPin())
 		{
@@ -128,7 +126,7 @@ void UCustomizableObjectNodeGroupProjectorParameter::BackwardsCompatibleFixup()
 		}
 	}
 
-	if (CustomizableObjectCustomVersion < FCustomizableObjectCustomVersion::GroupProjectorImagePinRemoved)
+	if (CustomizableObjectCustomVersion == FCustomizableObjectCustomVersion::GroupProjectorImagePinRemoved)
 	{
 		ReconstructNode();
 	}
diff --git a/Engine/Plugins/Experimental/Mutable/Source/CustomizableObjectEditor/Private/MuCOE/Nodes/CustomizableObjectNodeMaterial.cpp b/Engine/Plugins/Experimental/Mutable/Source/CustomizableObjectEditor/Private/MuCOE/Nodes/CustomizableObjectNodeMaterial.cpp
--- a/Engine/Plugins/Experimental/Mutable/Source/CustomizableObjectEditor/Private/MuCOE/Nodes/CustomizableObjectNodeMaterial.cpp
+++ b/Engine/Plugins/Experimental/Mutable/Source/CustomizableObjectEditor/Private/MuCOE/Nodes/CustomizableObjectNodeMaterial.cpp
@@ -290,13 +290,11 @@ bool UCustomizableObjectNodeMaterialPinDataImage::CanEditChange(const FProperty*
 }
 
 
-void UCustomizableObjectNodeMaterialPinDataImage::PostLoad()
+void UCustomizableObjectNodeMaterialPinDataImage::BackwardsCompatibleFixup(int32 CustomizableObjectCustomVersion)
 {
 	Super::PostLoad();
-
-	const int32 CustomizableObjectCustomVersion = GetLinkerCustomVersion(FCustomizableObjectCustomVersion::GUID);
 	
-	if (CustomizableObjectCustomVersion < FCustomizableObjectCustomVersion::NodeMaterialPinDataImageDetails)
+	if (CustomizableObjectCustomVersion == FCustomizableObjectCustomVersion::NodeMaterialPinDataImageDetails)
 	{
 		if (UVLayout == UV_LAYOUT_IGNORE)
 		{
@@ -371,13 +369,11 @@ UCustomizableObjectNodeRemapPinsByName* UCustomizableObjectNodeMaterial::CreateR
 }
 
 
-void UCustomizableObjectNodeMaterial::BackwardsCompatibleFixup()
+void UCustomizableObjectNodeMaterial::BackwardsCompatibleFixup(int32 CustomizableObjectCustomVersion)
 {
-	Super::BackwardsCompatibleFixup();
-
-	const int32 CustomizableObjectCustomVersion = GetLinkerCustomVersion(FCustomizableObjectCustomVersion::GUID);
+	Super::BackwardsCompatibleFixup(CustomizableObjectCustomVersion);
 	
-	if (CustomizableObjectCustomVersion < FCustomizableObjectCustomVersion::AutomaticNodeMaterial)
+	if (CustomizableObjectCustomVersion == FCustomizableObjectCustomVersion::AutomaticNodeMaterial)
 	{
 		for (const FCustomizableObjectNodeMaterialImage& Image : Images_DEPRECATED)
 		{
@@ -523,7 +519,7 @@ void UCustomizableObjectNodeMaterial::BackwardsCompatibleFixup()
 	}
 
 	// Fill PinsParameter.
-	if (CustomizableObjectCustomVersion < FCustomizableObjectCustomVersion::AutomaticNodeMaterialPerformanceBug) // || CustomizableObjectCustomVersion < FCustomizableObjectCustomVersion::AutomaticNodeMaterialPerformance
+	if (CustomizableObjectCustomVersion == FCustomizableObjectCustomVersion::AutomaticNodeMaterialPerformanceBug) // || CustomizableObjectCustomVersion < FCustomizableObjectCustomVersion::AutomaticNodeMaterialPerformance
 	{
 		for (const UEdGraphPin* Pin : GetAllNonOrphanPins())
 		{
@@ -534,14 +530,14 @@ void UCustomizableObjectNodeMaterial::BackwardsCompatibleFixup()
 		}
 	}
 
-	if (CustomizableObjectCustomVersion < FCustomizableObjectCustomVersion::AutomaticNodeMaterialUXImprovements)
+	if (CustomizableObjectCustomVersion == FCustomizableObjectCustomVersion::AutomaticNodeMaterialUXImprovements)
 	{
 		TextureParametersMode = bDefaultPinModeMutable_DEPRECATED ?
 			ENodePinMode::Mutable :
 			ENodePinMode::Passthrough;
 	}
 	
-	if (CustomizableObjectCustomVersion < FCustomizableObjectCustomVersion::ExtendMaterialOnlyMutableModeParameters)
+	if (CustomizableObjectCustomVersion == FCustomizableObjectCustomVersion::ExtendMaterialOnlyMutableModeParameters)
 	{
 		const uint32 NumTextureParameters = GetNumParameters(EMaterialParameterType::Texture);
 		for (uint32 ImageIndex = 0; ImageIndex < NumTextureParameters; ++ImageIndex)
@@ -554,7 +550,7 @@ void UCustomizableObjectNodeMaterial::BackwardsCompatibleFixup()
 		}
 	}
 	
-	if (CustomizableObjectCustomVersion < FCustomizableObjectCustomVersion::ExtendMaterialOnlyMutableModeParametersBug)
+	if (CustomizableObjectCustomVersion == FCustomizableObjectCustomVersion::ExtendMaterialOnlyMutableModeParametersBug)
 	{
 		for (const UEdGraphPin* Pin : GetAllPins())
 		{
@@ -565,12 +561,12 @@ void UCustomizableObjectNodeMaterial::BackwardsCompatibleFixup()
 		}
 	}
 	
-	if (CustomizableObjectCustomVersion < FCustomizableObjectCustomVersion::NodeMaterialAddTablePin)
+	if (CustomizableObjectCustomVersion == FCustomizableObjectCustomVersion::NodeMaterialAddTablePin)
 	{
 		Super::ReconstructNode();
 	}
 
-	if (CustomizableObjectCustomVersion < FCustomizableObjectCustomVersion::AddedTableMaterialSwitch)
+	if (CustomizableObjectCustomVersion == FCustomizableObjectCustomVersion::AddedTableMaterialSwitch)
 	{
 		UMaterialInstance* DefaultPinValue = nullptr;
 
@@ -625,12 +621,12 @@ void UCustomizableObjectNodeMaterial::BackwardsCompatibleFixup()
 		}
 	}
 
-	if (CustomizableObjectCustomVersion < FCustomizableObjectCustomVersion::NewComponentOptions)
+	if (CustomizableObjectCustomVersion == FCustomizableObjectCustomVersion::NewComponentOptions)
 	{
 		MeshComponentName = FName(FString::FromInt(MeshComponentIndex_DEPRECATED));
 	}
 
-	if (CustomizableObjectCustomVersion < FCustomizableObjectCustomVersion::NodeMaterialTypedImagePins)
+	if (CustomizableObjectCustomVersion == FCustomizableObjectCustomVersion::NodeMaterialTypedImagePins)
 	{
 		const UEdGraphSchema_CustomizableObject* Schema = GetDefault<UEdGraphSchema_CustomizableObject>();
 		
@@ -678,7 +674,7 @@ void UCustomizableObjectNodeMaterial::BackwardsCompatibleFixup()
 		}
 	}
 
-	if (CustomizableObjectCustomVersion < FCustomizableObjectCustomVersion::FixedMultilayerMaterialIds)
+	if (CustomizableObjectCustomVersion == FCustomizableObjectCustomVersion::FixedMultilayerMaterialIds)
 	{
 		if (Material && Material->GetCachedExpressionData().bHasMaterialLayers)
 		{
diff --git a/Engine/Plugins/Experimental/Mutable/Source/CustomizableObjectEditor/Private/MuCOE/Nodes/CustomizableObjectNodeMaterialVariation.cpp b/Engine/Plugins/Experimental/Mutable/Source/CustomizableObjectEditor/Private/MuCOE/Nodes/CustomizableObjectNodeMaterialVariation.cpp
--- a/Engine/Plugins/Experimental/Mutable/Source/CustomizableObjectEditor/Private/MuCOE/Nodes/CustomizableObjectNodeMaterialVariation.cpp
+++ b/Engine/Plugins/Experimental/Mutable/Source/CustomizableObjectEditor/Private/MuCOE/Nodes/CustomizableObjectNodeMaterialVariation.cpp
@@ -7,11 +7,11 @@
 #include "MuCOE/EdGraphSchema_CustomizableObject.h"
 
 
-void UCustomizableObjectNodeMaterialVariation::BackwardsCompatibleFixup()
+void UCustomizableObjectNodeMaterialVariation::BackwardsCompatibleFixup(int32 CustomizableObjectCustomVersion)
 {
-	const int32 CustomizableObjectCustomVersion = GetLinkerCustomVersion(FCustomizableObjectCustomVersion::GUID);
+	Super::BackwardsCompatibleFixup(CustomizableObjectCustomVersion);
 
-	if (CustomizableObjectCustomVersion < FCustomizableObjectCustomVersion::NodeVariationSerializationIssue)
+	if (CustomizableObjectCustomVersion == FCustomizableObjectCustomVersion::NodeVariationSerializationIssue)
 	{
 		for (const FCustomizableObjectMaterialVariation& OldVariation : Variations_DEPRECATED)
 		{
@@ -21,8 +21,6 @@ void UCustomizableObjectNodeMaterialVariation::BackwardsCompatibleFixup()
 			VariationsData.Add(Variation);
 		}
 	}
-
-	Super::BackwardsCompatibleFixup();
 }
 
 
diff --git a/Engine/Plugins/Experimental/Mutable/Source/CustomizableObjectEditor/Private/MuCOE/Nodes/CustomizableObjectNodeMeshSwitch.cpp b/Engine/Plugins/Experimental/Mutable/Source/CustomizableObjectEditor/Private/MuCOE/Nodes/CustomizableObjectNodeMeshSwitch.cpp
--- a/Engine/Plugins/Experimental/Mutable/Source/CustomizableObjectEditor/Private/MuCOE/Nodes/CustomizableObjectNodeMeshSwitch.cpp
+++ b/Engine/Plugins/Experimental/Mutable/Source/CustomizableObjectEditor/Private/MuCOE/Nodes/CustomizableObjectNodeMeshSwitch.cpp
@@ -8,13 +8,11 @@
 #define LOCTEXT_NAMESPACE "CustomizableObjectEditor"
 
 
-void UCustomizableObjectNodeMeshSwitch::BackwardsCompatibleFixup()
+void UCustomizableObjectNodeMeshSwitch::BackwardsCompatibleFixup(int32 CustomizableObjectCustomVersion)
 {
-	Super::BackwardsCompatibleFixup();
-
-	const int32 CustomizableObjectCustomVersion = GetLinkerCustomVersion(FCustomizableObjectCustomVersion::GUID);
-
-	if (CustomizableObjectCustomVersion < FCustomizableObjectCustomVersion::BugPinsSwitch)
+	Super::BackwardsCompatibleFixup(CustomizableObjectCustomVersion);
+	
+	if (CustomizableObjectCustomVersion == FCustomizableObjectCustomVersion::BugPinsSwitch)
 	{
 		OutputPinReference = FindPin(TEXT("Mesh"));	
 	}
diff --git a/Engine/Plugins/Experimental/Mutable/Source/CustomizableObjectEditor/Private/MuCOE/Nodes/CustomizableObjectNodeMeshVariation.cpp b/Engine/Plugins/Experimental/Mutable/Source/CustomizableObjectEditor/Private/MuCOE/Nodes/CustomizableObjectNodeMeshVariation.cpp
--- a/Engine/Plugins/Experimental/Mutable/Source/CustomizableObjectEditor/Private/MuCOE/Nodes/CustomizableObjectNodeMeshVariation.cpp
+++ b/Engine/Plugins/Experimental/Mutable/Source/CustomizableObjectEditor/Private/MuCOE/Nodes/CustomizableObjectNodeMeshVariation.cpp
@@ -7,11 +7,11 @@
 #include "MuCOE/EdGraphSchema_CustomizableObject.h"
 
 
-void UCustomizableObjectNodeMeshVariation::BackwardsCompatibleFixup()
+void UCustomizableObjectNodeMeshVariation::BackwardsCompatibleFixup(int32 CustomizableObjectCustomVersion)
 {
-	const int32 CustomizableObjectCustomVersion = GetLinkerCustomVersion(FCustomizableObjectCustomVersion::GUID);
+	Super::BackwardsCompatibleFixup(CustomizableObjectCustomVersion);
 
-	if (CustomizableObjectCustomVersion < FCustomizableObjectCustomVersion::NodeVariationSerializationIssue)
+	if (CustomizableObjectCustomVersion == FCustomizableObjectCustomVersion::NodeVariationSerializationIssue)
 	{
 		for (const FCustomizableObjectMeshVariation& OldVariation : Variations_DEPRECATED)
 		{
@@ -21,8 +21,6 @@ void UCustomizableObjectNodeMeshVariation::BackwardsCompatibleFixup()
 			VariationsData.Add(Variation);
 		}
 	}
-
-	Super::BackwardsCompatibleFixup();
 }
 
 
diff --git a/Engine/Plugins/Experimental/Mutable/Source/CustomizableObjectEditor/Private/MuCOE/Nodes/CustomizableObjectNodeModifierBase.cpp b/Engine/Plugins/Experimental/Mutable/Source/CustomizableObjectEditor/Private/MuCOE/Nodes/CustomizableObjectNodeModifierBase.cpp
--- a/Engine/Plugins/Experimental/Mutable/Source/CustomizableObjectEditor/Private/MuCOE/Nodes/CustomizableObjectNodeModifierBase.cpp
+++ b/Engine/Plugins/Experimental/Mutable/Source/CustomizableObjectEditor/Private/MuCOE/Nodes/CustomizableObjectNodeModifierBase.cpp
@@ -15,16 +15,14 @@ FLinearColor UCustomizableObjectNodeModifierBase::GetNodeTitleColor() const
 }
 
 
-void UCustomizableObjectNodeModifierBase::PostBackwardsCompatibleFixup()
+void UCustomizableObjectNodeModifierBase::BackwardsCompatibleFixup(int32 CustomizableObjectCustomVersion)
 {
-	Super::PostBackwardsCompatibleFixup();
-
-	int32 CustomizableObjectCustomVersion = GetLinkerCustomVersion(FCustomizableObjectCustomVersion::GUID);
-
+	Super::BackwardsCompatibleFixup(CustomizableObjectCustomVersion);
+	
 	const UEdGraphSchema_CustomizableObject* Schema = GetDefault<UEdGraphSchema_CustomizableObject>();
 
 	// Remove "Material" pin and add "Modifier"
-	if (CustomizableObjectCustomVersion < FCustomizableObjectCustomVersion::AddModifierPin)
+	if (CustomizableObjectCustomVersion == FCustomizableObjectCustomVersion::AddModifierPin)
 	{
 		// Replace the output pin
 		UEdGraphPin* OldPin = FindPin(TEXT("Material"));
@@ -118,6 +116,5 @@ void UCustomizableObjectNodeModifierBase::PostBackwardsCompatibleFixup()
 			}
 		}
 	}
-
 }
 
diff --git a/Engine/Plugins/Experimental/Mutable/Source/CustomizableObjectEditor/Private/MuCOE/Nodes/CustomizableObjectNodeMorphMaterial.cpp b/Engine/Plugins/Experimental/Mutable/Source/CustomizableObjectEditor/Private/MuCOE/Nodes/CustomizableObjectNodeMorphMaterial.cpp
--- a/Engine/Plugins/Experimental/Mutable/Source/CustomizableObjectEditor/Private/MuCOE/Nodes/CustomizableObjectNodeMorphMaterial.cpp
+++ b/Engine/Plugins/Experimental/Mutable/Source/CustomizableObjectEditor/Private/MuCOE/Nodes/CustomizableObjectNodeMorphMaterial.cpp
@@ -103,13 +103,11 @@ FText UCustomizableObjectNodeMorphMaterial::GetTooltipText() const
 }
 
 
-void UCustomizableObjectNodeMorphMaterial::BackwardsCompatibleFixup()
+void UCustomizableObjectNodeMorphMaterial::BackwardsCompatibleFixup(int32 CustomizableObjectCustomVersion)
 {
-	Super::BackwardsCompatibleFixup();
-
-	const int32 CustomizableObjectCustomVersion = GetLinkerCustomVersion(FCustomizableObjectCustomVersion::GUID);
-
-	if (CustomizableObjectCustomVersion < FCustomizableObjectCustomVersion::MorphMaterialAddFactorPin)
+	Super::BackwardsCompatibleFixup(CustomizableObjectCustomVersion);
+	
+	if (CustomizableObjectCustomVersion == FCustomizableObjectCustomVersion::MorphMaterialAddFactorPin)
 	{
 		ReconstructNode();
 	}
diff --git a/Engine/Plugins/Experimental/Mutable/Source/CustomizableObjectEditor/Private/MuCOE/Nodes/CustomizableObjectNodeObject.cpp b/Engine/Plugins/Experimental/Mutable/Source/CustomizableObjectEditor/Private/MuCOE/Nodes/CustomizableObjectNodeObject.cpp
--- a/Engine/Plugins/Experimental/Mutable/Source/CustomizableObjectEditor/Private/MuCOE/Nodes/CustomizableObjectNodeObject.cpp
+++ b/Engine/Plugins/Experimental/Mutable/Source/CustomizableObjectEditor/Private/MuCOE/Nodes/CustomizableObjectNodeObject.cpp
@@ -38,13 +38,11 @@ UCustomizableObjectNodeObject::UCustomizableObjectNodeObject()
 }
 
 
-void UCustomizableObjectNodeObject::BackwardsCompatibleFixup()
+void UCustomizableObjectNodeObject::BackwardsCompatibleFixup(int32 CustomizableObjectCustomVersion)
 {
-	Super::BackwardsCompatibleFixup();
-
-	const int32 CustomizableObjectCustomVersion = GetLinkerCustomVersion(FCustomizableObjectCustomVersion::GUID);
+	Super::BackwardsCompatibleFixup(CustomizableObjectCustomVersion);
 	
-	if (CustomizableObjectCustomVersion < FCustomizableObjectCustomVersion::StateTextureCompressionStrategyEnum)
+	if (CustomizableObjectCustomVersion == FCustomizableObjectCustomVersion::StateTextureCompressionStrategyEnum)
 	{
 		for (FCustomizableObjectState& State : States)
 		{
@@ -58,7 +56,7 @@ void UCustomizableObjectNodeObject::BackwardsCompatibleFixup()
 		}
 	}
 	
-	if (CustomizableObjectCustomVersion < FCustomizableObjectCustomVersion::RegenerateNodeObjectsIds)
+	if (CustomizableObjectCustomVersion == FCustomizableObjectCustomVersion::RegenerateNodeObjectsIds)
 	{
 		// This will regenerate all the Node Object Guids to finally remove the duplicated Guids warning.
 		// It is safe to do this here as Node Object do not use its node guid to link themeselves to other nodes.
@@ -75,7 +73,7 @@ void UCustomizableObjectNodeObject::BackwardsCompatibleFixup()
 	}
 
 	// Update state never-stream flag from deprecated enum
-	if (CustomizableObjectCustomVersion < FCustomizableObjectCustomVersion::CustomizableObjectStateHasSeparateNeverStreamFlag)
+	if (CustomizableObjectCustomVersion == FCustomizableObjectCustomVersion::CustomizableObjectStateHasSeparateNeverStreamFlag)
 	{
 		for (FCustomizableObjectState& s : States)
 		{
@@ -83,7 +81,7 @@ void UCustomizableObjectNodeObject::BackwardsCompatibleFixup()
 		}
 	}
 
-	if (CustomizableObjectCustomVersion < FCustomizableObjectCustomVersion::StateUIMetadata)
+	if (CustomizableObjectCustomVersion == FCustomizableObjectCustomVersion::StateUIMetadata)
 	{
 		for (FCustomizableObjectState& State : States)
 		{
@@ -96,7 +94,7 @@ void UCustomizableObjectNodeObject::BackwardsCompatibleFixup()
 		}
 	}
 
-	if (CustomizableObjectCustomVersion < FCustomizableObjectCustomVersion::NewComponentOptions)
+	if (CustomizableObjectCustomVersion == FCustomizableObjectCustomVersion::NewComponentOptions)
 	{
 		// Like we did in the CO components, we use the index of the component as the name of the component
 		for (int32 ComponentIndex = 0; ComponentIndex < ComponentSettings.Num(); ++ComponentIndex)
@@ -104,6 +102,21 @@ void UCustomizableObjectNodeObject::BackwardsCompatibleFixup()
 			ComponentSettings[ComponentIndex].ComponentName = FString::FromInt(ComponentIndex);
 		}
 	}
+
+	if (CustomizableObjectCustomVersion == FCustomizableObjectCustomVersion::MovedCompatibilityFromPostBackwardsCompatibleFixup)
+	{
+		// Fix up ComponentSettings. Only root nodes
+		if (ComponentSettings.IsEmpty() && bIsBase && !ParentObject)
+		{
+			FComponentSettings ComponentSettingsTemplate;
+			ComponentSettingsTemplate.LODReductionSettings.SetNum(NumLODs);
+
+			if (UCustomizableObject* CurrentObject = Cast<UCustomizableObject>(GetOutermostObject()))
+			{
+				ComponentSettings.Init(ComponentSettingsTemplate, CurrentObject->GetPrivate()->MutableMeshComponents.Num());
+			}
+		}
+	}
 }
 
 
@@ -300,18 +313,6 @@ void UCustomizableObjectNodeObject::PostBackwardsCompatibleFixup()
 {
 	Super::PostBackwardsCompatibleFixup();
 
-	// Fix up ComponentSettings. Only root nodes
-	if (ComponentSettings.IsEmpty() && bIsBase && !ParentObject)
-	{
-		FComponentSettings ComponentSettingsTemplate;
-		ComponentSettingsTemplate.LODReductionSettings.SetNum(NumLODs);
-
-		if (UCustomizableObject* CurrentObject = Cast<UCustomizableObject>(GetOutermostObject()))
-		{
-			ComponentSettings.Init(ComponentSettingsTemplate, CurrentObject->GetPrivate()->MutableMeshComponents.Num());
-		}
-	}
-
 	// Reconstruct in case any extension pins have changed
 	ReconstructNode();
 }
diff --git a/Engine/Plugins/Experimental/Mutable/Source/CustomizableObjectEditor/Private/MuCOE/Nodes/CustomizableObjectNodePassThroughTexture.cpp b/Engine/Plugins/Experimental/Mutable/Source/CustomizableObjectEditor/Private/MuCOE/Nodes/CustomizableObjectNodePassThroughTexture.cpp
--- a/Engine/Plugins/Experimental/Mutable/Source/CustomizableObjectEditor/Private/MuCOE/Nodes/CustomizableObjectNodePassThroughTexture.cpp
+++ b/Engine/Plugins/Experimental/Mutable/Source/CustomizableObjectEditor/Private/MuCOE/Nodes/CustomizableObjectNodePassThroughTexture.cpp
@@ -8,13 +8,11 @@
 #define LOCTEXT_NAMESPACE "CustomizableObjectEditor"
 
 
-void UCustomizableObjectNodePassThroughTexture::BackwardsCompatibleFixup()
+void UCustomizableObjectNodePassThroughTexture::BackwardsCompatibleFixup(int32 CustomizableObjectCustomVersion)
 {
-	Super::BackwardsCompatibleFixup();
-
-	const int32 CustomizableObjectCustomVersion = GetLinkerCustomVersion(FCustomizableObjectCustomVersion::GUID);
-
-	if (CustomizableObjectCustomVersion < FCustomizableObjectCustomVersion::AddedAnyTextureTypeToPassThroughTextures)
+	Super::BackwardsCompatibleFixup(CustomizableObjectCustomVersion);
+	
+	if (CustomizableObjectCustomVersion == FCustomizableObjectCustomVersion::AddedAnyTextureTypeToPassThroughTextures)
 	{
 		if (Texture)
 		{
diff --git a/Engine/Plugins/Experimental/Mutable/Source/CustomizableObjectEditor/Private/MuCOE/Nodes/CustomizableObjectNodeProjectorConstant.cpp b/Engine/Plugins/Experimental/Mutable/Source/CustomizableObjectEditor/Private/MuCOE/Nodes/CustomizableObjectNodeProjectorConstant.cpp
--- a/Engine/Plugins/Experimental/Mutable/Source/CustomizableObjectEditor/Private/MuCOE/Nodes/CustomizableObjectNodeProjectorConstant.cpp
+++ b/Engine/Plugins/Experimental/Mutable/Source/CustomizableObjectEditor/Private/MuCOE/Nodes/CustomizableObjectNodeProjectorConstant.cpp
@@ -51,13 +51,11 @@ void UCustomizableObjectNodeProjectorConstant::AllocateDefaultPins(UCustomizable
 }
 
 
-void UCustomizableObjectNodeProjectorConstant::BackwardsCompatibleFixup()
+void UCustomizableObjectNodeProjectorConstant::BackwardsCompatibleFixup(int32 CustomizableObjectCustomVersion)
 {
-	Super::BackwardsCompatibleFixup();
-
-	const int32 CustomizableObjectCustomVersion = GetLinkerCustomVersion(FCustomizableObjectCustomVersion::GUID);
-
-	if (CustomizableObjectCustomVersion < FCustomizableObjectCustomVersion::ProjectorNodesDefaultValueFix)
+	Super::BackwardsCompatibleFixup(CustomizableObjectCustomVersion);
+	
+	if (CustomizableObjectCustomVersion == FCustomizableObjectCustomVersion::ProjectorNodesDefaultValueFix)
 	{
 		Value.ProjectionType = ProjectionType_DEPRECATED;
 	}
diff --git a/Engine/Plugins/Experimental/Mutable/Source/CustomizableObjectEditor/Private/MuCOE/Nodes/CustomizableObjectNodeProjectorParameter.cpp b/Engine/Plugins/Experimental/Mutable/Source/CustomizableObjectEditor/Private/MuCOE/Nodes/CustomizableObjectNodeProjectorParameter.cpp
--- a/Engine/Plugins/Experimental/Mutable/Source/CustomizableObjectEditor/Private/MuCOE/Nodes/CustomizableObjectNodeProjectorParameter.cpp
+++ b/Engine/Plugins/Experimental/Mutable/Source/CustomizableObjectEditor/Private/MuCOE/Nodes/CustomizableObjectNodeProjectorParameter.cpp
@@ -40,13 +40,11 @@ void UCustomizableObjectNodeProjectorParameter::AllocateDefaultPins(UCustomizabl
 }
 
 
-void UCustomizableObjectNodeProjectorParameter::BackwardsCompatibleFixup()
+void UCustomizableObjectNodeProjectorParameter::BackwardsCompatibleFixup(int32 CustomizableObjectCustomVersion)
 {
-	Super::BackwardsCompatibleFixup();
-
-	const int32 CustomizableObjectCustomVersion = GetLinkerCustomVersion(FCustomizableObjectCustomVersion::GUID);
-
-	if (CustomizableObjectCustomVersion < FCustomizableObjectCustomVersion::ProjectorNodesDefaultValueFix)
+	Super::BackwardsCompatibleFixup(CustomizableObjectCustomVersion);
+	
+	if (CustomizableObjectCustomVersion == FCustomizableObjectCustomVersion::ProjectorNodesDefaultValueFix)
 	{
 		DefaultValue.ProjectionType = ProjectionType_DEPRECATED;
 	}
diff --git a/Engine/Plugins/Experimental/Mutable/Source/CustomizableObjectEditor/Private/MuCOE/Nodes/CustomizableObjectNodeRemoveMeshBlocks.cpp b/Engine/Plugins/Experimental/Mutable/Source/CustomizableObjectEditor/Private/MuCOE/Nodes/CustomizableObjectNodeRemoveMeshBlocks.cpp
--- a/Engine/Plugins/Experimental/Mutable/Source/CustomizableObjectEditor/Private/MuCOE/Nodes/CustomizableObjectNodeRemoveMeshBlocks.cpp
+++ b/Engine/Plugins/Experimental/Mutable/Source/CustomizableObjectEditor/Private/MuCOE/Nodes/CustomizableObjectNodeRemoveMeshBlocks.cpp
@@ -14,14 +14,12 @@ class UEdGraphPin;
 #define LOCTEXT_NAMESPACE "CustomizableObjectEditor"
 
 
-void UCustomizableObjectNodeRemoveMeshBlocks::BackwardsCompatibleFixup()
+void UCustomizableObjectNodeRemoveMeshBlocks::BackwardsCompatibleFixup(int32 CustomizableObjectCustomVersion)
 {
-	Super::BackwardsCompatibleFixup();
+	Super::BackwardsCompatibleFixup(CustomizableObjectCustomVersion);
 	
-	int32 CustomizableObjectCustomVersion = GetLinkerCustomVersion(FCustomizableObjectCustomVersion::GUID);
-
 	// Convert deprecated node index list to the node id list.
-	if (CustomizableObjectCustomVersion < FCustomizableObjectCustomVersion::PostLoadToCustomVersion
+	if (CustomizableObjectCustomVersion == FCustomizableObjectCustomVersion::PostLoadToCustomVersion
 		&& BlockIds_DEPRECATED.Num() < Blocks_DEPRECATED.Num())
 	{
 		if (UCustomizableObjectNodeMaterialBase* ParentMaterialNode = GetParentMaterialNode())
@@ -64,7 +62,7 @@ void UCustomizableObjectNodeRemoveMeshBlocks::BackwardsCompatibleFixup()
 	}
 
 	// Convert deprecated node id list to absolute rect list.
-	if (CustomizableObjectCustomVersion < FCustomizableObjectCustomVersion::UseUVRects)
+	if (CustomizableObjectCustomVersion == FCustomizableObjectCustomVersion::UseUVRects)
 	{
 		// If we are here, it means this node was loaded from a version that didn't have it's own layout.
 		check(Layout->Blocks.IsEmpty());
diff --git a/Engine/Plugins/Experimental/Mutable/Source/CustomizableObjectEditor/Private/MuCOE/Nodes/CustomizableObjectNodeSkeletalMesh.cpp b/Engine/Plugins/Experimental/Mutable/Source/CustomizableObjectEditor/Private/MuCOE/Nodes/CustomizableObjectNodeSkeletalMesh.cpp
--- a/Engine/Plugins/Experimental/Mutable/Source/CustomizableObjectEditor/Private/MuCOE/Nodes/CustomizableObjectNodeSkeletalMesh.cpp
+++ b/Engine/Plugins/Experimental/Mutable/Source/CustomizableObjectEditor/Private/MuCOE/Nodes/CustomizableObjectNodeSkeletalMesh.cpp
@@ -508,13 +508,11 @@ FText UCustomizableObjectNodeSkeletalMesh::GetTooltipText() const
 }
 
 
-void UCustomizableObjectNodeSkeletalMesh::BackwardsCompatibleFixup()
+void UCustomizableObjectNodeSkeletalMesh::BackwardsCompatibleFixup(int32 CustomizableObjectCustomVersion)
 {
-	Super::BackwardsCompatibleFixup();
-
-	const int32 CustomizableObjectCustomVersion = GetLinkerCustomVersion(FCustomizableObjectCustomVersion::GUID);
+	Super::BackwardsCompatibleFixup(CustomizableObjectCustomVersion);
 	
-	if (CustomizableObjectCustomVersion < FCustomizableObjectCustomVersion::PostLoadToCustomVersion)
+	if (CustomizableObjectCustomVersion == FCustomizableObjectCustomVersion::PostLoadToCustomVersion)
 	{
 		for (FCustomizableObjectNodeSkeletalMeshLOD& LOD : LODs_DEPRECATED)
 		{
@@ -568,7 +566,7 @@ void UCustomizableObjectNodeSkeletalMesh::BackwardsCompatibleFixup()
 		}
 	}
 
-	if (CustomizableObjectCustomVersion < FCustomizableObjectCustomVersion::ConvertAnimationSlotToFName)
+	if (CustomizableObjectCustomVersion == FCustomizableObjectCustomVersion::ConvertAnimationSlotToFName)
 	{
 		if (AnimBlueprintSlotName.IsNone() && AnimBlueprintSlot_DEPRECATED != -1)
 		{
@@ -577,7 +575,7 @@ void UCustomizableObjectNodeSkeletalMesh::BackwardsCompatibleFixup()
 		}
 	}
 
-	if (CustomizableObjectCustomVersion < FCustomizableObjectCustomVersion::AutomaticNodeSkeletalMesh)
+	if (CustomizableObjectCustomVersion == FCustomizableObjectCustomVersion::AutomaticNodeSkeletalMesh)
 	{
 		for (int32 LODIndex = 0; LODIndex < LODs_DEPRECATED.Num(); ++LODIndex)
 		{
@@ -642,22 +640,22 @@ void UCustomizableObjectNodeSkeletalMesh::BackwardsCompatibleFixup()
 		ReconstructNode();
 	}
 
-	if (CustomizableObjectCustomVersion < FCustomizableObjectCustomVersion::AutomaticNodeSkeletalMeshPinDataOuter)
+	if (CustomizableObjectCustomVersion == FCustomizableObjectCustomVersion::AutomaticNodeSkeletalMeshPinDataOuter)
 	{
 		ReconstructNode(); // Pins did not have Pin Data. Reconstruct them.
 	}
 
-	if (CustomizableObjectCustomVersion < FCustomizableObjectCustomVersion::AutomaticNodeSkeletalMeshPinDataUProperty)
+	if (CustomizableObjectCustomVersion == FCustomizableObjectCustomVersion::AutomaticNodeSkeletalMeshPinDataUProperty)
 	{
 		ReconstructNode(CreateRemapPinsByName()); // Correct pins but incorrect Pin Data. Reconstruct and remap pins only by name, no Pin Data.
 	}
 
-	if (CustomizableObjectCustomVersion < FCustomizableObjectCustomVersion::IgnoreDisabledSections)
+	if (CustomizableObjectCustomVersion == FCustomizableObjectCustomVersion::IgnoreDisabledSections)
 	{
 		ReconstructNode(); // Pins representing disabled sections could be present. 
 	}
 
-	if (CustomizableObjectCustomVersion < FCustomizableObjectCustomVersion::SkeletalMeshNodeDefaultPinWithoutPinData)
+	if (CustomizableObjectCustomVersion == FCustomizableObjectCustomVersion::SkeletalMeshNodeDefaultPinWithoutPinData)
 	{
 		if (const UEdGraphPin* Pin = DefaultPin.Get())
 		{
diff --git a/Engine/Plugins/Experimental/Mutable/Source/CustomizableObjectEditor/Private/MuCOE/Nodes/CustomizableObjectNodeTable.cpp b/Engine/Plugins/Experimental/Mutable/Source/CustomizableObjectEditor/Private/MuCOE/Nodes/CustomizableObjectNodeTable.cpp
--- a/Engine/Plugins/Experimental/Mutable/Source/CustomizableObjectEditor/Private/MuCOE/Nodes/CustomizableObjectNodeTable.cpp
+++ b/Engine/Plugins/Experimental/Mutable/Source/CustomizableObjectEditor/Private/MuCOE/Nodes/CustomizableObjectNodeTable.cpp
@@ -101,13 +101,11 @@ void UCustomizableObjectNodeTableRemapPins::RemapPins(const UCustomizableObjectN
 
 // Table Node -------
 
-void UCustomizableObjectNodeTable::BackwardsCompatibleFixup()
+void UCustomizableObjectNodeTable::BackwardsCompatibleFixup(int32 CustomizableObjectCustomVersion)
 {
-	Super::BackwardsCompatibleFixup();
-
-	const int32 CustomizableObjectCustomVersion = GetLinkerCustomVersion(FCustomizableObjectCustomVersion::GUID);
-
-	if (CustomizableObjectCustomVersion < FCustomizableObjectCustomVersion::AddedTableNodesTextureMode)
+	Super::BackwardsCompatibleFixup(CustomizableObjectCustomVersion);
+	
+	if (CustomizableObjectCustomVersion == FCustomizableObjectCustomVersion::AddedTableNodesTextureMode)
 	{
 		for (UEdGraphPin* Pin : Pins)
 		{
@@ -174,7 +172,7 @@ void UCustomizableObjectNodeTable::BackwardsCompatibleFixup()
 		}
 	}
 	
-	if (CustomizableObjectCustomVersion < FCustomizableObjectCustomVersion::NodeTablePinViewer)
+	if (CustomizableObjectCustomVersion == FCustomizableObjectCustomVersion::NodeTablePinViewer)
 	{
 		for (UEdGraphPin* Pin : GetAllPins())
 		{
@@ -1156,13 +1154,11 @@ TSoftClassPtr<UAnimInstance> UCustomizableObjectNodeTable::GetAnimInstanceAt(con
 }
 
 
-void UCustomizableObjectNodeTableImagePinData::PostLoad()
+void UCustomizableObjectNodeTableImagePinData::BackwardsCompatibleFixup(int32 CustomizableObjectCustomVersion)
 {
-	Super::PostLoad();
-
-	const int32 CustomizableObjectCustomVersion = GetLinkerCustomVersion(FCustomizableObjectCustomVersion::GUID);
-
-	if (CustomizableObjectCustomVersion < FCustomizableObjectCustomVersion::AddedAnyTextureTypeToPassThroughTextures)
+	Super::BackwardsCompatibleFixup(CustomizableObjectCustomVersion);
+	
+	if (CustomizableObjectCustomVersion == FCustomizableObjectCustomVersion::AddedAnyTextureTypeToPassThroughTextures)
 	{
 		if (bIsArrayTexture_DEPRECATED)
 		{
diff --git a/Engine/Plugins/Experimental/Mutable/Source/CustomizableObjectEditor/Private/MuCOE/Nodes/CustomizableObjectNodeTexture.cpp b/Engine/Plugins/Experimental/Mutable/Source/CustomizableObjectEditor/Private/MuCOE/Nodes/CustomizableObjectNodeTexture.cpp
--- a/Engine/Plugins/Experimental/Mutable/Source/CustomizableObjectEditor/Private/MuCOE/Nodes/CustomizableObjectNodeTexture.cpp
+++ b/Engine/Plugins/Experimental/Mutable/Source/CustomizableObjectEditor/Private/MuCOE/Nodes/CustomizableObjectNodeTexture.cpp
@@ -34,13 +34,11 @@ void UCustomizableObjectNodeTexture::AllocateDefaultPins(UCustomizableObjectNode
 }
 
 
-void UCustomizableObjectNodeTexture::BackwardsCompatibleFixup()
+void UCustomizableObjectNodeTexture::BackwardsCompatibleFixup(int32 CustomizableObjectCustomVersion)
 {
-	Super::BackwardsCompatibleFixup();
-
-	const int32 CustomizableObjectCustomVersion = GetLinkerCustomVersion(FCustomizableObjectCustomVersion::GUID);
-
-	if (CustomizableObjectCustomVersion < FCustomizableObjectCustomVersion::FixPinsNamesImageToTexture2)
+	Super::BackwardsCompatibleFixup(CustomizableObjectCustomVersion);
+	
+	if (CustomizableObjectCustomVersion == FCustomizableObjectCustomVersion::FixPinsNamesImageToTexture2)
 	{
 		if (UEdGraphPin* TexturePin = FindPin(TEXT("Image"))) {
 			TexturePin->PinName = TEXT("Texture");
diff --git a/Engine/Plugins/Experimental/Mutable/Source/CustomizableObjectEditor/Private/MuCOE/Nodes/CustomizableObjectNodeTextureBinarise.cpp b/Engine/Plugins/Experimental/Mutable/Source/CustomizableObjectEditor/Private/MuCOE/Nodes/CustomizableObjectNodeTextureBinarise.cpp
--- a/Engine/Plugins/Experimental/Mutable/Source/CustomizableObjectEditor/Private/MuCOE/Nodes/CustomizableObjectNodeTextureBinarise.cpp
+++ b/Engine/Plugins/Experimental/Mutable/Source/CustomizableObjectEditor/Private/MuCOE/Nodes/CustomizableObjectNodeTextureBinarise.cpp
@@ -28,13 +28,11 @@ void UCustomizableObjectNodeTextureBinarise::AllocateDefaultPins(UCustomizableOb
 }
 
 
-void UCustomizableObjectNodeTextureBinarise::BackwardsCompatibleFixup()
+void UCustomizableObjectNodeTextureBinarise::BackwardsCompatibleFixup(int32 CustomizableObjectCustomVersion)
 {
-	Super::BackwardsCompatibleFixup();
-
-	const int32 CustomizableObjectCustomVersion = GetLinkerCustomVersion(FCustomizableObjectCustomVersion::GUID);
-
-	if (CustomizableObjectCustomVersion < FCustomizableObjectCustomVersion::FixPinsNamesImageToTexture2)
+	Super::BackwardsCompatibleFixup(CustomizableObjectCustomVersion);
+	
+	if (CustomizableObjectCustomVersion == FCustomizableObjectCustomVersion::FixPinsNamesImageToTexture2)
 	{
 		if (UEdGraphPin* TexturePin = FindPin(TEXT("Image"))) {
 			TexturePin->PinName = TEXT("Texture");
@@ -42,7 +40,7 @@ void UCustomizableObjectNodeTextureBinarise::BackwardsCompatibleFixup()
 		}
 	}
 
-	if (CustomizableObjectCustomVersion < FCustomizableObjectCustomVersion::FixPinsNamesImageToTexture2)
+	if (CustomizableObjectCustomVersion == FCustomizableObjectCustomVersion::FixPinsNamesImageToTexture2)
 	{
 		if (UEdGraphPin* TexturePin = FindPin(TEXT("Base Image"))) {
 			TexturePin->PinName = TEXT("Base Texture");
diff --git a/Engine/Plugins/Experimental/Mutable/Source/CustomizableObjectEditor/Private/MuCOE/Nodes/CustomizableObjectNodeTextureColourMap.cpp b/Engine/Plugins/Experimental/Mutable/Source/CustomizableObjectEditor/Private/MuCOE/Nodes/CustomizableObjectNodeTextureColourMap.cpp
--- a/Engine/Plugins/Experimental/Mutable/Source/CustomizableObjectEditor/Private/MuCOE/Nodes/CustomizableObjectNodeTextureColourMap.cpp
+++ b/Engine/Plugins/Experimental/Mutable/Source/CustomizableObjectEditor/Private/MuCOE/Nodes/CustomizableObjectNodeTextureColourMap.cpp
@@ -32,13 +32,11 @@ void UCustomizableObjectNodeTextureColourMap::AllocateDefaultPins(UCustomizableO
 }
 
 
-void UCustomizableObjectNodeTextureColourMap::BackwardsCompatibleFixup()
+void UCustomizableObjectNodeTextureColourMap::BackwardsCompatibleFixup(int32 CustomizableObjectCustomVersion)
 {
-	Super::BackwardsCompatibleFixup();
-
-	const int32 CustomizableObjectCustomVersion = GetLinkerCustomVersion(FCustomizableObjectCustomVersion::GUID);
-
-	if (CustomizableObjectCustomVersion < FCustomizableObjectCustomVersion::FixPinsNamesImageToTexture2)
+	Super::BackwardsCompatibleFixup(CustomizableObjectCustomVersion);
+	
+	if (CustomizableObjectCustomVersion == FCustomizableObjectCustomVersion::FixPinsNamesImageToTexture2)
 	{
 		if (UEdGraphPin* TexturePin = FindPin(TEXT("Image")))
 		{
diff --git a/Engine/Plugins/Experimental/Mutable/Source/CustomizableObjectEditor/Private/MuCOE/Nodes/CustomizableObjectNodeTextureFromChannels.cpp b/Engine/Plugins/Experimental/Mutable/Source/CustomizableObjectEditor/Private/MuCOE/Nodes/CustomizableObjectNodeTextureFromChannels.cpp
--- a/Engine/Plugins/Experimental/Mutable/Source/CustomizableObjectEditor/Private/MuCOE/Nodes/CustomizableObjectNodeTextureFromChannels.cpp
+++ b/Engine/Plugins/Experimental/Mutable/Source/CustomizableObjectEditor/Private/MuCOE/Nodes/CustomizableObjectNodeTextureFromChannels.cpp
@@ -36,13 +36,11 @@ void UCustomizableObjectNodeTextureFromChannels::AllocateDefaultPins(UCustomizab
 }
 
 
-void UCustomizableObjectNodeTextureFromChannels::BackwardsCompatibleFixup()
+void UCustomizableObjectNodeTextureFromChannels::BackwardsCompatibleFixup(int32 CustomizableObjectCustomVersion)
 {
-	Super::BackwardsCompatibleFixup();
-
-	const int32 CustomizableObjectCustomVersion = GetLinkerCustomVersion(FCustomizableObjectCustomVersion::GUID);
-
-	if (CustomizableObjectCustomVersion < FCustomizableObjectCustomVersion::FixPinsNamesImageToTexture2)
+	Super::BackwardsCompatibleFixup(CustomizableObjectCustomVersion);
+	
+	if (CustomizableObjectCustomVersion == FCustomizableObjectCustomVersion::FixPinsNamesImageToTexture2)
 	{
 		if (UEdGraphPin* TexturePin = FindPin(TEXT("Image")))
 		{
diff --git a/Engine/Plugins/Experimental/Mutable/Source/CustomizableObjectEditor/Private/MuCOE/Nodes/CustomizableObjectNodeTextureFromColor.cpp b/Engine/Plugins/Experimental/Mutable/Source/CustomizableObjectEditor/Private/MuCOE/Nodes/CustomizableObjectNodeTextureFromColor.cpp
--- a/Engine/Plugins/Experimental/Mutable/Source/CustomizableObjectEditor/Private/MuCOE/Nodes/CustomizableObjectNodeTextureFromColor.cpp
+++ b/Engine/Plugins/Experimental/Mutable/Source/CustomizableObjectEditor/Private/MuCOE/Nodes/CustomizableObjectNodeTextureFromColor.cpp
@@ -24,13 +24,11 @@ void UCustomizableObjectNodeTextureFromColor::AllocateDefaultPins(UCustomizableO
 }
 
 
-void UCustomizableObjectNodeTextureFromColor::BackwardsCompatibleFixup()
+void UCustomizableObjectNodeTextureFromColor::BackwardsCompatibleFixup(int32 CustomizableObjectCustomVersion)
 {
-	Super::BackwardsCompatibleFixup();
-
-	const int32 CustomizableObjectCustomVersion = GetLinkerCustomVersion(FCustomizableObjectCustomVersion::GUID);
-
-	if (CustomizableObjectCustomVersion < FCustomizableObjectCustomVersion::FixPinsNamesImageToTexture2)
+	Super::BackwardsCompatibleFixup(CustomizableObjectCustomVersion);
+	
+	if (CustomizableObjectCustomVersion == FCustomizableObjectCustomVersion::FixPinsNamesImageToTexture2)
 	{
 		if (UEdGraphPin* TexturePin = FindPin(TEXT("Image")))
 		{
diff --git a/Engine/Plugins/Experimental/Mutable/Source/CustomizableObjectEditor/Private/MuCOE/Nodes/CustomizableObjectNodeTextureInvert.cpp b/Engine/Plugins/Experimental/Mutable/Source/CustomizableObjectEditor/Private/MuCOE/Nodes/CustomizableObjectNodeTextureInvert.cpp
--- a/Engine/Plugins/Experimental/Mutable/Source/CustomizableObjectEditor/Private/MuCOE/Nodes/CustomizableObjectNodeTextureInvert.cpp
+++ b/Engine/Plugins/Experimental/Mutable/Source/CustomizableObjectEditor/Private/MuCOE/Nodes/CustomizableObjectNodeTextureInvert.cpp
@@ -25,18 +25,16 @@ void UCustomizableObjectNodeTextureInvert::AllocateDefaultPins(UCustomizableObje
 }
 
 
-void UCustomizableObjectNodeTextureInvert::BackwardsCompatibleFixup()
+void UCustomizableObjectNodeTextureInvert::BackwardsCompatibleFixup(int32 CustomizableObjectCustomVersion)
 {
-	Super::BackwardsCompatibleFixup();
+	Super::BackwardsCompatibleFixup(CustomizableObjectCustomVersion);
 
-	const int32 CustomizableObjectCustomVersion = GetLinkerCustomVersion(FCustomizableObjectCustomVersion::GUID);
-
-	if (CustomizableObjectCustomVersion < FCustomizableObjectCustomVersion::PinsNamesImageToTexture)
+	if (CustomizableObjectCustomVersion == FCustomizableObjectCustomVersion::PinsNamesImageToTexture)
 	{
 		BaseImagePinReference = FEdGraphPinReference(FindPin(TEXT("Base Image")));
 	}
 	
-	if (CustomizableObjectCustomVersion < FCustomizableObjectCustomVersion::FixPinsNamesImageToTexture2)
+	if (CustomizableObjectCustomVersion == FCustomizableObjectCustomVersion::FixPinsNamesImageToTexture2)
 	{
 		bool Replaced = false;
 		if (UEdGraphPin* TexturePin = FindPin(TEXT("Image")))
diff --git a/Engine/Plugins/Experimental/Mutable/Source/CustomizableObjectEditor/Private/MuCOE/Nodes/CustomizableObjectNodeTextureLayer.cpp b/Engine/Plugins/Experimental/Mutable/Source/CustomizableObjectEditor/Private/MuCOE/Nodes/CustomizableObjectNodeTextureLayer.cpp
--- a/Engine/Plugins/Experimental/Mutable/Source/CustomizableObjectEditor/Private/MuCOE/Nodes/CustomizableObjectNodeTextureLayer.cpp
+++ b/Engine/Plugins/Experimental/Mutable/Source/CustomizableObjectEditor/Private/MuCOE/Nodes/CustomizableObjectNodeTextureLayer.cpp
@@ -70,13 +70,11 @@ void UCustomizableObjectNodeTextureLayer::AllocateDefaultPins(UCustomizableObjec
 }
 
 
-void UCustomizableObjectNodeTextureLayer::BackwardsCompatibleFixup()
+void UCustomizableObjectNodeTextureLayer::BackwardsCompatibleFixup(int32 CustomizableObjectCustomVersion)
 {
-	Super::BackwardsCompatibleFixup();
-
-	const int32 CustomizableObjectCustomVersion = GetLinkerCustomVersion(FCustomizableObjectCustomVersion::GUID);
-
-	if (CustomizableObjectCustomVersion < FCustomizableObjectCustomVersion::FixPinsNamesImageToTexture2)
+	Super::BackwardsCompatibleFixup(CustomizableObjectCustomVersion);
+	
+	if (CustomizableObjectCustomVersion == FCustomizableObjectCustomVersion::FixPinsNamesImageToTexture2)
 	{
 		if (UEdGraphPin* TexturePin = FindPin(TEXT("Image")))
 		{
diff --git a/Engine/Plugins/Experimental/Mutable/Source/CustomizableObjectEditor/Private/MuCOE/Nodes/CustomizableObjectNodeTextureParameter.cpp b/Engine/Plugins/Experimental/Mutable/Source/CustomizableObjectEditor/Private/MuCOE/Nodes/CustomizableObjectNodeTextureParameter.cpp
--- a/Engine/Plugins/Experimental/Mutable/Source/CustomizableObjectEditor/Private/MuCOE/Nodes/CustomizableObjectNodeTextureParameter.cpp
+++ b/Engine/Plugins/Experimental/Mutable/Source/CustomizableObjectEditor/Private/MuCOE/Nodes/CustomizableObjectNodeTextureParameter.cpp
@@ -59,13 +59,11 @@ void UCustomizableObjectNodeTextureParameter::OnRenameNode(const FString& NewNam
 }
 
 
-void UCustomizableObjectNodeTextureParameter::BackwardsCompatibleFixup()
+void UCustomizableObjectNodeTextureParameter::BackwardsCompatibleFixup(int32 CustomizableObjectCustomVersion)
 {
-	Super::BackwardsCompatibleFixup();
-
-	const int32 CustomizableObjectCustomVersion = GetLinkerCustomVersion(FCustomizableObjectCustomVersion::GUID);
-
-	if (CustomizableObjectCustomVersion < FCustomizableObjectCustomVersion::NodeTextureParameterDefaultToReferenceValue)
+	Super::BackwardsCompatibleFixup(CustomizableObjectCustomVersion);
+	
+	if (CustomizableObjectCustomVersion == FCustomizableObjectCustomVersion::NodeTextureParameterDefaultToReferenceValue)
 	{
 		ReferenceValue = DefaultValue;
 		DefaultValue = {};
diff --git a/Engine/Plugins/Experimental/Mutable/Source/CustomizableObjectEditor/Private/MuCOE/Nodes/CustomizableObjectNodeTextureProject.cpp b/Engine/Plugins/Experimental/Mutable/Source/CustomizableObjectEditor/Private/MuCOE/Nodes/CustomizableObjectNodeTextureProject.cpp
--- a/Engine/Plugins/Experimental/Mutable/Source/CustomizableObjectEditor/Private/MuCOE/Nodes/CustomizableObjectNodeTextureProject.cpp
+++ b/Engine/Plugins/Experimental/Mutable/Source/CustomizableObjectEditor/Private/MuCOE/Nodes/CustomizableObjectNodeTextureProject.cpp
@@ -117,13 +117,11 @@ void UCustomizableObjectNodeTextureProject::AllocateDefaultPins(UCustomizableObj
 }
 
 
-void UCustomizableObjectNodeTextureProject::BackwardsCompatibleFixup()
+void UCustomizableObjectNodeTextureProject::BackwardsCompatibleFixup(int32 CustomizableObjectCustomVersion)
 {
-	Super::BackwardsCompatibleFixup();
-
-	const int32 CustomizableObjectCustomVersion = GetLinkerCustomVersion(FCustomizableObjectCustomVersion::GUID);
-
-	if (CustomizableObjectCustomVersion < FCustomizableObjectCustomVersion::FixPinsNamesImageToTexture2)
+	Super::BackwardsCompatibleFixup(CustomizableObjectCustomVersion);
+	
+	if (CustomizableObjectCustomVersion == FCustomizableObjectCustomVersion::FixPinsNamesImageToTexture2)
 	{
 		TArray<UEdGraphPin*> AllPins = GetAllPins();
 		uint32 TexturePinsCount = 0;
diff --git a/Engine/Plugins/Experimental/Mutable/Source/CustomizableObjectEditor/Private/MuCOE/Nodes/CustomizableObjectNodeTextureSwitch.cpp b/Engine/Plugins/Experimental/Mutable/Source/CustomizableObjectEditor/Private/MuCOE/Nodes/CustomizableObjectNodeTextureSwitch.cpp
--- a/Engine/Plugins/Experimental/Mutable/Source/CustomizableObjectEditor/Private/MuCOE/Nodes/CustomizableObjectNodeTextureSwitch.cpp
+++ b/Engine/Plugins/Experimental/Mutable/Source/CustomizableObjectEditor/Private/MuCOE/Nodes/CustomizableObjectNodeTextureSwitch.cpp
@@ -8,18 +8,16 @@
 #define LOCTEXT_NAMESPACE "CustomizableObjectEditor"
 
 
-void UCustomizableObjectNodeTextureSwitch::BackwardsCompatibleFixup()
+void UCustomizableObjectNodeTextureSwitch::BackwardsCompatibleFixup(int32 CustomizableObjectCustomVersion)
 {
-	Super::BackwardsCompatibleFixup();
-
-	const int32 CustomizableObjectCustomVersion = GetLinkerCustomVersion(FCustomizableObjectCustomVersion::GUID);
-
-	if (CustomizableObjectCustomVersion < FCustomizableObjectCustomVersion::PinsNamesImageToTexture)
+	Super::BackwardsCompatibleFixup(CustomizableObjectCustomVersion);
+	
+	if (CustomizableObjectCustomVersion == FCustomizableObjectCustomVersion::PinsNamesImageToTexture)
 	{
 		OutputPinReference = FindPin(TEXT("Image"));
 	}
 
-	if (CustomizableObjectCustomVersion < FCustomizableObjectCustomVersion::FixPinsNamesImageToTexture2)
+	if (CustomizableObjectCustomVersion == FCustomizableObjectCustomVersion::FixPinsNamesImageToTexture2)
 	{
 		if (UEdGraphPin* TexturePin = FindPin(TEXT("Image")))
 		{
diff --git a/Engine/Plugins/Experimental/Mutable/Source/CustomizableObjectEditor/Private/MuCOE/Nodes/CustomizableObjectNodeTextureToChannels.cpp b/Engine/Plugins/Experimental/Mutable/Source/CustomizableObjectEditor/Private/MuCOE/Nodes/CustomizableObjectNodeTextureToChannels.cpp
--- a/Engine/Plugins/Experimental/Mutable/Source/CustomizableObjectEditor/Private/MuCOE/Nodes/CustomizableObjectNodeTextureToChannels.cpp
+++ b/Engine/Plugins/Experimental/Mutable/Source/CustomizableObjectEditor/Private/MuCOE/Nodes/CustomizableObjectNodeTextureToChannels.cpp
@@ -68,13 +68,11 @@ void UCustomizableObjectNodeTextureToChannels::AllocateDefaultPins(UCustomizable
 }
 
 
-void UCustomizableObjectNodeTextureToChannels::BackwardsCompatibleFixup()
+void UCustomizableObjectNodeTextureToChannels::BackwardsCompatibleFixup(int32 CustomizableObjectCustomVersion)
 {
-	Super::BackwardsCompatibleFixup();
-
-	const int32 CustomizableObjectCustomVersion = GetLinkerCustomVersion(FCustomizableObjectCustomVersion::GUID);
-
-	if (CustomizableObjectCustomVersion < FCustomizableObjectCustomVersion::FixPinsNamesImageToTexture2)
+	Super::BackwardsCompatibleFixup(CustomizableObjectCustomVersion);
+	
+	if (CustomizableObjectCustomVersion == FCustomizableObjectCustomVersion::FixPinsNamesImageToTexture2)
 	{
 		if (UEdGraphPin* TexturePin = FindPin(TEXT("Image"))) {
 			TexturePin->PinName = TEXT("Texture");
diff --git a/Engine/Plugins/Experimental/Mutable/Source/CustomizableObjectEditor/Private/MuCOE/Nodes/CustomizableObjectNodeVariation.cpp b/Engine/Plugins/Experimental/Mutable/Source/CustomizableObjectEditor/Private/MuCOE/Nodes/CustomizableObjectNodeVariation.cpp
--- a/Engine/Plugins/Experimental/Mutable/Source/CustomizableObjectEditor/Private/MuCOE/Nodes/CustomizableObjectNodeVariation.cpp
+++ b/Engine/Plugins/Experimental/Mutable/Source/CustomizableObjectEditor/Private/MuCOE/Nodes/CustomizableObjectNodeVariation.cpp
@@ -97,13 +97,11 @@ FText UCustomizableObjectNodeVariation::GetTooltipText() const
 }
 
 
-void UCustomizableObjectNodeVariation::BackwardsCompatibleFixup()
+void UCustomizableObjectNodeVariation::BackwardsCompatibleFixup(int32 CustomizableObjectCustomVersion)
 {
-	Super::BackwardsCompatibleFixup();
-
-	const int32 CustomizableObjectCustomVersion = GetLinkerCustomVersion(FCustomizableObjectCustomVersion::GUID);
-
-	if (CustomizableObjectCustomVersion < FCustomizableObjectCustomVersion::NodeVariationSerializationIssue)
+	Super::BackwardsCompatibleFixup(CustomizableObjectCustomVersion);
+	
+	if (CustomizableObjectCustomVersion == FCustomizableObjectCustomVersion::NodeVariationSerializationIssue)
 	{
 		ReconstructNode();
 	}
diff --git a/Engine/Plugins/Experimental/Mutable/Source/CustomizableObjectEditor/Public/MuCOE/Nodes/CustomizableObjectNode.h b/Engine/Plugins/Experimental/Mutable/Source/CustomizableObjectEditor/Public/MuCOE/Nodes/CustomizableObjectNode.h
--- a/Engine/Plugins/Experimental/Mutable/Source/CustomizableObjectEditor/Public/MuCOE/Nodes/CustomizableObjectNode.h
+++ b/Engine/Plugins/Experimental/Mutable/Source/CustomizableObjectEditor/Public/MuCOE/Nodes/CustomizableObjectNode.h
@@ -34,6 +34,10 @@ public:
 
 	/** Virtual function used to copy pin data when remapping pins. */
 	virtual void Copy(const UCustomizableObjectNodePinData& Other) {}
+	
+	/** Add backwards compatible code here.
+	 * When called, it is guaranteed that all nodes in this graph will have executed the PostLoad function. */
+	virtual void BackwardsCompatibleFixup(int32 CustomizableObjectCustomVersion) {}
 
 protected:
 	/** Virtual function used to perform comparision between different UCustomizableObjectNodePinData inherited types. */
@@ -107,10 +111,11 @@ public:
 
 	/** Add backwards compatible code here.
 	 * When called, it is guaranteed that all nodes in this graph will have executed the PostLoad function. */
-	virtual void BackwardsCompatibleFixup();
+	virtual void BackwardsCompatibleFixup(int32 CustomizableObjectCustomVersion);
 
 	/** Add post load work here.
-	 * When called, it is guaranteed that all nodes in this graph will have executed the BackwardsCompatibleFixup function. */
+	 * When called, it is guaranteed that all nodes in this graph will have executed the BackwardsCompatibleFixup function.
+	 * Notice that no compatibility code should not go here. Use BackwardsCompatibleFixup instead. */
 	virtual void PostBackwardsCompatibleFixup() {}
 	
 	/** Virtual implementation of RemovePin. Allows to do work before removing a pin.
diff --git a/Engine/Plugins/Experimental/Mutable/Source/CustomizableObjectEditor/Public/MuCOE/Nodes/CustomizableObjectNodeColorSwitch.h b/Engine/Plugins/Experimental/Mutable/Source/CustomizableObjectEditor/Public/MuCOE/Nodes/CustomizableObjectNodeColorSwitch.h
--- a/Engine/Plugins/Experimental/Mutable/Source/CustomizableObjectEditor/Public/MuCOE/Nodes/CustomizableObjectNodeColorSwitch.h
+++ b/Engine/Plugins/Experimental/Mutable/Source/CustomizableObjectEditor/Public/MuCOE/Nodes/CustomizableObjectNodeColorSwitch.h
@@ -14,7 +14,7 @@ public:
 	GENERATED_BODY()
 
 	// UCustomizableObjectNode interface
-	virtual void BackwardsCompatibleFixup() override;
+	virtual void BackwardsCompatibleFixup(int32 CustomizableObjectCustomVersion) override;
 
 	// UCustomizableObjectNodeSwitchBase interface
 	virtual FName GetCategory() const override;
diff --git a/Engine/Plugins/Experimental/Mutable/Source/CustomizableObjectEditor/Public/MuCOE/Nodes/CustomizableObjectNodeColorVariation.h b/Engine/Plugins/Experimental/Mutable/Source/CustomizableObjectEditor/Public/MuCOE/Nodes/CustomizableObjectNodeColorVariation.h
--- a/Engine/Plugins/Experimental/Mutable/Source/CustomizableObjectEditor/Public/MuCOE/Nodes/CustomizableObjectNodeColorVariation.h
+++ b/Engine/Plugins/Experimental/Mutable/Source/CustomizableObjectEditor/Public/MuCOE/Nodes/CustomizableObjectNodeColorVariation.h
@@ -20,7 +20,7 @@ public:
 	TArray<FCustomizableObjectColorVariation> Variations_DEPRECATED;
 
 	// UCustomizableObjectNode interface
-	virtual void BackwardsCompatibleFixup() override;
+	virtual void BackwardsCompatibleFixup(int32 CustomizableObjectCustomVersion) override;
 	
 	// UCustomizableObjectNodeVariation interface
 	virtual FName GetCategory() const override;
diff --git a/Engine/Plugins/Experimental/Mutable/Source/CustomizableObjectEditor/Public/MuCOE/Nodes/CustomizableObjectNodeEditMaterial.h b/Engine/Plugins/Experimental/Mutable/Source/CustomizableObjectEditor/Public/MuCOE/Nodes/CustomizableObjectNodeEditMaterial.h
--- a/Engine/Plugins/Experimental/Mutable/Source/CustomizableObjectEditor/Public/MuCOE/Nodes/CustomizableObjectNodeEditMaterial.h
+++ b/Engine/Plugins/Experimental/Mutable/Source/CustomizableObjectEditor/Public/MuCOE/Nodes/CustomizableObjectNodeEditMaterial.h
@@ -57,7 +57,7 @@ public:
 	virtual void PinConnectionListChanged(UEdGraphPin* Pin) override;
 
 	// UCustomizableObjectNode interface
-	virtual void BackwardsCompatibleFixup() override;
+	virtual void BackwardsCompatibleFixup(int32 CustomizableObjectCustomVersion) override;
 	virtual void PostBackwardsCompatibleFixup() override;
 	virtual void AllocateDefaultPins(UCustomizableObjectNodeRemapPins* RemapPins) override;
 	virtual bool IsNodeOutDatedAndNeedsRefresh() override;
diff --git a/Engine/Plugins/Experimental/Mutable/Source/CustomizableObjectEditor/Public/MuCOE/Nodes/CustomizableObjectNodeExtendMaterial.h b/Engine/Plugins/Experimental/Mutable/Source/CustomizableObjectEditor/Public/MuCOE/Nodes/CustomizableObjectNodeExtendMaterial.h
--- a/Engine/Plugins/Experimental/Mutable/Source/CustomizableObjectEditor/Public/MuCOE/Nodes/CustomizableObjectNodeExtendMaterial.h
+++ b/Engine/Plugins/Experimental/Mutable/Source/CustomizableObjectEditor/Public/MuCOE/Nodes/CustomizableObjectNodeExtendMaterial.h
@@ -38,7 +38,7 @@ public:
 	
 	// UCustomizableObjectNode
 	virtual void PostBackwardsCompatibleFixup() override;
-	virtual void BackwardsCompatibleFixup() override;
+	virtual void BackwardsCompatibleFixup(int32 CustomizableObjectCustomVersion) override;
 	virtual void UpdateReferencedNodeId(const FGuid& NewGuid) override;
 
 	// Begin EdGraphNode interface
diff --git a/Engine/Plugins/Experimental/Mutable/Source/CustomizableObjectEditor/Public/MuCOE/Nodes/CustomizableObjectNodeExternalPin.h b/Engine/Plugins/Experimental/Mutable/Source/CustomizableObjectEditor/Public/MuCOE/Nodes/CustomizableObjectNodeExternalPin.h
--- a/Engine/Plugins/Experimental/Mutable/Source/CustomizableObjectEditor/Public/MuCOE/Nodes/CustomizableObjectNodeExternalPin.h
+++ b/Engine/Plugins/Experimental/Mutable/Source/CustomizableObjectEditor/Public/MuCOE/Nodes/CustomizableObjectNodeExternalPin.h
@@ -33,7 +33,7 @@ public:
 	virtual FText GetTooltipText() const override;
 	
 	// UCustomizableObjectNode interface
-	virtual void BackwardsCompatibleFixup() override;
+	virtual void BackwardsCompatibleFixup(int32 CustomizableObjectCustomVersion) override;
 	virtual void PostBackwardsCompatibleFixup() override;
 	virtual void AllocateDefaultPins(UCustomizableObjectNodeRemapPins* RemapPins) override;
 	virtual UCustomizableObjectNodeRemapPins* CreateRemapPinsDefault() const override;
diff --git a/Engine/Plugins/Experimental/Mutable/Source/CustomizableObjectEditor/Public/MuCOE/Nodes/CustomizableObjectNodeFloatParameter.h b/Engine/Plugins/Experimental/Mutable/Source/CustomizableObjectEditor/Public/MuCOE/Nodes/CustomizableObjectNodeFloatParameter.h
--- a/Engine/Plugins/Experimental/Mutable/Source/CustomizableObjectEditor/Public/MuCOE/Nodes/CustomizableObjectNodeFloatParameter.h
+++ b/Engine/Plugins/Experimental/Mutable/Source/CustomizableObjectEditor/Public/MuCOE/Nodes/CustomizableObjectNodeFloatParameter.h
@@ -53,7 +53,7 @@ public:
 	// UCustomizableObjectNode interface
 	void AllocateDefaultPins(UCustomizableObjectNodeRemapPins* RemapPins) override;
 	bool IsAffectedByLOD() const override;
-	void BackwardsCompatibleFixup() override;
+	void BackwardsCompatibleFixup(int32 CustomizableObjectCustomVersion) override;
 	virtual void OnRenameNode(const FString& NewName) override;
 
 };
diff --git a/Engine/Plugins/Experimental/Mutable/Source/CustomizableObjectEditor/Public/MuCOE/Nodes/CustomizableObjectNodeFloatSwitch.h b/Engine/Plugins/Experimental/Mutable/Source/CustomizableObjectEditor/Public/MuCOE/Nodes/CustomizableObjectNodeFloatSwitch.h
--- a/Engine/Plugins/Experimental/Mutable/Source/CustomizableObjectEditor/Public/MuCOE/Nodes/CustomizableObjectNodeFloatSwitch.h
+++ b/Engine/Plugins/Experimental/Mutable/Source/CustomizableObjectEditor/Public/MuCOE/Nodes/CustomizableObjectNodeFloatSwitch.h
@@ -14,7 +14,7 @@ public:
 	GENERATED_BODY()
 
 	// UCustomizableObjectNode interface
-	virtual void BackwardsCompatibleFixup() override;
+	virtual void BackwardsCompatibleFixup(int32 CustomizableObjectCustomVersion) override;
 
 	// UCustomizableObjectNodeSwitchBase interface
 	virtual FName GetCategory() const override;
diff --git a/Engine/Plugins/Experimental/Mutable/Source/CustomizableObjectEditor/Public/MuCOE/Nodes/CustomizableObjectNodeFloatVariation.h b/Engine/Plugins/Experimental/Mutable/Source/CustomizableObjectEditor/Public/MuCOE/Nodes/CustomizableObjectNodeFloatVariation.h
--- a/Engine/Plugins/Experimental/Mutable/Source/CustomizableObjectEditor/Public/MuCOE/Nodes/CustomizableObjectNodeFloatVariation.h
+++ b/Engine/Plugins/Experimental/Mutable/Source/CustomizableObjectEditor/Public/MuCOE/Nodes/CustomizableObjectNodeFloatVariation.h
@@ -20,7 +20,7 @@ public:
 	TArray<FCustomizableObjectFloatVariation> Variations_DEPRECATED;
 
 	// UCustomizableObjectNode interface
-	virtual void BackwardsCompatibleFixup() override;
+	virtual void BackwardsCompatibleFixup(int32 CustomizableObjectCustomVersion) override;
 	
 	// UCustomizableObjectNodeVariation interface
 	virtual FName GetCategory() const override;
diff --git a/Engine/Plugins/Experimental/Mutable/Source/CustomizableObjectEditor/Public/MuCOE/Nodes/CustomizableObjectNodeGroupProjectorParameter.h b/Engine/Plugins/Experimental/Mutable/Source/CustomizableObjectEditor/Public/MuCOE/Nodes/CustomizableObjectNodeGroupProjectorParameter.h
--- a/Engine/Plugins/Experimental/Mutable/Source/CustomizableObjectEditor/Public/MuCOE/Nodes/CustomizableObjectNodeGroupProjectorParameter.h
+++ b/Engine/Plugins/Experimental/Mutable/Source/CustomizableObjectEditor/Public/MuCOE/Nodes/CustomizableObjectNodeGroupProjectorParameter.h
@@ -102,7 +102,7 @@ public:
 	virtual FText GetTooltipText() const override;
 
 	// UCustomizableObjectNode interface.
-	virtual void BackwardsCompatibleFixup() override;
+	virtual void BackwardsCompatibleFixup(int32 CustomizableObjectCustomVersion) override;
 	virtual void AllocateDefaultPins(UCustomizableObjectNodeRemapPins* RemapPins);
 };
 
diff --git a/Engine/Plugins/Experimental/Mutable/Source/CustomizableObjectEditor/Public/MuCOE/Nodes/CustomizableObjectNodeMaterial.h b/Engine/Plugins/Experimental/Mutable/Source/CustomizableObjectEditor/Public/MuCOE/Nodes/CustomizableObjectNodeMaterial.h
--- a/Engine/Plugins/Experimental/Mutable/Source/CustomizableObjectEditor/Public/MuCOE/Nodes/CustomizableObjectNodeMaterial.h
+++ b/Engine/Plugins/Experimental/Mutable/Source/CustomizableObjectEditor/Public/MuCOE/Nodes/CustomizableObjectNodeMaterial.h
@@ -120,7 +120,7 @@ public:
 	virtual bool CanConnect(const UEdGraphPin* InOwnedInputPin, const UEdGraphPin* InOutputPin, bool& bOutIsOtherNodeBlocklisted, bool& bOutArePinsCompatible) const override;
 
 	// UCustomizableObjectNode interface
-	virtual void BackwardsCompatibleFixup() override;
+	virtual void BackwardsCompatibleFixup(int32 CustomizableObjectCustomVersion) override;
 	virtual void AllocateDefaultPins(UCustomizableObjectNodeRemapPins* RemapPins) override;
 	virtual bool CanPinBeHidden(const UEdGraphPin& Pin) const override;
 	virtual bool HasPinViewer() const override;
@@ -259,17 +259,16 @@ class CUSTOMIZABLEOBJECTEDITOR_API UCustomizableObjectNodeMaterialPinDataImage :
 {
 	GENERATED_BODY()
 
-	friend void UCustomizableObjectNodeMaterial::BackwardsCompatibleFixup();
+	friend void UCustomizableObjectNodeMaterial::BackwardsCompatibleFixup(int32 CustomizableObjectCustomVersion);
 	
 public:
 	// UObject interface
 	virtual void PostEditChangeProperty(FPropertyChangedEvent& PropertyChangedEvent) override;
 	virtual bool CanEditChange(const FProperty* InProperty) const override;
 
-private:
-	virtual void PostLoad() override;
-	
-public:
+	// UCustomizableObjectNodePinData interface
+	virtual void BackwardsCompatibleFixup(int32 CustomizableObjectCustomVersion) override;
+
 	// NodePinDataParameter interface
 	/** Virtual function used to copy pin data when remapping pins. */
 	virtual void Copy(const UCustomizableObjectNodePinData& Other) override;
diff --git a/Engine/Plugins/Experimental/Mutable/Source/CustomizableObjectEditor/Public/MuCOE/Nodes/CustomizableObjectNodeMaterialVariation.h b/Engine/Plugins/Experimental/Mutable/Source/CustomizableObjectEditor/Public/MuCOE/Nodes/CustomizableObjectNodeMaterialVariation.h
--- a/Engine/Plugins/Experimental/Mutable/Source/CustomizableObjectEditor/Public/MuCOE/Nodes/CustomizableObjectNodeMaterialVariation.h
+++ b/Engine/Plugins/Experimental/Mutable/Source/CustomizableObjectEditor/Public/MuCOE/Nodes/CustomizableObjectNodeMaterialVariation.h
@@ -33,7 +33,7 @@ private:
 
 public:
 	// UCustomizableObjectNode interface
-	virtual void BackwardsCompatibleFixup() override;
+	virtual void BackwardsCompatibleFixup(int32 CustomizableObjectCustomVersion) override;
 	virtual bool IsSingleOutputNode() const override;
 	
 	// UCustomizableObjectNodeVariation interface
diff --git a/Engine/Plugins/Experimental/Mutable/Source/CustomizableObjectEditor/Public/MuCOE/Nodes/CustomizableObjectNodeMeshSwitch.h b/Engine/Plugins/Experimental/Mutable/Source/CustomizableObjectEditor/Public/MuCOE/Nodes/CustomizableObjectNodeMeshSwitch.h
--- a/Engine/Plugins/Experimental/Mutable/Source/CustomizableObjectEditor/Public/MuCOE/Nodes/CustomizableObjectNodeMeshSwitch.h
+++ b/Engine/Plugins/Experimental/Mutable/Source/CustomizableObjectEditor/Public/MuCOE/Nodes/CustomizableObjectNodeMeshSwitch.h
@@ -14,7 +14,7 @@ public:
 	GENERATED_BODY()
 
 	// UCustomizableObjectNode interface
-	virtual void BackwardsCompatibleFixup() override;
+	virtual void BackwardsCompatibleFixup(int32 CustomizableObjectCustomVersion) override;
 	
 	// UCustomizableObjectNodeSwitchBase interface
 	virtual FName GetCategory() const override;
diff --git a/Engine/Plugins/Experimental/Mutable/Source/CustomizableObjectEditor/Public/MuCOE/Nodes/CustomizableObjectNodeMeshVariation.h b/Engine/Plugins/Experimental/Mutable/Source/CustomizableObjectEditor/Public/MuCOE/Nodes/CustomizableObjectNodeMeshVariation.h
--- a/Engine/Plugins/Experimental/Mutable/Source/CustomizableObjectEditor/Public/MuCOE/Nodes/CustomizableObjectNodeMeshVariation.h
+++ b/Engine/Plugins/Experimental/Mutable/Source/CustomizableObjectEditor/Public/MuCOE/Nodes/CustomizableObjectNodeMeshVariation.h
@@ -20,7 +20,7 @@ public:
 	TArray<FCustomizableObjectMeshVariation> Variations_DEPRECATED;
 
 	// UCustomizableObjectNode interface
-	virtual void BackwardsCompatibleFixup() override;
+	virtual void BackwardsCompatibleFixup(int32 CustomizableObjectCustomVersion) override;
 	
 	// UCustomizableObjectNodeVariation interface
 	virtual FName GetCategory() const override;
diff --git a/Engine/Plugins/Experimental/Mutable/Source/CustomizableObjectEditor/Public/MuCOE/Nodes/CustomizableObjectNodeModifierBase.h b/Engine/Plugins/Experimental/Mutable/Source/CustomizableObjectEditor/Public/MuCOE/Nodes/CustomizableObjectNodeModifierBase.h
--- a/Engine/Plugins/Experimental/Mutable/Source/CustomizableObjectEditor/Public/MuCOE/Nodes/CustomizableObjectNodeModifierBase.h
+++ b/Engine/Plugins/Experimental/Mutable/Source/CustomizableObjectEditor/Public/MuCOE/Nodes/CustomizableObjectNodeModifierBase.h
@@ -19,7 +19,6 @@ public:
 	virtual FLinearColor GetNodeTitleColor() const override;
 
 	// UCustomizableObjectNode interface
-	virtual void PostBackwardsCompatibleFixup() override;
-
+	virtual void BackwardsCompatibleFixup(int32 CustomizableObjectCustomVersion) override;
 };
 
diff --git a/Engine/Plugins/Experimental/Mutable/Source/CustomizableObjectEditor/Public/MuCOE/Nodes/CustomizableObjectNodeMorphMaterial.h b/Engine/Plugins/Experimental/Mutable/Source/CustomizableObjectEditor/Public/MuCOE/Nodes/CustomizableObjectNodeMorphMaterial.h
--- a/Engine/Plugins/Experimental/Mutable/Source/CustomizableObjectEditor/Public/MuCOE/Nodes/CustomizableObjectNodeMorphMaterial.h
+++ b/Engine/Plugins/Experimental/Mutable/Source/CustomizableObjectEditor/Public/MuCOE/Nodes/CustomizableObjectNodeMorphMaterial.h
@@ -27,7 +27,7 @@ public:
 	FText GetTooltipText() const override;
 
 	// UCustomizableObjectNode interface
-	virtual void BackwardsCompatibleFixup() override;
+	virtual void BackwardsCompatibleFixup(int32 CustomizableObjectCustomVersion) override;
 	void AllocateDefaultPins(UCustomizableObjectNodeRemapPins* RemapPins) override;
 	void PinConnectionListChanged(UEdGraphPin* Pin) override;
 	void PostEditChangeProperty(FPropertyChangedEvent& PropertyChangedEvent) override;
diff --git a/Engine/Plugins/Experimental/Mutable/Source/CustomizableObjectEditor/Public/MuCOE/Nodes/CustomizableObjectNodeObject.h b/Engine/Plugins/Experimental/Mutable/Source/CustomizableObjectEditor/Public/MuCOE/Nodes/CustomizableObjectNodeObject.h
--- a/Engine/Plugins/Experimental/Mutable/Source/CustomizableObjectEditor/Public/MuCOE/Nodes/CustomizableObjectNodeObject.h
+++ b/Engine/Plugins/Experimental/Mutable/Source/CustomizableObjectEditor/Public/MuCOE/Nodes/CustomizableObjectNodeObject.h
@@ -221,7 +221,7 @@ public:
 	void PostDuplicate(bool bDuplicateForPIE) override;
 
 	// UCustomizableObjectNode interface
-	virtual void BackwardsCompatibleFixup() override;
+	virtual void BackwardsCompatibleFixup(int32 CustomizableObjectCustomVersion) override;
 	virtual void PostBackwardsCompatibleFixup() override;
 	void AllocateDefaultPins(UCustomizableObjectNodeRemapPins* RemapPins) override;
 
diff --git a/Engine/Plugins/Experimental/Mutable/Source/CustomizableObjectEditor/Public/MuCOE/Nodes/CustomizableObjectNodePassThroughTexture.h b/Engine/Plugins/Experimental/Mutable/Source/CustomizableObjectEditor/Public/MuCOE/Nodes/CustomizableObjectNodePassThroughTexture.h
--- a/Engine/Plugins/Experimental/Mutable/Source/CustomizableObjectEditor/Public/MuCOE/Nodes/CustomizableObjectNodePassThroughTexture.h
+++ b/Engine/Plugins/Experimental/Mutable/Source/CustomizableObjectEditor/Public/MuCOE/Nodes/CustomizableObjectNodePassThroughTexture.h
@@ -17,7 +17,7 @@ public:
 	TObjectPtr<UTexture> PassThroughTexture = nullptr;
 
 	// UCustomizableObjectNode interface
-	virtual void BackwardsCompatibleFixup() override;
+	virtual void BackwardsCompatibleFixup(int32 CustomizableObjectCustomVersion) override;
 
 	// UCustomizableObjectNode interface
 	void AllocateDefaultPins(UCustomizableObjectNodeRemapPins* RemapPins) override;
diff --git a/Engine/Plugins/Experimental/Mutable/Source/CustomizableObjectEditor/Public/MuCOE/Nodes/CustomizableObjectNodeProjectorConstant.h b/Engine/Plugins/Experimental/Mutable/Source/CustomizableObjectEditor/Public/MuCOE/Nodes/CustomizableObjectNodeProjectorConstant.h
--- a/Engine/Plugins/Experimental/Mutable/Source/CustomizableObjectEditor/Public/MuCOE/Nodes/CustomizableObjectNodeProjectorConstant.h
+++ b/Engine/Plugins/Experimental/Mutable/Source/CustomizableObjectEditor/Public/MuCOE/Nodes/CustomizableObjectNodeProjectorConstant.h
@@ -57,7 +57,7 @@ public:
 
 	// UCustomizableObjectNode interface
 	void AllocateDefaultPins(UCustomizableObjectNodeRemapPins* RemapPins) override;
-	virtual void BackwardsCompatibleFixup() override;
+	virtual void BackwardsCompatibleFixup(int32 CustomizableObjectCustomVersion) override;
 
 	// Own interface
 	UEdGraphPin* ValuePin();
diff --git a/Engine/Plugins/Experimental/Mutable/Source/CustomizableObjectEditor/Public/MuCOE/Nodes/CustomizableObjectNodeProjectorParameter.h b/Engine/Plugins/Experimental/Mutable/Source/CustomizableObjectEditor/Public/MuCOE/Nodes/CustomizableObjectNodeProjectorParameter.h
--- a/Engine/Plugins/Experimental/Mutable/Source/CustomizableObjectEditor/Public/MuCOE/Nodes/CustomizableObjectNodeProjectorParameter.h
+++ b/Engine/Plugins/Experimental/Mutable/Source/CustomizableObjectEditor/Public/MuCOE/Nodes/CustomizableObjectNodeProjectorParameter.h
@@ -60,7 +60,7 @@ public:
 	// UCustomizableObjectNode interface
 	void AllocateDefaultPins(UCustomizableObjectNodeRemapPins* RemapPins) override;
 	virtual bool IsAffectedByLOD() const override { return false; }
-	virtual void BackwardsCompatibleFixup() override;
+	virtual void BackwardsCompatibleFixup(int32 CustomizableObjectCustomVersion) override;
 
 	// Own interface
 	UEdGraphPin* ProjectorPin() const
diff --git a/Engine/Plugins/Experimental/Mutable/Source/CustomizableObjectEditor/Public/MuCOE/Nodes/CustomizableObjectNodeRemoveMeshBlocks.h b/Engine/Plugins/Experimental/Mutable/Source/CustomizableObjectEditor/Public/MuCOE/Nodes/CustomizableObjectNodeRemoveMeshBlocks.h
--- a/Engine/Plugins/Experimental/Mutable/Source/CustomizableObjectEditor/Public/MuCOE/Nodes/CustomizableObjectNodeRemoveMeshBlocks.h
+++ b/Engine/Plugins/Experimental/Mutable/Source/CustomizableObjectEditor/Public/MuCOE/Nodes/CustomizableObjectNodeRemoveMeshBlocks.h
@@ -33,7 +33,7 @@ public:
 	FText GetTooltipText() const override;
 
 	// UCustomizableObjectNode interface
-	virtual void BackwardsCompatibleFixup() override;
+	virtual void BackwardsCompatibleFixup(int32 CustomizableObjectCustomVersion) override;
 	void AllocateDefaultPins(UCustomizableObjectNodeRemapPins* RemapPins) override;
 	bool IsNodeOutDatedAndNeedsRefresh() override;
 	FString GetRefreshMessage() const override;
diff --git a/Engine/Plugins/Experimental/Mutable/Source/CustomizableObjectEditor/Public/MuCOE/Nodes/CustomizableObjectNodeSkeletalMesh.h b/Engine/Plugins/Experimental/Mutable/Source/CustomizableObjectEditor/Public/MuCOE/Nodes/CustomizableObjectNodeSkeletalMesh.h
--- a/Engine/Plugins/Experimental/Mutable/Source/CustomizableObjectEditor/Public/MuCOE/Nodes/CustomizableObjectNodeSkeletalMesh.h
+++ b/Engine/Plugins/Experimental/Mutable/Source/CustomizableObjectEditor/Public/MuCOE/Nodes/CustomizableObjectNodeSkeletalMesh.h
@@ -211,7 +211,7 @@ public:
 	virtual FText GetTooltipText() const override;
 
 	// UCustomizableObjectNode interface
-	virtual void BackwardsCompatibleFixup() override;
+	virtual void BackwardsCompatibleFixup(int32 CustomizableObjectCustomVersion) override;
 	virtual void AllocateDefaultPins(UCustomizableObjectNodeRemapPins* RemapPins) override;
 	virtual UCustomizableObjectNodeRemapPins* CreateRemapPinsDefault() const override;
 	virtual bool HasPinViewer() const override;
diff --git a/Engine/Plugins/Experimental/Mutable/Source/CustomizableObjectEditor/Public/MuCOE/Nodes/CustomizableObjectNodeTable.h b/Engine/Plugins/Experimental/Mutable/Source/CustomizableObjectEditor/Public/MuCOE/Nodes/CustomizableObjectNodeTable.h
--- a/Engine/Plugins/Experimental/Mutable/Source/CustomizableObjectEditor/Public/MuCOE/Nodes/CustomizableObjectNodeTable.h
+++ b/Engine/Plugins/Experimental/Mutable/Source/CustomizableObjectEditor/Public/MuCOE/Nodes/CustomizableObjectNodeTable.h
@@ -101,7 +101,7 @@ class CUSTOMIZABLEOBJECTEDITOR_API UCustomizableObjectNodeTableImagePinData : pu
 
 public:
 	// UObject interface
-	virtual void PostLoad() override;
+	virtual void BackwardsCompatibleFixup(int32 CustomizableObjectCustomVersion) override;
 	virtual void PostEditChangeProperty(FPropertyChangedEvent& PropertyChangedEvent) override;
 	virtual bool CanEditChange(const FProperty* InProperty) const override;
 
@@ -248,7 +248,7 @@ public:
 	virtual bool GetCanRenameNode() const override { return true; }
 	
 	// UCustomizableObjectNode interface
-	virtual void BackwardsCompatibleFixup() override;
+	virtual void BackwardsCompatibleFixup(int32 CustomizableObjectCustomVersion) override;
 	void AllocateDefaultPins(UCustomizableObjectNodeRemapPins* RemapPins) override;
 	bool IsNodeOutDatedAndNeedsRefresh() override;
 	FString GetRefreshMessage() const override;
diff --git a/Engine/Plugins/Experimental/Mutable/Source/CustomizableObjectEditor/Public/MuCOE/Nodes/CustomizableObjectNodeTexture.h b/Engine/Plugins/Experimental/Mutable/Source/CustomizableObjectEditor/Public/MuCOE/Nodes/CustomizableObjectNodeTexture.h
--- a/Engine/Plugins/Experimental/Mutable/Source/CustomizableObjectEditor/Public/MuCOE/Nodes/CustomizableObjectNodeTexture.h
+++ b/Engine/Plugins/Experimental/Mutable/Source/CustomizableObjectEditor/Public/MuCOE/Nodes/CustomizableObjectNodeTexture.h
@@ -83,7 +83,7 @@ public:
 
 	// UCustomizableObjectNode interface
 	virtual void AllocateDefaultPins(UCustomizableObjectNodeRemapPins* RemapPins) override;
-	virtual void BackwardsCompatibleFixup() override;
+	virtual void BackwardsCompatibleFixup(int32 CustomizableObjectCustomVersion) override;
 
 	// Begin EdGraphNode interface
 	FText GetNodeTitle(ENodeTitleType::Type TitleType) const override;
diff --git a/Engine/Plugins/Experimental/Mutable/Source/CustomizableObjectEditor/Public/MuCOE/Nodes/CustomizableObjectNodeTextureBinarise.h b/Engine/Plugins/Experimental/Mutable/Source/CustomizableObjectEditor/Public/MuCOE/Nodes/CustomizableObjectNodeTextureBinarise.h
--- a/Engine/Plugins/Experimental/Mutable/Source/CustomizableObjectEditor/Public/MuCOE/Nodes/CustomizableObjectNodeTextureBinarise.h
+++ b/Engine/Plugins/Experimental/Mutable/Source/CustomizableObjectEditor/Public/MuCOE/Nodes/CustomizableObjectNodeTextureBinarise.h
@@ -25,7 +25,7 @@ public:
 	
 	// UCustomizableObjectNode interface
 	virtual void AllocateDefaultPins(UCustomizableObjectNodeRemapPins* RemapPins) override;
-	virtual void BackwardsCompatibleFixup() override;
+	virtual void BackwardsCompatibleFixup(int32 CustomizableObjectCustomVersion) override;
 
 	UEdGraphPin* GetBaseImagePin() const;
 
diff --git a/Engine/Plugins/Experimental/Mutable/Source/CustomizableObjectEditor/Public/MuCOE/Nodes/CustomizableObjectNodeTextureColourMap.h b/Engine/Plugins/Experimental/Mutable/Source/CustomizableObjectEditor/Public/MuCOE/Nodes/CustomizableObjectNodeTextureColourMap.h
--- a/Engine/Plugins/Experimental/Mutable/Source/CustomizableObjectEditor/Public/MuCOE/Nodes/CustomizableObjectNodeTextureColourMap.h
+++ b/Engine/Plugins/Experimental/Mutable/Source/CustomizableObjectEditor/Public/MuCOE/Nodes/CustomizableObjectNodeTextureColourMap.h
@@ -26,7 +26,7 @@ public:
 
 	// UCustomizableObjectNode interface
 	virtual void AllocateDefaultPins(UCustomizableObjectNodeRemapPins* RemapPins) override;
-	virtual void BackwardsCompatibleFixup() override;
+	virtual void BackwardsCompatibleFixup(int32 CustomizableObjectCustomVersion) override;
 
 	UEdGraphPin* GetMapPin() const
 	{
diff --git a/Engine/Plugins/Experimental/Mutable/Source/CustomizableObjectEditor/Public/MuCOE/Nodes/CustomizableObjectNodeTextureFromChannels.h b/Engine/Plugins/Experimental/Mutable/Source/CustomizableObjectEditor/Public/MuCOE/Nodes/CustomizableObjectNodeTextureFromChannels.h
--- a/Engine/Plugins/Experimental/Mutable/Source/CustomizableObjectEditor/Public/MuCOE/Nodes/CustomizableObjectNodeTextureFromChannels.h
+++ b/Engine/Plugins/Experimental/Mutable/Source/CustomizableObjectEditor/Public/MuCOE/Nodes/CustomizableObjectNodeTextureFromChannels.h
@@ -26,7 +26,7 @@ public:
 
 	// UCustomizableObjectNode interface
 	virtual void AllocateDefaultPins(UCustomizableObjectNodeRemapPins* RemapPins) override;
-	virtual void BackwardsCompatibleFixup() override;
+	virtual void BackwardsCompatibleFixup(int32 CustomizableObjectCustomVersion) override;
 
 	UEdGraphPin* RPin() const
 	{
diff --git a/Engine/Plugins/Experimental/Mutable/Source/CustomizableObjectEditor/Public/MuCOE/Nodes/CustomizableObjectNodeTextureFromColor.h b/Engine/Plugins/Experimental/Mutable/Source/CustomizableObjectEditor/Public/MuCOE/Nodes/CustomizableObjectNodeTextureFromColor.h
--- a/Engine/Plugins/Experimental/Mutable/Source/CustomizableObjectEditor/Public/MuCOE/Nodes/CustomizableObjectNodeTextureFromColor.h
+++ b/Engine/Plugins/Experimental/Mutable/Source/CustomizableObjectEditor/Public/MuCOE/Nodes/CustomizableObjectNodeTextureFromColor.h
@@ -26,7 +26,7 @@ public:
 
 	// UCustomizableObjectNode interface
 	virtual void AllocateDefaultPins(UCustomizableObjectNodeRemapPins* RemapPins) override;
-	virtual void BackwardsCompatibleFixup() override;
+	virtual void BackwardsCompatibleFixup(int32 CustomizableObjectCustomVersion) override;
 
 	UEdGraphPin* ColorPin() const
 	{
diff --git a/Engine/Plugins/Experimental/Mutable/Source/CustomizableObjectEditor/Public/MuCOE/Nodes/CustomizableObjectNodeTextureInvert.h b/Engine/Plugins/Experimental/Mutable/Source/CustomizableObjectEditor/Public/MuCOE/Nodes/CustomizableObjectNodeTextureInvert.h
--- a/Engine/Plugins/Experimental/Mutable/Source/CustomizableObjectEditor/Public/MuCOE/Nodes/CustomizableObjectNodeTextureInvert.h
+++ b/Engine/Plugins/Experimental/Mutable/Source/CustomizableObjectEditor/Public/MuCOE/Nodes/CustomizableObjectNodeTextureInvert.h
@@ -24,7 +24,7 @@ public:
 	
 	// UCustomizableObjectNode interface
 	virtual void AllocateDefaultPins(UCustomizableObjectNodeRemapPins* RemapPins) override;
-	virtual void BackwardsCompatibleFixup() override;
+	virtual void BackwardsCompatibleFixup(int32 CustomizableObjectCustomVersion) override;
 
 	// Own interface
 	UEdGraphPin* GetBaseImagePin() const;
diff --git a/Engine/Plugins/Experimental/Mutable/Source/CustomizableObjectEditor/Public/MuCOE/Nodes/CustomizableObjectNodeTextureLayer.h b/Engine/Plugins/Experimental/Mutable/Source/CustomizableObjectEditor/Public/MuCOE/Nodes/CustomizableObjectNodeTextureLayer.h
--- a/Engine/Plugins/Experimental/Mutable/Source/CustomizableObjectEditor/Public/MuCOE/Nodes/CustomizableObjectNodeTextureLayer.h
+++ b/Engine/Plugins/Experimental/Mutable/Source/CustomizableObjectEditor/Public/MuCOE/Nodes/CustomizableObjectNodeTextureLayer.h
@@ -71,7 +71,7 @@ public:
 
 	// UCustomizableObjectNode interface
 	virtual void AllocateDefaultPins(UCustomizableObjectNodeRemapPins* RemapPins) override;
-	virtual void BackwardsCompatibleFixup() override;
+	virtual void BackwardsCompatibleFixup(int32 CustomizableObjectCustomVersion) override;
 	virtual bool CanConnect(const UEdGraphPin* InOwnedInputPin, const UEdGraphPin* InOutputPin, bool& bOutIsOtherNodeBlocklisted, bool& bOutArePinsCompatible) const override;
 		
 	UEdGraphPin* BasePin() const
diff --git a/Engine/Plugins/Experimental/Mutable/Source/CustomizableObjectEditor/Public/MuCOE/Nodes/CustomizableObjectNodeTextureParameter.h b/Engine/Plugins/Experimental/Mutable/Source/CustomizableObjectEditor/Public/MuCOE/Nodes/CustomizableObjectNodeTextureParameter.h
--- a/Engine/Plugins/Experimental/Mutable/Source/CustomizableObjectEditor/Public/MuCOE/Nodes/CustomizableObjectNodeTextureParameter.h
+++ b/Engine/Plugins/Experimental/Mutable/Source/CustomizableObjectEditor/Public/MuCOE/Nodes/CustomizableObjectNodeTextureParameter.h
@@ -49,7 +49,7 @@ public:
 	virtual bool GetCanRenameNode() const override { return true; }
 
 	// UCustomizableObjectNode interface
-	virtual void BackwardsCompatibleFixup() override;
+	virtual void BackwardsCompatibleFixup(int32 CustomizableObjectCustomVersion) override;
 	virtual void AllocateDefaultPins(UCustomizableObjectNodeRemapPins* RemapPins) override;
 	virtual bool IsAffectedByLOD() const override { return false; }
 };
diff --git a/Engine/Plugins/Experimental/Mutable/Source/CustomizableObjectEditor/Public/MuCOE/Nodes/CustomizableObjectNodeTextureProject.h b/Engine/Plugins/Experimental/Mutable/Source/CustomizableObjectEditor/Public/MuCOE/Nodes/CustomizableObjectNodeTextureProject.h
--- a/Engine/Plugins/Experimental/Mutable/Source/CustomizableObjectEditor/Public/MuCOE/Nodes/CustomizableObjectNodeTextureProject.h
+++ b/Engine/Plugins/Experimental/Mutable/Source/CustomizableObjectEditor/Public/MuCOE/Nodes/CustomizableObjectNodeTextureProject.h
@@ -87,7 +87,7 @@ public:
 
 	// UCustomizableObjectNode interface
 	virtual void AllocateDefaultPins(UCustomizableObjectNodeRemapPins* RemapPins) override;
-	virtual void BackwardsCompatibleFixup() override;
+	virtual void BackwardsCompatibleFixup(int32 CustomizableObjectCustomVersion) override;
 
 	UEdGraphPin* MeshPin() const
 	{
diff --git a/Engine/Plugins/Experimental/Mutable/Source/CustomizableObjectEditor/Public/MuCOE/Nodes/CustomizableObjectNodeTextureSwitch.h b/Engine/Plugins/Experimental/Mutable/Source/CustomizableObjectEditor/Public/MuCOE/Nodes/CustomizableObjectNodeTextureSwitch.h
--- a/Engine/Plugins/Experimental/Mutable/Source/CustomizableObjectEditor/Public/MuCOE/Nodes/CustomizableObjectNodeTextureSwitch.h
+++ b/Engine/Plugins/Experimental/Mutable/Source/CustomizableObjectEditor/Public/MuCOE/Nodes/CustomizableObjectNodeTextureSwitch.h
@@ -13,7 +13,7 @@ class CUSTOMIZABLEOBJECTEDITOR_API UCustomizableObjectNodeTextureSwitch : public
 public:
 	GENERATED_BODY()
 	// UCustomizableObjectNode interface
-	virtual void BackwardsCompatibleFixup() override;
+	virtual void BackwardsCompatibleFixup(int32 CustomizableObjectCustomVersion) override;
 	
 	// UCustomizableObjectNodeSwitchBase interface
 	virtual FName GetCategory() const override;
diff --git a/Engine/Plugins/Experimental/Mutable/Source/CustomizableObjectEditor/Public/MuCOE/Nodes/CustomizableObjectNodeTextureToChannels.h b/Engine/Plugins/Experimental/Mutable/Source/CustomizableObjectEditor/Public/MuCOE/Nodes/CustomizableObjectNodeTextureToChannels.h
--- a/Engine/Plugins/Experimental/Mutable/Source/CustomizableObjectEditor/Public/MuCOE/Nodes/CustomizableObjectNodeTextureToChannels.h
+++ b/Engine/Plugins/Experimental/Mutable/Source/CustomizableObjectEditor/Public/MuCOE/Nodes/CustomizableObjectNodeTextureToChannels.h
@@ -31,7 +31,7 @@ public:
 
 	// UCustomizableObjectNode interface
 	virtual void AllocateDefaultPins(UCustomizableObjectNodeRemapPins* RemapPins) override;
-	virtual void BackwardsCompatibleFixup() override;
+	virtual void BackwardsCompatibleFixup(int32 CustomizableObjectCustomVersion) override;
 
 	UEdGraphPin* InputPin() const;
 
diff --git a/Engine/Plugins/Experimental/Mutable/Source/CustomizableObjectEditor/Public/MuCOE/Nodes/CustomizableObjectNodeVariation.h b/Engine/Plugins/Experimental/Mutable/Source/CustomizableObjectEditor/Public/MuCOE/Nodes/CustomizableObjectNodeVariation.h
--- a/Engine/Plugins/Experimental/Mutable/Source/CustomizableObjectEditor/Public/MuCOE/Nodes/CustomizableObjectNodeVariation.h
+++ b/Engine/Plugins/Experimental/Mutable/Source/CustomizableObjectEditor/Public/MuCOE/Nodes/CustomizableObjectNodeVariation.h
@@ -50,7 +50,7 @@ public:
 	virtual FText GetTooltipText() const override;
 	
 	// UCustomizableObjectNode interface
-	virtual void BackwardsCompatibleFixup() override;
+	virtual void BackwardsCompatibleFixup(int32 CustomizableObjectCustomVersion) override;
 	virtual void AllocateDefaultPins(UCustomizableObjectNodeRemapPins* RemapPins) override;
 
 	// Own interface
