[Mutable] Allow using DDC storage to cache cooked platform data.
- Initial support to store and stream editor compilation data to/from the DDC.
- Move Participating objects, streamable resources, and other data derived from compilation into the ModelResources struct.

#rnx

diff --git a/Engine/Plugins/Experimental/Mutable/Source/CustomizableObject/Internal/MuCO/CustomizableObjectCompilerTypes.h b/Engine/Plugins/Experimental/Mutable/Source/CustomizableObject/Internal/MuCO/CustomizableObjectCompilerTypes.h
--- a/Engine/Plugins/Experimental/Mutable/Source/CustomizableObject/Internal/MuCO/CustomizableObjectCompilerTypes.h
+++ b/Engine/Plugins/Experimental/Mutable/Source/CustomizableObject/Internal/MuCO/CustomizableObjectCompilerTypes.h
@@ -4,6 +4,11 @@
 
 #include "UObject/WeakObjectPtr.h"
 
+#if WITH_EDITOR
+#include "DerivedDataCachePolicy.h"
+#include "DerivedDataCacheKey.h"
+#endif
+
 #include "CustomizableObjectCompilerTypes.generated.h"
 
 
@@ -110,6 +115,17 @@ struct FCompilationOptions
 
 	/** If true, gather all game asset references and save them in the Customizable Object. */
 	bool bGatherReferences = false;
+
+	/** Whether or not the compiler should query a request to load the compiled data from the DDC */
+	bool bQueryCompiledDatafromDDC = false;
+
+	/** Whether or not the compiler should store the compiled data to the DDC*/
+	bool bStoreCompiledDataInDDC = false;
+
+	/** High limit of the size in bytes of the packaged data when storing this object in the DDC.
+	 * This limit is before any pak or filesystem compression. This limit will be broken if a single piece of data is bigger because data is not fragmented for packaging purposes.
+	 */
+	uint64 DDCBytesLimit = 64 * 1024;
 };
 
 
@@ -141,6 +157,12 @@ struct CUSTOMIZABLEOBJECT_API FCompilationRequest
 
 	bool IsAsyncCompilation() const;
 
+	void SetDerivedDataCachePolicy(UE::DerivedData::ECachePolicy InCachePolicy);
+	UE::DerivedData::ECachePolicy GetDerivedDataCachePolicy() const;
+
+	void BuildDerivedDataCacheKey();
+	UE::DerivedData::FCacheKey GetDerivedDataCacheKey() const;
+
 	void SetCompilationState(ECompilationStatePrivate InState, ECompilationResultPrivate InResult);
 
 	ECompilationStatePrivate GetCompilationState() const;
@@ -165,6 +187,9 @@ private:
 
 	bool bAsync = true;
 
+	UE::DerivedData::ECachePolicy DDCPolicy = UE::DerivedData::ECachePolicy::None;
+	UE::DerivedData::FCacheKey DDCKey;
+
 	// Stores the only option of an Int Param that should be compiled
 	TMap<FString, FString> ParamNamesToSelectedOptions;
 
diff --git a/Engine/Plugins/Experimental/Mutable/Source/CustomizableObject/Internal/MuCO/CustomizableObjectPrivate.h b/Engine/Plugins/Experimental/Mutable/Source/CustomizableObject/Internal/MuCO/CustomizableObjectPrivate.h
--- a/Engine/Plugins/Experimental/Mutable/Source/CustomizableObject/Internal/MuCO/CustomizableObjectPrivate.h
+++ b/Engine/Plugins/Experimental/Mutable/Source/CustomizableObject/Internal/MuCO/CustomizableObjectPrivate.h
@@ -4,6 +4,7 @@
 
 #include "MuCO/CustomizableObject.h"
 #include "MuCO/CustomizableObjectCompilerTypes.h"
+#include "MuCO/CustomizableObjectStreamedResourceData.h"
 #include "MuCO/StateMachine.h"
 #include "MuCO/CustomizableObjectUIData.h"
 #include "MuCO/CustomizableObjectIdentifier.h"
@@ -26,6 +27,15 @@ namespace mu
 	struct FBoneName;
 }
 
+#if WITH_EDITOR
+namespace UE::DerivedData
+{
+	struct FValueId;
+	struct FCacheKey;
+	enum class ECachePolicy : uint32;
+}
+#endif
+
 class USkeletalMesh;
 class USkeleton;
 class UPhysicsAsset;
@@ -33,7 +43,6 @@ class UMaterialInterface;
 class UTexture;
 class UAnimInstance;
 class UAssetUserData;
-class UCustomizableObject;
 class USkeletalMeshLODSettings;
 struct FModelResources;
 struct FModelStreamableBulkData;
@@ -433,7 +442,7 @@ struct CUSTOMIZABLEOBJECT_API FMutableRefSkeletalMeshData
 #if WITH_EDITORONLY_DATA
 	friend FArchive& operator<<(FArchive& Ar, FMutableRefSkeletalMeshData& Data);
 
-	void InitResources(UCustomizableObject* InOuter, const ITargetPlatform* InTargetPlatform);
+	void InitResources(UCustomizableObject* InOuter, FModelResources& InModelResources, const ITargetPlatform* InTargetPlatform);
 #endif
 
 };
@@ -490,6 +499,34 @@ struct CUSTOMIZABLEOBJECT_API FMutableSkinWeightProfileInfo
 #endif
 };
 
+
+USTRUCT()
+struct CUSTOMIZABLEOBJECT_API FMutableStreamableBlock
+{
+	GENERATED_USTRUCT_BODY()
+
+	UPROPERTY()
+	uint32 FileId = 0;
+
+	/** Used to store properties of the data, necessary for its recovery. For instance if it is high-res. */
+	UPROPERTY()
+	uint32 Flags = 0;
+
+	UPROPERTY()
+	uint64 Offset = 0;
+
+	friend FArchive& operator<<(FArchive& Ar, FMutableStreamableBlock& Data)
+	{
+		Ar << Data.FileId;
+		Ar << Data.Flags;
+		Ar << Data.Offset;
+		return Ar;
+	}
+};
+template<> struct TCanBulkSerialize<FMutableStreamableBlock> { enum { Value = true }; };
+static_assert(sizeof(FMutableStreamableBlock) == 8 * 2);
+
+
 USTRUCT()
 struct FRealTimeMorphStreamable
 {
@@ -805,10 +842,16 @@ namespace MutablePrivate
 
 	struct CUSTOMIZABLEOBJECT_API FMutableCachedPlatformData
 	{
-		/** */
+		/** Serialized mu::Model */
 		TArray64<uint8> ModelData;
 
-		/** */
+		/** Serialized FModelResources */
+		TArray64<uint8> ModelResourcesData;
+
+		/** Streamable resources info such as files and offsets. */
+		TSharedPtr<FModelStreamableBulkData> ModelStreamables;
+
+		/** Struct containing map of RomId to RomBytes. */
 		FModelStreamableData ModelStreamableData;
 
 		/** */
@@ -827,7 +870,7 @@ namespace MutablePrivate
 	 */
 	void CUSTOMIZABLEOBJECT_API GenerateBulkDataFilesList(
 		TSharedPtr<const mu::Model, ESPMode::ThreadSafe> Model,
-		TSharedPtr<FModelStreamableBulkData> StreamableBulkData,
+		FModelStreamableBulkData& StreamableBulkData,
 		const ITargetPlatform* TargetPlatform,
 		uint64 TargetBulkDataFileBytes,
 		TArray<FFile>& OutBulkDataFiles);
@@ -838,6 +881,11 @@ namespace MutablePrivate
 		TFunctionRef<void(FFile&, TArray64<uint8>&)> WriteFile,
 		bool bDropData
 	);
+
+	UE::DerivedData::FValueId CUSTOMIZABLEOBJECT_API GetDerivedDataModelId();
+	UE::DerivedData::FValueId CUSTOMIZABLEOBJECT_API GetDerivedDataModelResourcesId();
+	UE::DerivedData::FValueId CUSTOMIZABLEOBJECT_API GetDerivedDataModelStreamableBulkDataId();
+	UE::DerivedData::FValueId CUSTOMIZABLEOBJECT_API GetDerivedDataBulkDataFilesId();
 }
 #endif
 
@@ -856,7 +904,16 @@ struct CUSTOMIZABLEOBJECT_API FModelStreamableBulkData
 	void Serialize(FArchive& Ar, UObject* Owner, bool bCooked);
 
 #if WITH_EDITORONLY_DATA
-	friend FArchive& operator<<(FArchive& Ar, FModelStreamableBulkData& Struct);
+	friend FArchive& operator<<(FArchive& Ar, FModelStreamableBulkData& Data)
+	{
+		Ar << Data.ModelStreamables;
+		Ar << Data.ClothingStreamables;
+		Ar << Data.RealTimeMorphStreamables;
+
+		// Don't serialize FByteBulkData manually, the data will be skipped.
+
+		return Ar;
+	}
 #endif
 };
 
@@ -885,7 +942,9 @@ struct FMutableParamNameSet
 };
 
 
-// Referenced materials, skeletons, passthrough textures...
+/** Struct containing all UE resources derived from a CO compilation. These resources will be embedded in the CO at cook time but not in the editor.
+  * Editor compilations will serialize this struct to disk using the Serialize methods. Ensure new fields are serialized, too.
+  * Variables and settings that should not change until the CO is re-compiled should be stored here. */
 USTRUCT()
 struct FModelResources
 {
@@ -993,8 +1052,35 @@ struct FModelResources
 	// Stores what param names use a certain table as a table can be used from multiple table nodes, useful for partial compilations to restrict params
 	UPROPERTY()
 	TMap<FString, FMutableParamNameSet> TableToParamNames;
+
+	/** Map to identify what CustomizableObject owns a parameter. Used to display a tooltip when hovering a parameter
+	 * in the Prev. instance panel */
+	TMap<FString, FString> CustomizableObjectPathMap;
+
+	TMap<FString, FCustomizableObjectIdPair> GroupNodeMap;
+
+	/** If the object is compiled with maximum optimizations. */
+	bool bIsCompiledWithOptimization = true;
+
+	/** This is a non-user-controlled flag to disable streaming (set at object compilation time, depending on optimization). */
+	bool bIsTextureStreamingDisabled = false;
+
+	/** List of external packages that if changed, a compilation is required.
+	  * Key is the package name. Value is the the UPackage::Guid, which is regenerated each time the packages is saved.
+	  *
+	  * Updated each time the CO is compiled and saved in the Derived Data. */
+	TMap<FName, FGuid> ParticipatingObjects;
+
+	// Used to know if roms and other resources must be streamed from the DDC.
+	bool bIsStoredInDDC = false;
+	UE::DerivedData::FCacheKey DDCKey = UE::DerivedData::FCacheKey::Empty;
+	UE::DerivedData::ECachePolicy DDCDefaultPolicy = UE::DerivedData::ECachePolicy::Default;
 #endif
 
+	// Constant Resources streamed in on demand when generating meshes
+	UPROPERTY()
+	TArray<FCustomizableObjectStreamedResourceData> StreamedResourceData;
+
 	/** Max number of components in the compiled Model. */
 	UPROPERTY()
 	uint8 NumComponents = 0; 
@@ -1010,6 +1096,42 @@ struct FModelResources
 	/** First LOD available, some platforms may remove lower LODs when cooking, this MinLOD represents the first LOD we can generate */
 	UPROPERTY()
 	uint8 FirstLODAvailable = 0;
+
+#if WITH_EDITORONLY_DATA
+	void CUSTOMIZABLEOBJECT_API Serialize(FObjectAndNameAsStringProxyArchive& Ar, bool bIsCooking);
+	bool CUSTOMIZABLEOBJECT_API Unserialize(FObjectAndNameAsStringProxyArchive& Ar, UCustomizableObject& Outer, const ITargetPlatform* InTargetPlatform, bool bIsCooking);
+#endif
+};
+
+
+UCLASS(config = Engine)
+class CUSTOMIZABLEOBJECT_API UCustomizableObjectBulk : public UObject
+{
+public:
+	GENERATED_BODY()
+
+	//~ Begin UObject Interface
+	virtual void PostLoad() override;
+	//~ End UObject Interface
+
+	/**  */
+	const FString& GetBulkFilePrefix() const { return BulkFilePrefix; }
+
+	TUniquePtr<IAsyncReadFileHandle> OpenFileAsyncRead(uint32 FileId, uint32 Flags) const;
+
+#if WITH_EDITOR
+
+	//~ Begin UObject Interface
+	virtual void CookAdditionalFilesOverride(const TCHAR*, const ITargetPlatform*, TFunctionRef<void(const TCHAR*, void*, int64)>) override;
+	//~ End UObject Interface
+#endif
+
+#if WITH_EDITOR
+private:
+#endif
+
+	/** Prefix to locate bulkfiles for loading, using the file ids in each FMutableStreamableBlock. */
+	FString BulkFilePrefix;
 };
 
 
@@ -1109,6 +1231,14 @@ public:
 #if WITH_EDITOR
 	/** Compose file name. */
 	FString GetCompiledDataFileName(bool bIsModel, const ITargetPlatform* InTargetPlatform = nullptr, bool bIsDiskStreamer = false);
+
+	/** DDC helpers. BuildDerivedDataKey is expensive, try to cache it as much as possible. */
+	FString BuildDerivedDataKey(FCompilationOptions Options);
+	UE::DerivedData::FCacheKey GetDerivedDataCacheKeyForOptions(FCompilationOptions InOptions);
+
+	/** Attempts to load the compiled data from DDC. Builds key if not supplied. */
+	void LoadCompiledDataFromDDC(FCompilationOptions Options, UE::DerivedData::ECachePolicy DefaultPolicy, UE::DerivedData::FCacheKey* DDCKey);
+
 #endif
 	
 	/** Rebuild ParameterProperties from the current compiled model. */
@@ -1122,13 +1252,10 @@ public:
 	
 	FGuid GetVersionId() const;
 
-	// Unless we are packaging there is no need for keeping all the data generated during compilation, this information is stored in the derived data.
-	void ClearCompiledData(bool bIsCooking);
-
 	void SaveEmbeddedData(FArchive& Ar);
 
 	// Compile the object for a specific platform - Compile for Cook Customizable Object
-	void CompileForTargetPlatform(const ITargetPlatform* TargetPlatform);
+	void CompileForTargetPlatform(UCustomizableObject& CustomizableObject, const ITargetPlatform& TargetPlatform);
 	
 	// Add a profile that stores the values of the parameters used by the CustomInstance.
 	FReply AddNewParameterProfile(FString Name, class UCustomizableObjectInstance& CustomInstance);
@@ -1136,9 +1263,10 @@ public:
 	// Compose folder name where the data is stored
 	static FString GetCompiledDataFolderPath();
 
-	/** Generic Save/Load methods to write/read compiled data */
-	void SaveCompiledData(FObjectAndNameAsStringProxyArchive& Ar, bool bSkipEditorOnlyData = false);
-	void LoadCompiledData(FObjectAndNameAsStringProxyArchive& Ar, const ITargetPlatform* InTargetPlatform, bool bSkipEditorOnlyData = false);
+	/** Generic Load methods to read compiled data */
+	bool LoadModelResources(FArchive& Ar, const ITargetPlatform* InTargetPlatform, bool bSkipEditorOnlyData = false);
+	void LoadModelStreamableBulk(FArchive& Ar, bool bIsCooking);
+	void LoadModel(FArchive& Ar);
 
 	/** Load compiled data for the running platform from disk, this is used to load Editor Compilations. */
 	void LoadCompiledDataFromDisk();
@@ -1184,7 +1312,8 @@ public:
 
 	TArray<FCustomizableObjectStreamedResourceData>& GetStreamedExtensionData();
 	
-	TArray<FCustomizableObjectStreamedResourceData>& GetStreamedResourceData();
+	const FCustomizableObjectResourceData* LoadStreamedResource(int32 ResourceIndex);
+	void UnloadStreamedResource(int32 ResourceIndex);
 
 #if WITH_EDITORONLY_DATA
 	TObjectPtr<UEdGraph>& GetSource() const;
@@ -1213,38 +1342,17 @@ public:
 	/** Cook requests. */
 	TArray<TSharedRef<FCompilationRequest>> CompileRequests;
 
-	/** List of external packages that if changed, a compilation is required.
-	 * Key is the package name. Value is the the UPackage::Guid, which is regenerated each time the packages is saved.
-	 *
-	 * Updated each time the CO is compiled and saved in the Derived Data. */
-	TMap<FName, FGuid> ParticipatingObjects;
-
 	/** List of Participating Objects (packages) has been marked as dirty since the last compilation. */
 	TArray<FName> DirtyParticipatingObjects;
-	
-	/** Map to identify what CustomizableObject owns a parameter. Used to display a tooltip when hovering a parameter
-	 * in the Prev. instance panel */
-	UPROPERTY(Transient)
-	TMap<FString, FString> CustomizableObjectPathMap;
 
-	UPROPERTY(Transient)
-	TMap<FString, FCustomizableObjectIdPair> GroupNodeMap;
-
-	/** If the object is compiled, this flag is false unless it was compiled with maximum optimizations. If the object is not compiled, its value is meaningless. */
-	bool bIsCompiledWithoutOptimization = true;
-
-	/** This is a non-user-controlled flag to disable streaming (set at object compilation time, depending on optimization). */
-	bool bDisableTextureStreaming = false;
-	
 	ECompilationStatePrivate CompilationState = ECompilationStatePrivate::None;
 	ECompilationResultPrivate CompilationResult = ECompilationResultPrivate::Unknown;
 	
 	FPostCompileDelegate PostCompileDelegate;
 
-#if WITH_EDITOR
 	/** Map of PlatformName to CachedPlatformData. Only valid while cooking. */
 	TMap<FString, MutablePrivate::FMutableCachedPlatformData> CachedPlatformsData;
-#endif
+
 #endif
 
 	FCustomizableObjectStatus Status;
@@ -1351,6 +1459,8 @@ public:
 
 		AddMaterialSlotNameIndexToSurfaceMetadata,
 		
+		DerivedDataCache,
+
 		// -----<new versions can be added above this line>--------
 		LastCustomizableObjectVersion
 	};
diff --git a/Engine/Plugins/Experimental/Mutable/Source/CustomizableObject/Internal/MuCO/UnrealMutableModelDiskStreamer.h b/Engine/Plugins/Experimental/Mutable/Source/CustomizableObject/Internal/MuCO/UnrealMutableModelDiskStreamer.h
--- a/Engine/Plugins/Experimental/Mutable/Source/CustomizableObject/Internal/MuCO/UnrealMutableModelDiskStreamer.h
+++ b/Engine/Plugins/Experimental/Mutable/Source/CustomizableObject/Internal/MuCO/UnrealMutableModelDiskStreamer.h
@@ -5,6 +5,10 @@
 #include "MuCO/CustomizableObjectPrivate.h"
 #include "MuR/Serialisation.h"
 
+#if WITH_EDITOR
+#include "DerivedDataRequestOwner.h"
+#endif
+
 class FArchive;
 class IAsyncReadFileHandle;
 class IAsyncReadRequest;
@@ -64,21 +68,32 @@ protected:
 
 	struct FReadRequest
 	{
-		TSharedPtr<IBulkDataIORequest> BulkReadRequest = nullptr;
-		TSharedPtr<IAsyncReadRequest> FileReadRequest = nullptr;
+		TSharedPtr<IBulkDataIORequest> BulkReadRequest;
+		TSharedPtr<IAsyncReadRequest> FileReadRequest;
 		TSharedPtr<FAsyncFileCallBack> FileCallback;
+
+#if WITH_EDITORONLY_DATA
+		TSharedPtr<UE::DerivedData::FRequestOwner> DDCReadRequest;
+#endif
 	};
 	
 	/** Streaming data for one object. */
 	struct FObjectData
 	{
 		TWeakPtr<const mu::Model> Model;
-		TMap<OPERATION_ID, FReadRequest> CurrentReadRequests;
-		FString BulkFilePrefix;
 
+		FString BulkFilePrefix;
+		TMap<OPERATION_ID, FReadRequest> CurrentReadRequests;
 		TMap<uint32, TSharedPtr<IAsyncReadFileHandle>> ReadFileHandles;
 
 		TSharedPtr<FModelStreamableBulkData> ModelStreamableBulkData;
+
+#if WITH_EDITORONLY_DATA
+		// DDC files streaming
+		bool bIsStoredInDDC = false;
+		UE::DerivedData::FCacheKey DDCKey;
+		UE::DerivedData::FCacheRecordPolicy DDCPolicy;
+#endif
 	};
 
 	TArray<FObjectData> Objects;
diff --git a/Engine/Plugins/Experimental/Mutable/Source/CustomizableObject/Private/MuCO/CustomizableObject.cpp b/Engine/Plugins/Experimental/Mutable/Source/CustomizableObject/Private/MuCO/CustomizableObject.cpp
--- a/Engine/Plugins/Experimental/Mutable/Source/CustomizableObject/Private/MuCO/CustomizableObject.cpp
+++ b/Engine/Plugins/Experimental/Mutable/Source/CustomizableObject/Private/MuCO/CustomizableObject.cpp
@@ -42,6 +42,11 @@
 
 #if WITH_EDITOR
 #include "Editor.h"
+
+#include "DerivedDataCache.h"
+#include "DerivedDataCacheInterface.h"
+#include "DerivedDataCacheKey.h"
+#include "DerivedDataRequestOwner.h"
 #endif
 
 
@@ -70,6 +75,16 @@ TAutoConsoleVariable<bool> CVarMutableUseBulkData(
 	TEXT("True - Use FBulkData to store streamable data.\n")
 	TEXT("False - Use Mut files to store streamable data\n"));
 
+
+TAutoConsoleVariable<int32> CVarMutableDerivedDataCacheUsage(
+	TEXT("mutable.DerivedDataCacheUsage"),
+	0,
+	TEXT("Derived data cache access for cooked data.")
+	TEXT("0 - None. Disables access to the cache.")
+	TEXT("1 - Local. Allow cache requests to query and store records and values in local caches.")
+	TEXT("2 - Default. Allow cache requests to query and store records and values in any caches."),
+	ECVF_Default);
+
 #endif
 
 #if WITH_EDITORONLY_DATA
@@ -175,22 +190,17 @@ void UCustomizableObject::PreSave(FObjectPreSaveContext ObjectSaveContext)
 		// Load cached data before saving
 		if (GetPrivate()->TryLoadCompiledCookDataForPlatform(TargetPlatform))
 		{
-			MutablePrivate::FMutableCachedPlatformData& CachedPlatformData = *GetPrivate()->CachedPlatformsData.Find(TargetPlatform->PlatformName());
-
-			// Generate list of files and update streamable blocks ids and offsets
-			FCompilationOptions Options = GetPrivate()->GetCompileOptions();
-			TSharedPtr<FModelStreamableBulkData> ModelStreamableBulkData = GetPrivate()->GetModelStreamableBulkData(true);
-
-			MutablePrivate::GenerateBulkDataFilesList(GetPrivate()->GetModel(), ModelStreamableBulkData, TargetPlatform, Options.PackagedDataBytesLimit, CachedPlatformData.BulkDataFiles);
-
 			const bool bUseBulkData = CVarMutableUseBulkData.GetValueOnAnyThread();
 			if (bUseBulkData)
 			{
+				MutablePrivate::FMutableCachedPlatformData& CachedPlatformData = *GetPrivate()->CachedPlatformsData.Find(TargetPlatform->PlatformName());
+				TSharedPtr<FModelStreamableBulkData> ModelStreamableBulkData = GetPrivate()->GetModelStreamableBulkData(true);
+
 				const int32 NumBulkDataFiles = CachedPlatformData.BulkDataFiles.Num();
 				ModelStreamableBulkData->HashToBulkData.Empty(NumBulkDataFiles);
 				ModelStreamableBulkData->HashToBulkData.Reserve(NumBulkDataFiles);
 
-				const auto WriteBulkData = [&ModelStreamableBulkData](MutablePrivate::FFile& File, TArray64<uint8>& FileBulkData)
+				const auto WriteBulkData = [ModelStreamableBulkData](MutablePrivate::FFile& File, TArray64<uint8>& FileBulkData)
 					{
 						MUTABLE_CPUPROFILER_SCOPE(WriteBulkData);
 
@@ -227,7 +237,10 @@ void UCustomizableObject::PreSave(FObjectPreSaveContext ObjectSaveContext)
 			UE_LOG(LogMutable, Warning, TEXT("Cook: Customizable Object [%s] is missing [%s] platform data."), *GetName(),
 				*ObjectSaveContext.GetTargetPlatform()->PlatformName());
 			
-			GetPrivate()->ClearCompiledData(true);
+			// Clear model resources
+			GetPrivate()->SetModel(nullptr, FGuid());
+			GetPrivate()->GetModelResources(true /* bIsCooking */) = FModelResources();
+			GetPrivate()->GetModelStreamableBulkData(true).Reset();
 		}
 	}
 }
@@ -282,10 +295,17 @@ bool UCustomizableObjectPrivate::TryLoadCompiledCookDataForPlatform(const ITarge
 		return false;
 	}
 
-	FMemoryReaderView MemoryReader(PlatformData->ModelData);
-	FObjectAndNameAsStringProxyArchive ObjectReader(MemoryReader, true);
-	LoadCompiledData(ObjectReader, TargetPlatform, true);
-	return true;
+	FMemoryReaderView ModelResourcesReader(PlatformData->ModelResourcesData);
+	if (LoadModelResources(ModelResourcesReader, TargetPlatform, true))
+	{
+		SetModelStreamableBulkData(PlatformData->ModelStreamables, true);
+
+		FMemoryReaderView ModelReader(PlatformData->ModelData);
+		LoadModel(ModelReader);
+		return GetModel() != nullptr;
+	}
+
+	return false;
 }
 
 #endif // End WITH_EDITOR
@@ -432,7 +452,7 @@ void UCustomizableObject::BeginCacheForCookedPlatformData(const ITargetPlatform*
 	}
 
 	// Compile and save in the CachedPlatformsData map
-	GetPrivate()->CompileForTargetPlatform(TargetPlatform);
+	GetPrivate()->CompileForTargetPlatform(*this, *TargetPlatform);
 }
 
 
@@ -465,432 +485,424 @@ FGuid GenerateIdentifier(const UCustomizableObject& CustomizableObject)
 }
 
 
-void UCustomizableObjectPrivate::ClearCompiledData(bool bIsCooking)
+bool UCustomizableObjectPrivate::LoadModelResources(FArchive& MemoryReader, const ITargetPlatform* InTargetPlatform, bool bIsCooking)
 {
-	GetModelResources(bIsCooking) = FModelResources();
+	// Make sure mutable has been initialised.
+	UCustomizableObjectSystem::GetInstance();
 
-#if WITH_EDITORONLY_DATA
-	CustomizableObjectPathMap.Empty();
-	GroupNodeMap.Empty();
-	ParticipatingObjects.Empty();
-#endif
+	FModelResources LocalModelResources;
+	
+	FObjectAndNameAsStringProxyArchive ObjectReader(MemoryReader, true);
+	const bool bLoadedSuccessfully = LocalModelResources.Unserialize(ObjectReader, *GetPublic(), InTargetPlatform, bIsCooking);
+
+	GetModelResources(bIsCooking) = MoveTemp(LocalModelResources);
 
-	GetPublic()->BulkData = nullptr;
+	return bLoadedSuccessfully;
 }
 
 
-void SerializeStreamedResources(FArchive& Ar, UObject* Object, TArray<FCustomizableObjectStreamedResourceData>& StreamedResources, bool bIsCooking)
+void UCustomizableObjectPrivate::LoadModelStreamableBulk(FArchive& MemoryReader, bool bIsCooking)
 {
-	if (Ar.IsSaving())
+	TSharedPtr<FModelStreamableBulkData> LocalModelStreamablesPtr = MakeShared<FModelStreamableBulkData>();
+	FModelStreamableBulkData& LocalModelStreamables = *LocalModelStreamablesPtr.Get();
+	MemoryReader << LocalModelStreamables;
+
+	SetModelStreamableBulkData(LocalModelStreamablesPtr, bIsCooking);
+}
+
+
+void UCustomizableObjectPrivate::LoadModel(FArchive& MemoryReader)
+{
+	TSharedPtr<mu::Model, ESPMode::ThreadSafe> LoadedModel;
+
+	UnrealMutableInputStream Stream(MemoryReader);
+	mu::InputArchive Arch(&Stream);
+	LoadedModel = mu::Model::StaticUnserialise(Arch);
+
+	SetModel(LoadedModel, GenerateIdentifier(*GetPublic()));
+}
+
+
+void SerializeStreamedResources(FArchive& Ar, TArray<FCustomizableObjectStreamedResourceData>& StreamedResources)
+{
+	check(Ar.IsSaving());
+
+	int32 NumStreamedResources = StreamedResources.Num();
+	Ar << NumStreamedResources;
+
+	for (const FCustomizableObjectStreamedResourceData& ResourceData : StreamedResources)
 	{
-		int32 NumStreamedResources = StreamedResources.Num();
-		Ar << NumStreamedResources;
+		const FCustomizableObjectResourceData& Data = ResourceData.GetLoadedData();
+		uint32 ResourceDataType = (uint32)Data.Type;
+		Ar << ResourceDataType;
 
-		for (const FCustomizableObjectStreamedResourceData& ResourceData : StreamedResources)
+		switch (Data.Type)
 		{
-			const FCustomizableObjectResourceData& Data = ResourceData.GetLoadedData();
-			uint32 ResourceDataType = (uint32)Data.Type;
-			Ar << ResourceDataType;
+		case ECOResourceDataType::AssetUserData:
+		{
+			const FCustomizableObjectAssetUserData* AssetUserData = Data.Data.GetPtr<FCustomizableObjectAssetUserData>();
+			FString AssetUserDataPath;
 
-			switch (Data.Type)
-			{
-			case ECOResourceDataType::AssetUserData:
+			if (AssetUserData->AssetUserDataEditor)
 			{
-				const FCustomizableObjectAssetUserData* AssetUserData = Data.Data.GetPtr<FCustomizableObjectAssetUserData>();
-				FString AssetUserDataPath;
-
-				if (AssetUserData->AssetUserDataEditor)
-				{
-					AssetUserDataPath = TSoftObjectPtr<UAssetUserData>(AssetUserData->AssetUserDataEditor).ToString();
-				}
-
-				Ar << AssetUserDataPath;
-				break;
-			}
-			default:
-				check(false);
-				break;
+				AssetUserDataPath = TSoftObjectPtr<UAssetUserData>(AssetUserData->AssetUserDataEditor).ToString();
 			}
+
+			Ar << AssetUserDataPath;
+			break;
+		}
+		default:
+			check(false);
+			break;
 		}
 	}
-	else 
-	{
-		const FString CustomizableObjectName = GetNameSafe(Object) + TEXT("_");
+}
 
-		int32 NumStreamedResources = 0;
-		Ar << NumStreamedResources;
 
-		StreamedResources.SetNum(NumStreamedResources);
+void UnserializeStreamedResources(FArchive& Ar, UObject* Object, TArray<FCustomizableObjectStreamedResourceData>& StreamedResources, bool bIsCooking)
+{
+	check(Ar.IsLoading());
+
+	const FString CustomizableObjectName = GetNameSafe(Object) + TEXT("_");
 
-		for (int32 ResourceIndex = 0; ResourceIndex < NumStreamedResources; ++ResourceIndex)
+	int32 NumStreamedResources = 0;
+	Ar << NumStreamedResources;
+
+	StreamedResources.SetNum(NumStreamedResources);
+
+	for (int32 ResourceIndex = 0; ResourceIndex < NumStreamedResources; ++ResourceIndex)
+	{
+		// Override existing containers
+		UCustomizableObjectResourceDataContainer* Container = StreamedResources[ResourceIndex].GetPath().Get();
+
+		// Create a new container if none.
+		if (!Container)
 		{
-			// Override existing containers
-			UCustomizableObjectResourceDataContainer* Container = StreamedResources[ResourceIndex].GetPath().Get();
+			// Generate a deterministic name to help with deterministic cooking
+			const FString ContainerName = CustomizableObjectName + FString::Printf(TEXT("SR_%d"), ResourceIndex);
 
-			// Create a new container if none.
-			if (!Container)
-			{
-				// Generate a deterministic name to help with deterministic cooking
-				const FString ContainerName = CustomizableObjectName + FString::Printf(TEXT("SR_%d"), ResourceIndex);
+			UCustomizableObjectResourceDataContainer* ExistingContainer = FindObject<UCustomizableObjectResourceDataContainer>(Object, *ContainerName);
+			Container = ExistingContainer ? ExistingContainer : NewObject<UCustomizableObjectResourceDataContainer>(
+				Object,
+				FName(*ContainerName),
+				RF_Public);
 
-				UCustomizableObjectResourceDataContainer* ExistingContainer = FindObject<UCustomizableObjectResourceDataContainer>(Object, *ContainerName);
-				Container = ExistingContainer ? ExistingContainer : NewObject<UCustomizableObjectResourceDataContainer>(
-					Object,
-					FName(*ContainerName),
-					RF_Public);
+			StreamedResources[ResourceIndex] = { Container };
+		}
 
-				StreamedResources[ResourceIndex] = { Container };
-			}
+		check(Container);
+		uint32 Type = 0;
+		Ar << Type;
 
-			check(Container);
-			uint32 Type = 0;
-			Ar << Type;
-			
-			Container->Data.Type = (ECOResourceDataType)Type;
-			switch (Container->Data.Type)
-			{
-				case ECOResourceDataType::AssetUserData:
-				{
-					FString AssetUserDataPath;
-					Ar << AssetUserDataPath;
-					
-					FCustomizableObjectAssetUserData ResourceData;
+		Container->Data.Type = (ECOResourceDataType)Type;
+		switch (Container->Data.Type)
+		{
+		case ECOResourceDataType::AssetUserData:
+		{
+			FString AssetUserDataPath;
+			Ar << AssetUserDataPath;
 
-					TSoftObjectPtr<UAssetUserData> SoftAssetUserData = TSoftObjectPtr<UAssetUserData>(FSoftObjectPath(AssetUserDataPath));
-					ResourceData.AssetUserDataEditor = !SoftAssetUserData.IsNull() ? SoftAssetUserData.LoadSynchronous() : nullptr;
+			FCustomizableObjectAssetUserData ResourceData;
 
-					if (!ResourceData.AssetUserDataEditor)
-					{
-						UE_LOG(LogMutable, Warning, TEXT("Failed to load streamed resource of type AssetUserData. Resource name: [%s]"), *AssetUserDataPath);
-					}
+			TSoftObjectPtr<UAssetUserData> SoftAssetUserData = TSoftObjectPtr<UAssetUserData>(FSoftObjectPath(AssetUserDataPath));
+			ResourceData.AssetUserDataEditor = !SoftAssetUserData.IsNull() ? SoftAssetUserData.LoadSynchronous() : nullptr;
 
-					if (bIsCooking)
-					{
-						// Rename the asset user data for duplicate
-						const FString AssetName = CustomizableObjectName + GetNameSafe(ResourceData.AssetUserDataEditor);
-						
-						// Find or duplicate the AUD replacing the outer
-						ResourceData.AssetUserData = FindObject<UAssetUserData>(Container, *AssetName);
-						if (!ResourceData.AssetUserData)
-						{
-							// AUD may be private objects within meshes. Duplicate changing the outer to avoid including meshes into the builds.
-							ResourceData.AssetUserData = DuplicateObject<UAssetUserData>(ResourceData.AssetUserDataEditor, Container, FName(*AssetName));
-						}
-					}
+			if (!ResourceData.AssetUserDataEditor)
+			{
+				UE_LOG(LogMutable, Warning, TEXT("Failed to load streamed resource of type AssetUserData. Resource name: [%s]"), *AssetUserDataPath);
+			}
 
-					Container->Data.Data = FInstancedStruct::Make(ResourceData);
-					break;
+			if (bIsCooking)
+			{
+				// Rename the asset user data for duplicate
+				const FString AssetName = CustomizableObjectName + GetNameSafe(ResourceData.AssetUserDataEditor);
+
+				// Find or duplicate the AUD replacing the outer
+				ResourceData.AssetUserData = FindObject<UAssetUserData>(Container, *AssetName);
+				if (!ResourceData.AssetUserData)
+				{
+					// AUD may be private objects within meshes. Duplicate changing the outer to avoid including meshes into the builds.
+					ResourceData.AssetUserData = DuplicateObject<UAssetUserData>(ResourceData.AssetUserDataEditor, Container, FName(*AssetName));
 				}
-				default:
-					check(false);
-					break;
 			}
+
+			Container->Data.Data = FInstancedStruct::Make(ResourceData);
+			break;
+		}
+		default:
+			check(false);
+			break;
 		}
 	}
 }
 
 
-void UCustomizableObjectPrivate::SaveCompiledData(FObjectAndNameAsStringProxyArchive& MemoryWriter, bool bIsCooking)
+void FModelResources::Serialize(FObjectAndNameAsStringProxyArchive& MemoryWriter, bool bIsCooking)
 {
-	int32 InternalVersion = UCustomizableObjectPrivate::CurrentSupportedVersion;
-	MutableCompiledDataStreamHeader Header(InternalVersion, GetVersionId());
-	MemoryWriter << Header;
+	MUTABLE_CPUPROFILER_SCOPE(FModelResources::Serialize);
+	check(IsInGameThread());
 
-	FModelResources& LocalModelResources = GetModelResources(false);
+	int32 SupportedVersion = UCustomizableObjectPrivate::CurrentSupportedVersion;
+	MemoryWriter << SupportedVersion;
 
-	MemoryWriter << LocalModelResources.ReferenceSkeletalMeshesData;
+	MemoryWriter << ReferenceSkeletalMeshesData;
 
-	SerializeStreamedResources(MemoryWriter, GetPublic(), GetPublic()->StreamedResourceData, bIsCooking);
-	
-	int32 NumReferencedMaterials = LocalModelResources.Materials.Num();
+	SerializeStreamedResources(MemoryWriter, StreamedResourceData);
+
+	int32 NumReferencedMaterials = Materials.Num();
 	MemoryWriter << NumReferencedMaterials;
 
-	for (const TSoftObjectPtr<UMaterialInterface>& Material : LocalModelResources.Materials)
+	for (const TSoftObjectPtr<UMaterialInterface>& Material : Materials)
 	{
 		FString StringRef = Material.ToString();
 		MemoryWriter << StringRef;
 	}
 
-	int32 NumReferencedSkeletons = LocalModelResources.Skeletons.Num();
+	int32 NumReferencedSkeletons = Skeletons.Num();
 	MemoryWriter << NumReferencedSkeletons;
 
-	for (const TSoftObjectPtr<USkeleton>& Skeleton : LocalModelResources.Skeletons)
+	for (const TSoftObjectPtr<USkeleton>& Skeleton : Skeletons)
 	{
 		FString StringRef = Skeleton.ToString();
 		MemoryWriter << StringRef;
 	}
 
-	int32 NumPassthroughTextures = LocalModelResources.PassThroughTextures.Num();
+	int32 NumPassthroughTextures = PassThroughTextures.Num();
 	MemoryWriter << NumPassthroughTextures;
 
-	for (const TSoftObjectPtr<UTexture>& PassthroughTexture : LocalModelResources.PassThroughTextures)
+	for (const TSoftObjectPtr<UTexture>& PassthroughTexture : PassThroughTextures)
 	{
 		FString StringRef = PassthroughTexture.ToString();
 		MemoryWriter << StringRef;
 	}
 
-	int32 NumPassthroughMeshes = LocalModelResources.PassThroughMeshes.Num();
+	int32 NumPassthroughMeshes = PassThroughMeshes.Num();
 	MemoryWriter << NumPassthroughMeshes;
 
-	for (const TSoftObjectPtr<USkeletalMesh>& PassthroughMesh : LocalModelResources.PassThroughMeshes)
+	for (const TSoftObjectPtr<USkeletalMesh>& PassthroughMesh : PassThroughMeshes)
 	{
 		FString StringRef = PassthroughMesh.ToString();
 		MemoryWriter << StringRef;
 	}
 
 #if WITH_EDITORONLY_DATA
-	int32 NumRuntimeReferencedTextures = LocalModelResources.RuntimeReferencedTextures.Num();
+	int32 NumRuntimeReferencedTextures = RuntimeReferencedTextures.Num();
 	MemoryWriter << NumRuntimeReferencedTextures;
-	
-	for (const TSoftObjectPtr<const UTexture>& RuntimeReferencedTexture : LocalModelResources.RuntimeReferencedTextures)
+
+	for (const TSoftObjectPtr<const UTexture>& RuntimeReferencedTexture : RuntimeReferencedTextures)
 	{
 		FString StringRef = RuntimeReferencedTexture.ToString();
 		MemoryWriter << StringRef;
 	}
 #endif
 
-	int32 NumPhysicsAssets = LocalModelResources.PhysicsAssets.Num();
+	int32 NumPhysicsAssets = PhysicsAssets.Num();
 	MemoryWriter << NumPhysicsAssets;
 
-	for (const TSoftObjectPtr<UPhysicsAsset>& PhysicsAsset : LocalModelResources.PhysicsAssets)
+	for (const TSoftObjectPtr<UPhysicsAsset>& PhysicsAsset : PhysicsAssets)
 	{
 		FString StringRef = PhysicsAsset.ToString();
 		MemoryWriter << StringRef;
 	}
 
-	int32 NumAnimBps = LocalModelResources.AnimBPs.Num();
+	int32 NumAnimBps = AnimBPs.Num();
 	MemoryWriter << NumAnimBps;
 
-	for (const TSoftClassPtr<UAnimInstance>& AnimBp : LocalModelResources.AnimBPs)
+	for (const TSoftClassPtr<UAnimInstance>& AnimBp : AnimBPs)
 	{
 		FString StringRef = AnimBp.ToString();
 		MemoryWriter << StringRef;
 	}
 
-	MemoryWriter << LocalModelResources.AnimBpOverridePhysiscAssetsInfo;
+	MemoryWriter << AnimBpOverridePhysiscAssetsInfo;
 
-	MemoryWriter << LocalModelResources.MaterialSlotNames;
-	MemoryWriter << LocalModelResources.BoneNamesMap;
-	MemoryWriter << LocalModelResources.SocketArray;
+	MemoryWriter << MaterialSlotNames;
+	MemoryWriter << BoneNamesMap;
+	MemoryWriter << SocketArray;
 
-	MemoryWriter << LocalModelResources.SkinWeightProfilesInfo;
+	MemoryWriter << SkinWeightProfilesInfo;
 
-	MemoryWriter << LocalModelResources.ImageProperties;
-	MemoryWriter << LocalModelResources.ParameterUIDataMap;
-	MemoryWriter << LocalModelResources.StateUIDataMap;
+	MemoryWriter << ImageProperties;
+	MemoryWriter << ParameterUIDataMap;
+	MemoryWriter << StateUIDataMap;
 
 #if WITH_EDITORONLY_DATA
-	MemoryWriter << LocalModelResources.IntParameterOptionDataTable;
+	MemoryWriter << IntParameterOptionDataTable;
 #endif
-	
-	MemoryWriter << LocalModelResources.ClothingAssetsData;
-	MemoryWriter << LocalModelResources.ClothSharedConfigsData;
 
-	MemoryWriter << LocalModelResources.NumComponents;
-	MemoryWriter << LocalModelResources.NumLODs;
-	MemoryWriter << LocalModelResources.NumLODsToStream;
-	MemoryWriter << LocalModelResources.FirstLODAvailable;
+	MemoryWriter << ClothingAssetsData;
+	MemoryWriter << ClothSharedConfigsData;
 
-	// Editor Only data
-	MemoryWriter << bDisableTextureStreaming;
-	MemoryWriter << bIsCompiledWithoutOptimization;
-	MemoryWriter << CustomizableObjectPathMap;
-	MemoryWriter << GroupNodeMap;
-	MemoryWriter << ParticipatingObjects;
-	MemoryWriter << LocalModelResources.TableToParamNames;
+	MemoryWriter << NumComponents;
+	MemoryWriter << NumLODs;
+	MemoryWriter << NumLODsToStream;
+	MemoryWriter << FirstLODAvailable;
 
+	// Editor Only data
 	if (!bIsCooking)
 	{
-		MemoryWriter << LocalModelResources.EditorOnlyMorphTargetReconstructionData;
-		MemoryWriter << LocalModelResources.EditorOnlyClothingMeshToMeshVertData;
-	}
+		MemoryWriter << bIsTextureStreamingDisabled;
+		MemoryWriter << bIsCompiledWithOptimization;
+		MemoryWriter << CustomizableObjectPathMap;
+		MemoryWriter << GroupNodeMap;
+		MemoryWriter << ParticipatingObjects;
+		MemoryWriter << TableToParamNames;
 
-	{
-		FModelStreamableBulkData& ModelStreamables = *GetModelStreamableBulkData(false).Get();
-		MemoryWriter << ModelStreamables;
+		MemoryWriter << EditorOnlyMorphTargetReconstructionData;
+		MemoryWriter << EditorOnlyClothingMeshToMeshVertData;
 	}
 }
 
 
-void UCustomizableObjectPrivate::LoadCompiledData(FObjectAndNameAsStringProxyArchive& MemoryReader, const ITargetPlatform* InTargetPlatform, bool bIsCooking)
+bool FModelResources::Unserialize(FObjectAndNameAsStringProxyArchive& MemoryReader, UCustomizableObject& Outer, const ITargetPlatform* InTargetPlatform, bool bIsCooking)
 {
-	TSharedPtr<mu::Model, ESPMode::ThreadSafe> LoadedModel;
-	ClearCompiledData(bIsCooking);
+	MUTABLE_CPUPROFILER_SCOPE(FModelResources::Unserialize);
+	check(IsInGameThread());
 
-	MutableCompiledDataStreamHeader Header;
-	MemoryReader << Header;
+	int32 SupportedVersion = 0;
+	MemoryReader << SupportedVersion;
 
-	if (UCustomizableObjectPrivate::CurrentSupportedVersion == Header.InternalVersion)
+	if (SupportedVersion != UCustomizableObjectPrivate::CurrentSupportedVersion)
 	{
-		// Make sure mutable has been initialised.
-		UCustomizableObjectSystem::GetInstance();
-
-		FModelResources& LocalModelResource = GetModelResources(bIsCooking);
-		LocalModelResource = FModelResources();
+		return false;
+	}
 
-		MemoryReader << LocalModelResource.ReferenceSkeletalMeshesData;
+	MemoryReader << ReferenceSkeletalMeshesData;
 
-		SerializeStreamedResources(MemoryReader, GetPublic(), GetPublic()->StreamedResourceData, bIsCooking);
+	UnserializeStreamedResources(MemoryReader, &Outer, StreamedResourceData, bIsCooking);
 
-		// Initialize resources. 
-		for(FMutableRefSkeletalMeshData& ReferenceSkeletalMeshData : LocalModelResource.ReferenceSkeletalMeshesData)
-		{
-			ReferenceSkeletalMeshData.InitResources(GetPublic(), InTargetPlatform);
-		}
+	// Initialize resources. 
+	for (FMutableRefSkeletalMeshData& ReferenceSkeletalMeshData : ReferenceSkeletalMeshesData)
+	{
+		ReferenceSkeletalMeshData.InitResources(&Outer, *this, InTargetPlatform);
+	}
 
-		int32 NumReferencedMaterials = 0;
-		MemoryReader << NumReferencedMaterials;
+	int32 NumReferencedMaterials = 0;
+	MemoryReader << NumReferencedMaterials;
+	Materials.Reset(NumReferencedMaterials);
 
-		for (int32 i = 0; i < NumReferencedMaterials; ++i)
-		{
-			FString StringRef;
-			MemoryReader << StringRef;
+	for (int32 i = 0; i < NumReferencedMaterials; ++i)
+	{
+		FString StringRef;
+		MemoryReader << StringRef;
 
-			LocalModelResource.Materials.Add(TSoftObjectPtr<UMaterialInterface>(FSoftObjectPath(StringRef)));
-		}
+		Materials.Add(TSoftObjectPtr<UMaterialInterface>(FSoftObjectPath(StringRef)));
+	}
 
-		int32 NumReferencedSkeletons = 0;
-		MemoryReader << NumReferencedSkeletons;
+	int32 NumReferencedSkeletons = 0;
+	MemoryReader << NumReferencedSkeletons;
+	Skeletons.Reset(NumReferencedMaterials);
 
-		for (int32 SkeletonIndex = 0; SkeletonIndex < NumReferencedSkeletons; ++SkeletonIndex)
-		{
-			FString StringRef;
-			MemoryReader << StringRef;
+	for (int32 SkeletonIndex = 0; SkeletonIndex < NumReferencedSkeletons; ++SkeletonIndex)
+	{
+		FString StringRef;
+		MemoryReader << StringRef;
 
-			LocalModelResource.Skeletons.Add(TSoftObjectPtr<USkeleton>(FSoftObjectPath(StringRef)));
-		}
+		Skeletons.Add(TSoftObjectPtr<USkeleton>(FSoftObjectPath(StringRef)));
+	}
 
-		int32 NumPassthroughTextures = 0;
-		MemoryReader << NumPassthroughTextures;
+	int32 NumPassthroughTextures = 0;
+	MemoryReader << NumPassthroughTextures;
+	PassThroughTextures.Reset(NumPassthroughTextures);
 
-		for (int32 Index = 0; Index < NumPassthroughTextures; ++Index)
-		{
-			FString StringRef;
-			MemoryReader << StringRef;
+	for (int32 Index = 0; Index < NumPassthroughTextures; ++Index)
+	{
+		FString StringRef;
+		MemoryReader << StringRef;
 
-			LocalModelResource.PassThroughTextures.Add(TSoftObjectPtr<UTexture>(FSoftObjectPath(StringRef)));
-		}
+		PassThroughTextures.Add(TSoftObjectPtr<UTexture>(FSoftObjectPath(StringRef)));
+	}
 
-		int32 NumPassthroughMeshes = 0;
-		MemoryReader << NumPassthroughMeshes;
+	int32 NumPassthroughMeshes = 0;
+	MemoryReader << NumPassthroughMeshes;
+	PassThroughMeshes.Reset(NumPassthroughMeshes);
 
-		for (int32 Index = 0; Index < NumPassthroughMeshes; ++Index)
-		{
-			FString StringRef;
-			MemoryReader << StringRef;
+	for (int32 Index = 0; Index < NumPassthroughMeshes; ++Index)
+	{
+		FString StringRef;
+		MemoryReader << StringRef;
 
-			LocalModelResource.PassThroughMeshes.Add(TSoftObjectPtr<USkeletalMesh>(FSoftObjectPath(StringRef)));
-		}
+		PassThroughMeshes.Add(TSoftObjectPtr<USkeletalMesh>(FSoftObjectPath(StringRef)));
+	}
 
 #if WITH_EDITORONLY_DATA
-		int32 NumRuntimeReferencedTextures = 0;
-		MemoryReader << NumRuntimeReferencedTextures;
+	int32 NumRuntimeReferencedTextures = 0;
+	MemoryReader << NumRuntimeReferencedTextures;
+	RuntimeReferencedTextures.Reset(NumRuntimeReferencedTextures);
 
-		for (int32 Index = 0; Index < NumRuntimeReferencedTextures; ++Index)
-		{
-			FString StringRef;
-			MemoryReader << StringRef;
+	for (int32 Index = 0; Index < NumRuntimeReferencedTextures; ++Index)
+	{
+		FString StringRef;
+		MemoryReader << StringRef;
 
-			LocalModelResource.RuntimeReferencedTextures.Add(TSoftObjectPtr<const UTexture>(FSoftObjectPath(StringRef)));
-		}
+		RuntimeReferencedTextures.Add(TSoftObjectPtr<const UTexture2D>(FSoftObjectPath(StringRef)));
+	}
 #endif
-		
-		int32 NumPhysicsAssets = 0;
-		MemoryReader << NumPhysicsAssets;
 
-		for (int32 i = 0; i < NumPhysicsAssets; ++i)
-		{
-			FString StringRef;
-			MemoryReader << StringRef;
+	int32 NumPhysicsAssets = 0;
+	MemoryReader << NumPhysicsAssets;
+	PhysicsAssets.Reset(NumPhysicsAssets);
 
-			LocalModelResource.PhysicsAssets.Add(TSoftObjectPtr<UPhysicsAsset>(FSoftObjectPath(StringRef)));
-		}
+	for (int32 i = 0; i < NumPhysicsAssets; ++i)
+	{
+		FString StringRef;
+		MemoryReader << StringRef;
 
+		PhysicsAssets.Add(TSoftObjectPtr<UPhysicsAsset>(FSoftObjectPath(StringRef)));
+	}
 
-		int32 NumAnimBps = 0;
-		MemoryReader << NumAnimBps;
 
-		for (int32 Index = 0; Index < NumAnimBps; ++Index)
-		{
-			FString StringRef;
-			MemoryReader << StringRef;
+	int32 NumAnimBps = 0;
+	MemoryReader << NumAnimBps;
+	AnimBPs.Reset(NumAnimBps);
 
-			LocalModelResource.AnimBPs.Add(TSoftClassPtr<UAnimInstance>(StringRef));
-		}
+	for (int32 Index = 0; Index < NumAnimBps; ++Index)
+	{
+		FString StringRef;
+		MemoryReader << StringRef;
+
+		AnimBPs.Add(TSoftClassPtr<UAnimInstance>(StringRef));
+	}
 
-		MemoryReader << LocalModelResource.AnimBpOverridePhysiscAssetsInfo;
+	MemoryReader << AnimBpOverridePhysiscAssetsInfo;
 
-		MemoryReader << LocalModelResource.MaterialSlotNames;
-		MemoryReader << LocalModelResource.BoneNamesMap;
-		MemoryReader << LocalModelResource.SocketArray;
+	MemoryReader << MaterialSlotNames;
+	MemoryReader << BoneNamesMap;
+	MemoryReader << SocketArray;
 
-		MemoryReader << LocalModelResource.SkinWeightProfilesInfo;
+	MemoryReader << SkinWeightProfilesInfo;
 
-		MemoryReader << LocalModelResource.ImageProperties;
-		MemoryReader << LocalModelResource.ParameterUIDataMap;
-		MemoryReader << LocalModelResource.StateUIDataMap;
+	MemoryReader << ImageProperties;
+	MemoryReader << ParameterUIDataMap;
+	MemoryReader << StateUIDataMap;
 
 #if WITH_EDITORONLY_DATA
-		MemoryReader << LocalModelResource.IntParameterOptionDataTable;
+	MemoryReader << IntParameterOptionDataTable;
 #endif
 
-		MemoryReader << LocalModelResource.ClothingAssetsData; 
-		MemoryReader << LocalModelResource.ClothSharedConfigsData; 
+	MemoryReader << ClothingAssetsData;
+	MemoryReader << ClothSharedConfigsData;
 
-		MemoryReader << LocalModelResource.NumComponents;
-		MemoryReader << LocalModelResource.NumLODs;
-		MemoryReader << LocalModelResource.NumLODsToStream;
-		MemoryReader << LocalModelResource.FirstLODAvailable;
+	MemoryReader << NumComponents;
+	MemoryReader << NumLODs;
+	MemoryReader << NumLODsToStream;
+	MemoryReader << FirstLODAvailable;
 
-		bool bInvalidateModel = false;
-
-		// Editor Only data
-		{
-			MemoryReader << bDisableTextureStreaming;
-			MemoryReader << bIsCompiledWithoutOptimization;
-			MemoryReader << CustomizableObjectPathMap;
-			MemoryReader << GroupNodeMap;
-			MemoryReader << ParticipatingObjects;
-			MemoryReader << LocalModelResource.TableToParamNames;
-
-			if (!bIsCooking)
-			{
-				MemoryReader << LocalModelResource.EditorOnlyMorphTargetReconstructionData;
-				MemoryReader << LocalModelResource.EditorOnlyClothingMeshToMeshVertData;
-				
-				DirtyParticipatingObjects.Empty();
-
-				TArray<FName> OutOfDatePackages;
-				bInvalidateModel = IsCompilationOutOfDate(&OutOfDatePackages);
-
-				if (bInvalidateModel)
-				{
-					UE_LOG(LogMutable, Display, TEXT("Invalidating compiled data due to changes in %s."), *OutOfDatePackages[0].ToString());
-				}
-			}
-		}
-
-		{
-			TSharedPtr<FModelStreamableBulkData> LocalModelStreamablesPtr = MakeShared<FModelStreamableBulkData>();
-			FModelStreamableBulkData& LocalModelStreamables = *LocalModelStreamablesPtr.Get();
-			MemoryReader << LocalModelStreamables;
-
-			SetModelStreamableBulkData(LocalModelStreamablesPtr, bIsCooking);
-		}
-
-		bool bModelSerialized = false;
-		MemoryReader << bModelSerialized;
+	// Editor Only data
+	if (!bIsCooking)
+	{
+		MemoryReader << bIsTextureStreamingDisabled;
+		MemoryReader << bIsCompiledWithOptimization;
+		MemoryReader << CustomizableObjectPathMap;
+		MemoryReader << GroupNodeMap;
+		MemoryReader << ParticipatingObjects;
+		MemoryReader << TableToParamNames;
 
-		if (bModelSerialized && !bInvalidateModel)
-		{
-			UnrealMutableInputStream Stream(MemoryReader);
-			mu::InputArchive Arch(&Stream);
-			LoadedModel = mu::Model::StaticUnserialise(Arch);
-		}
+		MemoryReader << EditorOnlyMorphTargetReconstructionData;
+		MemoryReader << EditorOnlyClothingMeshToMeshVertData;
 	}
-	
-	UpdateParameterPropertiesFromModel(LoadedModel);
-	SetModel(LoadedModel, GenerateIdentifier(*GetPublic()));
+
+	return true;
 }
 
 
@@ -948,8 +960,19 @@ void UCustomizableObjectPrivate::LoadCompiledDataFromDisk()
 				CompiledDataFileHandle->Read(CompiledDataBytes.GetData(), CompiledDataSize);
 
 				FMemoryReaderView MemoryReader(CompiledDataBytes);
-				FObjectAndNameAsStringProxyArchive ObjectReader(MemoryReader, true);
-				LoadCompiledData(ObjectReader, RunningPlatform);
+
+				DirtyParticipatingObjects.Empty();
+
+				TArray<FName> OutOfDatePackages;
+				if (LoadModelResources(MemoryReader, RunningPlatform) && !IsCompilationOutOfDate(&OutOfDatePackages))
+				{
+					LoadModelStreamableBulk(MemoryReader, /* bIsCooking */false);
+					LoadModel(MemoryReader);
+				}
+				else if (!OutOfDatePackages.IsEmpty())
+				{
+					UE_LOG(LogMutable, Display, TEXT("Invalidating compiled data due to changes in %s."), *OutOfDatePackages[0].ToString());
+				}
 			}
 		}
 	}
@@ -961,17 +984,10 @@ void UCustomizableObjectPrivate::LoadCompiledDataFromDisk()
 }
 
 
-void UCustomizableObjectPrivate::CompileForTargetPlatform(const ITargetPlatform* TargetPlatform)
+void UCustomizableObjectPrivate::CompileForTargetPlatform(UCustomizableObject& CustomizableObject, const ITargetPlatform& TargetPlatform)
 {
-	if (!TargetPlatform)
-	{
-		return;
-	}
-
-	UCustomizableObject* CustomizableObject = GetPublic();
-
 	ICustomizableObjectEditorModule* EditorModule = ICustomizableObjectEditorModule::Get();
-	if (!EditorModule || !EditorModule->IsRootObject(*CustomizableObject))
+	if (!EditorModule || !EditorModule->IsRootObject(CustomizableObject))
 	{
 		SetIsChildObject(true);
 		return;
@@ -979,12 +995,25 @@ void UCustomizableObjectPrivate::CompileForTargetPlatform(const ITargetPlatform*
 
 	const bool bAsync = false; // TODO PERE
 
-	TSharedRef<FCompilationRequest> CompileRequest = MakeShared<FCompilationRequest>(*CustomizableObject, bAsync);
+	TSharedRef<FCompilationRequest> CompileRequest = MakeShared<FCompilationRequest>(CustomizableObject, bAsync);
 	FCompilationOptions& Options = CompileRequest->GetCompileOptions();
 	Options.OptimizationLevel = UE_MUTABLE_MAX_OPTIMIZATION;	// Force max optimization when packaging.
 	Options.TextureCompression = ECustomizableObjectTextureCompression::HighQuality;
 	Options.bIsCooking = true;
-	Options.TargetPlatform = TargetPlatform;
+	Options.TargetPlatform = &TargetPlatform;
+
+	const int32 DDCUsage = CVarMutableDerivedDataCacheUsage.GetValueOnAnyThread();
+	UE::DerivedData::ECachePolicy DefaultCachePolicy = UE::DerivedData::ECachePolicy::None;
+	if (DDCUsage == 1)
+	{
+		DefaultCachePolicy = UE::DerivedData::ECachePolicy::Local;
+	}
+	else if (DDCUsage == 2)
+	{
+		DefaultCachePolicy = UE::DerivedData::ECachePolicy::Default;
+	}
+	CompileRequest->SetDerivedDataCachePolicy(DefaultCachePolicy);
+
 	CompileRequests.Add(CompileRequest);
 
 	EditorModule->CompileCustomizableObject(CompileRequest, true);
@@ -1152,9 +1181,6 @@ void UCustomizableObjectPrivate::LoadEmbeddedData(FArchive& Ar)
 		mu::InputArchive Arch(&Stream);
 		TSharedPtr<mu::Model, ESPMode::ThreadSafe> Model = mu::Model::StaticUnserialise(Arch);
 
-		// Create parameter properties
-		UpdateParameterPropertiesFromModel(Model);
-
 		SetModel(Model, FGuid());
 	}
 }
@@ -2057,19 +2083,6 @@ FArchive& operator<<(FArchive& Ar, FMutableStateData& Struct)
 	return Ar;
 }
 
-#if WITH_EDITORONLY_DATA
-FArchive& operator<<(FArchive& Ar, FModelStreamableBulkData& Data)
-{
-	Ar << Data.ModelStreamables;
-	Ar << Data.ClothingStreamables;
-	Ar << Data.RealTimeMorphStreamables;
-
-	// Don't serialize FByteBulkData manually, the data will be skipped.
-
-	return Ar;
-}
-#endif
-
 
 void FModelStreamableBulkData::Serialize(FArchive& Ar, UObject* Owner, bool bCooked)
 {
@@ -2159,6 +2172,9 @@ void UCustomizableObjectPrivate::SetModel(const TSharedPtr<mu::Model, ESPMode::T
 	
 	MutableModel = Model;
 
+	// Create parameter properties
+	UpdateParameterPropertiesFromModel(Model);
+
 	using EState = FCustomizableObjectStatus::EState;
 	Status.NextState(Model ? EState::ModelLoaded : EState::NoModel);
 }
@@ -2241,9 +2257,9 @@ bool UCustomizableObjectPrivate::IsCompilationOutOfDate(TArray<FName>* OutOfDate
 
 void UCustomizableObjectPrivate::OnParticipatingObjectDirty(UPackage* Package, bool)
 {
-	if (ParticipatingObjects.Contains(Package->GetFName()))
+	if (GetModelResources().ParticipatingObjects.Contains(Package->GetFName()))
 	{
-		DirtyParticipatingObjects.AddUnique(Package->GetFName());		
+		DirtyParticipatingObjects.AddUnique(Package->GetFName());
 	}
 }
 #endif
@@ -2287,9 +2303,46 @@ TArray<FCustomizableObjectStreamedResourceData>& UCustomizableObjectPrivate::Get
 }
 
 
-TArray<FCustomizableObjectStreamedResourceData>& UCustomizableObjectPrivate::GetStreamedResourceData()
+const FCustomizableObjectResourceData* UCustomizableObjectPrivate::LoadStreamedResource(int32 ResourceIndex)
+{
+#if WITH_EDITORONLY_DATA
+	FModelResources& LocalModelResources = ModelResourcesEditor;
+#else
+	FModelResources& LocalModelResources = ModelResources;
+#endif
+
+	if (LocalModelResources.StreamedResourceData.IsValidIndex(ResourceIndex))
+	{
+		FCustomizableObjectStreamedResourceData& Resource = LocalModelResources.StreamedResourceData[ResourceIndex];
+		if (!Resource.IsLoaded())
+		{
+			Resource.NotifyLoaded(Resource.GetPath().Get());
+		}
+
+		return &Resource.GetLoadedData();
+	}
+
+	return nullptr;
+}
+
+void UCustomizableObjectPrivate::UnloadStreamedResource(int32 ResourceIndex)
 {
-	return GetPublic()->StreamedResourceData;
+	// Only Unload in cooked builds. Unloading them when in the editor will trigger an assert. 
+	if (FPlatformProperties::RequiresCookedData())
+	{
+		return;
+	}
+
+#if WITH_EDITORONLY_DATA
+	FModelResources& LocalModelResources = ModelResourcesEditor;
+#else
+	FModelResources& LocalModelResources = ModelResources;
+#endif
+
+	if (LocalModelResources.StreamedResourceData.IsValidIndex(ResourceIndex))
+	{
+		LocalModelResources.StreamedResourceData[ResourceIndex].Unload();
+	}
 }
 
 
@@ -2306,7 +2359,9 @@ FCompilationOptions UCustomizableObjectPrivate::GetCompileOptions() const
 	Options.TextureCompression = TextureCompression;
 	Options.OptimizationLevel = OptimizationLevel;
 	Options.bUseDiskCompilation = bUseDiskCompilation;
-	
+
+	Options.TargetPlatform = GetTargetPlatformManagerRef().GetRunningTargetPlatform();
+
 	const int32 TargetBulkDataFileBytesOverride = CVarPackagedDataBytesLimitOverride.GetValueOnAnyThread();
 	if ( TargetBulkDataFileBytesOverride >= 0)
 	{
@@ -2317,7 +2372,7 @@ FCompilationOptions UCustomizableObjectPrivate::GetCompileOptions() const
 	{
 		Options.PackagedDataBytesLimit =  PackagedDataBytesLimit;
 	}
-	
+
 	Options.EmbeddedDataBytesLimit = EmbeddedDataBytesLimit;
 	Options.CustomizableObjectNumBoneInfluences = ICustomizableObjectModule::Get().GetNumBoneInfluences();
 	Options.bRealTimeMorphTargetsEnabled = GetPublic()->bEnableRealTimeMorphTargets;
@@ -2411,7 +2466,7 @@ namespace MutablePrivate
 
 	void GenerateBulkDataFilesList(
 		TSharedPtr<const mu::Model, ESPMode::ThreadSafe> Model,
-		TSharedPtr<FModelStreamableBulkData> ModelStreamableBulkData,
+		FModelStreamableBulkData& ModelStreamableBulkData,
 		const ITargetPlatform* TargetPlatform,
 		uint64 TargetBulkDataFileBytes,
 		TArray<FFile>& OutBulkDataFiles)
@@ -2420,7 +2475,7 @@ namespace MutablePrivate
 
 		OutBulkDataFiles.Empty();
 
-		if (!Model || !ModelStreamableBulkData)
+		if (!Model)
 		{
 			return;
 		}
@@ -2428,6 +2483,9 @@ namespace MutablePrivate
 		const uint64 MaxChunkSize = UCustomizableObjectSystem::GetInstance()->GetMaxChunkSizeForPlatform(TargetPlatform);
 		TargetBulkDataFileBytes = FMath::Min(TargetBulkDataFileBytes, MaxChunkSize);
 
+		// TODO: Temp. Remove after unifying generated output files code between editor an package. UE-222777 
+		const bool bRequiresCookedData = TargetPlatform->RequiresCookedData();
+
 		// Root nodes by flags.
 		const int32 NumRoms = Model->GetRomCount();
 		TMap<uint32, FClassifyNode> RootNode;
@@ -2438,7 +2496,7 @@ namespace MutablePrivate
 			{
 				uint32 BlockId = Model->GetRomId(RomIndex);
 				const uint32 BlockSize = Model->GetRomSize(RomIndex);
-				const mu::ERomFlags BlockFlags = Model->GetRomFlags(RomIndex);
+				const mu::ERomFlags BlockFlags = bRequiresCookedData ? Model->GetRomFlags(RomIndex) : mu::ERomFlags::None;
 
 				FBlock CurrentBlock = { EDataType::Model, BlockId, BlockSize, uint32(BlockFlags), 0 };
 				AddNode(RootNode, NumRoms, CurrentBlock);
@@ -2452,7 +2510,7 @@ namespace MutablePrivate
 
 			TArray<FMutableStreamableBlock> RealTimeMorphTargetsBlocks;
 
-			const TMap<uint32, FRealTimeMorphStreamable>& RealTimeMorphStreamables = ModelStreamableBulkData->RealTimeMorphStreamables;
+			const TMap<uint32, FRealTimeMorphStreamable>& RealTimeMorphStreamables = ModelStreamableBulkData.RealTimeMorphStreamables;
 
 			const int32 NumBlocks = RealTimeMorphTargetsBlocks.Num();
 			for (const TPair<uint32, FRealTimeMorphStreamable>& MorphStreamable : RealTimeMorphStreamables)
@@ -2474,7 +2532,7 @@ namespace MutablePrivate
 
 			TArray<FMutableStreamableBlock> ClothingBlocks;
 
-			const TMap<uint32, FClothingStreamable>& ClothingStreamables = ModelStreamableBulkData->ClothingStreamables;
+			const TMap<uint32, FClothingStreamable>& ClothingStreamables = ModelStreamableBulkData.ClothingStreamables;
 
 			const int32 NumBlocks = ClothingBlocks.Num();
 			for (const TPair<uint32, FClothingStreamable>& ClothStreamable : ClothingStreamables)
@@ -2574,7 +2632,7 @@ namespace MutablePrivate
 					FBlock& ThisBlock = CurrentFile.Blocks[FileBlockIndex];
 					ThisBlock.Offset = OffsetInFile;
 
-					FMutableStreamableBlock* StreamableBlock = ModelStreamableBulkData->ModelStreamables.Find(ThisBlock.Id);
+					FMutableStreamableBlock* StreamableBlock = ModelStreamableBulkData.ModelStreamables.Find(ThisBlock.Id);
 					StreamableBlock->FileId = FileId;
 					StreamableBlock->Offset = OffsetInFile;
 					check(StreamableBlock->Flags == CurrentFile.Flags);
@@ -2583,7 +2641,7 @@ namespace MutablePrivate
 			}
 			else if (CurrentFile.DataType == EDataType::RealTimeMorph)
 			{
-				TMap<uint32, FRealTimeMorphStreamable>& MorphBlocks = ModelStreamableBulkData->RealTimeMorphStreamables;
+				TMap<uint32, FRealTimeMorphStreamable>& MorphBlocks = ModelStreamableBulkData.RealTimeMorphStreamables;
 				// Set it to all streamable blocks
 				uint32 OffsetInFile = 0;
 				for (int32 FileBlockIndex = 0; FileBlockIndex < CurrentFile.Blocks.Num(); ++FileBlockIndex)
@@ -2599,7 +2657,7 @@ namespace MutablePrivate
 			}
 			else if (CurrentFile.DataType == EDataType::Clothing)
 			{
-				TMap<uint32, FClothingStreamable>& ClothBlocks = ModelStreamableBulkData->ClothingStreamables;
+				TMap<uint32, FClothingStreamable>& ClothBlocks = ModelStreamableBulkData.ClothingStreamables;
 				// Set it to all streamable blocks
 				uint32 OffsetInFile = 0;
 				for (int32 FileBlockIndex = 0; FileBlockIndex < CurrentFile.Blocks.Num(); ++FileBlockIndex)
@@ -2644,7 +2702,337 @@ namespace MutablePrivate
 			WriteFile(CurrentFile, FileBulkData);
 		}
 	}
+
+	UE::DerivedData::FValueId GetDerivedDataModelId()
+	{
+		UE::DerivedData::FValueId::ByteArray ValueIdBytes = {};
+		ValueIdBytes[0] = 1;
+		return UE::DerivedData::FValueId(ValueIdBytes);
+	}
+
+	UE::DerivedData::FValueId GetDerivedDataModelResourcesId()
+	{
+		UE::DerivedData::FValueId::ByteArray ValueIdBytes = {};
+		ValueIdBytes[0] = 2;
+		return UE::DerivedData::FValueId(ValueIdBytes);
+	}
+
+	UE::DerivedData::FValueId GetDerivedDataModelStreamableBulkDataId()
+	{
+		UE::DerivedData::FValueId::ByteArray ValueIdBytes = {};
+		ValueIdBytes[0] = 3;
+		return UE::DerivedData::FValueId(ValueIdBytes);
+	}
+
+	UE::DerivedData::FValueId GetDerivedDataBulkDataFilesId()
+	{
+		UE::DerivedData::FValueId::ByteArray ValueIdBytes = {};
+		ValueIdBytes[0] = 4;
+		return UE::DerivedData::FValueId(ValueIdBytes);
+	}
+}
+
+
+
+void SerializeCompilationOptionsForDDC(FArchive& Ar, FCompilationOptions& Options)
+{
+	FString PlatformName = Options.TargetPlatform ? Options.TargetPlatform->PlatformName() : FString();
+	Ar << PlatformName;
+	Ar << Options.TextureCompression;
+	Ar << Options.OptimizationLevel;
+	Ar << Options.DDCBytesLimit;
+	Ar << Options.CustomizableObjectNumBoneInfluences;
+	Ar << Options.bRealTimeMorphTargetsEnabled;
+	Ar << Options.bClothingEnabled;
+	Ar << Options.b16BitBoneWeightsEnabled;
+	Ar << Options.bSkinWeightProfilesEnabled;
+	Ar << Options.bPhysicsAssetMergeEnabled;
+	Ar << Options.bAnimBpPhysicsManipulationEnabled;
+	Ar << Options.ImageTiling;
+}
+
+
+const FString& GetCustomizableObjectDerivedDataVersion()
+{
+	static FString CachedVersionString = TEXT("53DD1D7C-C040-4A52-162B-A3F03E2E8320");
+	return CachedVersionString;
+}
+
+
+FString UCustomizableObjectPrivate::BuildDerivedDataKey(FCompilationOptions Options)
+{
+	FString KeySuffix;
+
+	TArray<uint8> TempBytes;
+	TempBytes.Reserve(160);
+
+	KeySuffix += FString::Printf(TEXT("%08x-"), (uint32)CurrentSupportedVersion);
+
+	check(IsInGameThread());
+	UCustomizableObject& CustomizableObject = *GetPublic();
+
+	// Customizable Object Ids
+	KeySuffix += GenerateIdentifier(CustomizableObject).ToString();
+	KeySuffix += CustomizableObject.VersionId.ToString();
+
+	// Compile Options
+	{
+		TempBytes.Reset();
+		FMemoryWriter Ar(TempBytes, /*bIsPersistent=*/ true);
+		SerializeCompilationOptionsForDDC(Ar, Options);
+		KeySuffix += FString::Printf(TEXT("OPT_%016x"), CityHash64(reinterpret_cast<const char*>(TempBytes.GetData()), TempBytes.Num()));
+	}
+
+	// Content Version
+	if (const ICustomizableObjectEditorModule* Module = ICustomizableObjectEditorModule::Get())
+	{
+		KeySuffix += TEXT("VER_") + Module->GetCurrentContentVersionForObject(CustomizableObject);
+	}
+
+	// Participating objects hash
+	{
+		TempBytes.Reset();
+		FMemoryWriter Ar(TempBytes, /*bIsPersistent=*/ true);
+		//Ar << CachedParticipatingObjectsHash; // TODO PERE: Build ParticipatingObjects Hash
+		KeySuffix += FString::FromBlob(TempBytes.GetData(), TempBytes.Num() * TempBytes.GetTypeSize());
+	}
+
+	return FDerivedDataCacheInterface::BuildCacheKey(
+		TEXT("CUSTOMIZABLEOBJECT"),
+		*GetCustomizableObjectDerivedDataVersion(),
+		*KeySuffix
+	);
+}
+
+
+UE::DerivedData::FCacheKey UCustomizableObjectPrivate::GetDerivedDataCacheKeyForOptions(FCompilationOptions InOptions)
+{
+	using namespace UE::DerivedData;
+
+	// Cache key as string
+	FString DerivedDataKey = BuildDerivedDataKey(InOptions);
+
+	FCacheKey CacheKey;
+	CacheKey.Bucket = FCacheBucket(TEXT("CustomizableObject"));
+	CacheKey.Hash = FIoHashBuilder::HashBuffer(MakeMemoryView(FTCHARToUTF8(DerivedDataKey)));
+	return CacheKey;
+}
+
+
+void UCustomizableObjectPrivate::LoadCompiledDataFromDDC(FCompilationOptions Options, UE::DerivedData::ECachePolicy DefaultPolicy, UE::DerivedData::FCacheKey* DDCKey)
+{
+	MUTABLE_CPUPROFILER_SCOPE(UCustomizableObjectPrivate::LoadCompiledDataFromDDC);
+
+	using namespace UE::DerivedData;
+
+	/* Overview.
+	*	1. Create an initial pull request to look for the compiled data in the DDC. Skip streamable binary blobs.
+	*	2. Try to load the compiled data.
+	*	3. (Cooking) Create a second request to pull all streamable blobs and cache the compiled data.
+	*/
+
+	FCacheKey CacheKey = DDCKey ? *DDCKey : GetDerivedDataCacheKeyForOptions(Options);
+	check(CacheKey.Hash.IsZero() == false);
+
+	// Buffers with the compiled data
+	FSharedBuffer ModelBytesDDC;
+	FSharedBuffer ModelResourcesBytesDDC;
+	FSharedBuffer ModelStreamablesBytesDDC;
+	FSharedBuffer BulkDataFilesBytesDDC;
+
+	{	// Create a (sync) request to get the serialized Model, ModelResources, and ModelStreamable files to validate versioning and resources 
+		MUTABLE_CPUPROFILER_SCOPE(CheckDDC);
+
+		// Set the request policy to Default + SkipData to avoid pulling the streamable files until we know the compiled data can be used.
+		FCacheRecordPolicyBuilder PolicyBuilder(DefaultPolicy | ECachePolicy::SkipData);
+
+		// Overwrite the request policy for the resources we want to pull
+		PolicyBuilder.AddValuePolicy(MutablePrivate::GetDerivedDataModelResourcesId(), DefaultPolicy);
+		PolicyBuilder.AddValuePolicy(MutablePrivate::GetDerivedDataModelId(), DefaultPolicy);
+		PolicyBuilder.AddValuePolicy(MutablePrivate::GetDerivedDataModelStreamableBulkDataId(), DefaultPolicy);
+		PolicyBuilder.AddValuePolicy(MutablePrivate::GetDerivedDataBulkDataFilesId(), DefaultPolicy);
+
+		FCacheGetRequest Request;
+		Request.Name = GetPathNameSafe(GetPublic());
+		Request.Key = CacheKey;
+		Request.Policy = PolicyBuilder.Build();
+
+		// Sync request to retrieve the compiled data for validation. Streamable resources are excluded.
+		FRequestOwner RequestOwner(EPriority::Blocking);
+		GetCache().Get(MakeArrayView(&Request, 1), RequestOwner,
+			[&ModelBytesDDC, &ModelResourcesBytesDDC, &ModelStreamablesBytesDDC, &BulkDataFilesBytesDDC](FCacheGetResponse&& Response)
+			{
+				if (Response.Status == EStatus::Ok)
+				{
+					const FCompressedBuffer& ModelCompressedBuffer = Response.Record.GetValue(MutablePrivate::GetDerivedDataModelId()).GetData();
+					ModelBytesDDC = ModelCompressedBuffer.Decompress();
+
+					const FCompressedBuffer& ModelResourcesCompressedBuffer = Response.Record.GetValue(MutablePrivate::GetDerivedDataModelResourcesId()).GetData();
+					ModelResourcesBytesDDC = ModelResourcesCompressedBuffer.Decompress();
+					
+					const FCompressedBuffer& ModelStreamablesCompressedBuffer = Response.Record.GetValue(MutablePrivate::GetDerivedDataModelStreamableBulkDataId()).GetData();
+					ModelStreamablesBytesDDC = ModelStreamablesCompressedBuffer.Decompress();
+
+					const FCompressedBuffer& BulkDataFilesCompressedBuffer = Response.Record.GetValue(MutablePrivate::GetDerivedDataBulkDataFilesId()).GetData();
+					BulkDataFilesBytesDDC = BulkDataFilesCompressedBuffer.Decompress();
+				}
+			});
+		RequestOwner.Wait();
+	}
+
+	// Check if all the requested buffers were found.
+	if (!ModelBytesDDC.IsNull() && !ModelResourcesBytesDDC.IsNull() && !BulkDataFilesBytesDDC.IsNull()  && !ModelStreamablesBytesDDC.IsNull())
+	{
+		// Load the compiled data to validate it.
+		FMemoryReaderView ModelResourcesReader(ModelResourcesBytesDDC.GetView());
+		if (LoadModelResources(ModelResourcesReader, Options.TargetPlatform, Options.bIsCooking))
+		{
+			FModelResources& LocalModelResources = GetModelResources(Options.bIsCooking);
+			LocalModelResources.bIsStoredInDDC = true;
+			LocalModelResources.DDCKey = CacheKey;
+			LocalModelResources.DDCDefaultPolicy = ECachePolicy::Default;
+
+			FMemoryReaderView ModelStreamablesReader(ModelStreamablesBytesDDC.GetView());
+			LoadModelStreamableBulk(ModelStreamablesReader, Options.bIsCooking);
+
+			FMemoryReaderView ModelReader(ModelBytesDDC.GetView());
+			LoadModel(ModelReader);
+		}
+
+		TSharedPtr<mu::Model> Model = GetModel();
+		TSharedPtr<FModelStreamableBulkData> ModelStreamables = GetModelStreamableBulkData(Options.bIsCooking);
+
+		// Cache CookedPlatfomData. 
+		if (Options.bIsCooking && Model && ModelStreamables)
+		{
+			// Sync cache cooked platform data
+			// TODO UE-220138: Sync -> Async
+			MUTABLE_CPUPROFILER_SCOPE(CacheCookedPlatformData);
+
+			MutablePrivate::FMutableCachedPlatformData CachedData;
+			
+			// Cache Model, ModelResources and ModelStreamables
+			CachedData.ModelData.Append(reinterpret_cast<const uint8*>(ModelBytesDDC.GetData()), ModelBytesDDC.GetSize());
+			CachedData.ModelResourcesData.Append(reinterpret_cast<const uint8*>(ModelResourcesBytesDDC.GetData()), ModelResourcesBytesDDC.GetSize());
+			CachedData.ModelStreamables = ModelStreamables;
+
+			// Value Id to file mapping to reconstruct the cached data
+			TMap<FValueId, MutablePrivate::FFile> ValueIdToFile;
+
+			{
+				MUTABLE_CPUPROFILER_SCOPE(BuildValueIdToFile);
+				TArray<MutablePrivate::FFile> BulkDataFiles;
+				FMemoryReaderView FilesReader(BulkDataFilesBytesDDC.GetView());
+				FilesReader << BulkDataFiles;
+
+				ValueIdToFile.Reserve(BulkDataFiles.Num());
+
+				FValueId::ByteArray ValueIdBytes = {};
+				for (MutablePrivate::FFile& File : BulkDataFiles)
+				{
+					FMemory::Memcpy(&ValueIdBytes, &File.DataType, sizeof(File.DataType));
+					FMemory::Memcpy(&ValueIdBytes[4], &File.Id, sizeof(File.Id));
+					MutablePrivate::FFile& DestFile = ValueIdToFile.Add(FValueId(ValueIdBytes));
+					DestFile = MoveTemp(File);
+				}
+				BulkDataFiles.Empty();
+			}
+
+			// Create a new pull request to retrieve all compiled data. Streamable bulk data included
+			FCacheGetRequest Request;
+			Request.Name = GetPathNameSafe(GetPublic());
+			Request.Key = CacheKey;
+			Request.Policy = ECachePolicy::Default;
+
+			FRequestOwner RequestOwner(EPriority::Blocking);
+			GetCache().Get(MakeArrayView(&Request, 1), RequestOwner,
+				[&CachedData, &ValueIdToFile, ModelStreamables](FCacheGetResponse&& Response)
+				{
+					MUTABLE_CPUPROFILER_SCOPE(CacheBulkDataFromDDC);
+
+					if (ensure(Response.Status == EStatus::Ok))
+					{
+						// Get all values and convert them to FMutableCachedPlatformData's format
+						TConstArrayView<FValueWithId> Values = Response.Record.GetValues();
+
+						TArray64<uint8> TempData;
+						for (const FValueWithId& Value : Values)
+						{
+							check(Value.IsValid());
+
+							const MutablePrivate::FFile* File = ValueIdToFile.Find(Value.GetId());
+							if (!File) // Skip value. It is not a streamable binary blob.
+							{
+								continue;
+							}
+
+							const uint64 RawSize = Value.GetRawSize();
+							TempData.SetNumUninitialized(RawSize, EAllowShrinking::No);
+
+							// Decompress streamable binary blobs
+							const bool bDecompressedSuccessfully = Value.GetData().TryDecompressTo(MakeMemoryView(TempData.GetData(), RawSize));
+							check(bDecompressedSuccessfully);
+
+							// Filter and cache the data by DataType
+							switch (File->DataType)
+							{
+							case MutablePrivate::EDataType::Model:
+							{
+								for (const MutablePrivate::FBlock& Block : File->Blocks)
+								{
+									CachedData.ModelStreamableData.Set(Block.Id, TempData.GetData() + Block.Offset, Block.Size);
+								}
+								break;
+							}
+							case MutablePrivate::EDataType::RealTimeMorph:
+							{
+								// Store the MorphData in a single array.
+								uint64 Offset = CachedData.MorphData.Num();
+								CachedData.MorphData.Append(TempData.GetData(), RawSize);
+
+								// Fix offsets
+								for (const MutablePrivate::FBlock& Block : File->Blocks)
+								{
+									ModelStreamables->RealTimeMorphStreamables[Block.Id].Block.Offset = Offset;
+									Offset += Block.Size;
+								}
+								break;
+							}
+							case MutablePrivate::EDataType::Clothing:
+							{
+								// Store the ClothingData in a single array.
+								uint64 Offset = CachedData.ClothingData.Num();
+								CachedData.ClothingData.Append(TempData.GetData(), RawSize);
+
+								// Fix offsets
+								for (const MutablePrivate::FBlock& Block : File->Blocks)
+								{
+									ModelStreamables->ClothingStreamables[Block.Id].Block.Offset = Offset;
+									Offset += Block.Size;
+								}
+								break;
+							}
+
+							default:
+								unimplemented();
+								break;
+							}
+						}
+					}
+				});
+			RequestOwner.Wait();
+
+			// Generate list of files and update streamable blocks ids and offsets
+			MutablePrivate::GenerateBulkDataFilesList(Model, *ModelStreamables.Get(), Options.TargetPlatform, Options.PackagedDataBytesLimit, CachedData.BulkDataFiles);
+
+			MutablePrivate::FMutableCachedPlatformData& CachedPlatformData = CachedPlatformsData.Add(Options.TargetPlatform->PlatformName(), {});
+			CachedPlatformData = MoveTemp(CachedData);
+		}
+	}
+
+	return;
 }
+
 #endif // WITH_EDITOR
 
 //-------------------------------------------------------------------------------------------------
@@ -2935,7 +3323,7 @@ FArchive& operator<<(FArchive& Ar, FMutableSkinWeightProfileInfo& Info)
 //-------------------------------------------------------------------------------------------------
 //-------------------------------------------------------------------------------------------------
 //-------------------------------------------------------------------------------------------------
-void FMutableRefSkeletalMeshData::InitResources(UCustomizableObject* InOuter, const ITargetPlatform* InTargetPlatform)
+void FMutableRefSkeletalMeshData::InitResources(UCustomizableObject* InOuter, FModelResources& InModelResources, const ITargetPlatform* InTargetPlatform)
 {
 	check(InOuter);
 
@@ -2948,14 +3336,13 @@ void FMutableRefSkeletalMeshData::InitResources(UCustomizableObject* InOuter, co
 	// Initialize AssetUserData
 	for (FMutableRefAssetUserData& Data : AssetUserData)
 	{
-		if (!InOuter->GetPrivate()->GetStreamedResourceData().IsValidIndex(Data.AssetUserDataIndex))
+		if (!InModelResources.StreamedResourceData.IsValidIndex(Data.AssetUserDataIndex))
 		{
 			check(false);
 			continue;
 		}
 
-		FCustomizableObjectStreamedResourceData& StreamedResource = InOuter->GetPrivate()->GetStreamedResourceData()[Data.AssetUserDataIndex];
-		Data.AssetUserData = StreamedResource.GetPath().LoadSynchronous();
+		Data.AssetUserData = InModelResources.StreamedResourceData[Data.AssetUserDataIndex].GetPath().LoadSynchronous();
 		check(Data.AssetUserData);
 		check(Data.AssetUserData->Data.Type == ECOResourceDataType::AssetUserData);
 	}
@@ -3030,6 +3417,7 @@ FCompilationRequest::FCompilationRequest(UCustomizableObject& InCustomizableObje
 	CustomizableObject = &InCustomizableObject;
 	Options = InCustomizableObject.GetPrivate()->GetCompileOptions();
 	bAsync = bAsyncCompile;
+	DDCPolicy = UE::DerivedData::ECachePolicy::None;
 }
 
 
@@ -3051,6 +3439,34 @@ bool FCompilationRequest::IsAsyncCompilation() const
 }
 
 
+void FCompilationRequest::SetDerivedDataCachePolicy(UE::DerivedData::ECachePolicy InCachePolicy)
+{
+	DDCPolicy = InCachePolicy;
+	Options.bQueryCompiledDatafromDDC = EnumHasAnyFlags(InCachePolicy, UE::DerivedData::ECachePolicy::Query);
+	Options.bStoreCompiledDataInDDC = EnumHasAnyFlags(InCachePolicy, UE::DerivedData::ECachePolicy::Store);
+}
+
+
+UE::DerivedData::ECachePolicy FCompilationRequest::GetDerivedDataCachePolicy() const
+{
+	return DDCPolicy;
+}
+
+void FCompilationRequest::BuildDerivedDataCacheKey()
+{
+	if (UCustomizableObject* Object = CustomizableObject.Get())
+	{
+		DDCKey = Object->GetPrivate()->GetDerivedDataCacheKeyForOptions(Options);
+	}
+}
+
+
+UE::DerivedData::FCacheKey FCompilationRequest::GetDerivedDataCacheKey() const
+{
+	return DDCKey;
+}
+
+
 void FCompilationRequest::SetCompilationState(ECompilationStatePrivate InState, ECompilationResultPrivate InResult)
 {
 	State = InState;
diff --git a/Engine/Plugins/Experimental/Mutable/Source/CustomizableObject/Private/MuCO/CustomizableObjectInstance.cpp b/Engine/Plugins/Experimental/Mutable/Source/CustomizableObject/Private/MuCO/CustomizableObjectInstance.cpp
--- a/Engine/Plugins/Experimental/Mutable/Source/CustomizableObject/Private/MuCO/CustomizableObjectInstance.cpp
+++ b/Engine/Plugins/Experimental/Mutable/Source/CustomizableObject/Private/MuCO/CustomizableObjectInstance.cpp
@@ -5262,19 +5262,18 @@ UE::Tasks::FTask UCustomizableInstancePrivate::LoadAdditionalAssetsAndData(
 
 				for (uint64 ResourceId : StreamedResources)
 				{
-					const TArray<FCustomizableObjectStreamedResourceData>& StreamedResourcesData = CustomizableObject->GetPrivate()->GetStreamedResourceData();
 					FCustomizableObjectStreameableResourceId TypedResourceId = BitCast<FCustomizableObjectStreameableResourceId>(ResourceId);	
 		
 					if (TypedResourceId.Type == (uint8)FCustomizableObjectStreameableResourceId::EType::AssetUserData)
 					{
 						const uint32 ResourceIndex = TypedResourceId.Id;
-						if (!StreamedResourcesData.IsValidIndex(ResourceIndex))
+						if (!ModelResources.StreamedResourceData.IsValidIndex(ResourceIndex))
 						{
-							UE_LOG(LogMutable, Error, TEXT("Invalid streamed resource index. Max Index [%d]. Resource Index [%d]."), StreamedResourcesData.Num(), ResourceIndex);
+							UE_LOG(LogMutable, Error, TEXT("Invalid streamed resource index. Max Index [%d]. Resource Index [%d]."), ModelResources.StreamedResourceData.Num(), ResourceIndex);
 							continue; 
 						}
 
-						const FCustomizableObjectStreamedResourceData& StreamedResource = StreamedResourcesData[ResourceIndex];
+						const FCustomizableObjectStreamedResourceData& StreamedResource = ModelResources.StreamedResourceData[ResourceIndex];
 						if (!StreamedResource.IsLoaded())
 						{
 							AssetsToStream.AddUnique(StreamedResource.GetPath().ToSoftObjectPath());
@@ -5820,9 +5819,9 @@ void UCustomizableInstancePrivate::AdditionalAssetsAsyncLoaded(UCustomizableObje
 {
 	MUTABLE_CPUPROFILER_SCOPE(UCustomizableInstancePrivate::AdditionalAssetsAsyncLoaded);
 
-	UCustomizableObject* CustomizableObject = Public->GetCustomizableObject();
+	UCustomizableObjectPrivate* CustomizableObjectPrivate = Public->GetCustomizableObject()->GetPrivate();
 
-	const FModelResources& ModelResources = CustomizableObject->GetPrivate()->GetModelResources();
+	const FModelResources& ModelResources = CustomizableObjectPrivate->GetModelResources();
 
 	// Loaded Materials
 	check(ObjectToInstanceIndexMap.Num() == ReferencedMaterials.Num());
@@ -5851,25 +5850,22 @@ void UCustomizableInstancePrivate::AdditionalAssetsAsyncLoaded(UCustomizableObje
 #endif
 	}
 
-	
-	TArray<FCustomizableObjectStreamedResourceData>& StreamedResources = CustomizableObject->GetPrivate()->GetStreamedResourceData();
-
 	for (FCustomizableInstanceComponentData& ComponentData : ComponentsData)
 	{
 		for (int32 ResourceIndex : ComponentData.StreamedResourceIndex)
 		{
-			FCustomizableObjectStreamedResourceData& Resource = StreamedResources[ResourceIndex];
-			if (!Resource.IsLoaded())
+			const FCustomizableObjectResourceData* ResourceData = CustomizableObjectPrivate->LoadStreamedResource(ResourceIndex);
+			if (!ResourceData)
 			{
-				Resource.NotifyLoaded(Resource.GetPath().Get());
+				check(false); // Invalid resource index?
+				continue;
 			}
 
-			const FCustomizableObjectResourceData& ResourceData = Resource.GetLoadedData();
-			switch (ResourceData.Type)
+			switch (ResourceData->Type)
 			{
 				case ECOResourceDataType::AssetUserData:
 				{
-					const FCustomizableObjectAssetUserData* AUDResource = ResourceData.Data.GetPtr<FCustomizableObjectAssetUserData>();
+					const FCustomizableObjectAssetUserData* AUDResource = ResourceData->Data.GetPtr<FCustomizableObjectAssetUserData>();
 #if WITH_EDITORONLY_DATA
 					ComponentData.AssetUserDataArray.Add(AUDResource->AssetUserDataEditor);
 #else
@@ -5880,6 +5876,9 @@ void UCustomizableInstancePrivate::AdditionalAssetsAsyncLoaded(UCustomizableObje
 				default:
 					break;
 			}
+
+			// Unload by removing the reference to the container. Only if the platform has cooked data.
+			CustomizableObjectPrivate->UnloadStreamedResource(ResourceIndex);
 		}
 
 		// Loaded Skeletons
@@ -5984,15 +5983,6 @@ void UCustomizableInstancePrivate::AdditionalAssetsAsyncLoaded(UCustomizableObje
 	}
 
 	PassThroughMeshesToLoad.Empty();
-
-	// Only Unload in cooked builds. Unloading them when in the editor will trigger an assert. 
-	if (FPlatformProperties::RequiresCookedData())
-	{
-		for (FCustomizableObjectStreamedResourceData& ResourceData : StreamedResources)
-		{
-			ResourceData.Unload();
-		}
-	}
 }
 
 
diff --git a/Engine/Plugins/Experimental/Mutable/Source/CustomizableObject/Private/MuCO/CustomizableObjectSystem.cpp b/Engine/Plugins/Experimental/Mutable/Source/CustomizableObject/Private/MuCO/CustomizableObjectSystem.cpp
--- a/Engine/Plugins/Experimental/Mutable/Source/CustomizableObject/Private/MuCO/CustomizableObjectSystem.cpp
+++ b/Engine/Plugins/Experimental/Mutable/Source/CustomizableObject/Private/MuCO/CustomizableObjectSystem.cpp
@@ -991,7 +991,7 @@ void UCustomizableObjectSystemPrivate::GetMipStreamingConfig(const UCustomizable
 
 #if WITH_EDITORONLY_DATA
 	// Was streaming disabled at object-compilation time? 
-	if (Instance.GetCustomizableObject()->GetPrivate()->bDisableTextureStreaming)
+	if (Instance.GetCustomizableObject()->GetPrivate()->GetModelResources().bIsTextureStreamingDisabled)
 	{
 		bOutNeverStream = true;
 	}
@@ -4210,6 +4210,11 @@ int32 UCustomizableObjectSystem::GetWorkingMemory() const
 
 uint64 UCustomizableObjectSystem::GetMaxChunkSizeForPlatform(const ITargetPlatform* TargetPlatform)
 {
+	if (!TargetPlatform || !TargetPlatform->RequiresCookedData())
+	{
+		return MAX_uint64;
+	}
+
 	const FString& PlatformName = TargetPlatform ? TargetPlatform->IniPlatformName() : FPlatformProperties::IniPlatformName();
 
 	if (const int64* CachedMaxChunkSize = GetPrivate()->PlatformMaxChunkSize.Find(PlatformName))
diff --git a/Engine/Plugins/Experimental/Mutable/Source/CustomizableObject/Private/MuCO/UnrealMutableModelDiskStreamer.cpp b/Engine/Plugins/Experimental/Mutable/Source/CustomizableObject/Private/MuCO/UnrealMutableModelDiskStreamer.cpp
--- a/Engine/Plugins/Experimental/Mutable/Source/CustomizableObject/Private/MuCO/UnrealMutableModelDiskStreamer.cpp
+++ b/Engine/Plugins/Experimental/Mutable/Source/CustomizableObject/Private/MuCO/UnrealMutableModelDiskStreamer.cpp
@@ -12,6 +12,12 @@
 #include "HAL/PlatformFileManager.h"
 #include "Serialization/BulkData.h"
 
+#if WITH_EDITOR
+#include "DerivedDataCache.h"
+#include "DerivedDataCacheInterface.h"
+#include "DerivedDataCacheKey.h"
+#include "DerivedDataRequestOwner.h"
+#endif
 
 DECLARE_DWORD_ACCUMULATOR_STAT(TEXT("Streaming Ops"), STAT_MutableStreamingOps, STATGROUP_Mutable);
 
@@ -83,6 +89,8 @@ bool FUnrealMutableModelBulkReader::PrepareStreamingForObject(UCustomizableObjec
 		!=
 		nullptr;
 
+	const FModelResources& ModelResources = CustomizableObject->GetPrivate()->GetModelResources();
+
 	if (!bAlreadyStreaming)
 	{
 		FObjectData NewData;
@@ -103,6 +111,10 @@ bool FUnrealMutableModelBulkReader::PrepareStreamingForObject(UCustomizableObjec
 		FString FolderPath = CustomizableObject->GetPrivate()->GetCompiledDataFolderPath();
 		FString FullFileName = FolderPath + CustomizableObject->GetPrivate()->GetCompiledDataFileName(false, nullptr, true);
 		NewData.BulkFilePrefix = *FullFileName;
+
+		NewData.bIsStoredInDDC = ModelResources.bIsStoredInDDC;
+		NewData.DDCKey = ModelResources.DDCKey;
+		NewData.DDCPolicy = ModelResources.DDCDefaultPolicy;
 #else
 		if (NewData.ModelStreamableBulkData->HashToBulkData.IsEmpty())
 		{
@@ -218,7 +230,7 @@ mu::ModelReader::OPERATION_ID FUnrealMutableModelBulkReader::BeginReadBlock(cons
 		// object that have progressive mip generation.
 		if (CompletionCallback)
 		{
-			(*CompletionCallback)(false);    		
+			(*CompletionCallback)(false);
 		}
 		return -1;
 	}
@@ -232,7 +244,7 @@ mu::ModelReader::OPERATION_ID FUnrealMutableModelBulkReader::BeginReadBlock(cons
 
 		if (CompletionCallback)
 		{
-			(*CompletionCallback)(false);    		
+			(*CompletionCallback)(false);
 		}
 		return -1;
 	}
@@ -262,59 +274,121 @@ mu::ModelReader::OPERATION_ID FUnrealMutableModelBulkReader::BeginReadBlock(cons
 	}
 	else
 	{
-		int32 BulkDataOffsetInFile = 0;
 #if WITH_EDITOR
-		BulkDataOffsetInFile = sizeof(MutableCompiledDataStreamHeader);
-#endif
+		if (ObjectData->bIsStoredInDDC)
+		{
+			MUTABLE_CPUPROFILER_SCOPE(FUnrealMutableModelBulkStreamer::OpenReadFileDDC);
 
-		FReadRequest ReadRequest;
+			using namespace UE::DerivedData;
 
-		if (CompletionCallback)
-		{
-			ReadRequest.FileCallback = MakeShared<TFunction<void(bool, IAsyncReadRequest*)>>([CompletionCallbackCapture = *CompletionCallback](bool bWasCancelled, IAsyncReadRequest*) -> void
+			FValueId::ByteArray ValueIdBytes = {};
+			check(sizeof(FValueId::ByteArray) >= 4 + sizeof(Block->FileId));
+
+			MutablePrivate::EDataType DataType = MutablePrivate::EDataType::Model;
+			FMemory::Memcpy(&ValueIdBytes, &DataType, sizeof(DataType));
+			FMemory::Memcpy(&ValueIdBytes[4], &Block->FileId, sizeof(Block->FileId));
+			FValueId ResourceId = FValueId(ValueIdBytes);
+
+			// Skip loading values by default
+			FCacheRecordPolicyBuilder PolicyBuilder(ECachePolicy::Default | ECachePolicy::SkipData);
+
+			PolicyBuilder.AddValuePolicy(ResourceId, ECachePolicy::Default); // Override policy for the Model. Always load the compiled model data.
+
+			FCacheGetRequest Request;
+			Request.Name = ObjectData->BulkFilePrefix;
+			Request.Key = ObjectData->DDCKey;
+			Request.Policy = PolicyBuilder.Build();
+
+			FReadRequest ReadRequest;
+			ReadRequest.DDCReadRequest = MakeShared<FRequestOwner>(EPriority::High);
+
+			GetCache().Get(MakeArrayView(&Request, 1), *ReadRequest.DDCReadRequest.Get(),
+				[CompletionCallbackCapture = *CompletionCallback, Block, ResourceId, pBuffer, size](FCacheGetResponse&& Response)
 				{
-					CompletionCallbackCapture(!bWasCancelled);
+					bool bSuccess = Response.Status == EStatus::Ok;
+					if (bSuccess)
+					{
+						const FCompressedBuffer& CompressedBuffer = Response.Record.GetValue(ResourceId).GetData();
+						if (size < CompressedBuffer.GetRawSize())
+						{
+							check(CompressedBuffer.GetRawSize() >= Block->Offset + size);
+							FSharedBuffer DecompressedBuffer = CompressedBuffer.Decompress();
+							FMemory::Memcpy(pBuffer, reinterpret_cast<const uint8*>(DecompressedBuffer.GetData()) + Block->Offset, size);
+						}
+						else
+						{
+							check(CompressedBuffer.TryDecompressTo(MakeMemoryView(pBuffer, size)));
+						}
+					}
+					else
+					{
+						check(false);
+					}
+
+					CompletionCallbackCapture(bSuccess);
 				});
-		}
 
-		TSharedPtr<IAsyncReadFileHandle> FileHandle;
+			ObjectData->CurrentReadRequests.Add(Result, ReadRequest);
+		}
+		else
+#endif
 		{
-			FScopeLock Lock(&FileHandlesCritical);
 
-			TSharedPtr<IAsyncReadFileHandle>& Found = ObjectData->ReadFileHandles.FindOrAdd(Block->FileId);
-			if (!Found)
-			{
+			int32 BulkDataOffsetInFile = 0;
+
 #if WITH_EDITOR
-				FString FilePath = ObjectData->BulkFilePrefix;
-#else
-				FString FilePath = FString::Printf(TEXT("%s-%08x.mut"), *ObjectData->BulkFilePrefix, Block->FileId);
-				if (Block->Flags == uint32(mu::ERomFlags::HighRes))
-				{
-					FilePath += TEXT(".high");
-				}
+			BulkDataOffsetInFile = sizeof(MutableCompiledDataStreamHeader);
 #endif
 
-				Found = MakeShareable(FPlatformFileManager::Get().GetPlatformFile().OpenAsyncRead(*FilePath));
+			FReadRequest ReadRequest;
+
+			if (CompletionCallback)
+			{
+				ReadRequest.FileCallback = MakeShared<TFunction<void(bool, IAsyncReadRequest*)>>([CompletionCallbackCapture = *CompletionCallback](bool bWasCancelled, IAsyncReadRequest*) -> void
+					{
+						CompletionCallbackCapture(!bWasCancelled);
+					});
+			}
 
+			TSharedPtr<IAsyncReadFileHandle> FileHandle;
+			{
+				FScopeLock Lock(&FileHandlesCritical);
+
+				TSharedPtr<IAsyncReadFileHandle>& Found = ObjectData->ReadFileHandles.FindOrAdd(Block->FileId);
 				if (!Found)
 				{
-					UE_LOG(LogMutable, Error, TEXT("Failed to create AsyncReadFileHandle. File Path [%s]."), *FilePath);
-					check(false);
-					return -1;
+#if WITH_EDITOR
+					FString FilePath = ObjectData->BulkFilePrefix;
+#else
+					FString FilePath = FString::Printf(TEXT("%s-%08x.mut"), *ObjectData->BulkFilePrefix, Block->FileId);
+					if (Block->Flags == uint32(mu::ERomFlags::HighRes))
+					{
+						FilePath += TEXT(".high");
+					}
+#endif
+
+					Found = MakeShareable(FPlatformFileManager::Get().GetPlatformFile().OpenAsyncRead(*FilePath));
+
+					if (!Found)
+					{
+						UE_LOG(LogMutable, Error, TEXT("Failed to create AsyncReadFileHandle. File Path [%s]."), *FilePath);
+						check(false);
+						return -1;
+					}
 				}
+
+				FileHandle = Found;
 			}
 
-			FileHandle = Found;
-		}
+			ReadRequest.FileReadRequest = MakeShareable(FileHandle->ReadRequest(
+				BulkDataOffsetInFile + Block->Offset,
+				size,
+				(EAsyncIOPriorityAndFlags)StreamPriority,
+				ReadRequest.FileCallback.Get(),
+				reinterpret_cast<uint8*>(pBuffer)));
 
-		ReadRequest.FileReadRequest = MakeShareable(FileHandle->ReadRequest(
-			BulkDataOffsetInFile + Block->Offset,
-			size,
-			(EAsyncIOPriorityAndFlags)StreamPriority,
-			ReadRequest.FileCallback.Get(),
-			reinterpret_cast<uint8*>(pBuffer)));
-
-		ObjectData->CurrentReadRequests.Add(Result, ReadRequest);
+			ObjectData->CurrentReadRequests.Add(Result, ReadRequest);
+		}
 	}
 
 	INC_DWORD_STAT(STAT_MutableStreamingOps);
@@ -335,10 +409,16 @@ bool FUnrealMutableModelBulkReader::IsReadCompleted(mu::ModelReader::OPERATION_I
 			{
 				return ReadRequest->FileReadRequest->PollCompletion();
 			}
-			else if(ReadRequest->BulkReadRequest)
+			else if (ReadRequest->BulkReadRequest)
 			{
 				return ReadRequest->BulkReadRequest->PollCompletion();
 			}
+#if WITH_EDITORONLY_DATA
+			else if (ReadRequest->DDCReadRequest)
+			{
+				return ReadRequest->DDCReadRequest->Poll(); // ReadLock
+			}
+#endif
 		}
 	}
 
@@ -359,36 +439,37 @@ bool FUnrealMutableModelBulkReader::EndRead(mu::ModelReader::OPERATION_ID Operat
 		FReadRequest* ReadRequest = o.CurrentReadRequests.Find(OperationId);
 		if (ReadRequest)
 		{
+			bool bCompleted = false;
+			bool bResult = false;
+
 			if (ReadRequest->FileReadRequest)
 			{
-				bool bCompleted = ReadRequest->FileReadRequest->WaitCompletion();
-				if (!bCompleted)
-				{
-					UE_LOG(LogMutable, Error, TEXT("Operation failed to complete in EndRead."));
-					check(false);
-					bSuccess = false;
-				}
-				else if (ReadRequest->FileReadRequest->GetReadResults()==nullptr)
-				{
-					// This means we have failed: file not found?
-					bSuccess = false;
-				}
+				bCompleted = ReadRequest->FileReadRequest->WaitCompletion();
+				bResult = ReadRequest->FileReadRequest->GetReadResults() != nullptr;
 			}
 			else if (ReadRequest->BulkReadRequest)
 			{
-				bool bCompleted = ReadRequest->BulkReadRequest->WaitCompletion();
-				if (!bCompleted)
-				{
-					UE_LOG(LogMutable, Error, TEXT("Operation failed to complete in EndRead."));
-					check(false);
-					bSuccess = false;
-				}
-				else if (ReadRequest->BulkReadRequest->GetReadResults() == nullptr)
-				{
-					// This means we have failed: file not found?
-					bSuccess = false;
-				}
+				bCompleted = ReadRequest->BulkReadRequest->WaitCompletion();
+				bResult = ReadRequest->BulkReadRequest->GetReadResults() != nullptr;
+			}
+#if WITH_EDITORONLY_DATA
+			else if (ReadRequest->DDCReadRequest)
+			{
+				ReadRequest->DDCReadRequest->Wait();
+				
+				// bSuccess TODO Pere: find a way to know if the request completed successfully
+				bCompleted = true;
+				bResult = true;
 			}
+#endif
+			if (!bCompleted)
+			{
+				UE_LOG(LogMutable, Error, TEXT("Operation failed to complete in EndRead."));
+				check(false);
+			}
+
+			bSuccess = bCompleted && bResult;
+			
 			o.CurrentReadRequests.Remove(OperationId);
 			bFound = true;
 			break;
diff --git a/Engine/Plugins/Experimental/Mutable/Source/CustomizableObject/Public/MuCO/CustomizableObject.h b/Engine/Plugins/Experimental/Mutable/Source/CustomizableObject/Public/MuCO/CustomizableObject.h
--- a/Engine/Plugins/Experimental/Mutable/Source/CustomizableObject/Public/MuCO/CustomizableObject.h
+++ b/Engine/Plugins/Experimental/Mutable/Source/CustomizableObject/Public/MuCO/CustomizableObject.h
@@ -24,6 +24,7 @@ class UCustomizableObject;
 class UEdGraph;
 class USkeletalMesh;
 class UCustomizableObjectPrivate;
+class UCustomizableObjectBulk;
 struct FFrame;
 struct FStreamableHandle;
 template <typename FuncType> class TFunctionRef;
@@ -199,32 +200,6 @@ struct FCustomizableObjectMeshToMeshVertData
 static_assert(sizeof(FCustomizableObjectMeshToMeshVertData) == sizeof(float)*4*3 + sizeof(uint16)*4 + sizeof(float) + sizeof(uint32));
 template<> struct TCanBulkSerialize<FCustomizableObjectMeshToMeshVertData> { enum { Value = true }; };
 
-USTRUCT()
-struct CUSTOMIZABLEOBJECT_API FMutableStreamableBlock
-{
-	GENERATED_USTRUCT_BODY()
-
-	UPROPERTY()
-	uint32 FileId = 0;
-
-	/** Used to store properties of the data, necessary for its recovery. For instance if it is high-res. */
-	UPROPERTY()
-	uint32 Flags = 0;
-
-	UPROPERTY()
-	uint64 Offset = 0;
-
-	friend FArchive& operator<<(FArchive& Ar, FMutableStreamableBlock& Data)
-	{
-		Ar << Data.FileId;
-		Ar << Data.Flags;
-		Ar << Data.Offset;
-		return Ar;
-	}
-};
-template<> struct TCanBulkSerialize<FMutableStreamableBlock> { enum { Value = true }; };
-static_assert(sizeof(FMutableStreamableBlock) == 8*2);
-
 USTRUCT()
 struct FMutableLODSettings
 {
@@ -259,37 +234,6 @@ struct FMutableLODSettings
 };
 
 
-UCLASS( config=Engine )
-class CUSTOMIZABLEOBJECT_API UCustomizableObjectBulk : public UObject
-{
-public:
-	GENERATED_BODY()
-
-	//~ Begin UObject Interface
-	virtual void PostLoad() override;
-	//~ End UObject Interface
-
-	/**  */
-	const FString& GetBulkFilePrefix() const { return BulkFilePrefix; }
-	
-	TUniquePtr<IAsyncReadFileHandle> OpenFileAsyncRead(uint32 FileId, uint32 Flags) const;
-
-#if WITH_EDITOR
-
-	//~ Begin UObject Interface
-	virtual void CookAdditionalFilesOverride(const TCHAR*, const ITargetPlatform*, TFunctionRef<void(const TCHAR*, void*, int64)> ) override;
-	//~ End UObject Interface
-#endif
-
-#if WITH_EDITOR
-private:
-#endif
-
-	/** Prefix to locate bulkfiles for loading, using the file ids in each FMutableStreamableBlock. */
-	FString BulkFilePrefix;
-};
-
-
 UCLASS( BlueprintType, config=Engine )
 class CUSTOMIZABLEOBJECT_API UCustomizableObject : public UObject
 {
@@ -340,10 +284,6 @@ private:
 	UPROPERTY()
 	TArray<FCustomizableObjectStreamedResourceData> StreamedExtensionData;
 
-	// Constant Resources streamed in on demand when generating meshes
-	UPROPERTY()
-	TArray<FCustomizableObjectStreamedResourceData> StreamedResourceData;
-
 public:
 	/** Use the SkeletalMesh of reference as a placeholder until the custom mesh is ready to use.
 	  * Note: If disabled, a null mesh will be used to replace the discarded mesh due to 'ReplaceDiscardedWithReferenceMesh' being enabled. */
diff --git a/Engine/Plugins/Experimental/Mutable/Source/CustomizableObjectEditor/CustomizableObjectEditor.Build.cs b/Engine/Plugins/Experimental/Mutable/Source/CustomizableObjectEditor/CustomizableObjectEditor.Build.cs
--- a/Engine/Plugins/Experimental/Mutable/Source/CustomizableObjectEditor/CustomizableObjectEditor.Build.cs
+++ b/Engine/Plugins/Experimental/Mutable/Source/CustomizableObjectEditor/CustomizableObjectEditor.Build.cs
@@ -88,6 +88,7 @@ public class CustomizableObjectEditor : ModuleRules
 			new string[] {
 				"PropertyEditor",
 				"CustomizableObject",
+				"DerivedDataCache",
 				"MutableTools",
 			}
 		);
diff --git a/Engine/Plugins/Experimental/Mutable/Source/CustomizableObjectEditor/Private/MuCOE/CustomizableInstanceDetails.cpp b/Engine/Plugins/Experimental/Mutable/Source/CustomizableObjectEditor/Private/MuCOE/CustomizableInstanceDetails.cpp
--- a/Engine/Plugins/Experimental/Mutable/Source/CustomizableObjectEditor/Private/MuCOE/CustomizableInstanceDetails.cpp
+++ b/Engine/Plugins/Experimental/Mutable/Source/CustomizableObjectEditor/Private/MuCOE/CustomizableInstanceDetails.cpp
@@ -989,9 +989,9 @@ TSharedRef<SWidget> FCustomizableInstanceDetails::GenerateIntWidget(const UCusto
 
 	// Tooltip for the selected option
 	FString ToolTipText = FString("None");
-	if (const FString* Identifier = CustomizableObject.GetPrivate()->GroupNodeMap.FindKey(FCustomizableObjectIdPair(ParamName, SelectedOption)))
+	if (const FString* Identifier = CustomizableObject.GetPrivate()->GetModelResources().GroupNodeMap.FindKey(FCustomizableObjectIdPair(ParamName, SelectedOption)))
 	{
-		if (FString* CustomizableObjectPath = CustomizableObject.GetPrivate()->CustomizableObjectPathMap.Find(*Identifier))
+		if (const FString* CustomizableObjectPath = CustomizableObject.GetPrivate()->GetModelResources().CustomizableObjectPathMap.Find(*Identifier))
 		{
 			ToolTipText = *CustomizableObjectPath;
 		}
diff --git a/Engine/Plugins/Experimental/Mutable/Source/CustomizableObjectEditor/Private/MuCOE/CustomizableObjectCompileRunnable.cpp b/Engine/Plugins/Experimental/Mutable/Source/CustomizableObjectEditor/Private/MuCOE/CustomizableObjectCompileRunnable.cpp
--- a/Engine/Plugins/Experimental/Mutable/Source/CustomizableObjectEditor/Private/MuCOE/CustomizableObjectCompileRunnable.cpp
+++ b/Engine/Plugins/Experimental/Mutable/Source/CustomizableObjectEditor/Private/MuCOE/CustomizableObjectCompileRunnable.cpp
@@ -15,6 +15,12 @@
 #include "Containers/Ticker.h"
 #include "Serialization/ObjectAndNameAsStringProxyArchive.h"
 #include "Trace/Trace.inl"
+#include "UObject/SoftObjectPtr.h"
+#include "Engine/AssetUserData.h"
+
+#include "DerivedDataCacheInterface.h"
+#include "DerivedDataCache.h"
+#include "DerivedDataRequestOwner.h"
 
 class ITargetPlatform;
 
@@ -267,133 +273,86 @@ void FCustomizableObjectCompileRunnable::Tick()
 }
 
 
-FCustomizableObjectSaveDDRunnable::FCustomizableObjectSaveDDRunnable(UCustomizableObject* CustomizableObject, const FCompilationOptions& InOptions, TSharedPtr<mu::Model> InModel)
+FCustomizableObjectSaveDDRunnable::FCustomizableObjectSaveDDRunnable(const TSharedPtr<FCompilationRequest>& InRequest,
+	TSharedPtr<mu::Model> InModel,
+	FModelResources& InModelResources,
+	TSharedPtr<FModelStreamableBulkData> InModelStreamables)
 {
 	MUTABLE_CPUPROFILER_SCOPE(FCustomizableObjectSaveDDRunnable::FCustomizableObjectSaveDDRunnable)
-		
+
 	Model = InModel;
-	Options = InOptions;
-	
+	ModelStreamables = InModelStreamables;
+
+	Options = InRequest->GetCompileOptions();
+
+	DDCKey = InRequest->GetDerivedDataCacheKey();
+	DefaultDDCPolicy = InRequest->GetDerivedDataCachePolicy();
+
+	UCustomizableObject* CustomizableObject = InRequest->GetCustomizableObject();
+	CustomizableObjectName = GetNameSafe(CustomizableObject);
+
 	CustomizableObjectHeader.InternalVersion = CustomizableObject->GetPrivate()->CurrentSupportedVersion;
-	CustomizableObjectHeader.VersionId = Options.bIsCooking? FGuid::NewGuid() : CustomizableObject->GetPrivate()->GetVersionId();
+	CustomizableObjectHeader.VersionId = CustomizableObject->GetPrivate()->GetVersionId();
 
-	if (!Options.bIsCooking)
+	// Cache ModelResources
 	{
-		// We will be saving all compilation data in two separate files, write CO Data
-		FolderPath = CustomizableObject->GetPrivate()->GetCompiledDataFolderPath();
-		CompileDataFullFileName = FolderPath + CustomizableObject->GetPrivate()->GetCompiledDataFileName(true, InOptions.TargetPlatform);
-		StreamableDataFullFileName = FolderPath + CustomizableObject->GetPrivate()->GetCompiledDataFileName(false, InOptions.TargetPlatform);
-
-		// Serialize Customizable Object's data
-		FMemoryWriter64 MemoryWriter(ModelBytes);
+		FMemoryWriter64 MemoryWriter(PlatformData.ModelResourcesData);
 		FObjectAndNameAsStringProxyArchive ObjectWriter(MemoryWriter, true);
-		CustomizableObject->GetPrivate()->SaveCompiledData(ObjectWriter, Options.bIsCooking);
+		InModelResources.Serialize(ObjectWriter, Options.bIsCooking);
 	}
-#if WITH_EDITORONLY_DATA
-	else
+
+	// Cache Morphs and Clothing. 
 	{
 		// Do a copy of the Morph and Clothing Data generated at compile time. Only needed when cooking.
-		
-		constexpr bool bGetCookedFalse = false;
-		FModelResources& ModelResources = CustomizableObject->GetPrivate()->GetModelResources(bGetCookedFalse);
-		
+		// Morphs
 		static_assert(TCanBulkSerialize<FMorphTargetVertexData>::Value);
-		const TArray<FMorphTargetVertexData>& MorphVertexData = ModelResources.EditorOnlyMorphTargetReconstructionData;
+		TArray<FMorphTargetVertexData>& MorphVertexData = InModelResources.EditorOnlyMorphTargetReconstructionData;
 
-		MorphDataBytes.SetNum(MorphVertexData.Num() * sizeof(FMorphTargetVertexData));
-		FMemory::Memcpy(MorphDataBytes.GetData(), MorphVertexData.GetData(), MorphDataBytes.Num());
-		
+		PlatformData.MorphData.SetNum(MorphVertexData.Num() * sizeof(FMorphTargetVertexData));
+		FMemory::Memcpy(PlatformData.MorphData.GetData(), MorphVertexData.GetData(), PlatformData.MorphData.Num());
+
+		// Cloth
 		static_assert(TCanBulkSerialize<FCustomizableObjectMeshToMeshVertData>::Value);
-		const TArray<FCustomizableObjectMeshToMeshVertData>& ClothingVertexData = ModelResources.EditorOnlyClothingMeshToMeshVertData;
+		TArray<FCustomizableObjectMeshToMeshVertData>& ClothingVertexData = InModelResources.EditorOnlyClothingMeshToMeshVertData;
+
+		PlatformData.ClothingData.SetNum(ClothingVertexData.Num() * sizeof(FCustomizableObjectMeshToMeshVertData));
+		FMemory::Memcpy(PlatformData.ClothingData.GetData(), ClothingVertexData.GetData(), PlatformData.ClothingData.Num());
+
+		if (Options.bIsCooking)
+		{
+			InModelResources.EditorOnlyMorphTargetReconstructionData.Empty();
+			InModelResources.EditorOnlyClothingMeshToMeshVertData.Empty();
+		}
+	}
 
-		ClothingDataBytes.SetNum(ClothingVertexData.Num() * sizeof(FCustomizableObjectMeshToMeshVertData));
-		FMemory::Memcpy(ClothingDataBytes.GetData(), ClothingVertexData.GetData(), ClothingDataBytes.Num());
+	if (!Options.bIsCooking)
+	{
+		// We will be saving all compilation data in two separate files, write CO Data
+		FolderPath = CustomizableObject->GetPrivate()->GetCompiledDataFolderPath();
+		CompileDataFullFileName = FolderPath + CustomizableObject->GetPrivate()->GetCompiledDataFileName(true, Options.TargetPlatform);
+		StreamableDataFullFileName = FolderPath + CustomizableObject->GetPrivate()->GetCompiledDataFileName(false, Options.TargetPlatform);
 	}
-#endif // WITH_EDITORONLY_DATA
 }
 
-
 uint32 FCustomizableObjectSaveDDRunnable::Run()
 {
 	MUTABLE_CPUPROFILER_SCOPE(FCustomizableObjectSaveDDRunnable::Run)
 
-	// MorphDataBytes and ClothingDataBytes has data only if cooking. 
-	check(!!Options.bIsCooking || MorphDataBytes.IsEmpty());
-	check(!!Options.bIsCooking || ClothingDataBytes.IsEmpty());
-
-	bool bModelSerialized = Model.Get() != nullptr;
-
-	if (Options.bIsCooking)
+	if (Model)
 	{
-		// Serialize mu::Model and streamable resources 
-		FMemoryWriter64 ModelMemoryWriter(ModelBytes, false, true);
-		ModelMemoryWriter << bModelSerialized;
-		if (bModelSerialized)
-		{
-			FUnrealMutableModelBulkWriterCook Streamer(&ModelMemoryWriter, &ModelStreamableData);
-			constexpr bool bDropData = true;
-			mu::Model::Serialise(Model.Get(), Streamer, bDropData);
+		CachePlatfromData();
 
-			// Morph and Clothing are already in the corresponding buffer copied from the compilation thread.
-		}
-	}
-	else if (bModelSerialized) // Save CO data + mu::Model and streamable resources to disk
-	{
-		// Create folder...
-		IFileManager& FileManager = IFileManager::Get();
-		FileManager.MakeDirectory(*FolderPath, true);
-
-		// Delete files...
-		bool bFilesDeleted = true;
-		if (FileManager.FileExists(*CompileDataFullFileName)
-			&& !FileManager.Delete(*CompileDataFullFileName, true, false, true))
-		{
-			UE_LOG(LogMutable, Error, TEXT("Failed to delete compiled data in file [%s]."), *CompileDataFullFileName);
-			bFilesDeleted = false;
-		}
+		bool bStoredSuccessfully = false;
 
-		if (FileManager.FileExists(*StreamableDataFullFileName)
-			&& !FileManager.Delete(*StreamableDataFullFileName, true, false, true))
+		// TODO UE-222775: Allow using DDC in editor builds, not just for cooking.
+		if (Options.bIsCooking && Options.bStoreCompiledDataInDDC && !DDCKey.Hash.IsZero())
 		{
-			UE_LOG(LogMutable, Error, TEXT("Failed to delete streamed data in file [%s]."), *StreamableDataFullFileName);
-			bFilesDeleted = false;
+			StoreCachedPlatformDataInDDC(bStoredSuccessfully);
 		}
 
-		// Store current compiled data
-		if (bFilesDeleted)
+		if (!Options.bIsCooking && !bStoredSuccessfully)
 		{
-			// Create file writers...
-			TUniquePtr<FArchive> ModelMemoryWriter(FileManager.CreateFileWriter(*CompileDataFullFileName));
-			TUniquePtr<FArchive> StreamableMemoryWriter(FileManager.CreateFileWriter(*StreamableDataFullFileName));
-			check(ModelMemoryWriter);
-			check(StreamableMemoryWriter);
-
-			// Serailize headers to validate data
-			*ModelMemoryWriter << CustomizableObjectHeader;
-			*StreamableMemoryWriter << CustomizableObjectHeader;
-
-			// Serialize Customizable Object's Data to disk
-			ModelMemoryWriter->Serialize(ModelBytes.GetData(), ModelBytes.Num() * sizeof(uint8));
-			ModelBytes.Empty();
-
-			// Serialize mu::Model and streamable resources
-			*ModelMemoryWriter << bModelSerialized;
-
-			FUnrealMutableModelBulkWriterEditor Streamer(ModelMemoryWriter.Get(), StreamableMemoryWriter.Get());
-			constexpr bool bDropData = true;
-			mu::Model::Serialise(Model.Get(), Streamer, bDropData);
-
-			// Save to disk
-			ModelMemoryWriter->Flush();
-			StreamableMemoryWriter->Flush();
-
-			ModelMemoryWriter->Close();
-			StreamableMemoryWriter->Close();
-		}
-		else
-		{
-			// Remove old data if there.
-			Model.Reset();
+			StoreCachedPlatformDataToDisk(bStoredSuccessfully);
 		}
 	}
 
@@ -414,5 +373,254 @@ const ITargetPlatform* FCustomizableObjectSaveDDRunnable::GetTargetPlatform() co
 	return Options.TargetPlatform;
 }
 
+
+void FCustomizableObjectSaveDDRunnable::CachePlatfromData()
+{
+	MUTABLE_CPUPROFILER_SCOPE(CachePlatfromData);
+
+	if (!Model || !ModelStreamables)
+	{
+		check(false);
+		return;
+	}
+
+	// Cache ModelStreamalbes
+	{
+		PlatformData.ModelStreamables = ModelStreamables;
+
+		const uint64 PackageDataBytesLimit = Options.bIsCooking ? Options.PackagedDataBytesLimit : MAX_uint64;
+
+		// Generate list of files and update streamable blocks ids and offsets
+		MutablePrivate::GenerateBulkDataFilesList(Model, *PlatformData.ModelStreamables.Get(), Options.TargetPlatform, PackageDataBytesLimit, PlatformData.BulkDataFiles);
+	}
+
+
+	// Cache Model and Model Roms
+	{
+		FMemoryWriter64 ModelMemoryWriter(PlatformData.ModelData);
+		FUnrealMutableModelBulkWriterCook Streamer(&ModelMemoryWriter, &PlatformData.ModelStreamableData);
+
+		// Serialize mu::Model and streamable resources 
+		constexpr bool bDropData = true;
+		mu::Model::Serialise(Model.Get(), Streamer, bDropData);
+	}
+}
+
+
+void FCustomizableObjectSaveDDRunnable::StoreCachedPlatformDataInDDC(bool& bStoredSuccessfully)
+{
+	MUTABLE_CPUPROFILER_SCOPE(StoreCachedPlatformDataInDDC);
+
+	using namespace UE::DerivedData;
+
+	check(Model.Get() != nullptr);
+	check(DDCKey.Hash.IsZero() == false);
+
+	bStoredSuccessfully = false;
+
+	// DDC record
+	FCacheRecordBuilder RecordBuilder(DDCKey);
+
+	// Store streamable resources info as FValues
+	{
+		MUTABLE_CPUPROFILER_SCOPE(SerializeModelStreamables);
+
+		// ModelStreamable  will be modified for the DDC record. Modify a copy
+		FModelStreamableBulkData ModelStreamablesDDC = *ModelStreamables.Get();
+
+		// Generate list of files and update streamable blocks ids and offsets
+		MutablePrivate::GenerateBulkDataFilesList(Model, ModelStreamablesDDC, Options.TargetPlatform, Options.DDCBytesLimit, BulkDataFilesDDC);
+
+		TArray64<uint8> ModelStreamablesBytesDDC;
+		FMemoryWriter64 MemoryWriterDDC(ModelStreamablesBytesDDC);
+		MemoryWriterDDC << ModelStreamablesDDC;
+
+		const FValue ModelStreamablesValue = FValue::Compress(FSharedBuffer::MakeView(ModelStreamablesBytesDDC.GetData(), ModelStreamablesBytesDDC.Num()));
+		RecordBuilder.AddValue(MutablePrivate::GetDerivedDataModelStreamableBulkDataId(), ModelStreamablesValue);
+	}
+
+	// Store streamable resources as FValues
+	{
+		MUTABLE_CPUPROFILER_SCOPE(SerializeBulkDataForDDC);
+		
+		FValueId::ByteArray ValueIdBytes = {};
+		const auto WriteBulkDataDDC = [&RecordBuilder, &ValueIdBytes](MutablePrivate::FFile& File, TArray64<uint8>& FileBulkData)
+			{
+				FMemory::Memcpy(&ValueIdBytes, &File.DataType, sizeof(File.DataType));
+				FMemory::Memcpy(&ValueIdBytes[4], &File.Id, sizeof(File.Id));
+				const FValueId ValueId(ValueIdBytes);
+				const FValue Value = FValue::Compress(FSharedBuffer::MakeView(FileBulkData.GetData(), FileBulkData.Num()));
+
+				RecordBuilder.AddValue(ValueId, Value);
+			};
+
+		constexpr bool bDropData = false;
+		MutablePrivate::SerializeBulkDataFiles(PlatformData, BulkDataFilesDDC, WriteBulkDataDDC, bDropData);
+	}
+
+
+	// Store BulkData Files as a FValue to reconstruct the data later on
+	{
+		MUTABLE_CPUPROFILER_SCOPE(SerializeBulkDataFilesForDDC);
+
+		TArray<uint8> BulkDataFilesBytes;
+		FMemoryWriter MemoryWriter(BulkDataFilesBytes);
+
+		MemoryWriter << BulkDataFilesDDC;
+
+		const FValue BulkDataFilesValue = FValue::Compress(FSharedBuffer::MakeView(BulkDataFilesBytes.GetData(), BulkDataFilesBytes.Num()));
+		RecordBuilder.AddValue(MutablePrivate::GetDerivedDataBulkDataFilesId(), BulkDataFilesValue);
+	}
+
+	// Store ModelResources bytes as a FValue
+	{
+		MUTABLE_CPUPROFILER_SCOPE(SerializeModelResourcesForDDC);
+
+		const FValue ModelResourcesValue = FValue::Compress(FSharedBuffer::MakeView(PlatformData.ModelResourcesData.GetData(), PlatformData.ModelResourcesData.Num()));
+		RecordBuilder.AddValue(MutablePrivate::GetDerivedDataModelResourcesId(), ModelResourcesValue);
+	}
+
+	// Store Model bytes as a FValue
+	{
+		MUTABLE_CPUPROFILER_SCOPE(SerializeModelForDDC);
+
+		const FValue ModelValue = FValue::Compress(FSharedBuffer::MakeView(PlatformData.ModelData.GetData(), PlatformData.ModelData.Num()));
+		RecordBuilder.AddValue(MutablePrivate::GetDerivedDataModelId(), ModelValue);
+	}
+
+	// Push record to the DDC
+	{
+		MUTABLE_CPUPROFILER_SCOPE(PushRecordToDDC);
+
+		FRequestOwner RequestOwner(UE::DerivedData::EPriority::Blocking);
+		const FCachePutRequest PutRequest = { UE::FSharedString(CustomizableObjectName), RecordBuilder.Build(), DefaultDDCPolicy };
+		GetCache().Put(MakeArrayView(&PutRequest, 1), RequestOwner,
+			[&bStoredSuccessfully](FCachePutResponse&& Response)
+			{
+				if (Response.Status == EStatus::Ok)
+				{
+					bStoredSuccessfully = true;
+				}
+			});
+
+		RequestOwner.Wait();
+		check(bStoredSuccessfully);
+	}
+}
+
+
+void FCustomizableObjectSaveDDRunnable::StoreCachedPlatformDataToDisk(bool& bStoredSuccessfully)
+{
+	MUTABLE_CPUPROFILER_SCOPE(StoreCachedPlatformDataToDisk);
+
+	check(Model.Get() != nullptr);
+	check(!Options.bIsCooking);
+
+	bStoredSuccessfully = false;
+
+	// Create folder...
+	IFileManager& FileManager = IFileManager::Get();
+	FileManager.MakeDirectory(*FolderPath, true);
+
+	// Delete files...
+	
+	
+	bool bFilesDeleted = true;
+	if (FileManager.FileExists(*CompileDataFullFileName)
+		&& !FileManager.Delete(*CompileDataFullFileName, true, false, true))
+	{
+		UE_LOG(LogMutable, Error, TEXT("Failed to delete compiled data in file [%s]."), *CompileDataFullFileName);
+		bFilesDeleted = false;
+	}
+
+	if (FileManager.FileExists(*StreamableDataFullFileName)
+		&& !FileManager.Delete(*StreamableDataFullFileName, true, false, true))
+	{
+		UE_LOG(LogMutable, Error, TEXT("Failed to delete streamed data in file [%s]."), *StreamableDataFullFileName);
+		bFilesDeleted = false;
+	}
+
+	if (!bFilesDeleted)
+	{
+		// Couldn't delete files. Delete model and return.
+		Model.Reset();
+		return;
+	}
+
+	// Serialize Streamable resources
+	{
+		// Create file writer
+		TUniquePtr<FArchive> StreamableMemoryWriter(FileManager.CreateFileWriter(*StreamableDataFullFileName));
+		check(StreamableMemoryWriter);
+
+		// Serailize headers to validate data
+		*StreamableMemoryWriter << CustomizableObjectHeader;
+
+		const auto WriteBulkDataToDisk = [&StreamableMemoryWriter](MutablePrivate::FFile& File, TArray64<uint8>& FileBulkData)
+			{
+				switch (File.DataType)
+				{
+				case MutablePrivate::EDataType::Model:
+				{
+					StreamableMemoryWriter->Serialize(FileBulkData.GetData(), FileBulkData.Num() * sizeof(uint8));
+					break;
+				}
+				case MutablePrivate::EDataType::RealTimeMorph:
+				{
+					// TODO: Store clothing streamable to disk. UE-222777
+					break;
+				}
+				case MutablePrivate::EDataType::Clothing:
+				{
+					// TODO: Store clothing streamable to disk. UE-222777
+					break;
+				}
+				default:
+					break;
+				}
+				
+			};
+
+		// Serialize streamable resources into a single file and fix offsets
+		constexpr bool bDropData = true;
+		MutablePrivate::SerializeBulkDataFiles(PlatformData, PlatformData.BulkDataFiles, WriteBulkDataToDisk, bDropData);
+		StreamableMemoryWriter->Flush();
+		StreamableMemoryWriter->Close();
+
+		PlatformData.MorphData.Empty();
+		PlatformData.ClothingData.Empty();
+	}
+
+	// Serialize Model and ModelResources. Store after SerializeBulkDataFiles fixes the HashToStreamableFiles offsets.
+	{
+		// Create file writer
+		TUniquePtr<FArchive> ModelMemoryWriter(FileManager.CreateFileWriter(*CompileDataFullFileName));
+		check(ModelMemoryWriter);
+
+		// Serailize headers to validate data
+		*ModelMemoryWriter << CustomizableObjectHeader;
+
+		ModelMemoryWriter->Serialize(PlatformData.ModelResourcesData.GetData(), PlatformData.ModelResourcesData.Num() * sizeof(uint8));
+
+		{
+			// ModelMemoryWriter (Writer to disk) doesn't handle FNames properly. Serialize them ModelStreamables in two steps.
+			TArray64<uint8> ModelStreamablesBytes;
+			FMemoryWriter64 ModelStreamablesMemoryWriter(ModelStreamablesBytes);
+			ModelStreamablesMemoryWriter << *PlatformData.ModelStreamables.Get();
+			ModelMemoryWriter->Serialize(ModelStreamablesBytes.GetData(), ModelStreamablesBytes.Num() * sizeof(uint8));
+		}
+
+		ModelMemoryWriter->Serialize(PlatformData.ModelData.GetData(), PlatformData.ModelData.Num() * sizeof(uint8));
+		
+		ModelMemoryWriter->Flush();
+		ModelMemoryWriter->Close();
+
+		PlatformData.ModelData.Empty();
+	}
+
+	bStoredSuccessfully = true;
+}
+
+
 #undef LOCTEXT_NAMESPACE
 
diff --git a/Engine/Plugins/Experimental/Mutable/Source/CustomizableObjectEditor/Private/MuCOE/CustomizableObjectCompileRunnable.h b/Engine/Plugins/Experimental/Mutable/Source/CustomizableObjectEditor/Private/MuCOE/CustomizableObjectCompileRunnable.h
--- a/Engine/Plugins/Experimental/Mutable/Source/CustomizableObjectEditor/Private/MuCOE/CustomizableObjectCompileRunnable.h
+++ b/Engine/Plugins/Experimental/Mutable/Source/CustomizableObjectEditor/Private/MuCOE/CustomizableObjectCompileRunnable.h
@@ -10,11 +10,12 @@
 #include "MuT/Node.h"
 #include "MuCOE/CustomizableObjectEditorLogger.h"
 
+#include "DerivedDataCacheKey.h"
+#include "DerivedDataCachePolicy.h"
+
 #include <atomic>
 
 class ITargetPlatform;
-class UCustomizableObject;
-
 
 class FCustomizableObjectCompileRunnable : public FRunnable
 {
@@ -89,7 +90,7 @@ class FCustomizableObjectSaveDDRunnable : public FRunnable
 {
 public:
 
-	FCustomizableObjectSaveDDRunnable(UCustomizableObject* CustomizableObject, const FCompilationOptions& Options, TSharedPtr<mu::Model> InModel);
+	FCustomizableObjectSaveDDRunnable(const TSharedPtr<FCompilationRequest>& InRequest, TSharedPtr<mu::Model> InModel, FModelResources& ModelResources, TSharedPtr<FModelStreamableBulkData> ModelStreamables);
 
 	// FRunnable interface
 	uint32 Run() override;
@@ -101,16 +102,26 @@ public:
 	const ITargetPlatform* GetTargetPlatform() const;
 
 private:
+
+	void CachePlatfromData();
+
+	void StoreCachedPlatformDataInDDC(bool& bStoredSuccessfully);
+
+	void StoreCachedPlatformDataToDisk(bool& bStoredSuccessfully);
+
 	FCompilationOptions Options;
 
 	MutableCompiledDataStreamHeader CustomizableObjectHeader;
 
+	FString CustomizableObjectName;
+
 	// Paths used to save files to disk
 	FString FolderPath;
 	FString CompileDataFullFileName;
 	FString StreamableDataFullFileName;
 
-	bool bIsCooking = false;
+	UE::DerivedData::FCacheKey DDCKey;
+	UE::DerivedData::ECachePolicy DefaultDDCPolicy;
 
 	// Whether the thread has finished running
 	std::atomic<bool> bThreadCompleted = false;
@@ -118,16 +129,11 @@ private:
 public:
 
 	TSharedPtr<mu::Model, ESPMode::ThreadSafe> Model;
+	TSharedPtr<FModelStreamableBulkData> ModelStreamables;
 
-	// Bytes where the model is stored
-	TArray64<uint8> ModelBytes;
-
-	// Model streamed data
-	MutablePrivate::FModelStreamableData ModelStreamableData;
-
-	// Bytes store streameable files coming form the CO itself.
-	TArray64<uint8> MorphDataBytes;
+	// Cached platform data
+	MutablePrivate::FMutableCachedPlatformData PlatformData;
 
-	// Bytes store streameable files coming form the CO itself.
-	TArray64<uint8> ClothingDataBytes;
+	// DDC Helpers
+	TArray<MutablePrivate::FFile> BulkDataFilesDDC;
 };
diff --git a/Engine/Plugins/Experimental/Mutable/Source/CustomizableObjectEditor/Private/MuCOE/CustomizableObjectCompiler.cpp b/Engine/Plugins/Experimental/Mutable/Source/CustomizableObjectEditor/Private/MuCOE/CustomizableObjectCompiler.cpp
--- a/Engine/Plugins/Experimental/Mutable/Source/CustomizableObjectEditor/Private/MuCOE/CustomizableObjectCompiler.cpp
+++ b/Engine/Plugins/Experimental/Mutable/Source/CustomizableObjectEditor/Private/MuCOE/CustomizableObjectCompiler.cpp
@@ -230,15 +230,6 @@ void FCustomizableObjectCompiler::Compile(const TSharedRef<FCompilationRequest>&
 		}
 	}
 
-	// If we don't have the target platform yet (in editor) we need to get it
-	if (!CurrentOptions.TargetPlatform)
-	{
-		check(!CurrentOptions.bIsCooking);
-
-		CurrentOptions.TargetPlatform = GetTargetPlatformManagerRef().GetRunningTargetPlatform();
-		check(CurrentOptions.TargetPlatform != nullptr);
-	}
-
 	UE_LOG(LogMutable, Display, TEXT("Compiling Customizable Object %s for platform %s."), *CurrentObject->GetName(), *CurrentOptions.TargetPlatform->PlatformName());
 
 	if (CurrentOptions.bForceLargeLODBias)
@@ -257,6 +248,17 @@ void FCustomizableObjectCompiler::Compile(const TSharedRef<FCompilationRequest>&
 		const int32 NumCompletedRequests = NumCompilationRequests - GetNumRemainingWork();
 		FSlateNotificationManager::Get().UpdateProgressNotification(CompileNotificationHandle, NumCompletedRequests, NumCompilationRequests, UpdateMsg);
 	}
+
+	// DDC check
+	if (TryLoadCompiledDataFromDDC(*CurrentObject))
+	{
+		UE_LOG(LogMutable, Verbose, TEXT("PROFILE: [ %16.8f ] Finishing Compilation task for CO [%s]."), FPlatformTime::Seconds(), *CurrentObject->GetName());
+		TRACE_END_REGION(UE_MUTABLE_COMPILE_REGION);
+
+		UE_LOG(LogMutable, Display, TEXT("Compiled data loaded from DDC"));
+		CompleteRequest(ECompilationStatePrivate::Completed, ECompilationResultPrivate::Success);
+		return;
+	}
 	
 	TRACE_BEGIN_REGION(UE_MUTABLE_PRELOAD_REGION);
 	UE_LOG(LogMutable, Verbose, TEXT("PROFILE: [ %16.8f ] Preload asynchronously assets start."), FPlatformTime::Seconds());
@@ -931,9 +933,8 @@ void FCustomizableObjectCompiler::CompileInternal(bool bAsync)
 	}
 	else
 	{
+		ModelResources = FModelResources();
 		ModelStreamableBulkData = MakeShared<FModelStreamableBulkData>();
-
-		FModelResources ModelResources;
 		
 		ModelResources.ReferenceSkeletalMeshesData = MoveTemp(GenerationContext.ReferenceSkeletalMeshesData);
 
@@ -1118,20 +1119,13 @@ void FCustomizableObjectCompiler::CompileInternal(bool bAsync)
 		ModelResources.MeshMetadata = MoveTemp(GenerationContext.MeshMetadata);
 		ModelResources.SurfaceMetadata = MoveTemp(GenerationContext.SurfaceMetadata);
 
-		CurrentObject->GetPrivate()->GroupNodeMap = GenerationContext.GroupNodeMap;
+		ModelResources.GroupNodeMap = GenerationContext.GroupNodeMap;
 
-		if (GenerationContext.Options.OptimizationLevel == 0)
-		{
-			// If the optimization level is "none" disable texture streaming, because textures are all referenced
-			// unreal assets and progressive generation is not supported.
-			CurrentObject->GetPrivate()->bDisableTextureStreaming = true;
-		}
-		else
-		{
-			CurrentObject->GetPrivate()->bDisableTextureStreaming = false;
-		}
+		// If the optimization level is "none" disable texture streaming, because textures are all referenced
+		// unreal assets and progressive generation is not supported.
+		ModelResources.bIsTextureStreamingDisabled = GenerationContext.Options.OptimizationLevel == 0;
 		
-		CurrentObject->GetPrivate()->bIsCompiledWithoutOptimization = GenerationContext.Options.OptimizationLevel < UE_MUTABLE_MAX_OPTIMIZATION;
+		ModelResources.bIsCompiledWithOptimization = GenerationContext.Options.OptimizationLevel == UE_MUTABLE_MAX_OPTIMIZATION;
 
 		CurrentObject->GetPrivate()->GetAlwaysLoadedExtensionData() = MoveTemp(GenerationContext.AlwaysLoadedExtensionData);
 
@@ -1142,10 +1136,9 @@ void FCustomizableObjectCompiler::CompileInternal(bool bAsync)
 		}
 
 #if WITH_EDITORONLY_DATA
-		CurrentObject->GetPrivate()->CustomizableObjectPathMap = GenerationContext.CustomizableObjectPathMap;
-
 		// Cache the tables that are used by more than one param so that CompileOnlySelected can work properly
-		ModelResources.TableToParamNames = GenerationContext.TableToParamNames;
+		ModelResources.TableToParamNames = MoveTemp(GenerationContext.TableToParamNames);
+		ModelResources.CustomizableObjectPathMap = MoveTemp(GenerationContext.CustomizableObjectPathMap);
 #endif
 
 		ModelResources.NumComponents = GenerationContext.NumMeshComponentsInRoot + GenerationContext.NumExplicitMeshComponents;
@@ -1160,10 +1153,7 @@ void FCustomizableObjectCompiler::CompileInternal(bool bAsync)
 			CurrentObject->Modify();
 		}
 
-		// Always work with the ModelResources (Editor) when compiling. They'll be copied to the cooked version during PreSave.
-		CurrentObject->GetPrivate()->GetModelResources(false) = MoveTemp(ModelResources);
-
-		CurrentObject->GetPrivate()->GetStreamedResourceData() = MoveTemp(GenerationContext.StreamedResourceData);
+		ModelResources.StreamedResourceData = MoveTemp(GenerationContext.StreamedResourceData);
 
 		// Pass-through textures
 		TArray<FMutableSourceTextureData> NewCompileTimeReferencedTextures;
@@ -1196,7 +1186,7 @@ void FCustomizableObjectCompiler::CompileInternal(bool bAsync)
 			}
 			
 			// Copy final array of participating objects
-			CurrentObject->GetPrivate()->ParticipatingObjects = MoveTemp(GenerationContext.ParticipatingObjects);
+			ModelResources.ParticipatingObjects = MoveTemp(GenerationContext.ParticipatingObjects);
 			CurrentObject->GetPrivate()->DirtyParticipatingObjects.Empty();
 		}
 
@@ -1259,14 +1249,14 @@ void FCustomizableObjectCompiler::CompleteRequest(ECompilationStatePrivate State
 			System->UnlockObject(CurrentObject);
 		}
 
-		if (CurrentModel)
+		if (Model)
 		{
-			CurrentModel->GetPrivate()->UnloadRoms();
+			Model->GetPrivate()->UnloadRoms();
 		}
 
 		if (Result == ECompilationResultPrivate::Success || Result == ECompilationResultPrivate::Warnings)
 		{
-			CurrentObject->GetPrivate()->SetModel(CurrentModel, GenerateIdentifier(*CurrentObject));
+			CurrentObject->GetPrivate()->SetModel(Model, GenerateIdentifier(*CurrentObject));
 		}
 		else
 		{
@@ -1317,7 +1307,8 @@ void FCustomizableObjectCompiler::CompleteRequest(ECompilationStatePrivate State
 	// Request completed, reset pointers and state
 	CurrentObject = nullptr;
 	CurrentRequest.Reset();
-	CurrentModel.Reset();
+	Model.Reset();
+	ModelResources = {};
 
 	UE_LOG(LogMutable, Verbose, TEXT("PROFILE: [ %16.8f ] Completed compile request."), FPlatformTime::Seconds());
 	UE_LOG(LogMutable, Verbose, TEXT("PROFILE: -----------------------------------------------------------"));
@@ -1342,6 +1333,32 @@ bool FCustomizableObjectCompiler::TryPopCompileRequest()
 }
 
 
+bool FCustomizableObjectCompiler::TryLoadCompiledDataFromDDC(UCustomizableObject& CustomizableObject)
+{
+	if (!CurrentRequest.IsValid())
+	{
+		return false;
+	}
+
+	using namespace UE::DerivedData;
+
+	ECachePolicy DefaultPolicy = CurrentRequest->GetDerivedDataCachePolicy();
+	if (!CurrentOptions.bQueryCompiledDatafromDDC)
+	{
+		// Compilation not allowed to query DDC requests. 
+		return false;
+	}
+	
+	CurrentRequest->BuildDerivedDataCacheKey();
+
+	FCacheKey CacheKey = CurrentRequest->GetDerivedDataCacheKey();
+	CustomizableObject.GetPrivate()->LoadCompiledDataFromDDC(CurrentOptions, DefaultPolicy, &CacheKey);
+	
+	Model = CustomizableObject.GetPrivate()->GetModel();
+	return CustomizableObject.IsCompiled();
+}
+
+
 mu::NodePtr FCustomizableObjectCompiler::Export(UCustomizableObject* Object, const FCompilationOptions& InCompilerOptions, 
 	TArray<TSoftObjectPtr<const UTexture>>& OutRuntimeReferencedTextures,
 	TArray<FMutableSourceTextureData>& OutCompilerReferencedTextures )
@@ -1413,23 +1430,26 @@ void FCustomizableObjectCompiler::FinishCompilationTask()
 	check(CompileTask.IsValid());
 
 	UpdateCompilerLogData();
-	CurrentModel = CompileTask->Model;
+	Model = CompileTask->Model;
 
 	// Generate a map that using the resource id tells the offset and size of the resource inside the bulk data
 	// At this point it is assumed that all data goes into a single file.
-	if (CurrentModel)
+	if (Model)
 	{
-		const int32 NumStreamingFiles = CurrentModel->GetRomCount();
+		const int32 NumStreamingFiles = Model->GetRomCount();
 
 		TMap<uint32, FMutableStreamableBlock>& ModelStreamables = ModelStreamableBulkData->ModelStreamables;
 		ModelStreamables.Empty(NumStreamingFiles);
 
+		// TODO: Temp. Remove after unifying generated output files code between editor an package. UE-222777
+		const bool bRequiresCookedData = CurrentOptions.TargetPlatform->RequiresCookedData();
+
 		uint64 Offset = 0;
 		for (int32 FileIndex = 0; FileIndex < NumStreamingFiles; ++FileIndex)
 		{
-			const uint32 ResourceId = CurrentModel->GetRomId(FileIndex);
-			const uint32 ResourceSize = CurrentModel->GetRomSize(FileIndex);
-			mu::ERomFlags Flags = CurrentModel->GetRomFlags(FileIndex);
+			const uint32 ResourceId = Model->GetRomId(FileIndex);
+			const uint32 ResourceSize = Model->GetRomSize(FileIndex);
+			mu::ERomFlags Flags = bRequiresCookedData ? Model->GetRomFlags(FileIndex) : mu::ERomFlags::None;
 			ModelStreamables.Add(ResourceId, FMutableStreamableBlock{ 0, uint32(Flags), Offset });
 			Offset += ResourceSize;
 		}
@@ -1438,9 +1458,6 @@ void FCustomizableObjectCompiler::FinishCompilationTask()
 		CurrentObject->GetPrivate()->SetModelStreamableBulkData(ModelStreamableBulkData, false);
 	}
 
-	// Generate ParameterProperties and IntParameterLookUpTable
-	CurrentObject->GetPrivate()->UpdateParameterPropertiesFromModel(CurrentModel);
-
 	// Order matters
 	CompileThread.Reset();
 	CompileTask.Reset();
@@ -1450,7 +1467,7 @@ void FCustomizableObjectCompiler::FinishCompilationTask()
 
 	// Create SaveDD task
 	TRACE_BEGIN_REGION(UE_MUTABLE_SAVEDD_REGION);
-	SaveDDTask = MakeShareable(new FCustomizableObjectSaveDDRunnable(CurrentObject, CurrentOptions, CurrentModel));
+	SaveDDTask = MakeShareable(new FCustomizableObjectSaveDDRunnable(CurrentRequest, Model, ModelResources, ModelStreamableBulkData));
 }
 
 
@@ -1470,18 +1487,11 @@ void FCustomizableObjectCompiler::FinishSavingDerivedDataTask()
 		check(!CurrentObject->GetPrivate()->CachedPlatformsData.Find(PlatformName));
 
 		MutablePrivate::FMutableCachedPlatformData& Data = CurrentObject->GetPrivate()->CachedPlatformsData.Add(PlatformName);
-
-		// Cache CO data and mu::Model
-		FMemoryWriter64 MemoryWriter(Data.ModelData);
-		FObjectAndNameAsStringProxyArchive ObjectWriter(MemoryWriter, true);
-		CurrentObject->GetPrivate()->SaveCompiledData(ObjectWriter, true);
-		Data.ModelData.Append(SaveDDTask->ModelBytes);
-
-		// Cache streamable bulk data
-		Data.ModelStreamableData = MoveTemp(SaveDDTask->ModelStreamableData);
-		Data.MorphData = MoveTemp(SaveDDTask->MorphDataBytes);
-		Data.ClothingData = MoveTemp(SaveDDTask->ClothingDataBytes);
+		Data = MoveTemp(SaveDDTask->PlatformData);
 	}
+	
+	// Always work with the ModelResources (Editor) when compiling. They'll be copied to the cooked version during PreSave.
+	CurrentObject->GetPrivate()->GetModelResources(CurrentOptions.bIsCooking) = MoveTemp(ModelResources);
 
 	// Order matters
 	SaveDDThread.Reset();
diff --git a/Engine/Plugins/Experimental/Mutable/Source/CustomizableObjectEditor/Private/MuCOE/CustomizableObjectCompiler.h b/Engine/Plugins/Experimental/Mutable/Source/CustomizableObjectEditor/Private/MuCOE/CustomizableObjectCompiler.h
--- a/Engine/Plugins/Experimental/Mutable/Source/CustomizableObjectEditor/Private/MuCOE/CustomizableObjectCompiler.h
+++ b/Engine/Plugins/Experimental/Mutable/Source/CustomizableObjectEditor/Private/MuCOE/CustomizableObjectCompiler.h
@@ -87,6 +87,8 @@ private:
 	void CompleteRequest(ECompilationStatePrivate State, ECompilationResultPrivate Result);
 	bool TryPopCompileRequest();
 
+	bool TryLoadCompiledDataFromDDC(UCustomizableObject& CustomizableObject);
+
 	void PreloadingReferencerAssetsCallback(bool bAsync);
 	
 	void ProcessChildObjectsRecursively(UCustomizableObject* Object, FMutableGraphGenerationContext &GenerationContext);
@@ -167,7 +169,8 @@ private:
 	/** Array used to protect from garbage collection those COs loaded asynchronously */
 	TArray<TObjectPtr<UCustomizableObject>> ArrayGCProtect;
 
-	TSharedPtr<mu::Model, ESPMode::ThreadSafe> CurrentModel;
+	TSharedPtr<mu::Model, ESPMode::ThreadSafe> Model;
+	FModelResources ModelResources;
 	TSharedPtr<FModelStreamableBulkData> ModelStreamableBulkData;
 
 	// Protected from GC with FCustomizableObjectCompiler::AddReferencedObjects
diff --git a/Engine/Plugins/Experimental/Mutable/Source/CustomizableObjectEditor/Private/MuCOE/CustomizableObjectCookPackageSplitter.cpp b/Engine/Plugins/Experimental/Mutable/Source/CustomizableObjectEditor/Private/MuCOE/CustomizableObjectCookPackageSplitter.cpp
--- a/Engine/Plugins/Experimental/Mutable/Source/CustomizableObjectEditor/Private/MuCOE/CustomizableObjectCookPackageSplitter.cpp
+++ b/Engine/Plugins/Experimental/Mutable/Source/CustomizableObjectEditor/Private/MuCOE/CustomizableObjectCookPackageSplitter.cpp
@@ -152,8 +152,9 @@ void GenerateNewPackage(const FCustomizableObjectStreamedResourceData& StreamedD
 bool FCustomizableObjectCookPackageSplitter::ShouldSplit(UObject* SplitData)
 {
 	const UCustomizableObject* Object = CastChecked<UCustomizableObject>(SplitData);
+	const FModelResources& ModelResources = Object->GetPrivate()->GetModelResources(true);
 
-	return !Object->IsChildObject() && (Object->GetPrivate()->GetStreamedResourceData().Num() > 0 || Object->GetPrivate()->GetStreamedExtensionData().Num() > 0);
+	return !Object->IsChildObject() && (ModelResources.StreamedResourceData.Num() > 0 || Object->GetPrivate()->GetStreamedExtensionData().Num() > 0);
 }
 
 TArray<ICookPackageSplitter::FGeneratedPackage> FCustomizableObjectCookPackageSplitter::GetGenerateList(
@@ -164,11 +165,12 @@ TArray<ICookPackageSplitter::FGeneratedPackage> FCustomizableObjectCookPackageSp
 	StrongObject.Reset(OwnerObject);
 	
 	const UCustomizableObject* Object = CastChecked<UCustomizableObject>(OwnerObject);
+	FModelResources& ModelResources = Object->GetPrivate()->GetModelResources(true);
 
 	TArray<ICookPackageSplitter::FGeneratedPackage> Result;
 
 	// Generate a new package for each streamed Resource Data
-	for (const FCustomizableObjectStreamedResourceData& StreamedData : Object->GetPrivate()->GetStreamedResourceData())
+	for (const FCustomizableObjectStreamedResourceData& StreamedData : ModelResources.StreamedResourceData)
 	{
 		GenerateNewPackage(StreamedData, OwnerPackage, OwnerObject, Result);
 	}
@@ -227,9 +229,10 @@ bool FCustomizableObjectCookPackageSplitter::PreSaveGeneratorPackage(
 	};
 
 	UCustomizableObject* Object = CastChecked<UCustomizableObject>(OwnerObject);
+	FModelResources& ModelResources = Object->GetPrivate()->GetModelResources(true);
 
 	// There should be one generated package per streamed Resource Data
-	const int32 NumStreamedData = Object->GetPrivate()->GetStreamedResourceData().Num();
+	const int32 NumStreamedData = ModelResources.StreamedResourceData.Num();
 	const int32 NumStreamedExtensionData = Object->GetPrivate()->GetStreamedExtensionData().Num();
 	
 	check(NumStreamedData + NumStreamedExtensionData == PlaceholderPackages.Num());
@@ -244,7 +247,7 @@ bool FCustomizableObjectCookPackageSplitter::PreSaveGeneratorPackage(
 	for (int32 Index = 0; Index < NumStreamedData; ++Index)
 	{
 		const ICookPackageSplitter::FGeneratedPackageForPreSave& GeneratedPackage = PlaceholderPackages[Index];
-		if (!PreSavePackage(GeneratedPackage, Object->GetPrivate()->GetStreamedResourceData()))
+		if (!PreSavePackage(GeneratedPackage, ModelResources.StreamedResourceData))
 		{
 			return false;
 		}
@@ -272,17 +275,18 @@ void FCustomizableObjectCookPackageSplitter::PostSaveGeneratorPackage(UPackage*
 	// array on the CO to how it was before PreSaveGeneratorPackage.
 
 	UCustomizableObject* Object = CastChecked<UCustomizableObject>(OwnerObject);
+	FModelResources& ModelResources = Object->GetPrivate()->GetModelResources(true);
 
 	TArray<FCustomizableObjectStreamedResourceData> NewArray;
 	NewArray.Reset(SavedContainerNames.Num());
 
 	for (const FString& ContainerName : SavedContainerNames)
 	{
-		FCustomizableObjectStreamedResourceData* ResourceData = FindStreamedResourceData(Object->GetPrivate()->GetStreamedResourceData(), ContainerName);
+		FCustomizableObjectStreamedResourceData* ResourceData = FindStreamedResourceData(ModelResources.StreamedResourceData, ContainerName);
 		if (!ResourceData)
 		{
 			UE_LOG(LogMutable, Error, TEXT("Couldn't find streamed Resource Data container with name %s in array of %d entries"),
-				*ContainerName, Object->GetPrivate()->GetStreamedResourceData().Num());
+				*ContainerName, ModelResources.StreamedResourceData.Num());
 
 			continue;
 		}
@@ -294,7 +298,7 @@ void FCustomizableObjectCookPackageSplitter::PostSaveGeneratorPackage(UPackage*
 		NewArray.Emplace(Container);
 	}
 
-	Object->GetPrivate()->GetStreamedResourceData() = NewArray;
+	ModelResources.StreamedResourceData = NewArray;
 
 	NewArray.Reset(SavedExtensionContainerNames.Num());
 
@@ -329,8 +333,9 @@ bool FCustomizableObjectCookPackageSplitter::PopulateGeneratedPackage(
 	// Move the container into its newly generated package
 
 	UCustomizableObject* Object = CastChecked<UCustomizableObject>(OwnerObject);
-	
-	FCustomizableObjectStreamedResourceData* ResourceData = FindStreamedResourceData(Object->GetPrivate()->GetStreamedResourceData(), GeneratedPackage.RelativePath);
+	FModelResources& ModelResources = Object->GetPrivate()->GetModelResources(true);
+
+	FCustomizableObjectStreamedResourceData* ResourceData = FindStreamedResourceData(ModelResources.StreamedResourceData, GeneratedPackage.RelativePath);
 	if (!ResourceData)
 	{
 		ResourceData = FindStreamedResourceData(Object->GetPrivate()->GetStreamedExtensionData(), GeneratedPackage.RelativePath);
@@ -339,7 +344,7 @@ bool FCustomizableObjectCookPackageSplitter::PopulateGeneratedPackage(
 	if (!ResourceData)
 	{
 		UE_LOG(LogMutable, Error, TEXT("Couldn't find streamed resource Data container with name %s in arrays of %d and %d entries"),
-			*GeneratedPackage.RelativePath, Object->GetPrivate()->GetStreamedResourceData().Num(), Object->GetPrivate()->GetStreamedExtensionData().Num());
+			*GeneratedPackage.RelativePath, ModelResources.StreamedResourceData.Num(), Object->GetPrivate()->GetStreamedExtensionData().Num());
 
 		return false;
 	}
@@ -376,8 +381,9 @@ void FCustomizableObjectCookPackageSplitter::PostSaveGeneratedPackage(
 	// that everything is the same as it was before cooking.
 
 	UCustomizableObject* Object = CastChecked<UCustomizableObject>(OwnerObject);
+	FModelResources& ModelResources = Object->GetPrivate()->GetModelResources(true);
 
-	FCustomizableObjectStreamedResourceData* ResourceData = FindStreamedResourceData(Object->GetPrivate()->GetStreamedResourceData(), GeneratedPackage.RelativePath);
+	FCustomizableObjectStreamedResourceData* ResourceData = FindStreamedResourceData(ModelResources.StreamedResourceData, GeneratedPackage.RelativePath);
 	if (!ResourceData)
 	{
 		ResourceData = FindStreamedResourceData(Object->GetPrivate()->GetStreamedExtensionData(), GeneratedPackage.RelativePath);
@@ -386,7 +392,7 @@ void FCustomizableObjectCookPackageSplitter::PostSaveGeneratedPackage(
 	if (!ResourceData)
 	{
 		UE_LOG(LogMutable, Error, TEXT("Couldn't find streamed resource Data container with name %s in arrays of %d and %d entries"),
-			*GeneratedPackage.RelativePath, Object->GetPrivate()->GetStreamedResourceData().Num(), Object->GetPrivate()->GetStreamedExtensionData().Num());
+			*GeneratedPackage.RelativePath, ModelResources.StreamedResourceData.Num(), Object->GetPrivate()->GetStreamedExtensionData().Num());
 
 		return;
 	}
diff --git a/Engine/Plugins/Experimental/Mutable/Source/CustomizableObjectEditor/Private/MuCOE/CustomizableObjectEditorModule.cpp b/Engine/Plugins/Experimental/Mutable/Source/CustomizableObjectEditor/Private/MuCOE/CustomizableObjectEditorModule.cpp
--- a/Engine/Plugins/Experimental/Mutable/Source/CustomizableObjectEditor/Private/MuCOE/CustomizableObjectEditorModule.cpp
+++ b/Engine/Plugins/Experimental/Mutable/Source/CustomizableObjectEditor/Private/MuCOE/CustomizableObjectEditorModule.cpp
@@ -147,7 +147,7 @@ void ShowOnScreenCompileWarnings()
 	
 		// Show a warning if the compilation was not done with optimizations.
 		const uint64 KeyCompiledWithOptimization = reinterpret_cast<uint64>(Object);
-		if (Object->GetPrivate()->bIsCompiledWithoutOptimization)
+		if (!Object->GetPrivate()->GetModelResources().bIsCompiledWithOptimization)
 		{
 			FString Msg = FString::Printf(TEXT("Customizable Object [%s] was compiled without optimization."), *Object->GetName());
 			GEngine->AddOnScreenDebugMessage(KeyCompiledWithOptimization, ShowOnScreenCompileWarningsTickerTime * 2.0f, FColor::Yellow, Msg);
@@ -558,7 +558,8 @@ bool FCustomizableObjectEditorModule::IsCompilationOutOfDate(const UCustomizable
 	
 	const FAssetRegistryModule& AssetRegistryModule = FModuleManager::LoadModuleChecked<FAssetRegistryModule>("AssetRegistry");
 
-	for (const TTuple<FName, FGuid>& ParticipatingObject : Object.GetPrivate()->ParticipatingObjects)
+	const FModelResources& ModelResources = Object.GetPrivate()->GetModelResources();
+	for (const TTuple<FName, FGuid>& ParticipatingObject : ModelResources.ParticipatingObjects)
 	{
 		TSoftObjectPtr<UObject> SoftObjectPtr = TSoftObjectPtr<UObject>(FSoftObjectPath(ParticipatingObject.Key.ToString()));
 		if (SoftObjectPtr) // If loaded
@@ -622,7 +623,7 @@ bool FCustomizableObjectEditorModule::IsCompilationOutOfDate(const UCustomizable
 		TSoftObjectPtr<UObject> ReferencingObject = TSoftObjectPtr<UObject>(FSoftObjectPath(ObjectName.ToString())); 
 
 		if ((ReferencingObject && ReferencingObject->GetPackage()->IsDirty()) ||
-			!Object.GetPrivate()->ParticipatingObjects.Contains(ObjectName)) // Must be in the participating objects, if not it means it did not exist when compiling the object.
+			!ModelResources.ParticipatingObjects.Contains(ObjectName)) // Must be in the participating objects, if not it means it did not exist when compiling the object.
 		{
 			if (OutOfDatePackages)
 			{
diff --git a/Engine/Plugins/Experimental/Mutable/Source/CustomizableObjectEditor/Private/MuCOE/SCustomizableObjectEditorViewport.cpp b/Engine/Plugins/Experimental/Mutable/Source/CustomizableObjectEditor/Private/MuCOE/SCustomizableObjectEditorViewport.cpp
--- a/Engine/Plugins/Experimental/Mutable/Source/CustomizableObjectEditor/Private/MuCOE/SCustomizableObjectEditorViewport.cpp
+++ b/Engine/Plugins/Experimental/Mutable/Source/CustomizableObjectEditor/Private/MuCOE/SCustomizableObjectEditorViewport.cpp
@@ -221,7 +221,7 @@ FText SCustomizableObjectEditorViewport::GetWarningText() const
 	{
 		if (const UCustomizableObject* Object = Instance->GetCustomizableObject())
 		{
-			if (Object->GetPrivate()->bIsCompiledWithoutOptimization)
+			if (!Object->GetPrivate()->GetModelResources().bIsCompiledWithOptimization)
 			{
 				return LOCTEXT("CompiledWithoutOptimization", "Compiled without maximum optimization. Updates will be slower!");
 			}
