[mutable] Rename user facing properties from "image" to "texture.
#rnx

diff --git a/Engine/Plugins/Experimental/Mutable/Config/BaseMutable.ini b/Engine/Plugins/Experimental/Mutable/Config/BaseMutable.ini
new file mode 100644
--- /dev/null
+++ b/Engine/Plugins/Experimental/Mutable/Config/BaseMutable.ini
@@ -0,0 +1,8 @@
+
+[CoreRedirects]
++StructRedirects=(OldName="/Script/CustomizableObject.ParameterUIData",NewName="/Script/CustomizableObject.MutableParameterData")
++StructRedirects=(OldName="/Script/CustomizableObject.CompilationOptions",NewName="/Script/CustomizableObject.CompilationOptions_DEPRECATED")
+
++PropertyRedirects=(OldName="GroupProjectorParameterImage.OptionImage", NewName="OptionTexture")
++PropertyRedirects=(OldName="CustomizableObjectNodeGroupProjectorParameter.OptionImages", NewName="OptionTextures")
++PropertyRedirects=(OldName="CustomizableObjectNodeGroupProjectorParameter.OptionImagesDataTable", NewName="OptionTexturesDataTable")
diff --git a/Engine/Plugins/Experimental/Mutable/Config/DefaultMutable.ini b/Engine/Plugins/Experimental/Mutable/Config/DefaultMutable.ini
--- a/Engine/Plugins/Experimental/Mutable/Config/DefaultMutable.ini
+++ b/Engine/Plugins/Experimental/Mutable/Config/DefaultMutable.ini
@@ -1,9 +1,4 @@
 ; Settings for the Mutable plugin
-
-[CoreRedirects]
-+StructRedirects=(OldName="/Script/CustomizableObject.ParameterUIData",NewName="/Script/CustomizableObject.MutableParameterData")
-+StructRedirects=(OldName="/Script/CustomizableObject.CompilationOptions",NewName="/Script/CustomizableObject.CompilationOptions_DEPRECATED")
-
 [Features]
 
 ; If CustomizableObjectNumBoneInfluences is set to N, CustomizableObjects will be allowed to use a max of N bone influences 
diff --git a/Engine/Plugins/Experimental/Mutable/Source/CustomizableObjectEditor/Private/MuCOE/GenerateMutableSource/GenerateMutableSourceGroupProjector.cpp b/Engine/Plugins/Experimental/Mutable/Source/CustomizableObjectEditor/Private/MuCOE/GenerateMutableSource/GenerateMutableSourceGroupProjector.cpp
--- a/Engine/Plugins/Experimental/Mutable/Source/CustomizableObjectEditor/Private/MuCOE/GenerateMutableSource/GenerateMutableSourceGroupProjector.cpp
+++ b/Engine/Plugins/Experimental/Mutable/Source/CustomizableObjectEditor/Private/MuCOE/GenerateMutableSource/GenerateMutableSourceGroupProjector.cpp
@@ -337,11 +337,11 @@ bool GenerateMutableSourceGroupProjector(const UEdGraphPin* Pin, FMutableGraphGe
 			FGuid SelectedImageParamUid = Node->NodeGuid;
 			SelectedImageParamUid.D += 4;
 
-			// Add to UCustomizableObjectNodeGroupProjectorParameter::OptionImages those textures that are present in
-			// UCustomizableObjectNodeGroupProjectorParameter::OptionImagesDataTable avoiding any repeated element
-			TArray<FGroupProjectorParameterImage> ArrayOptionImage = ProjParamNode->GetFinalOptionImagesNoRepeat();
+			// Add to UCustomizableObjectNodeGroupProjectorParameter::OptionTextures those textures that are present in
+			// UCustomizableObjectNodeGroupProjectorParameter::OptionTexturesDataTable avoiding any repeated element
+			TArray<FGroupProjectorParameterImage> ArrayOptionTexture = ProjParamNode->GetFinalOptionTexturesNoRepeat();
 
-			if ((ProjParamNode->OptionImagesDataTable != nullptr) &&
+			if ((ProjParamNode->OptionTexturesDataTable != nullptr) &&
 				(ProjParamNode->DataTableTextureColumnName.ToString().IsEmpty() || (ProjParamNode->DataTableTextureColumnName.ToString() == "None")))
 			{
 				FString msg = FString::Printf(TEXT("The group projection node has a table assigned to the Option Images Data Table property, but no column to read textures is specified at the Data Table Texture Column Name property."));
@@ -400,7 +400,7 @@ bool GenerateMutableSourceGroupProjector(const UEdGraphPin* Pin, FMutableGraphGe
 				GenerationContext.ParameterUIDataMap.Add(OpacityParameterNodeName, ParameterUIData);
 			}
 			
-			if (ArrayOptionImage.Num() == 0)
+			if (ArrayOptionTexture.Num() == 0)
 			{
 				FString msg = FString::Printf(TEXT("The group projection node must have at least one option image connected to a texture or at least one valid element in Option Images Data Table."));
 				GenerationContext.Compiler->CompilerLog(FText::FromString(msg), ProjParamNode, EMessageSeverity::Error, true);
@@ -445,7 +445,7 @@ bool GenerateMutableSourceGroupProjector(const UEdGraphPin* Pin, FMutableGraphGe
 			EnumParameterNode->SetName(NodeEnumParamName);
 			EnumParameterNode->SetUid(SelectedImageParamUid.ToString());
 			GenerationContext.AddParameterNameUnique(originalGroup, NodeEnumParamName);
-			EnumParameterNode->SetValueCount(ArrayOptionImage.Num());
+			EnumParameterNode->SetValueCount(ArrayOptionTexture.Num());
 			EnumParameterNode->SetDefaultValueIndex(0);
 			EnumParameterNode->SetRangeCount(1);
 			EnumParameterNode->SetRange(0, NodeRangeFromScalar);
@@ -457,14 +457,14 @@ bool GenerateMutableSourceGroupProjector(const UEdGraphPin* Pin, FMutableGraphGe
 				ParameterUIData.IntegerParameterGroupType = ECustomizableObjectGroupType::COGT_ONE;
 				ParameterUIData.ParamUIMetadata.ExtraInformation.Add(FString("UseThumbnails"));
 
-				for (int ImageIndex = 0; ImageIndex < ArrayOptionImage.Num(); ++ImageIndex)
+				for (int ImageIndex = 0; ImageIndex < ArrayOptionTexture.Num(); ++ImageIndex)
 				{
-					EnumParameterNode->SetValue(ImageIndex, (float)ImageIndex, ArrayOptionImage[ImageIndex].OptionName);
+					EnumParameterNode->SetValue(ImageIndex, (float)ImageIndex, ArrayOptionTexture[ImageIndex].OptionName);
 
 					FMutableParamUIMetadata OptionMetadata = ParameterUIData.ParamUIMetadata;
-					OptionMetadata.UIThumbnail = ArrayOptionImage[ImageIndex].OptionImage;
+					OptionMetadata.UIThumbnail = ArrayOptionTexture[ImageIndex].OptionTexture;
 					ParameterUIData.ArrayIntegerParameterOption.Add(
-						ArrayOptionImage[ImageIndex].OptionName,
+						ArrayOptionTexture[ImageIndex].OptionName,
 						FIntegerParameterUIData(OptionMetadata));
 				}
 
@@ -473,14 +473,14 @@ bool GenerateMutableSourceGroupProjector(const UEdGraphPin* Pin, FMutableGraphGe
 
 			mu::NodeImageSwitchPtr SwitchNode = new mu::NodeImageSwitch;
 			SwitchNode->SetParameter(EnumParameterNode);
-			SwitchNode->SetOptionCount(ArrayOptionImage.Num());
+			SwitchNode->SetOptionCount(ArrayOptionTexture.Num());
 
 			const uint32 AdditionalLODBias = GenerationContext.Options.bUseLODAsBias ? GenerationContext.FirstLODAvailable : 0;
-			for (int32 SelectorIndex = 0; SelectorIndex < ArrayOptionImage.Num(); ++SelectorIndex)
+			for (int32 SelectorIndex = 0; SelectorIndex < ArrayOptionTexture.Num(); ++SelectorIndex)
 			{
-				if (const TObjectPtr<UTexture2D>& Texture = ArrayOptionImage[SelectorIndex].OptionImage)
+				if (const TObjectPtr<UTexture2D>& Texture = ArrayOptionTexture[SelectorIndex].OptionTexture)
 				{
-					mu::Ptr<mu::Image> ImageConstant = GenerateImageConstant(ArrayOptionImage[SelectorIndex].OptionImage, GenerationContext, false);
+					mu::Ptr<mu::Image> ImageConstant = GenerateImageConstant(ArrayOptionTexture[SelectorIndex].OptionTexture, GenerationContext, false);
 
 					mu::Ptr<mu::NodeImageConstant> ImageNode = new mu::NodeImageConstant();
 					ImageNode->SetValue(ImageConstant.get());
diff --git a/Engine/Plugins/Experimental/Mutable/Source/CustomizableObjectEditor/Private/MuCOE/Nodes/CustomizableObjectNodeGroupProjectorParameter.cpp b/Engine/Plugins/Experimental/Mutable/Source/CustomizableObjectEditor/Private/MuCOE/Nodes/CustomizableObjectNodeGroupProjectorParameter.cpp
--- a/Engine/Plugins/Experimental/Mutable/Source/CustomizableObjectEditor/Private/MuCOE/Nodes/CustomizableObjectNodeGroupProjectorParameter.cpp
+++ b/Engine/Plugins/Experimental/Mutable/Source/CustomizableObjectEditor/Private/MuCOE/Nodes/CustomizableObjectNodeGroupProjectorParameter.cpp
@@ -11,18 +11,18 @@ class UCustomizableObjectNodeRemapPins;
 #define LOCTEXT_NAMESPACE "CustomizableObjectEditor"
 
 
-TArray<FGroupProjectorParameterImage> UCustomizableObjectNodeGroupProjectorParameter::GetOptionImagesFromTable() const
+TArray<FGroupProjectorParameterImage> UCustomizableObjectNodeGroupProjectorParameter::GetOptionTexturesFromTable() const
 {
 	TArray<FGroupProjectorParameterImage> ArrayResult;
 
-	if (OptionImagesDataTable == nullptr)
+	if (OptionTexturesDataTable == nullptr)
 	{
 		return ArrayResult;
 	}
 
-	TArray<FName> ArrayRowName = OptionImagesDataTable->GetRowNames();
+	TArray<FName> ArrayRowName = OptionTexturesDataTable->GetRowNames();
 
-	FProperty* PropertyTexturePath = OptionImagesDataTable->FindTableProperty(DataTableTextureColumnName);
+	FProperty* PropertyTexturePath = OptionTexturesDataTable->FindTableProperty(DataTableTextureColumnName);
 
 	if (PropertyTexturePath == nullptr)
 	{
@@ -32,7 +32,7 @@ TArray<FGroupProjectorParameterImage> UCustomizableObjectNodeGroupProjectorParam
 
 	int32 NameIndex = 0;
 
-	for (TMap<FName, uint8*>::TConstIterator RowIt = OptionImagesDataTable->GetRowMap().CreateConstIterator(); RowIt; ++RowIt)
+	for (TMap<FName, uint8*>::TConstIterator RowIt = OptionTexturesDataTable->GetRowMap().CreateConstIterator(); RowIt; ++RowIt)
 	{
 		uint8* RowData = RowIt.Value();
 		FString PropertyValue(TEXT(""));
@@ -47,7 +47,7 @@ TArray<FGroupProjectorParameterImage> UCustomizableObjectNodeGroupProjectorParam
 		{
 			FGroupProjectorParameterImage GroupProjectorParameterImage;
 			GroupProjectorParameterImage.OptionName = ArrayRowName[NameIndex].ToString();
-			GroupProjectorParameterImage.OptionImage = Texture;
+			GroupProjectorParameterImage.OptionTexture = Texture;
 			ArrayResult.Add(GroupProjectorParameterImage);
 		}
 
@@ -58,16 +58,16 @@ TArray<FGroupProjectorParameterImage> UCustomizableObjectNodeGroupProjectorParam
 }
 
 
-TArray<FGroupProjectorParameterImage> UCustomizableObjectNodeGroupProjectorParameter::GetFinalOptionImagesNoRepeat() const
+TArray<FGroupProjectorParameterImage> UCustomizableObjectNodeGroupProjectorParameter::GetFinalOptionTexturesNoRepeat() const
 {
-	TArray<FGroupProjectorParameterImage> ArrayDataTable = GetOptionImagesFromTable();
+	TArray<FGroupProjectorParameterImage> ArrayDataTable = GetOptionTexturesFromTable();
 
-	for (int32 i = 0; i < OptionImages.Num(); ++i)
+	for (int32 i = 0; i < OptionTextures.Num(); ++i)
 	{
 		bool AlreadyAdded = false;
 		for (int32 j = 0; j < ArrayDataTable.Num(); ++j)
 		{
-			if (OptionImages[i].OptionName == ArrayDataTable[j].OptionName)
+			if (OptionTextures[i].OptionName == ArrayDataTable[j].OptionName)
 			{
 				AlreadyAdded = true;
 				break;
@@ -76,7 +76,7 @@ TArray<FGroupProjectorParameterImage> UCustomizableObjectNodeGroupProjectorParam
 
 		if (!AlreadyAdded)
 		{
-			ArrayDataTable.Add(OptionImages[i]);
+			ArrayDataTable.Add(OptionTextures[i]);
 		}
 	}
 
diff --git a/Engine/Plugins/Experimental/Mutable/Source/CustomizableObjectEditor/Public/MuCOE/Nodes/CustomizableObjectNodeGroupProjectorParameter.h b/Engine/Plugins/Experimental/Mutable/Source/CustomizableObjectEditor/Public/MuCOE/Nodes/CustomizableObjectNodeGroupProjectorParameter.h
--- a/Engine/Plugins/Experimental/Mutable/Source/CustomizableObjectEditor/Public/MuCOE/Nodes/CustomizableObjectNodeGroupProjectorParameter.h
+++ b/Engine/Plugins/Experimental/Mutable/Source/CustomizableObjectEditor/Public/MuCOE/Nodes/CustomizableObjectNodeGroupProjectorParameter.h
@@ -22,7 +22,7 @@ struct FGroupProjectorParameterImage
 	FString OptionName;
 
 	UPROPERTY(EditAnywhere, Category = ProjectorGroup)
-	TObjectPtr<UTexture2D> OptionImage = nullptr;
+	TObjectPtr<UTexture2D> OptionTexture = nullptr;
 };
 
 
@@ -46,11 +46,11 @@ public:
 	GENERATED_BODY()
 
 	/** Return array with the sticker name and UTexture2D for projection */
-	TArray<FGroupProjectorParameterImage> GetOptionImagesFromTable() const;
+	TArray<FGroupProjectorParameterImage> GetOptionTexturesFromTable() const;
 
 	/** Returns the final option images without repeated elements in the option names,
-	* the data table has preference over elements in the OptionImages array */
-	TArray<FGroupProjectorParameterImage> GetFinalOptionImagesNoRepeat() const;
+	* the data table has preference over elements in the OptionTextures array */
+	TArray<FGroupProjectorParameterImage> GetFinalOptionTexturesNoRepeat() const;
 
 	UPROPERTY(EditAnywhere, Category = ProjectorGroup, Meta = (ToolTip = "Only used when connected to a Group node. Specifies which material channel in the Group node's child material nodes will be connected to the projection."))
 	FString MaterialChannelNameToConnect;
@@ -71,7 +71,7 @@ public:
 	bool bShareProjectionTexturesBetweenLODs = false;
 
 	UPROPERTY(EditAnywhere, Category = ProjectorGroup)
-	TArray<FGroupProjectorParameterImage> OptionImages;
+	TArray<FGroupProjectorParameterImage> OptionTextures;
 
 	/** Poses will only affect component 0 of the CO */ // TODO UE-206803
 	UPROPERTY(EditAnywhere, Category = ProjectorGroup)
@@ -85,7 +85,7 @@ public:
 	* over elements from Option Images in case of duplicity. Use the "Data Table Texture Column Name" property to specify
 	* the name of the column where textures are read in the table. */
 	UPROPERTY(EditAnywhere, Category = ProjectorGroup)
-	TObjectPtr<UDataTable> OptionImagesDataTable = nullptr;
+	TObjectPtr<UDataTable> OptionTexturesDataTable = nullptr;
 
 	UPROPERTY(EditAnywhere, Category = ProjectorGroup)
 	FString AlternateProjectionResolutionStateName;
