[mutable] Fix modifier pin connections with legacy data.
#rnx

diff --git a/Engine/Plugins/Experimental/Mutable/Source/CustomizableObject/Internal/MuCO/CustomizableObjectCustomVersion.h b/Engine/Plugins/Experimental/Mutable/Source/CustomizableObject/Internal/MuCO/CustomizableObjectCustomVersion.h
--- a/Engine/Plugins/Experimental/Mutable/Source/CustomizableObject/Internal/MuCO/CustomizableObjectCustomVersion.h
+++ b/Engine/Plugins/Experimental/Mutable/Source/CustomizableObject/Internal/MuCO/CustomizableObjectCustomVersion.h
@@ -135,7 +135,9 @@ struct CUSTOMIZABLEOBJECT_API FCustomizableObjectCustomVersion
 		CorrectlySerializeTableToParamNames,
 
 		MovedCompatibilityFromPostBackwardsCompatibleFixup,
-		
+
+		FixModifierPin, 
+
 		// -----<new versions can be added above this line>-------------------------------------------------
 		VersionPlusOne,
 		LatestVersion = VersionPlusOne - 1
diff --git a/Engine/Plugins/Experimental/Mutable/Source/CustomizableObjectEditor/Private/MuCOE/Nodes/CustomizableObjectNodeModifierBase.cpp b/Engine/Plugins/Experimental/Mutable/Source/CustomizableObjectEditor/Private/MuCOE/Nodes/CustomizableObjectNodeModifierBase.cpp
--- a/Engine/Plugins/Experimental/Mutable/Source/CustomizableObjectEditor/Private/MuCOE/Nodes/CustomizableObjectNodeModifierBase.cpp
+++ b/Engine/Plugins/Experimental/Mutable/Source/CustomizableObjectEditor/Private/MuCOE/Nodes/CustomizableObjectNodeModifierBase.cpp
@@ -67,20 +67,82 @@ void UCustomizableObjectNodeModifierBase::BackwardsCompatibleFixup(int32 Customi
 						LinksToRemove.Add(LinkedToPin);
 						NewPin->MakeLinkTo(ToModifierPin);
 					}
+					else
+					{
+						ensure(false);
+					}
+				}
+			}
+
+			// Remove reconnected links
+			for (UEdGraphPin* LinkToRemove : LinksToRemove)
+			{
+				NewPin->BreakLinkTo(LinkToRemove);
+			}
+		}
+	}
+
+	// Fix the "Modifier" pin connection
+	if (CustomizableObjectCustomVersion == FCustomizableObjectCustomVersion::FixModifierPin)
+	{
+		// Replace the output pin
+		UEdGraphPin* ModifierPin = FindPin(TEXT("Modifier"));
+		if (ModifierPin)
+		{
+			// Reconnect it to the correct output of the target node
+			TArray<UEdGraphPin*> LinksToRemove;
+			for (int32 LinkIndex = 0; LinkIndex < ModifierPin->LinkedTo.Num(); ++LinkIndex)
+			{
+				UEdGraphPin* LinkedToPin = ModifierPin->LinkedTo[LinkIndex];
+
+				if (!LinkedToPin)
+				{
+					continue;
+				}
+
+				UEdGraphNode* ToNode = LinkedToPin->GetOwningNode();
+				if (!ToNode)
+				{
+					continue;
+				}
+
+				if (UCustomizableObjectNodeObject* ToObjectNode = Cast<UCustomizableObjectNodeObject>(ToNode))
+				{
+					UEdGraphPin* ToModifierPin = ToObjectNode->ModifiersPin();
+					if (ToModifierPin == LinkedToPin)
+					{
+						// It is already correctly connected
+						continue;
+					}
+
+					if (ToModifierPin)
+					{
+						LinksToRemove.Add(LinkedToPin);
+						ModifierPin->MakeLinkTo(ToModifierPin);
+					}
+					else
+					{
+						ensure(false);
+					}
 				}
 				else
 				{
-					// The modifier is connected to a node for which automatic upgrade support is not implemented. This will be warned below
+					// The modifier is connected to a node for which automatic upgrade support is not implemented. This will be warned in PostBackwardsCompatibleFixup
 				}
 			}
 
 			// Remove reconnected links
 			for (UEdGraphPin* LinkToRemove : LinksToRemove)
 			{
-				NewPin->BreakLinkTo(LinkToRemove);
+				ModifierPin->BreakLinkTo(LinkToRemove);
 			}
 		}
 	}
+}
+
+void UCustomizableObjectNodeModifierBase::PostBackwardsCompatibleFixup()
+{
+	const UEdGraphSchema_CustomizableObject* Schema = GetDefault<UEdGraphSchema_CustomizableObject>();
 
 	// Check for old legacy connections that need manual update
 	{
diff --git a/Engine/Plugins/Experimental/Mutable/Source/CustomizableObjectEditor/Private/MuCOE/Nodes/CustomizableObjectNodeObject.cpp b/Engine/Plugins/Experimental/Mutable/Source/CustomizableObjectEditor/Private/MuCOE/Nodes/CustomizableObjectNodeObject.cpp
--- a/Engine/Plugins/Experimental/Mutable/Source/CustomizableObjectEditor/Private/MuCOE/Nodes/CustomizableObjectNodeObject.cpp
+++ b/Engine/Plugins/Experimental/Mutable/Source/CustomizableObjectEditor/Private/MuCOE/Nodes/CustomizableObjectNodeObject.cpp
@@ -117,6 +117,18 @@ void UCustomizableObjectNodeObject::BackwardsCompatibleFixup(int32 CustomizableO
 			}
 		}
 	}
+
+	// Add the "Modifiers" pin
+	if (CustomizableObjectCustomVersion == FCustomizableObjectCustomVersion::AddModifierPin)
+	{
+		const UEdGraphSchema_CustomizableObject* Schema = GetDefault<UEdGraphSchema_CustomizableObject>();
+
+		UEdGraphPin* OldPin = FindPin(TEXT("Modifiers"));
+		if (!OldPin)
+		{
+			CustomCreatePin(EGPD_Input, Schema->PC_Modifier, ModifiersPinName, true);
+		}
+	}
 }
 
 
diff --git a/Engine/Plugins/Experimental/Mutable/Source/CustomizableObjectEditor/Public/MuCOE/Nodes/CustomizableObjectNodeModifierBase.h b/Engine/Plugins/Experimental/Mutable/Source/CustomizableObjectEditor/Public/MuCOE/Nodes/CustomizableObjectNodeModifierBase.h
--- a/Engine/Plugins/Experimental/Mutable/Source/CustomizableObjectEditor/Public/MuCOE/Nodes/CustomizableObjectNodeModifierBase.h
+++ b/Engine/Plugins/Experimental/Mutable/Source/CustomizableObjectEditor/Public/MuCOE/Nodes/CustomizableObjectNodeModifierBase.h
@@ -20,5 +20,6 @@ public:
 
 	// UCustomizableObjectNode interface
 	virtual void BackwardsCompatibleFixup(int32 CustomizableObjectCustomVersion) override;
+	virtual void PostBackwardsCompatibleFixup() override;
 };
 
