[Mutable] Fix crash when serializing streamable resources due to GCed references.

#rnx

diff --git a/Engine/Plugins/Experimental/Mutable/Source/CustomizableObjectEditor/Private/MuCOE/CustomizableObjectCompiler.cpp b/Engine/Plugins/Experimental/Mutable/Source/CustomizableObjectEditor/Private/MuCOE/CustomizableObjectCompiler.cpp
--- a/Engine/Plugins/Experimental/Mutable/Source/CustomizableObjectEditor/Private/MuCOE/CustomizableObjectCompiler.cpp
+++ b/Engine/Plugins/Experimental/Mutable/Source/CustomizableObjectEditor/Private/MuCOE/CustomizableObjectCompiler.cpp
@@ -1190,6 +1190,8 @@ void FCustomizableObjectCompiler::CompileInternal(bool bAsync)
 			CurrentObject->GetPrivate()->DirtyParticipatingObjects.Empty();
 		}
 
+		CurrentObject->GetPrivate()->GetModelResources(CurrentOptions.bIsCooking);
+
 		CompileTask = MakeShareable(new FCustomizableObjectCompileRunnable(MutableRoot));
 		CompileTask->Options = CurrentOptions;
 		CompileTask->ReferencedTextures = NewCompileTimeReferencedTextures;
