[mutable] Convert the following nodes to modifiers: EditMaterial, ExtendMaterial, RemoveMesh, RemoveMeshBlocks, MorphMaterial.
- Rename some existing modifier nodes to actually be prefixed "Modifier" instead of "Mesh".
- Added additional  entry points for modifiers during code generation. Now there are 4: right after generating a mesh constant, right after all mesh operations have been generated, right after a texture block has been generated and right after a texture has been generated.
- Fix bug with mesh clip with UV mask when layout blocks were used.
- For legacy data conversion autogenerated tags are created and assigned to the material nodes and the new modifier nodes.
- MorphMaterial code generation has been redesigned passing all morphs to the MutableTools layer.
#rnx

diff --git a/Engine/Plugins/Experimental/Mutable/Config/BaseMutable.ini b/Engine/Plugins/Experimental/Mutable/Config/BaseMutable.ini
--- a/Engine/Plugins/Experimental/Mutable/Config/BaseMutable.ini
+++ b/Engine/Plugins/Experimental/Mutable/Config/BaseMutable.ini
@@ -1,8 +1,20 @@
-
 [CoreRedirects]
+
 +StructRedirects=(OldName="/Script/CustomizableObject.ParameterUIData",NewName="/Script/CustomizableObject.MutableParameterData")
 +StructRedirects=(OldName="/Script/CustomizableObject.CompilationOptions",NewName="/Script/CustomizableObject.CompilationOptions_DEPRECATED")
 
 +PropertyRedirects=(OldName="GroupProjectorParameterImage.OptionImage", NewName="OptionTexture")
 +PropertyRedirects=(OldName="CustomizableObjectNodeGroupProjectorParameter.OptionImages", NewName="OptionTextures")
 +PropertyRedirects=(OldName="CustomizableObjectNodeGroupProjectorParameter.OptionImagesDataTable", NewName="OptionTexturesDataTable")
+
+; Node renames
++ClassRedirects=(OldName="/Script/CustomizableObjectEditor.CustomizableObjectNodeMeshClipDeform", NewName="/Script/CustomizableObjectEditor.CustomizableObjectNodeModifierClipDeform")
++ClassRedirects=(OldName="/Script/CustomizableObjectEditor.CustomizableObjectNodeMeshClipWithMesh", NewName="/Script/CustomizableObjectEditor.CustomizableObjectNodeModifierClipWithMesh")
++ClassRedirects=(OldName="/Script/CustomizableObjectEditor.CustomizableObjectNodeMeshClipMorph", NewName="/Script/CustomizableObjectEditor.CustomizableObjectNodeModifierClipMorph")
++ClassRedirects=(OldName="/Script/CustomizableObjectEditor.CustomizableObjectNodeExtendMaterial", NewName="/Script/CustomizableObjectEditor.CustomizableObjectNodeModifierExtendMeshSection")
++ClassRedirects=(OldName="/Script/CustomizableObjectEditor.CustomizableObjectNodeEditMaterial", NewName="/Script/CustomizableObjectEditor.CustomizableObjectNodeModifierEditMeshSection")
++ClassRedirects=(OldName="/Script/CustomizableObjectEditor.CustomizableObjectNodeMorphMaterial", NewName="/Script/CustomizableObjectEditor.CustomizableObjectNodeModifierMorphMeshSection")
++ClassRedirects=(OldName="/Script/CustomizableObjectEditor.CustomizableObjectNodeRemoveMesh", NewName="/Script/CustomizableObjectEditor.CustomizableObjectNodeModifierRemoveMesh")
++ClassRedirects=(OldName="/Script/CustomizableObjectEditor.CustomizableObjectNodeRemoveMeshBlocks", NewName="/Script/CustomizableObjectEditor.CustomizableObjectNodeModifierRemoveMeshBlocks")
++ClassRedirects=(OldName="/Script/CustomizableObjectEditor.CustomizableObjectNodeEditMaterialBase", NewName="/Script/CustomizableObjectEditor.CustomizableObjectNodeModifierEditMeshSectionBase")
++ClassRedirects=(OldName="/Script/CustomizableObjectEditor.CustomizableObjectNodeEditLayoutBlocks", NewName="/Script/CustomizableObjectEditor.CustomizableObjectNodeModifierEditLayoutBlocks")
diff --git a/Engine/Plugins/Experimental/Mutable/Source/CustomizableObject/Internal/MuCO/CustomizableObjectCustomVersion.h b/Engine/Plugins/Experimental/Mutable/Source/CustomizableObject/Internal/MuCO/CustomizableObjectCustomVersion.h
--- a/Engine/Plugins/Experimental/Mutable/Source/CustomizableObject/Internal/MuCO/CustomizableObjectCustomVersion.h
+++ b/Engine/Plugins/Experimental/Mutable/Source/CustomizableObject/Internal/MuCO/CustomizableObjectCustomVersion.h
@@ -137,6 +137,8 @@ struct CUSTOMIZABLEOBJECT_API FCustomizableObjectCustomVersion
 		MovedCompatibilityFromPostBackwardsCompatibleFixup,
 
 		FixModifierPin, 
+		
+		ConvertEditAndExtendToModifiers,
 
 		// -----<new versions can be added above this line>-------------------------------------------------
 		VersionPlusOne,
diff --git a/Engine/Plugins/Experimental/Mutable/Source/CustomizableObjectEditor/Private/MuCOE/CustomizableObjectEditor.cpp b/Engine/Plugins/Experimental/Mutable/Source/CustomizableObjectEditor/Private/MuCOE/CustomizableObjectEditor.cpp
--- a/Engine/Plugins/Experimental/Mutable/Source/CustomizableObjectEditor/Private/MuCOE/CustomizableObjectEditor.cpp
+++ b/Engine/Plugins/Experimental/Mutable/Source/CustomizableObjectEditor/Private/MuCOE/CustomizableObjectEditor.cpp
@@ -43,8 +43,8 @@
 #include "MuCOE/Nodes/CustomizableObjectNodeGroupProjectorParameter.h"
 #include "MuCOE/Nodes/CustomizableObjectNodeMaterial.h"
 #include "MuCOE/Nodes/CustomizableObjectNodeMaterialVariation.h"
-#include "MuCOE/Nodes/CustomizableObjectNodeMeshClipMorph.h"
-#include "MuCOE/Nodes/CustomizableObjectNodeMeshClipWithMesh.h"
+#include "MuCOE/Nodes/CustomizableObjectNodeModifierClipMorph.h"
+#include "MuCOE/Nodes/CustomizableObjectNodeModifierClipWithMesh.h"
 #include "MuCOE/Nodes/CustomizableObjectNodeObject.h"
 #include "MuCOE/Nodes/CustomizableObjectNodeObjectGroup.h"
 #include "MuCOE/Nodes/CustomizableObjectNodeProjectorConstant.h"
@@ -1000,7 +1000,7 @@ void FCustomizableObjectEditor::HideGizmoProjectorParameter()
 }
 
 
-void FCustomizableObjectEditor::ShowGizmoClipMorph(UCustomizableObjectNodeMeshClipMorph& Node)
+void FCustomizableObjectEditor::ShowGizmoClipMorph(UCustomizableObjectNodeModifierClipMorph& Node)
 {
 	if (Node.BoneName == FName())
 	{	
@@ -1034,7 +1034,7 @@ void FCustomizableObjectEditor::HideGizmoClipMorph()
 	for (FGraphPanelSelectionSet::TConstIterator NodeIt( SelectedNodes ); NodeIt; ++NodeIt)
 	{
 		const UObject* Node = *NodeIt;
-		if (Node->IsA<UCustomizableObjectNodeMeshClipMorph>())
+		if (Node->IsA<UCustomizableObjectNodeModifierClipMorph>())
 		{
 			GraphEditor->ClearSelectionSet();
 			break;
@@ -1043,7 +1043,7 @@ void FCustomizableObjectEditor::HideGizmoClipMorph()
 }
 
 
-void FCustomizableObjectEditor::ShowGizmoClipMesh(UCustomizableObjectNodeMeshClipWithMesh& Node)
+void FCustomizableObjectEditor::ShowGizmoClipMesh(UCustomizableObjectNodeModifierClipWithMesh& Node)
 {
 	UObject* ClipMesh = nullptr;
 	int32 LODIndex = 0;
@@ -1113,7 +1113,7 @@ void FCustomizableObjectEditor::HideGizmoClipMesh()
 	for (FGraphPanelSelectionSet::TConstIterator NodeIt( SelectedNodes ); NodeIt; ++NodeIt)
 	{
 		const UObject* Node = *NodeIt;
-		if (Node->IsA<UCustomizableObjectNodeMeshClipWithMesh>())
+		if (Node->IsA<UCustomizableObjectNodeModifierClipWithMesh>())
 		{
 			GraphEditor->ClearSelectionSet();
 			break;
@@ -1717,13 +1717,13 @@ void FCustomizableObjectEditor::OnSelectedGraphNodesChanged(const FGraphPanelSel
 			return;
 		}
 
-		if (UCustomizableObjectNodeMeshClipMorph* NodeMeshClipMorph = Cast<UCustomizableObjectNodeMeshClipMorph>(Objects[0]))
+		if (UCustomizableObjectNodeModifierClipMorph* NodeModifierClipMorph = Cast<UCustomizableObjectNodeModifierClipMorph>(Objects[0]))
 		{		
-			ShowGizmoClipMorph(*NodeMeshClipMorph);			
+			ShowGizmoClipMorph(*NodeModifierClipMorph);
 		}
-		else if (UCustomizableObjectNodeMeshClipWithMesh* NodeMeshClipWithMesh = Cast<UCustomizableObjectNodeMeshClipWithMesh>(Objects[0]))
+		else if (UCustomizableObjectNodeModifierClipWithMesh* NodeModifierClipWithMesh = Cast<UCustomizableObjectNodeModifierClipWithMesh>(Objects[0]))
 		{
-			ShowGizmoClipMesh(*NodeMeshClipWithMesh);			
+			ShowGizmoClipMesh(*NodeModifierClipWithMesh);
 		}
 		else if (UCustomizableObjectNodeProjectorConstant* NodeProjectorConstant = Cast<UCustomizableObjectNodeProjectorConstant>(Objects[0]))
 		{
diff --git a/Engine/Plugins/Experimental/Mutable/Source/CustomizableObjectEditor/Private/MuCOE/CustomizableObjectEditor.h b/Engine/Plugins/Experimental/Mutable/Source/CustomizableObjectEditor/Private/MuCOE/CustomizableObjectEditor.h
--- a/Engine/Plugins/Experimental/Mutable/Source/CustomizableObjectEditor/Private/MuCOE/CustomizableObjectEditor.h
+++ b/Engine/Plugins/Experimental/Mutable/Source/CustomizableObjectEditor/Private/MuCOE/CustomizableObjectEditor.h
@@ -28,8 +28,8 @@ class SWidget;
 class UCustomizableObject;
 class UCustomizableObjectInstance;
 class UCustomizableObjectNode;
-class UCustomizableObjectNodeMeshClipMorph;
-class UCustomizableObjectNodeMeshClipWithMesh;
+class UCustomizableObjectNodeModifierClipMorph;
+class UCustomizableObjectNodeModifierClipWithMesh;
 class UCustomizableObjectNodeObject;
 class UCustomizableObjectNodeProjectorConstant;
 class UCustomizableObjectNodeProjectorParameter;
@@ -206,9 +206,9 @@ public:
 	virtual void HideGizmoProjectorNodeProjectorParameter() override;
 	virtual void ShowGizmoProjectorParameter(const FString& ParamName, int32 RangeIndex = -1) override;
 	virtual void HideGizmoProjectorParameter() override;
-	virtual void ShowGizmoClipMorph(UCustomizableObjectNodeMeshClipMorph& Node) override;
+	virtual void ShowGizmoClipMorph(UCustomizableObjectNodeModifierClipMorph& Node) override;
 	virtual void HideGizmoClipMorph() override;
-	virtual void ShowGizmoClipMesh(UCustomizableObjectNodeMeshClipWithMesh& Node) override;
+	virtual void ShowGizmoClipMesh(UCustomizableObjectNodeModifierClipWithMesh& Node) override;
 	virtual void HideGizmoClipMesh() override;
 	virtual void ShowGizmoLight(ULightComponent& SelectedLight) override;
 	virtual void HideGizmoLight() override;
diff --git a/Engine/Plugins/Experimental/Mutable/Source/CustomizableObjectEditor/Private/MuCOE/CustomizableObjectEditorModule.cpp b/Engine/Plugins/Experimental/Mutable/Source/CustomizableObjectEditor/Private/MuCOE/CustomizableObjectEditorModule.cpp
--- a/Engine/Plugins/Experimental/Mutable/Source/CustomizableObjectEditor/Private/MuCOE/CustomizableObjectEditorModule.cpp
+++ b/Engine/Plugins/Experimental/Mutable/Source/CustomizableObjectEditor/Private/MuCOE/CustomizableObjectEditorModule.cpp
@@ -30,32 +30,33 @@
 #include "MuCOE/CustomizableObjectVersionBridge.h"
 #include "MuCOE/CustomizableObjectNodeObjectGroupDetails.h"
 #include "MuCOE/Nodes/CustomizableObjectNodeCopyMaterial.h"
-#include "MuCOE/Nodes/CustomizableObjectNodeEditMaterial.h"
-#include "MuCOE/Nodes/CustomizableObjectNodeEditMaterialDetails.h"
-#include "MuCOE/Nodes/CustomizableObjectNodeEditMaterialBaseDetails.h"
-#include "MuCOE/Nodes/CustomizableObjectNodeExtendMaterial.h"
 #include "MuCOE/Nodes/CustomizableObjectNodeExternalPin.h"
 #include "MuCOE/Nodes/CustomizableObjectNodeExternalPinDetails.h"
 #include "MuCOE/Nodes/CustomizableObjectNodeLayoutBlocks.h"
 #include "MuCOE/Nodes/CustomizableObjectNodeLayoutBlocksDetails.h"
 #include "MuCOE/Nodes/CustomizableObjectNodeMaterialDetails.h"
-#include "MuCOE/Nodes/CustomizableObjectNodeMeshClipMorph.h"
-#include "MuCOE/Nodes/CustomizableObjectNodeMeshClipMorphDetails.h"
-#include "MuCOE/Nodes/CustomizableObjectNodeMeshClipWithMesh.h"
-#include "MuCOE/Nodes/CustomizableObjectNodeMeshClipWithMeshDetails.h"
+#include "MuCOE/Nodes/CustomizableObjectNodeModifierEditMeshSection.h"
+#include "MuCOE/Nodes/CustomizableObjectNodeModifierEditMeshSectionDetails.h"
+#include "MuCOE/Nodes/CustomizableObjectNodeModifierExtendMeshSection.h"
+#include "MuCOE/Nodes/CustomizableObjectNodeModifierMorphMeshSection.h"
+#include "MuCOE/Nodes/CustomizableObjectNodeModifierMorphMeshSectionDetails.h"
+#include "MuCOE/Nodes/CustomizableObjectNodeModifierRemoveMeshBlocks.h"
+#include "MuCOE/Nodes/CustomizableObjectNodeModifierRemoveMeshBlocksDetails.h"
+#include "MuCOE/Nodes/CustomizableObjectNodeModifierRemoveMesh.h"
+#include "MuCOE/Nodes/CustomizableObjectNodeModifierRemoveMeshDetails.h"
+#include "MuCOE/Nodes/CustomizableObjectNodeModifierClipMorph.h"
+#include "MuCOE/Nodes/CustomizableObjectNodeModifierClipMorphDetails.h"
+#include "MuCOE/Nodes/CustomizableObjectNodeModifierClipWithMesh.h"
+#include "MuCOE/Nodes/CustomizableObjectNodeModifierClipWithMeshDetails.h"
 #include "MuCOE/Nodes/CustomizableObjectNodeMeshMorph.h"
 #include "MuCOE/Nodes/CustomizableObjectNodeMeshMorphDetails.h"
 #include "MuCOE/Nodes/CustomizableObjectNodeMeshReshapeCommon.h"
 #include "MuCOE/Nodes/CustomizableObjectNodeMeshReshapeSelectionDetails.h"
-#include "MuCOE/Nodes/CustomizableObjectNodeMorphMaterial.h"
-#include "MuCOE/Nodes/CustomizableObjectNodeMorphMaterialDetails.h"
 #include "MuCOE/Nodes/CustomizableObjectNodeObjectDetails.h"
 #include "MuCOE/Nodes/CustomizableObjectNodeObjectGroup.h"
 #include "MuCOE/Nodes/CustomizableObjectNodeDetails.h"
 #include "MuCOE/Nodes/CustomizableObjectNodeProjectorConstant.h"
 #include "MuCOE/Nodes/CustomizableObjectNodeProjectorParameterDetails.h"
-#include "MuCOE/Nodes/CustomizableObjectNodeRemoveMeshBlocks.h"
-#include "MuCOE/Nodes/CustomizableObjectNodeRemoveMeshBlocksDetails.h"
 #include "MuCOE/Nodes/CustomizableObjectNodeSkeletalMesh.h"
 #include "MuCOE/Nodes/CustomizableObjectNodeSkeletalMeshDetails.h"
 #include "MuCOE/Nodes/CustomizableObjectNodeStaticMesh.h"
@@ -72,7 +73,6 @@
 #include "MuCOE/GraphTraversal.h"
 #include "MuCOE/CustomizableObjectInstanceBaker.h"
 #include "Editor.h"
-#include "MuCOE/Nodes/CustomizableObjectNodeRemoveMesh.h"
 
 class AActor;
 class FString;
@@ -238,18 +238,17 @@ void FCustomizableObjectEditorModule::StartupModule()
 	// Nodes
 	FPropertyEditorModule& PropertyModule = FModuleManager::LoadModuleChecked<FPropertyEditorModule>("PropertyEditor");
 	RegisterCustomDetails(PropertyModule, UCustomizableObjectNodeLayoutBlocks::StaticClass(), FOnGetDetailCustomizationInstance::CreateStatic(&FCustomizableObjectNodeLayoutBlocksDetails::MakeInstance));
-	RegisterCustomDetails(PropertyModule, UCustomizableObjectNodeEditMaterial::StaticClass(), FOnGetDetailCustomizationInstance::CreateStatic(&FCustomizableObjectNodeEditMaterialDetails::MakeInstance));
-	RegisterCustomDetails(PropertyModule, UCustomizableObjectNodeRemoveMesh::StaticClass(), FOnGetDetailCustomizationInstance::CreateStatic(&FCustomizableObjectNodeEditMaterialBaseDetails::MakeInstance));
-	RegisterCustomDetails(PropertyModule, UCustomizableObjectNodeRemoveMeshBlocks::StaticClass(), FOnGetDetailCustomizationInstance::CreateStatic(&FCustomizableObjectNodeRemoveMeshBlocksDetails::MakeInstance));
-	RegisterCustomDetails(PropertyModule, UCustomizableObjectNodeExtendMaterial::StaticClass(), FOnGetDetailCustomizationInstance::CreateStatic(&FCustomizableObjectNodeParentedMaterialDetails::MakeInstance));
-	RegisterCustomDetails(PropertyModule, UCustomizableObjectNodeMorphMaterial::StaticClass(), FOnGetDetailCustomizationInstance::CreateStatic(&FCustomizableObjectNodeMorphMaterialDetails::MakeInstance));
+	RegisterCustomDetails(PropertyModule, UCustomizableObjectNodeModifierEditMeshSection::StaticClass(), FOnGetDetailCustomizationInstance::CreateStatic(&FCustomizableObjectNodeModifierEditMeshSectionDetails::MakeInstance));
+	RegisterCustomDetails(PropertyModule, UCustomizableObjectNodeModifierRemoveMesh::StaticClass(), FOnGetDetailCustomizationInstance::CreateStatic(&FCustomizableObjectNodeModifierRemoveMeshDetails::MakeInstance));
+	RegisterCustomDetails(PropertyModule, UCustomizableObjectNodeModifierRemoveMeshBlocks::StaticClass(), FOnGetDetailCustomizationInstance::CreateStatic(&FCustomizableObjectNodeModifierRemoveMeshBlocksDetails::MakeInstance));
+	RegisterCustomDetails(PropertyModule, UCustomizableObjectNodeModifierMorphMeshSection::StaticClass(), FOnGetDetailCustomizationInstance::CreateStatic(&FCustomizableObjectNodeModifierMorphMeshSectionDetails::MakeInstance));
 	RegisterCustomDetails(PropertyModule, UCustomizableObjectNodeObject::StaticClass(), FOnGetDetailCustomizationInstance::CreateStatic(&FCustomizableObjectNodeObjectDetails::MakeInstance));
 	RegisterCustomDetails(PropertyModule, UCustomizableObjectNodeObjectGroup::StaticClass(), FOnGetDetailCustomizationInstance::CreateStatic(&FCustomizableObjectNodeObjectGroupDetails::MakeInstance));
 	RegisterCustomDetails(PropertyModule, UCustomizableObjectNodeProjectorParameter::StaticClass(), FOnGetDetailCustomizationInstance::CreateStatic(&FCustomizableObjectNodeProjectorParameterDetails::MakeInstance));
 	RegisterCustomDetails(PropertyModule, UCustomizableObjectNodeProjectorConstant::StaticClass(), FOnGetDetailCustomizationInstance::CreateStatic(&FCustomizableObjectNodeProjectorParameterDetails::MakeInstance));
 	RegisterCustomDetails(PropertyModule, UCustomizableObjectNodeMeshMorph::StaticClass(), FOnGetDetailCustomizationInstance::CreateStatic(&FCustomizableObjectNodeMeshMorphDetails::MakeInstance));
-	RegisterCustomDetails(PropertyModule, UCustomizableObjectNodeMeshClipMorph::StaticClass(), FOnGetDetailCustomizationInstance::CreateStatic(&FCustomizableObjectNodeMeshClipMorphDetails::MakeInstance));
-	RegisterCustomDetails(PropertyModule, UCustomizableObjectNodeMeshClipWithMesh::StaticClass(), FOnGetDetailCustomizationInstance::CreateStatic(&FCustomizableObjectNodeMeshClipWithMeshDetails::MakeInstance));
+	RegisterCustomDetails(PropertyModule, UCustomizableObjectNodeModifierClipMorph::StaticClass(), FOnGetDetailCustomizationInstance::CreateStatic(&FCustomizableObjectNodeModifierClipMorphDetails::MakeInstance));
+	RegisterCustomDetails(PropertyModule, UCustomizableObjectNodeModifierClipWithMesh::StaticClass(), FOnGetDetailCustomizationInstance::CreateStatic(&FCustomizableObjectNodeModifierClipWithMeshDetails::MakeInstance));
 	RegisterCustomDetails(PropertyModule, UCustomizableObjectNodeExternalPin::StaticClass(), FOnGetDetailCustomizationInstance::CreateStatic(&FCustomizableObjectNodeExternalPinDetails::MakeInstance));
 	RegisterCustomDetails(PropertyModule, UCustomizableObjectNodeMaterial::StaticClass(), FOnGetDetailCustomizationInstance::CreateStatic(&FCustomizableObjectNodeMaterialDetails::MakeInstance));
 	RegisterCustomDetails(PropertyModule, UCustomizableObjectNodeSkeletalMesh::StaticClass(), FOnGetDetailCustomizationInstance::CreateStatic(&FCustomizableObjectNodeSkeletalMeshDetails::MakeInstance));
diff --git a/Engine/Plugins/Experimental/Mutable/Source/CustomizableObjectEditor/Private/MuCOE/CustomizableObjectEditorViewportClient.cpp b/Engine/Plugins/Experimental/Mutable/Source/CustomizableObjectEditor/Private/MuCOE/CustomizableObjectEditorViewportClient.cpp
--- a/Engine/Plugins/Experimental/Mutable/Source/CustomizableObjectEditor/Private/MuCOE/CustomizableObjectEditorViewportClient.cpp
+++ b/Engine/Plugins/Experimental/Mutable/Source/CustomizableObjectEditor/Private/MuCOE/CustomizableObjectEditorViewportClient.cpp
@@ -34,8 +34,8 @@
 #include "MuCO/CustomizableObjectSystem.h"
 #include "MuCOE/CustomizableObjectPreviewScene.h"
 #include "MuCOE/ICustomizableObjectInstanceEditor.h"
-#include "MuCOE/Nodes/CustomizableObjectNodeMeshClipMorph.h"
-#include "MuCOE/Nodes/CustomizableObjectNodeMeshClipWithMesh.h"
+#include "MuCOE/Nodes/CustomizableObjectNodeModifierClipMorph.h"
+#include "MuCOE/Nodes/CustomizableObjectNodeModifierClipWithMesh.h"
 #include "MuCOE/Nodes/CustomizableObjectNodeProjectorConstant.h"
 #include "MuCOE/Nodes/CustomizableObjectNodeProjectorParameter.h"
 #include "MuCOE/UnrealEditorPortabilityHelpers.h"
@@ -561,7 +561,7 @@ void FCustomizableObjectEditorViewportClient::DrawUVs(FViewport* InViewport, FCa
 }
 
 
-void FCustomizableObjectEditorViewportClient::ShowGizmoClipMorph(UCustomizableObjectNodeMeshClipMorph& NodeMeshClipMorph)
+void FCustomizableObjectEditorViewportClient::ShowGizmoClipMorph(UCustomizableObjectNodeModifierClipMorph& NodeMeshClipMorph)
 {
 	SetWidgetType(EWidgetType::ClipMorph);
 
@@ -599,7 +599,7 @@ void FCustomizableObjectEditorViewportClient::HideGizmoClipMorph()
 }
 
 
-void FCustomizableObjectEditorViewportClient::ShowGizmoClipMesh(UCustomizableObjectNodeMeshClipWithMesh& InClipMeshNode, UObject& ClipMesh, int32 LODIndex, int32 SectionIndex, int32 MaterialSlotIndex)
+void FCustomizableObjectEditorViewportClient::ShowGizmoClipMesh(UCustomizableObjectNodeModifierClipWithMesh& InClipMeshNode, UObject& ClipMesh, int32 LODIndex, int32 SectionIndex, int32 MaterialSlotIndex)
 {
 	HideGizmoClipMesh();
 
diff --git a/Engine/Plugins/Experimental/Mutable/Source/CustomizableObjectEditor/Private/MuCOE/EdGraphSchema_CustomizableObject.cpp b/Engine/Plugins/Experimental/Mutable/Source/CustomizableObjectEditor/Private/MuCOE/EdGraphSchema_CustomizableObject.cpp
--- a/Engine/Plugins/Experimental/Mutable/Source/CustomizableObjectEditor/Private/MuCOE/EdGraphSchema_CustomizableObject.cpp
+++ b/Engine/Plugins/Experimental/Mutable/Source/CustomizableObjectEditor/Private/MuCOE/EdGraphSchema_CustomizableObject.cpp
@@ -23,10 +23,8 @@
 #include "MuCOE/Nodes/CustomizableObjectNodeColorVariation.h"
 #include "MuCOE/Nodes/CustomizableObjectNodeCopyMaterial.h"
 #include "MuCOE/Nodes/CustomizableObjectNodeCurve.h"
-#include "MuCOE/Nodes/CustomizableObjectNodeEditMaterial.h"
 #include "MuCOE/Nodes/CustomizableObjectNodeEnumParameter.h"
 #include "MuCOE/Nodes/CustomizableObjectNodeExposePin.h"
-#include "MuCOE/Nodes/CustomizableObjectNodeExtendMaterial.h"
 #include "MuCOE/Nodes/CustomizableObjectNodeExternalPin.h"
 #include "MuCOE/Nodes/CustomizableObjectNodeFloatConstant.h"
 #include "MuCOE/Nodes/CustomizableObjectNodeFloatParameter.h"
@@ -37,9 +35,14 @@
 #include "MuCOE/Nodes/CustomizableObjectNodeLayoutBlocks.h"
 #include "MuCOE/Nodes/CustomizableObjectNodeMaterialVariation.h"
 #include "MuCOE/Nodes/CustomizableObjectNodeMaterialSwitch.h"
-#include "MuCOE/Nodes/CustomizableObjectNodeMeshClipDeform.h"
-#include "MuCOE/Nodes/CustomizableObjectNodeMeshClipMorph.h"
-#include "MuCOE/Nodes/CustomizableObjectNodeMeshClipWithMesh.h"
+#include "MuCOE/Nodes/CustomizableObjectNodeModifierExtendMeshSection.h"
+#include "MuCOE/Nodes/CustomizableObjectNodeModifierEditMeshSection.h"
+#include "MuCOE/Nodes/CustomizableObjectNodeModifierRemoveMesh.h"
+#include "MuCOE/Nodes/CustomizableObjectNodeModifierRemoveMeshBlocks.h"
+#include "MuCOE/Nodes/CustomizableObjectNodeModifierMorphMeshSection.h"
+#include "MuCOE/Nodes/CustomizableObjectNodeModifierClipDeform.h"
+#include "MuCOE/Nodes/CustomizableObjectNodeModifierClipMorph.h"
+#include "MuCOE/Nodes/CustomizableObjectNodeModifierClipWithMesh.h"
 #include "MuCOE/Nodes/CustomizableObjectNodeModifierClipWithUVMask.h"
 #include "MuCOE/Nodes/CustomizableObjectNodeMeshGeometryOperation.h"
 #include "MuCOE/Nodes/CustomizableObjectNodeMeshMorph.h"
@@ -48,15 +51,12 @@
 #include "MuCOE/Nodes/CustomizableObjectNodeMeshReshape.h"
 #include "MuCOE/Nodes/CustomizableObjectNodeMeshSwitch.h"
 #include "MuCOE/Nodes/CustomizableObjectNodeMeshVariation.h"
-#include "MuCOE/Nodes/CustomizableObjectNodeMorphMaterial.h"
 #include "MuCOE/Nodes/CustomizableObjectNodeObjectChild.h"
 #include "MuCOE/Nodes/CustomizableObjectNodeObjectGroup.h"
 #include "MuCOE/Nodes/CustomizableObjectNodeComponentMesh.h"
 #include "MuCOE/Nodes/CustomizableObjectNodeComponentVariation.h"
 #include "MuCOE/Nodes/CustomizableObjectNodeComponentSwitch.h"
 #include "MuCOE/Nodes/CustomizableObjectNodeProjectorConstant.h"
-#include "MuCOE/Nodes/CustomizableObjectNodeRemoveMesh.h"
-#include "MuCOE/Nodes/CustomizableObjectNodeRemoveMeshBlocks.h"
 #include "MuCOE/Nodes/CustomizableObjectNodeSkeletalMesh.h"
 #include "MuCOE/Nodes/CustomizableObjectNodeStaticMesh.h"
 #include "MuCOE/Nodes/CustomizableObjectNodeTable.h"
@@ -383,24 +383,32 @@ void UEdGraphSchema_CustomizableObject::GetGraphContextActions(FGraphContextMenu
 		UCustomizableObjectNode* ObjectTemplateNodes[]
 		{
 			ContextMenuBuilder.CreateTemplateNode<UCustomizableObjectNodeMaterial>(),
-			ContextMenuBuilder.CreateTemplateNode<UCustomizableObjectNodeExtendMaterial>(),
-			ContextMenuBuilder.CreateTemplateNode<UCustomizableObjectNodeRemoveMesh>(),
-			ContextMenuBuilder.CreateTemplateNode<UCustomizableObjectNodeRemoveMeshBlocks>(),
-			ContextMenuBuilder.CreateTemplateNode<UCustomizableObjectNodeEditMaterial>(),
 			ContextMenuBuilder.CreateTemplateNode<UCustomizableObjectNodeMaterialVariation>(),
 			ContextMenuBuilder.CreateTemplateNode<UCustomizableObjectNodeMaterialSwitch>(),
-			ContextMenuBuilder.CreateTemplateNode<UCustomizableObjectNodeMorphMaterial>(),
 			ContextMenuBuilder.CreateTemplateNode<UCustomizableObjectNodeCopyMaterial>(),
-			ContextMenuBuilder.CreateTemplateNode<UCustomizableObjectNodeMeshClipMorph>(),
-			ContextMenuBuilder.CreateTemplateNode<UCustomizableObjectNodeMeshClipWithMesh>(),
-			ContextMenuBuilder.CreateTemplateNode<UCustomizableObjectNodeMeshClipDeform>(),
-			ContextMenuBuilder.CreateTemplateNode<UCustomizableObjectNodeModifierClipWithUVMask>(),
 			ContextMenuBuilder.CreateTemplateNode<UCustomizableObjectNodeTable>(),
 		};
 
 		AddNewNodeCategoryActionsFiltered(ObjectTemplateNodes, ContextMenuBuilder, TEXT("Object"), GeneralGrouping, Filter);
 	}
 
+	{
+		UCustomizableObjectNode* ModifierTemplateNodes[]
+		{
+			ContextMenuBuilder.CreateTemplateNode<UCustomizableObjectNodeModifierExtendMeshSection>(),
+			ContextMenuBuilder.CreateTemplateNode<UCustomizableObjectNodeModifierRemoveMesh>(),
+			ContextMenuBuilder.CreateTemplateNode<UCustomizableObjectNodeModifierRemoveMeshBlocks>(),
+			ContextMenuBuilder.CreateTemplateNode<UCustomizableObjectNodeModifierEditMeshSection>(),
+			ContextMenuBuilder.CreateTemplateNode<UCustomizableObjectNodeModifierMorphMeshSection>(),
+			ContextMenuBuilder.CreateTemplateNode<UCustomizableObjectNodeModifierClipMorph>(),
+			ContextMenuBuilder.CreateTemplateNode<UCustomizableObjectNodeModifierClipWithMesh>(),
+			ContextMenuBuilder.CreateTemplateNode<UCustomizableObjectNodeModifierClipDeform>(),
+			ContextMenuBuilder.CreateTemplateNode<UCustomizableObjectNodeModifierClipWithUVMask>(),
+		};
+
+		AddNewNodeCategoryActionsFiltered(ModifierTemplateNodes, ContextMenuBuilder, TEXT("Modifier"), GeneralGrouping, Filter);
+	}
+
 	{
 		UCustomizableObjectNode* ComponentTemplateNodes[]
 		{
diff --git a/Engine/Plugins/Experimental/Mutable/Source/CustomizableObjectEditor/Private/MuCOE/GenerateMutableSource/GenerateMutableSource.cpp b/Engine/Plugins/Experimental/Mutable/Source/CustomizableObjectEditor/Private/MuCOE/GenerateMutableSource/GenerateMutableSource.cpp
--- a/Engine/Plugins/Experimental/Mutable/Source/CustomizableObjectEditor/Private/MuCOE/GenerateMutableSource/GenerateMutableSource.cpp
+++ b/Engine/Plugins/Experimental/Mutable/Source/CustomizableObjectEditor/Private/MuCOE/GenerateMutableSource/GenerateMutableSource.cpp
@@ -17,13 +17,13 @@
 #include "MuCOE/GenerateMutableSource/GenerateMutableSourceSurface.h"
 #include "MuCOE/GraphTraversal.h"
 #include "MuCOE/ICustomizableObjectExtensionNode.h"
-#include "MuCOE/Nodes/CustomizableObjectNodeEditMaterialBase.h"
-#include "MuCOE/Nodes/CustomizableObjectNodeExtendMaterial.h"
+#include "MuCOE/Nodes/CustomizableObjectNodeModifierEditMeshSectionBase.h"
+#include "MuCOE/Nodes/CustomizableObjectNodeModifierExtendMeshSection.h"
 #include "MuCOE/Nodes/CustomizableObjectNodeMaterial.h"
 #include "MuCOE/Nodes/CustomizableObjectNodeMaterialVariation.h"
 #include "MuCOE/Nodes/CustomizableObjectNodeMaterialSwitch.h"
 #include "MuCOE/Nodes/CustomizableObjectNodeMesh.h"
-#include "MuCOE/Nodes/CustomizableObjectNodeMeshClipWithMesh.h"
+#include "MuCOE/Nodes/CustomizableObjectNodeModifierClipWithMesh.h"
 #include "MuCOE/Nodes/CustomizableObjectNodeObjectGroup.h"
 #include "MuCOE/Nodes/CustomizableObjectNodeTable.h"
 #include "MuCOE/Nodes/CustomizableObjectNodeTexture.h"
@@ -48,7 +48,6 @@
 #include "MuT/NodeMeshConstant.h"
 #include "MuT/NodeObjectGroup.h"
 #include "MuT/NodeObjectNew.h"
-#include "MuT/NodeSurfaceEdit.h"
 #include "MuR/Mesh.h"
 #include "PhysicsEngine/PhysicsAsset.h"
 #include "PlatformInfo.h"
@@ -294,12 +293,12 @@ void FMutableGraphGenerationContext::GenerateClippingCOInternalTags()
 			It->Value[j]->RequiredTags.Add(TagName);
 		}
 
-		UCustomizableObjectNodeMeshClipWithMesh* CustomizableObjectNodeMeshClipWithMesh = It->Key;
+		UCustomizableObjectNodeModifierClipWithMesh* CustomizableObjectNodeModifierClipWithMesh = It->Key;
 
 		// Tag also those nodes in the material side of the clipping functionality (since a CO can have several materials with the same LOD,
 		// all of them are considered, so CustomizableObjectNodeMeshClipWithMesh->ArrayMaterialNodeToClipWithID contains the Guids of all
 		// the material nodes with the same name selected for being affected by the clipping operation
-		TArray<FGuid>& ArrayMaterialNodeToClipWithID = CustomizableObjectNodeMeshClipWithMesh->ArrayMaterialNodeToClipWithID;
+		TArray<FGuid>& ArrayMaterialNodeToClipWithID = CustomizableObjectNodeModifierClipWithMesh->ArrayMaterialNodeToClipWithID;
 
 		for (j = 0; j < ArrayMaterialNodeToClipWithID.Num(); ++j)
 		{
@@ -448,7 +447,16 @@ FMutableComponentInfo& FMutableGraphGenerationContext::GetCurrentComponentInfo()
 {
 	FMutableComponentInfo* CurrentComponentInfo = ComponentInfos.FindByPredicate(
 		[this](const FMutableComponentInfo& Component) { return Component.ComponentName == CurrentMeshComponent; });
-	check(CurrentComponentInfo);
+	
+	// Temp workaround to the problem of modifiers that generate meshes (like ExtendMeshSection) not having a known component at generation time.
+	// TODO: Actually detect all possible components and change this query in GenerationContext state to accomodate a set of "current components" instead of one.
+	//check(CurrentComponentInfo);
+
+	if (!CurrentComponentInfo)
+	{
+		check(!ComponentInfos.IsEmpty());
+		CurrentComponentInfo = &ComponentInfos[0];
+	}
 
 	return *CurrentComponentInfo;
 }
@@ -832,7 +840,7 @@ mu::Ptr<mu::NodeObject> GenerateMutableSource(const UEdGraphPin * Pin, FMutableG
 		return static_cast<mu::NodeObject*>(Generated->Node.get());
 	}
 
-	mu::NodeObjectPtr Result;
+	mu::Ptr<mu::NodeObject> Result;
 	
 	if (const UCustomizableObjectNodeObject* TypedNodeObj = Cast<UCustomizableObjectNodeObject>(Node))
 	{
@@ -1302,7 +1310,7 @@ mu::Ptr<mu::NodeObject> GenerateMutableSource(const UEdGraphPin * Pin, FMutableG
 			FString* SelectedOptionName = GenerationContext.ParamNamesToSelectedOptions.Find(TypedNodeGroup->GroupName); // If the param is in the map restrict to only the selected option
 			mu::NodeObjectPtr ChildNode = nullptr;
 
-			if (bConnectAtLeastTheLastChild || !SelectedOptionName || *SelectedOptionName == CustomizableObjectNodeObject->ObjectName)
+			if (bConnectAtLeastTheLastChild || !SelectedOptionName || (CustomizableObjectNodeObject && *SelectedOptionName == CustomizableObjectNodeObject->ObjectName) )
 			{
 				bAtLeastOneConnected = true;
 
@@ -1410,7 +1418,7 @@ mu::Ptr<mu::NodeObject> GenerateMutableSource(const UEdGraphPin * Pin, FMutableG
 			FString* SelectedOptionName = GenerationContext.ParamNamesToSelectedOptions.Find(TypedNodeGroup->GroupName); // If the param is in the map restrict to only the selected option
 			mu::NodeObjectPtr ChildNode = nullptr;
 
-			if (bConnectAtLeastTheLastChild || !SelectedOptionName || *SelectedOptionName == CustomizableObjectNodeObject->ObjectName)
+			if (bConnectAtLeastTheLastChild || !SelectedOptionName || (CustomizableObjectNodeObject  && *SelectedOptionName == CustomizableObjectNodeObject->ObjectName) )
 			{
 				bAtLeastOneConnected = true;
 
@@ -1595,32 +1603,12 @@ bool AffectsCurrentComponent(const UEdGraphPin* Pin, FMutableGraphGenerationCont
 	{
 		ComponentName = TypedNodeMat->GetMeshComponentName();
 	}
-	else if (const UCustomizableObjectNodeExtendMaterial* TypedNodeExt = Cast<UCustomizableObjectNodeExtendMaterial>(Node))
-	{
-		OriginalParentMaterialNode = TypedNodeExt->GetParentMaterialNode();
-		ComponentName = OriginalParentMaterialNode ? OriginalParentMaterialNode->GetMeshComponentName() : GenerationContext.CurrentMeshComponent;
-	}
-	else if (const UCustomizableObjectNodeEditMaterialBase* TypedNodeEdit = Cast<UCustomizableObjectNodeEditMaterialBase>(Node))
-	{
-		OriginalParentMaterialNode = TypedNodeEdit->GetParentMaterialNode();
-		ComponentName = OriginalParentMaterialNode ? OriginalParentMaterialNode->GetMeshComponentName() : GenerationContext.CurrentMeshComponent;
-	}
+
 	else if (const UCustomizableObjectNodeModifierBase* TypedNodeModifier = Cast<UCustomizableObjectNodeModifierBase>(Node))
 	{
-		// Because of the current implementation, modifiers affect all components at lod level. If there is only one component it is ok, but otherwise rise an error.
-		if (GenerationContext.NumMeshComponentsInRoot == 1)
-		{
-			ComponentName = GenerationContext.ComponentInfos[0].ComponentName;
-			return true;
-		}
-		else
-		{
-			// This case is not supported yet
-			FString Msg = FString::Printf(TEXT("Error! Node has modifiers when using multiple components in root object. This is currently not supported."));
-			GenerationContext.Compiler->CompilerLog(FText::FromString(Msg), Node, EMessageSeverity::Error);
-			ComponentName = GenerationContext.ComponentInfos[0].ComponentName;
-			return false;
-		}
+		// Modifiers can affect any component depending on tags.
+		// TODO: Scan
+		return true;
 	}
 	else
 	{
diff --git a/Engine/Plugins/Experimental/Mutable/Source/CustomizableObjectEditor/Private/MuCOE/GenerateMutableSource/GenerateMutableSource.h b/Engine/Plugins/Experimental/Mutable/Source/CustomizableObjectEditor/Private/MuCOE/GenerateMutableSource/GenerateMutableSource.h
--- a/Engine/Plugins/Experimental/Mutable/Source/CustomizableObjectEditor/Private/MuCOE/GenerateMutableSource/GenerateMutableSource.h
+++ b/Engine/Plugins/Experimental/Mutable/Source/CustomizableObjectEditor/Private/MuCOE/GenerateMutableSource/GenerateMutableSource.h
@@ -14,6 +14,7 @@
 #include "MuCOE/Nodes/CustomizableObjectNodeObject.h"
 #include "MuT/NodeComponentNew.h"
 #include "MuT/NodeImageConstant.h"
+#include "MuT/NodeMeshConstant.h"
 #include "MuT/NodeMeshApplyPose.h"
 #include "MuT/NodeModifierMeshClipWithMesh.h"
 #include "MuT/NodeObject.h"
@@ -29,7 +30,7 @@
 #include "UObject/Package.h"
 
 class UCustomizableObjectNodeMaterialBase;
-class UCustomizableObjectNodeMeshClipWithMesh;
+class UCustomizableObjectNodeModifierClipWithMesh;
 class UCustomizableObjectNodeTable;
 struct FCustomizableObjectClothingAssetData;
 
@@ -717,8 +718,8 @@ struct FMutableGraphGenerationContext
 	// Check if the Id of the node Node already exists, if it's new adds it to NodeIds array, otherwise, returns new Id
 	const FGuid GetNodeIdUnique(const UCustomizableObjectNode* Node);
 
-	/** Generates new tags for the UCustomizableObjectNodeMeshClipWithMesh nodes which have assigned a CO in the
-	* UCustomizableObjectNodeMeshClipWithMesh::CustomizableObjectToClipWith field
+	/** Generates new tags for the UCustomizableObjectNodeModifierClipWithMesh nodes which have assigned a CO in the
+	* UCustomizableObjectNodeModifierClipWithMesh::CustomizableObjectToClipWith field
 	* @return nothing */
 	void GenerateClippingCOInternalTags();
 
@@ -772,7 +773,7 @@ struct FMutableGraphGenerationContext
 	TMap<class UCustomizableObjectNodeMaterial*, TArray<mu::Ptr<mu::NodeSurfaceNew>>> MapMaterialNodeToMutableSurfaceNodeArray;
 
 	// Map with pairs (Unreal Mutable clip mesh node, array with its corresponding Mutable mesh modifier nodes) to add tags
-	TMap<class UCustomizableObjectNodeMeshClipWithMesh*, TArray<mu::Ptr<mu::NodeModifierMeshClipWithMesh>>> MapClipMeshNodeToMutableClipMeshNodeArray;
+	TMap<class UCustomizableObjectNodeModifierClipWithMesh*, TArray<mu::Ptr<mu::NodeModifierMeshClipWithMesh>>> MapClipMeshNodeToMutableClipMeshNodeArray;
 
 	// Data used for MorphTarget reconstruction.
 	TMap<uint32, FRealTimeMorphMeshData> RealTimeMorphTargetPerMeshData;
diff --git a/Engine/Plugins/Experimental/Mutable/Source/CustomizableObjectEditor/Private/MuCOE/GenerateMutableSource/GenerateMutableSourceGroupProjector.cpp b/Engine/Plugins/Experimental/Mutable/Source/CustomizableObjectEditor/Private/MuCOE/GenerateMutableSource/GenerateMutableSourceGroupProjector.cpp
--- a/Engine/Plugins/Experimental/Mutable/Source/CustomizableObjectEditor/Private/MuCOE/GenerateMutableSource/GenerateMutableSourceGroupProjector.cpp
+++ b/Engine/Plugins/Experimental/Mutable/Source/CustomizableObjectEditor/Private/MuCOE/GenerateMutableSource/GenerateMutableSourceGroupProjector.cpp
@@ -9,7 +9,7 @@
 #include "MuCOE/GraphTraversal.h"
 #include "MuCOE/GenerateMutableSource/GenerateMutableSourceProjector.h"
 #include "MuCOE/Nodes/CustomizableObjectNodeAnimationPose.h"
-#include "MuCOE/Nodes/CustomizableObjectNodeExtendMaterial.h"
+#include "MuCOE/Nodes/CustomizableObjectNodeModifierExtendMeshSection.h"
 #include "MuCOE/Nodes/CustomizableObjectNodeGroupProjectorParameter.h"
 #include "MuCOE/Nodes/CustomizableObjectNodeMaterial.h"
 #include "MuCOE/Nodes/CustomizableObjectNodeObjectGroup.h"
@@ -32,7 +32,7 @@ class UPoseAsset;
 
 
 mu::NodeImagePtr GenerateMutableGroupProjection(const int32 NodeLOD, const int32 ImageIndex, mu::NodeMeshPtr MeshNode, FMutableGraphGenerationContext& GenerationContext,
-	UCustomizableObjectNodeMaterialBase* TypedNodeMat, UCustomizableObjectNodeExtendMaterial* TypedNodeExt, bool& bShareProjectionTexturesBetweenLODs, bool& bIsGroupProjectorImage,
+	UCustomizableObjectNodeMaterialBase* TypedNodeMat, UCustomizableObjectNodeModifierExtendMeshSection* TypedNodeExt, bool& bShareProjectionTexturesBetweenLODs, bool& bIsGroupProjectorImage,
 	UTexture2D*& GroupProjectionReferenceTexture, TMap<FString, float>& TextureNameToProjectionResFactor, FString& AlternateResStateName)
 {
 	check(static_cast<bool>(TypedNodeMat) != static_cast<bool>(TypedNodeExt)); // XOr. TypedNodeMat valid or TypedNodeExt valid. At least one valid.
@@ -42,8 +42,6 @@ mu::NodeImagePtr GenerateMutableGroupProjection(const int32 NodeLOD, const int32
 		return mu::NodeImagePtr();
 	}
 
-	const UCustomizableObjectNodeMaterialBase* NodeMaterial = TypedNodeMat ? TypedNodeMat : TypedNodeExt->GetParentMaterialNode();
-
 	TArray<mu::Ptr<mu::NodeImageProject>> ImageNodes;
 	TArray<FGroupProjectorTempData> ImageNodes_ProjectorTempData;
 
@@ -73,7 +71,17 @@ mu::NodeImagePtr GenerateMutableGroupProjection(const int32 NodeLOD, const int32
 
 		const bool bProjectToImage = [&]
 		{
-			const FString ParameterName = NodeMaterial->GetParameterName(EMaterialParameterType::Texture, ImageIndex).ToString();
+			FString ParameterName;
+			
+			if (TypedNodeMat)
+			{
+				ParameterName = TypedNodeMat->GetParameterName(EMaterialParameterType::Texture, ImageIndex).ToString();
+			}
+			else
+			{
+				const FNodeMaterialParameterId ImageId = TypedNodeExt->GetParameterId(EMaterialParameterType::Texture, ImageIndex);
+				return TypedNodeExt->UsesImage(ImageId);
+			}
 
 			return ParameterName == ProjectorTempData.CustomizableObjectNodeGroupProjectorParameter->MaterialChannelNameToConnect;
 		}();
@@ -88,7 +96,7 @@ mu::NodeImagePtr GenerateMutableGroupProjection(const int32 NodeLOD, const int32
 				}
 				else
 				{
-					const FNodeMaterialParameterId ImageId = TypedNodeExt->GetParentMaterialNode()->GetParameterId(EMaterialParameterType::Texture, ImageIndex);
+					const FNodeMaterialParameterId ImageId = TypedNodeExt->GetParameterId(EMaterialParameterType::Texture, ImageIndex);
 					return TypedNodeExt->UsesImage(ImageId);
 				}
 			}();
@@ -103,12 +111,12 @@ mu::NodeImagePtr GenerateMutableGroupProjection(const int32 NodeLOD, const int32
 					}
 					else
 					{
-						return TypedNodeExt->GetParentMaterialNode()->GetParameterName(EMaterialParameterType::Texture, ImageIndex).ToString();
+						return TypedNodeExt->GetParameterName(EMaterialParameterType::Texture, ImageIndex).ToString();
 					}
 				}();
 				
 				FString msg = FString::Printf(TEXT("Material image [%s] is connected to an image but will be replaced by a Group Projector."), *ImageName);
-				GenerationContext.Compiler->CompilerLog(FText::FromString(msg), TypedNodeMat ? TypedNodeMat : TypedNodeExt->GetParentMaterialNode());
+				GenerationContext.Compiler->CompilerLog(FText::FromString(msg), TypedNodeMat);
 				continue;
 			}
 			
@@ -165,7 +173,7 @@ mu::NodeImagePtr GenerateMutableGroupProjection(const int32 NodeLOD, const int32
 					if (!NodeMeshApplyPose)
 					{
 						FString msg = FString::Printf(TEXT("Couldn't get bone transform information from a Pose Asset."));
-						GenerationContext.Compiler->CompilerLog(FText::FromString(msg), TypedNodeMat ? TypedNodeMat : TypedNodeExt->GetParentMaterialNode());
+						GenerationContext.Compiler->CompilerLog(FText::FromString(msg), TypedNodeMat);
 					}
 
 					MeshSwitchNode->SetOption(SelectorIndex + 1, NodeMeshApplyPose);
diff --git a/Engine/Plugins/Experimental/Mutable/Source/CustomizableObjectEditor/Private/MuCOE/GenerateMutableSource/GenerateMutableSourceGroupProjector.h b/Engine/Plugins/Experimental/Mutable/Source/CustomizableObjectEditor/Private/MuCOE/GenerateMutableSource/GenerateMutableSourceGroupProjector.h
--- a/Engine/Plugins/Experimental/Mutable/Source/CustomizableObjectEditor/Private/MuCOE/GenerateMutableSource/GenerateMutableSourceGroupProjector.h
+++ b/Engine/Plugins/Experimental/Mutable/Source/CustomizableObjectEditor/Private/MuCOE/GenerateMutableSource/GenerateMutableSourceGroupProjector.h
@@ -5,7 +5,7 @@
 #include "Containers/Map.h"
 #include "MuT/NodeSurfaceNew.h"
 
-class UCustomizableObjectNodeExtendMaterial;
+class UCustomizableObjectNodeModifierExtendMeshSection;
 class FString;
 class UCustomizableObjectNodeMaterial;
 class UCustomizableObjectNodeMaterialBase;
@@ -16,7 +16,7 @@ struct FMutableGraphGenerationContext;
 
 
 mu::NodeImagePtr GenerateMutableGroupProjection(const int32 NodeLOD, const int32 ImageIndex, mu::NodeMeshPtr MeshNode, FMutableGraphGenerationContext& GenerationContext,
-	UCustomizableObjectNodeMaterialBase* TypedNodeMat, UCustomizableObjectNodeExtendMaterial* TypedNodeExt, bool& bShareProjectionTexturesBetweenLODs, bool& bIsGroupProjectorImage,
+	UCustomizableObjectNodeMaterialBase* TypedNodeMat, UCustomizableObjectNodeModifierExtendMeshSection* TypedNodeExt, bool& bShareProjectionTexturesBetweenLODs, bool& bIsGroupProjectorImage,
 	UTexture2D*& GroupProjectionReferenceTexture, TMap<FString, float>& TextureNameToProjectionResFactor, FString& AlternateResStateName);
 
 
diff --git a/Engine/Plugins/Experimental/Mutable/Source/CustomizableObjectEditor/Private/MuCOE/GenerateMutableSource/GenerateMutableSourceMesh.cpp b/Engine/Plugins/Experimental/Mutable/Source/CustomizableObjectEditor/Private/MuCOE/GenerateMutableSource/GenerateMutableSourceMesh.cpp
--- a/Engine/Plugins/Experimental/Mutable/Source/CustomizableObjectEditor/Private/MuCOE/GenerateMutableSource/GenerateMutableSourceMesh.cpp
+++ b/Engine/Plugins/Experimental/Mutable/Source/CustomizableObjectEditor/Private/MuCOE/GenerateMutableSource/GenerateMutableSourceMesh.cpp
@@ -755,12 +755,7 @@ mu::MeshPtr ConvertSkeletalMeshToMutable(const USkeletalMesh* InSkeletalMesh, co
 	const bool bIgnorePhysics = EnumHasAnyFlags(CurrentFlags, EMutableMeshConversionFlags::IgnorePhysics);
 	const bool bDoNotCreateMeshMetadata = EnumHasAnyFlags(CurrentFlags, EMutableMeshConversionFlags::DoNotCreateMeshMetadata);
 
-	mu::MeshPtr MutableMesh = new mu::Mesh();
-
-	// CurrentMeshComponent == None implies IgnoreSkeleton flag.
-	// CurrentMeshComponent == None will only happen with modifiers and, for now, any mesh generated from a modifier
-	// should ignore skinning.
-	check(!(GenerationContext.CurrentMeshComponent.IsNone()) || bIgnoreSkeleton);
+	mu::Ptr<mu::Mesh> MutableMesh = new mu::Mesh();
 
 	bool bBoneMapModified = false;
 	TArray<FBoneIndexType> BoneMap;
@@ -2465,45 +2460,28 @@ mu::Ptr<mu::Mesh> GenerateMutableMesh(UObject * Mesh, const TSoftClassPtr<UAnimI
 	EMutableMeshConversionFlags CurrentFlags = GenerationContext.MeshGenerationFlags.Last();
 
 	FMutableGraphGenerationContext::FGeneratedMeshData::FKey Key = { Mesh, LODIndex, GenerationContext.CurrentLOD, SectionIndex, CurrentFlags, UniqueTags, CurrentNode };
-	mu::MeshPtr MutableMesh = GenerationContext.FindGeneratedMesh(Key);
-	
-	if (!MutableMesh)
+	mu::Ptr<mu::Mesh> MutableMesh = GenerationContext.FindGeneratedMesh(Key);	
+	if (MutableMesh)
 	{
-		if (USkeletalMesh* SkeletalMesh = Cast<USkeletalMesh>(Mesh))
-		{
-			// At some point we will want all meshes to be references at compile-time. For now, just create the actual pass-through meshes.
-			if (bIsReference)
-			{
-				MutableMesh = GenerateMeshConstant(SkeletalMesh, GenerationContext, bIsReference);
-			}
-			else
-			{
-				MutableMesh = ConvertSkeletalMeshToMutable(SkeletalMesh, AnimInstance, LODIndexConnected, SectionIndexConnected, LODIndex, SectionIndex, GenerationContext, CurrentNode, TableReferenceSkeletalMesh);
-
-				FSkeletalMeshModel* ImportedModel = SkeletalMesh->GetImportedModel();
+		return MutableMesh;
+	}
 
-				if (MutableMesh &&
-					ImportedModel->LODModels.IsValidIndex(LODIndex) &&
-					ImportedModel->LODModels[LODIndex].Sections.IsValidIndex(SectionIndex))
-				{
-					FMeshData MeshData;
-					MeshData.Mesh = Mesh;
-					MeshData.LOD = LODIndex;
-					MeshData.MaterialIndex = SectionIndex;
-					MeshData.Node = CurrentNode;
-					GenerationContext.PinData.GetCurrent().MeshesData.Add(MeshData); // Set::Emplace only supports single element constructors
-				}
-			}
+	if (USkeletalMesh* SkeletalMesh = Cast<USkeletalMesh>(Mesh))
+	{
+		// At some point we will want all meshes to be references at compile-time. For now, just create the actual pass-through meshes.
+		if (bIsReference)
+		{
+			MutableMesh = GenerateMeshConstant(SkeletalMesh, GenerationContext, bIsReference);
 		}
-		else if (const UStaticMesh* StaticMesh = Cast<UStaticMesh>(Mesh))
+		else
 		{
-			MutableMesh = ConvertStaticMeshToMutable(StaticMesh, LODIndex, SectionIndex, GenerationContext, CurrentNode);
+			MutableMesh = ConvertSkeletalMeshToMutable(SkeletalMesh, AnimInstance, LODIndexConnected, SectionIndexConnected, LODIndex, SectionIndex, GenerationContext, CurrentNode, TableReferenceSkeletalMesh);
 
-			const FStaticMeshRenderData* RenderData = StaticMesh->GetRenderData();
+			FSkeletalMeshModel* ImportedModel = SkeletalMesh->GetImportedModel();
 
 			if (MutableMesh &&
-				RenderData->LODResources.IsValidIndex(LODIndex) &&
-				RenderData->LODResources[LODIndex].Sections.IsValidIndex(SectionIndex))
+				ImportedModel->LODModels.IsValidIndex(LODIndex) &&
+				ImportedModel->LODModels[LODIndex].Sections.IsValidIndex(SectionIndex))
 			{
 				FMeshData MeshData;
 				MeshData.Mesh = Mesh;
@@ -2513,27 +2491,106 @@ mu::Ptr<mu::Mesh> GenerateMutableMesh(UObject * Mesh, const TSoftClassPtr<UAnimI
 				GenerationContext.PinData.GetCurrent().MeshesData.Add(MeshData); // Set::Emplace only supports single element constructors
 			}
 		}
-		else
-		{
-			GenerationContext.Compiler->CompilerLog(LOCTEXT("UnimplementedMesh", "Mesh type not implemented yet."), CurrentNode);
-		}
+	}
+	else if (const UStaticMesh* StaticMesh = Cast<UStaticMesh>(Mesh))
+	{
+		MutableMesh = ConvertStaticMeshToMutable(StaticMesh, LODIndex, SectionIndex, GenerationContext, CurrentNode);
 
-		if (MutableMesh)
+		const FStaticMeshRenderData* RenderData = StaticMesh->GetRenderData();
+
+		if (MutableMesh &&
+			RenderData->LODResources.IsValidIndex(LODIndex) &&
+			RenderData->LODResources[LODIndex].Sections.IsValidIndex(SectionIndex))
 		{
-			GenerationContext.GeneratedMeshes.Push({ Key, MutableMesh });
+			FMeshData MeshData;
+			MeshData.Mesh = Mesh;
+			MeshData.LOD = LODIndex;
+			MeshData.MaterialIndex = SectionIndex;
+			MeshData.Node = CurrentNode;
+			GenerationContext.PinData.GetCurrent().MeshesData.Add(MeshData); // Set::Emplace only supports single element constructors
 		}
 	}
+	else
+	{
+		GenerationContext.Compiler->CompilerLog(LOCTEXT("UnimplementedMesh", "Mesh type not implemented yet."), CurrentNode);
+	}
+
+	if (MutableMesh)
+	{
+		GenerationContext.GeneratedMeshes.Push({ Key, MutableMesh });
+	}
 	
 	return MutableMesh;
 }
 
 
-mu::MeshPtr BuildMorphedMutableMesh(const UEdGraphPin* BaseSourcePin, const FString& MorphTargetName, FMutableGraphGenerationContext & GenerationContext, const bool bOnlyConnectedLOD, const FName& RowName)
+mu::Ptr<mu::Mesh> BuildMorphedMutableMeshFromMesh(mu::Ptr<mu::Mesh> BaseSourceMesh, USkeletalMesh* SkeletalMesh, const FString& MorphTargetName, int32 LODIndex, int32 SectionIndex)
+{
+	// Clone it (it will probably be shared)
+	mu::Ptr<mu::Mesh> MorphedSourceMesh = BaseSourceMesh->Clone();
+
+	// Bake the morph in the new mutable mesh
+	UMorphTarget* MorphTarget = SkeletalMesh ? SkeletalMesh->FindMorphTarget(*MorphTargetName) : nullptr;
+
+	if (MorphTarget && MorphTarget->GetMorphLODModels().IsValidIndex(LODIndex))
+	{
+		int32 PosBuf = -1;
+		int32 PosChannel = -1;
+		MorphedSourceMesh->GetVertexBuffers().FindChannel(mu::MBS_POSITION, 0, &PosBuf, &PosChannel);
+		check(PosBuf >= 0 && PosChannel >= 0);
+
+		int32 PosElemSize = MorphedSourceMesh->GetVertexBuffers().GetElementSize(PosBuf);
+		int32 PosOffset = MorphedSourceMesh->GetVertexBuffers().GetChannelOffset(PosBuf, PosChannel);
+		uint8* PosBuffer = MorphedSourceMesh->GetVertexBuffers().GetBufferData(PosBuf) + PosOffset;
+
+		int32 NorBuf = -1;
+		int32 NorChannel = -1;
+		MorphedSourceMesh->GetVertexBuffers().FindChannel(mu::MBS_NORMAL, 0, &NorBuf, &NorChannel);
+
+		const bool bHasNormals = NorBuf >= 0 && NorChannel >= 0;
+
+		int32 NorElemSize = bHasNormals ? MorphedSourceMesh->GetVertexBuffers().GetElementSize(NorBuf) : 0;
+		int32 NorOffset = bHasNormals ? MorphedSourceMesh->GetVertexBuffers().GetChannelOffset(NorBuf, NorChannel) : 0;
+		uint8* NorBuffer = bHasNormals ? MorphedSourceMesh->GetVertexBuffers().GetBufferData(NorBuf) + NorOffset : nullptr;
+
+		int32 MaterialVertexStart = SkeletalMesh->GetImportedModel()->LODModels[LODIndex].Sections[SectionIndex].GetVertexBufferIndex();
+		int32 MeshVertexCount = MorphedSourceMesh->GetVertexBuffers().GetElementCount();
+
+		const FMorphTargetLODModel& MorphLODModel = MorphTarget->GetMorphLODModels()[LODIndex];
+		for (const FMorphTargetDelta& MorphDelta : MorphLODModel.Vertices)
+		{
+			const int32 VertexIndex = MorphDelta.SourceIdx - MaterialVertexStart;
+			if (VertexIndex >= 0 && VertexIndex < MeshVertexCount)
+			{
+				{
+					float* const PosData = reinterpret_cast<float*>(PosBuffer + PosElemSize * VertexIndex);
+					const FVector3f MorphedPosition = FVector3f(PosData[0], PosData[1], PosData[2]) + MorphDelta.PositionDelta;
+					PosData[0] = MorphedPosition.X;
+					PosData[1] = MorphedPosition.Y;
+					PosData[2] = MorphedPosition.Z;
+				}
+
+				if (bHasNormals)
+				{
+					float* const NorData = reinterpret_cast<float*>(NorBuffer + NorElemSize * VertexIndex);
+					const FVector3f MorphedNormal = FVector3f(NorData[0], NorData[1], NorData[2]) + MorphDelta.TangentZDelta;
+					NorData[0] = MorphedNormal.X;
+					NorData[1] = MorphedNormal.Y;
+					NorData[2] = MorphedNormal.Z;
+				}
+			}
+		}
+	}
+
+	return MorphedSourceMesh;
+}
+
+
+mu::Ptr<mu::Mesh> BuildMorphedMutableMesh(const UEdGraphPin* BaseSourcePin, const FString& MorphTargetName, FMutableGraphGenerationContext& GenerationContext, const bool bOnlyConnectedLOD, const FName& RowName)
 {
 	check(BaseSourcePin);
 	SCOPED_PIN_DATA(GenerationContext, BaseSourcePin)
 
-	mu::MeshPtr MorphedSourceMesh;
 
 	if (!BaseSourcePin)
 	{
@@ -2543,10 +2600,10 @@ mu::MeshPtr BuildMorphedMutableMesh(const UEdGraphPin* BaseSourcePin, const FStr
 
 	int32 LODIndexConnected = -1; // LOD which the pin is connected to
 	int32 SectionIndexConnected = -1;
-	
+
 	int32 LODIndex = -1; // Initialization required to remove uninitialized warning.
 	int32 SectionIndex = -1;
-	
+
 	USkeletalMesh* SkeletalMesh = nullptr;
 	UCustomizableObjectNode* Node = Cast<UCustomizableObjectNode>(BaseSourcePin->GetOwningNode());
 
@@ -2568,70 +2625,18 @@ mu::MeshPtr BuildMorphedMutableMesh(const UEdGraphPin* BaseSourcePin, const FStr
 		}
 	}
 
+	mu::Ptr<mu::Mesh> MorphedSourceMesh;
+
 	if (SkeletalMesh)
 	{
 		GetLODAndSectionForAutomaticLODs(GenerationContext, *Node, *SkeletalMesh, LODIndexConnected, SectionIndexConnected, LODIndex, SectionIndex, bOnlyConnectedLOD);
 		// Get the base mesh
 		constexpr bool bIsReference = false;
-		mu::Ptr<mu::Mesh> BaseSourceMesh = GenerateMutableMesh(SkeletalMesh, TSoftClassPtr<UAnimInstance>(), LODIndexConnected, SectionIndexConnected, 
-															   LODIndex, SectionIndex, FString(), GenerationContext, Node, nullptr, bIsReference);
+		mu::Ptr<mu::Mesh> BaseSourceMesh = GenerateMutableMesh(SkeletalMesh, TSoftClassPtr<UAnimInstance>(), LODIndexConnected, SectionIndexConnected,
+			LODIndex, SectionIndex, FString(), GenerationContext, Node, nullptr, bIsReference);
 		if (BaseSourceMesh)
 		{
-			// Clone it (it will probably be shared)
-			MorphedSourceMesh = BaseSourceMesh->Clone();
-
-			// Bake the morph in the new mutable mesh
-			UMorphTarget* MorphTarget = SkeletalMesh ? SkeletalMesh->FindMorphTarget(*MorphTargetName) : nullptr;
-
-			if (MorphTarget && MorphTarget->GetMorphLODModels().IsValidIndex(LODIndex))
-			{
-				int32 PosBuf = -1;
-				int32 PosChannel = -1;
-				MorphedSourceMesh->GetVertexBuffers().FindChannel(mu::MBS_POSITION, 0, &PosBuf, &PosChannel);
-				check(PosBuf >= 0 && PosChannel >= 0);
-
-				int32 PosElemSize = MorphedSourceMesh->GetVertexBuffers().GetElementSize(PosBuf);
-				int32 PosOffset = MorphedSourceMesh->GetVertexBuffers().GetChannelOffset(PosBuf, PosChannel);
-				uint8* PosBuffer = MorphedSourceMesh->GetVertexBuffers().GetBufferData(PosBuf) + PosOffset;
-
-				int32 NorBuf = -1;
-				int32 NorChannel = -1;
-				MorphedSourceMesh->GetVertexBuffers().FindChannel(mu::MBS_NORMAL, 0, &NorBuf, &NorChannel);
-
-				const bool bHasNormals = NorBuf >= 0 && NorChannel >= 0;
-
-				int32 NorElemSize = bHasNormals ? MorphedSourceMesh->GetVertexBuffers().GetElementSize(NorBuf) : 0;
-				int32 NorOffset	  = bHasNormals ? MorphedSourceMesh->GetVertexBuffers().GetChannelOffset(NorBuf, NorChannel) : 0;
-				uint8* NorBuffer  = bHasNormals ? MorphedSourceMesh->GetVertexBuffers().GetBufferData(NorBuf) + NorOffset : nullptr;
-
-				int32 MaterialVertexStart = SkeletalMesh->GetImportedModel()->LODModels[LODIndex].Sections[SectionIndex].GetVertexBufferIndex();
-				int32 MeshVertexCount = MorphedSourceMesh->GetVertexBuffers().GetElementCount();
-
-				const FMorphTargetLODModel& MorphLODModel = MorphTarget->GetMorphLODModels()[LODIndex];
-				for (const FMorphTargetDelta& MorphDelta : MorphLODModel.Vertices)
-				{
-					const int32 VertexIndex = MorphDelta.SourceIdx - MaterialVertexStart;
-					if (VertexIndex >= 0 && VertexIndex < MeshVertexCount)
-					{
-						{
-							float* const PosData = reinterpret_cast<float*>(PosBuffer + PosElemSize * VertexIndex);
-							const FVector3f MorphedPosition = FVector3f(PosData[0], PosData[1], PosData[2]) + MorphDelta.PositionDelta;
-							PosData[0] = MorphedPosition.X;
-							PosData[1] = MorphedPosition.Y;
-							PosData[2] = MorphedPosition.Z;
-						}
-
-						if (bHasNormals)
-						{
-							float* const NorData = reinterpret_cast<float*>(NorBuffer + NorElemSize * VertexIndex);
-							const FVector3f MorphedNormal = FVector3f(NorData[0], NorData[1], NorData[2]) + MorphDelta.TangentZDelta;
-							NorData[0] = MorphedNormal.X;
-							NorData[1] = MorphedNormal.Y;
-							NorData[2] = MorphedNormal.Z;
-						}
-					}
-				}
-			}
+			MorphedSourceMesh = BuildMorphedMutableMeshFromMesh(BaseSourceMesh, SkeletalMesh, MorphTargetName, LODIndex, SectionIndex);
 		}
 	}
 
@@ -3438,6 +3443,21 @@ mu::NodeMeshPtr GenerateMutableSourceMesh(const UEdGraphPin* Pin,
 			{
 				MeshNode->SetValue(MutableMesh);
 
+				// Add the potentially required morphs. For now always add all morphs. A better implementation can narrow down with:
+				// - detecting morph names in all "ModifierMorphMeshSection" nodes
+				// - detecting morph names in "ModifierMorphMeshSection" nodes relevant for this current context.
+				{
+					MUTABLE_CPUPROFILER_SCOPE(GenerateMutableSourceMesh_AddAllMorphs);
+
+					const TArray<TObjectPtr<UMorphTarget>>& Morphs = TypedNodeSkel->SkeletalMesh->GetMorphTargets();
+					for (TObjectPtr<UMorphTarget> Morph : Morphs)
+					{
+						FString MorphTargetName = Morph->GetName();
+						mu::Ptr<mu::Mesh> MorphedMesh = BuildMorphedMutableMeshFromMesh(MutableMesh, TypedNodeSkel->SkeletalMesh, MorphTargetName, LODIndex, SectionIndex);
+						MeshNode->AddMorph(MorphTargetName,MorphedMesh);
+					}
+				}
+
 				if (TypedNodeSkel->SkeletalMesh->GetPhysicsAsset() && 
 					MutableMesh->GetPhysicsBody() && 
 					MutableMesh->GetPhysicsBody()->GetBodyCount())
diff --git a/Engine/Plugins/Experimental/Mutable/Source/CustomizableObjectEditor/Private/MuCOE/GenerateMutableSource/GenerateMutableSourceModifier.cpp b/Engine/Plugins/Experimental/Mutable/Source/CustomizableObjectEditor/Private/MuCOE/GenerateMutableSource/GenerateMutableSourceModifier.cpp
--- a/Engine/Plugins/Experimental/Mutable/Source/CustomizableObjectEditor/Private/MuCOE/GenerateMutableSource/GenerateMutableSourceModifier.cpp
+++ b/Engine/Plugins/Experimental/Mutable/Source/CustomizableObjectEditor/Private/MuCOE/GenerateMutableSource/GenerateMutableSourceModifier.cpp
@@ -6,18 +6,33 @@
 #include "MuCOE/CustomizableObjectCompiler.h"
 #include "MuCOE/GenerateMutableSource/GenerateMutableSourceMesh.h"
 #include "MuCOE/GenerateMutableSource/GenerateMutableSourceImage.h"
+#include "MuCOE/GenerateMutableSource/GenerateMutableSourceFloat.h"
+#include "MuCOE/GenerateMutableSource/GenerateMutableSourceSurface.h"
+#include "MuCOE/GenerateMutableSource/GenerateMutableSourceLayout.h"
+#include "MuCOE/GenerateMutableSource/GenerateMutableSourceGroupProjector.h"
 #include "MuCOE/GraphTraversal.h"
 #include "MuCOE/MutableUtils.h"
-#include "MuCOE/Nodes/CustomizableObjectNodeMeshClipDeform.h"
-#include "MuCOE/Nodes/CustomizableObjectNodeMeshClipMorph.h"
-#include "MuCOE/Nodes/CustomizableObjectNodeMeshClipWithMesh.h"
+#include "MuCOE/Nodes/CustomizableObjectNodeModifierClipDeform.h"
+#include "MuCOE/Nodes/CustomizableObjectNodeModifierClipMorph.h"
+#include "MuCOE/Nodes/CustomizableObjectNodeModifierClipWithMesh.h"
 #include "MuCOE/Nodes/CustomizableObjectNodeModifierClipWithUVMask.h"
+#include "MuCOE/Nodes/CustomizableObjectNodeModifierExtendMeshSection.h"
+#include "MuCOE/Nodes/CustomizableObjectNodeModifierEditMeshSection.h"
+#include "MuCOE/Nodes/CustomizableObjectNodeModifierMorphMeshSection.h"
+#include "MuCOE/Nodes/CustomizableObjectNodeModifierRemoveMesh.h"
+#include "MuCOE/Nodes/CustomizableObjectNodeModifierRemoveMeshBlocks.h"
 #include "MuCOE/Nodes/CustomizableObjectNodeStaticMesh.h"
+#include "MuCOE/Nodes/CustomizableObjectNodeFloatParameter.h"
+#include "MuCOE/Nodes/CustomizableObjectNodeFloatConstant.h"
 #include "MuR/Mesh.h"
 #include "MuT/NodeMeshTransform.h"
 #include "MuT/NodeModifierMeshClipDeform.h"
 #include "MuT/NodeModifierMeshClipMorphPlane.h"
 #include "MuT/NodeModifierMeshClipWithUVMask.h"
+#include "MuT/NodeModifierSurfaceEdit.h"
+#include "MuT/NodeMeshConstant.h"
+#include "MuT/NodeMeshFragment.h"
+#include "MuT/NodeMeshFormat.h"
 
 #define LOCTEXT_NAMESPACE "CustomizableObjectEditor"
 
@@ -37,20 +52,18 @@ mu::Ptr<mu::NodeModifier> GenerateMutableSourceModifier(const UEdGraphPin * Pin,
 		return static_cast<mu::NodeModifier*>(Generated->Node.get());
 	}
 	
-	mu::NodeModifierPtr Result;
+	mu::Ptr<mu::NodeModifier> Result;
 
 	bool bDoNotAddToGeneratedCache = false; // TODO Remove on MTBL-829 
 
 	
-	// We don't need all the data for the modifiers meshes
-	const EMutableMeshConversionFlags ModifiersMeshFlags = 
+	if (const UCustomizableObjectNodeModifierClipMorph* TypedNodeClip = Cast<UCustomizableObjectNodeModifierClipMorph>(Node))
+	{
+		const EMutableMeshConversionFlags ModifiersMeshFlags =
 			EMutableMeshConversionFlags::IgnoreSkinning |
 			EMutableMeshConversionFlags::IgnorePhysics;
+		GenerationContext.MeshGenerationFlags.Push(ModifiersMeshFlags);
 
-	GenerationContext.MeshGenerationFlags.Push(ModifiersMeshFlags);
-
-	if (const UCustomizableObjectNodeMeshClipMorph* TypedNodeClip = Cast<UCustomizableObjectNodeMeshClipMorph>(Node))
-	{
 		mu::Ptr<mu::NodeModifierMeshClipMorphPlane> ClipNode = new mu::NodeModifierMeshClipMorphPlane();
 		Result = ClipNode;
 
@@ -68,10 +81,17 @@ mu::Ptr<mu::NodeModifier> GenerateMutableSourceModifier(const UEdGraphPin * Pin,
 		{
 			 ClipNode->RequiredTags.Add(Tag);
 		}
+
+		GenerationContext.MeshGenerationFlags.Pop();
 	}
 
-	else if (const UCustomizableObjectNodeMeshClipDeform* TypedNodeClipDeform = Cast<UCustomizableObjectNodeMeshClipDeform>(Node))
+	else if (const UCustomizableObjectNodeModifierClipDeform* TypedNodeClipDeform = Cast<UCustomizableObjectNodeModifierClipDeform>(Node))
 	{
+		const EMutableMeshConversionFlags ModifiersMeshFlags =
+			EMutableMeshConversionFlags::IgnoreSkinning |
+			EMutableMeshConversionFlags::IgnorePhysics;
+		GenerationContext.MeshGenerationFlags.Push(ModifiersMeshFlags);
+
 		mu::Ptr<mu::NodeModifierMeshClipDeform> ClipNode = new mu::NodeModifierMeshClipDeform();
 		Result = ClipNode;
 	
@@ -113,10 +133,17 @@ mu::Ptr<mu::NodeModifier> GenerateMutableSourceModifier(const UEdGraphPin * Pin,
 		{
 			ClipNode->RequiredTags.Add(Tag);
 		}		
+
+		GenerationContext.MeshGenerationFlags.Pop();
 	}
 
-	else if (const UCustomizableObjectNodeMeshClipWithMesh* TypedNodeClipMesh = Cast<UCustomizableObjectNodeMeshClipWithMesh>(Node))
+	else if (const UCustomizableObjectNodeModifierClipWithMesh* TypedNodeClipMesh = Cast<UCustomizableObjectNodeModifierClipWithMesh>(Node))
 	{
+		const EMutableMeshConversionFlags ModifiersMeshFlags =
+			EMutableMeshConversionFlags::IgnoreSkinning |
+			EMutableMeshConversionFlags::IgnorePhysics;
+		GenerationContext.MeshGenerationFlags.Push(ModifiersMeshFlags);
+
 		// MeshClipWithMesh can be connected to multiple objects, so the compiled NodeModifierMeshClipWithMesh
 		// needs to be different for each object. If it were added to the Generated cache, all the objects would get the same.
 		bDoNotAddToGeneratedCache = true;
@@ -182,13 +209,13 @@ mu::Ptr<mu::NodeModifier> GenerateMutableSourceModifier(const UEdGraphPin * Pin,
 
 		if (TypedNodeClipMesh->CustomizableObjectToClipWith != nullptr)
 		{
-			TArray<mu::Ptr<mu::NodeModifierMeshClipWithMesh>>* ArrayDataPtr = GenerationContext.MapClipMeshNodeToMutableClipMeshNodeArray.Find(Cast<UCustomizableObjectNodeMeshClipWithMesh>(Pin->GetOwningNode()));
+			TArray<mu::Ptr<mu::NodeModifierMeshClipWithMesh>>* ArrayDataPtr = GenerationContext.MapClipMeshNodeToMutableClipMeshNodeArray.Find(Cast<UCustomizableObjectNodeModifierClipWithMesh>(Pin->GetOwningNode()));
 
 			if (ArrayDataPtr == nullptr)
 			{
 				TArray<mu::Ptr<mu::NodeModifierMeshClipWithMesh>> ArrayData;
 				ArrayData.Add(ClipNode);
-				UCustomizableObjectNodeMeshClipWithMesh* CastedNode = Cast<UCustomizableObjectNodeMeshClipWithMesh>(Pin->GetOwningNode());
+				UCustomizableObjectNodeModifierClipWithMesh* CastedNode = Cast<UCustomizableObjectNodeModifierClipWithMesh>(Pin->GetOwningNode());
 				GenerationContext.MapClipMeshNodeToMutableClipMeshNodeArray.Add(CastedNode, ArrayData);
 			}
 			else
@@ -196,10 +223,17 @@ mu::Ptr<mu::NodeModifier> GenerateMutableSourceModifier(const UEdGraphPin * Pin,
 				ArrayDataPtr->AddUnique(ClipNode);
 			}
 		}
+
+		GenerationContext.MeshGenerationFlags.Pop();
 	}
 
 	else if (const UCustomizableObjectNodeModifierClipWithUVMask* TypedNodeClipUVMask = Cast<UCustomizableObjectNodeModifierClipWithUVMask>(Node))
 	{
+		const EMutableMeshConversionFlags ModifiersMeshFlags =
+			EMutableMeshConversionFlags::IgnoreSkinning |
+			EMutableMeshConversionFlags::IgnorePhysics;
+		GenerationContext.MeshGenerationFlags.Push(ModifiersMeshFlags);
+
 		// This modifier can be connected to multiple objects, so the compiled node
 		// needs to be different for each object. If it were added to the Generated cache, all the objects would get the same.
 		bDoNotAddToGeneratedCache = true;
@@ -229,6 +263,325 @@ mu::Ptr<mu::NodeModifier> GenerateMutableSourceModifier(const UEdGraphPin * Pin,
 		{
 			ClipNode->RequiredTags.Add(Tag);
 		}
+
+		GenerationContext.MeshGenerationFlags.Pop();
+	}
+
+	else if (UCustomizableObjectNodeModifierExtendMeshSection* TypedNodeExt = Cast<UCustomizableObjectNodeModifierExtendMeshSection>(Node))
+	{
+		const EMutableMeshConversionFlags ModifiersMeshFlags = EMutableMeshConversionFlags::None;
+		GenerationContext.MeshGenerationFlags.Push(ModifiersMeshFlags);
+
+		mu::Ptr<mu::NodeModifierSurfaceEdit> SurfNode = new mu::NodeModifierSurfaceEdit();
+		Result = SurfNode;
+
+		// TODO: This was used in the non-modifier version for group projectors. It may affect the "drop projection from LOD" feature.
+		const int32 LOD = Node->IsAffectedByLOD() ? GenerationContext.CurrentLOD : 0;
+
+		[&] // Using a lambda so control flow is easier to manage.
+		{
+			SurfNode->MultipleTagsPolicy = TypedNodeExt->MultipleTagPolicy;
+			for (const FString& Tag : TypedNodeExt->RequiredTags)
+			{
+				SurfNode->RequiredTags.Add(Tag);
+			}
+
+			mu::Ptr<mu::NodeMesh> AddMeshNode;
+			FMutableGraphMeshGenerationData MeshData;
+			if (const UEdGraphPin* ConnectedPin = FollowInputPin(*TypedNodeExt->AddMeshPin()))
+			{
+				AddMeshNode = GenerateMutableSourceMesh(ConnectedPin, GenerationContext, MeshData, true, false);
+			}
+
+			if (AddMeshNode)
+			{
+				mu::Ptr<mu::NodeMesh> MeshPtr = AddMeshNode;
+
+				const TArray<UCustomizableObjectLayout*> Layouts = TypedNodeExt->GetLayouts();
+
+				if (Layouts.Num())
+				{
+					mu::Ptr<mu::NodeMeshFragment> MeshFrag = new mu::NodeMeshFragment();
+
+					MeshFrag->SourceMesh = MeshPtr;
+					//TODO: Implement support for multiple UV channels (e.g. Add warning for vertices which have a block in a layout but not in the other)
+					MeshFrag->LayoutIndex = 0;
+
+					// For this case we don't want to create another layout: we will use the one defined in the mesh to be added since we want to add
+					// any block defined there.
+					//bool bWasEmpty = false;
+					//MeshFrag->Layout = CreateMutableLayoutNode(GenerationContext, Layouts[MeshFrag->LayoutIndex], true, bWasEmpty);
+
+					MeshPtr = MeshFrag;
+				}
+				else
+				{
+					GenerationContext.Compiler->CompilerLog(LOCTEXT("ExtendMaterialLayoutMissing", "Skeletal Mesh without Layout Node linked to an Extend Material. A 4x4 layout will be added as default layout."), Node);
+				}
+
+				mu::Ptr<mu::NodeMeshFormat> MeshFormat = new mu::NodeMeshFormat();
+				SetSurfaceFormat(GenerationContext,
+					MeshFormat->GetVertexBuffers(), MeshFormat->GetIndexBuffers(), MeshData,
+					GenerationContext.Options.CustomizableObjectNumBoneInfluences,
+					GenerationContext.Options.b16BitBoneWeightsEnabled);
+
+				MeshFormat->SetSource(MeshPtr.get());
+
+				SurfNode->MeshAdd = MeshFormat;
+			}
+
+			const int32 NumImages = TypedNodeExt->GetNumParameters(EMaterialParameterType::Texture);
+			SurfNode->Textures.SetNum(NumImages);
+			for (int32 ImageIndex = 0; ImageIndex < NumImages; ++ImageIndex)
+			{
+				mu::NodeImagePtr ImageNode;
+
+				if (!ImageNode) // If
+				{
+					const FString MaterialImageId = FGroupProjectorImageInfo::GenerateId(TypedNodeExt, ImageIndex);
+					const FGroupProjectorImageInfo* ProjectorInfo = GenerationContext.GroupProjectorLODCache.Find(MaterialImageId);
+
+					if (ProjectorInfo)
+					{
+						ensure(LOD > GenerationContext.FirstLODAvailable);
+						check(ProjectorInfo->SurfNode->Images[ImageIndex].Image == ProjectorInfo->ImageNode);
+						ImageNode = ProjectorInfo->ImageNode;
+
+						//TextureNameToProjectionResFactor.Add(ProjectorInfo->RealTextureName, ProjectorInfo->AlternateProjectionResolutionFactor);
+						//AlternateResStateName = ProjectorInfo->AlternateResStateName;
+					}
+				}
+
+				if (!ImageNode) // Else if
+				{
+					bool bShareProjectionTexturesBetweenLODs = false;
+					bool bIsGroupProjectorImage = false;
+					UTexture2D* GroupProjectionReferenceTexture = nullptr;
+					TMap<FString, float> TextureNameToProjectionResFactor;
+					FString AlternateResStateName;
+
+					ImageNode = GenerateMutableGroupProjection(LOD, ImageIndex, AddMeshNode, GenerationContext,
+						nullptr, TypedNodeExt, bShareProjectionTexturesBetweenLODs, bIsGroupProjectorImage,
+						GroupProjectionReferenceTexture, TextureNameToProjectionResFactor, AlternateResStateName);
+				}
+
+				if (!ImageNode) // Else if
+				{
+					const FNodeMaterialParameterId ImageId = TypedNodeExt->GetParameterId(EMaterialParameterType::Texture, ImageIndex);
+
+					if (TypedNodeExt->UsesImage(ImageId))
+					{
+						// TODO
+						//check(ParentMaterialNode->IsImageMutableMode(ImageIndex)); // Ensured at graph time. If it fails, something is wrong.
+
+						if (const UEdGraphPin* ConnectedPin = FollowInputPin(*TypedNodeExt->GetUsedImagePin(ImageId)))
+						{
+							// ReferenceTextureSize is used to limit the size of textures contributing to the final image.
+							const int32 ReferenceTextureSize = 0; // TODO GetBaseTextureSize(GenerationContext, TypedNodeExt, ImageIndex);
+
+							ImageNode = GenerateMutableSourceImage(ConnectedPin, GenerationContext, ReferenceTextureSize);
+						}
+					}
+				}
+
+				if (ImageNode)
+				{
+					SurfNode->Textures[ImageIndex].Extend = ImageNode;
+				}
+
+			}
+
+			for (const FString& Tag : TypedNodeExt->Tags)
+			{
+				SurfNode->EnableTags.Add(Tag);
+			}
+		}();
+
+		GenerationContext.MeshGenerationFlags.Pop();
+	}
+
+	else if (const UCustomizableObjectNodeModifierRemoveMesh* TypedNodeRem = Cast<UCustomizableObjectNodeModifierRemoveMesh>(Node))
+	{
+		const EMutableMeshConversionFlags ModifiersMeshFlags =
+			EMutableMeshConversionFlags::IgnoreSkinning |
+			EMutableMeshConversionFlags::IgnorePhysics;
+		GenerationContext.MeshGenerationFlags.Push(ModifiersMeshFlags);
+
+		mu::Ptr<mu::NodeModifierSurfaceEdit> SurfNode = new mu::NodeModifierSurfaceEdit();
+		Result = SurfNode;
+
+		SurfNode->MultipleTagsPolicy = TypedNodeRem->MultipleTagPolicy;
+		for (const FString& Tag : TypedNodeRem->RequiredTags)
+		{
+			SurfNode->RequiredTags.Add(Tag);
+		}
+
+		mu::Ptr<mu::NodeMesh> RemoveMeshNode;
+
+		if (const UEdGraphPin* ConnectedPin = FollowInputPin(*TypedNodeRem->RemoveMeshPin()))
+		{
+			FMutableGraphMeshGenerationData DummyMeshData;
+			RemoveMeshNode = GenerateMutableSourceMesh(ConnectedPin, GenerationContext, DummyMeshData, false, true);
+		}
+
+		SurfNode->MeshRemove = RemoveMeshNode;
+
+		GenerationContext.MeshGenerationFlags.Pop();
+	}
+
+	else if (const UCustomizableObjectNodeModifierRemoveMeshBlocks* TypedNodeRemBlocks = Cast<UCustomizableObjectNodeModifierRemoveMeshBlocks>(Node))
+	{
+		const EMutableMeshConversionFlags ModifiersMeshFlags =
+			EMutableMeshConversionFlags::IgnoreSkinning |
+			EMutableMeshConversionFlags::IgnorePhysics;
+		GenerationContext.MeshGenerationFlags.Push(ModifiersMeshFlags);
+
+		mu::Ptr<mu::NodeModifierMeshClipWithUVMask> ClipNode = new mu::NodeModifierMeshClipWithUVMask();
+		Result = ClipNode;
+
+		ClipNode->MultipleTagsPolicy = TypedNodeRemBlocks->MultipleTagPolicy;
+		for (const FString& Tag : TypedNodeRemBlocks->RequiredTags)
+		{
+			ClipNode->RequiredTags.Add(Tag);
+		}
+
+		bool bWasEmpty = false;
+		mu::Ptr<mu::NodeLayout> SourceLayout = CreateMutableLayoutNode(GenerationContext, TypedNodeRemBlocks->Layout, true, bWasEmpty);
+		ClipNode->ClipLayout = SourceLayout;
+		ClipNode->LayoutIndex = TypedNodeRemBlocks->ParentLayoutIndex;
+
+		GenerationContext.MeshGenerationFlags.Pop();
+	}
+
+	else if (UCustomizableObjectNodeModifierEditMeshSection* TypedNodeEdit = Cast<UCustomizableObjectNodeModifierEditMeshSection>(Node))
+	{
+		const EMutableMeshConversionFlags ModifiersMeshFlags =
+			EMutableMeshConversionFlags::IgnoreSkinning |
+			EMutableMeshConversionFlags::IgnorePhysics;
+		GenerationContext.MeshGenerationFlags.Push(ModifiersMeshFlags);
+
+		mu::Ptr<mu::NodeModifierSurfaceEdit> SurfNode = new mu::NodeModifierSurfaceEdit();
+		Result = SurfNode;
+
+		SurfNode->MultipleTagsPolicy = TypedNodeEdit->MultipleTagPolicy;
+		for (const FString& Tag : TypedNodeEdit->RequiredTags)
+		{
+			SurfNode->RequiredTags.Add(Tag);
+		}
+
+		const int32 NumImages = TypedNodeEdit->GetNumParameters(EMaterialParameterType::Texture);
+		SurfNode->Textures.SetNum(NumImages);
+		for (int32 ImageIndex = 0; ImageIndex < NumImages; ++ImageIndex)
+		{
+			const FNodeMaterialParameterId ImageId = TypedNodeEdit->GetParameterId(EMaterialParameterType::Texture, ImageIndex);
+
+			if (TypedNodeEdit->UsesImage(ImageId))
+			{
+				// TODO
+				//check(ParentMaterialNode->IsImageMutableMode(ImageIndex)); // Ensured at graph time. If it fails, something is wrong.
+
+				const UEdGraphPin* ConnectedImagePin = FollowInputPin(*TypedNodeEdit->GetUsedImagePin(ImageId));
+
+				mu::NodeModifierSurfaceEdit::FTexture& ImagePatch = SurfNode->Textures[ImageIndex];
+
+				// \todo: expose these two options?
+				ImagePatch.PatchBlendType = mu::EBlendType::BT_BLEND;
+				ImagePatch.bPatchApplyToAlpha = true;
+
+				// ReferenceTextureSize is used to limit the size of textures contributing to the final image.
+				const int32 ReferenceTextureSize = 0; //TODO GetBaseTextureSize(GenerationContext, ParentMaterialNode, ImageIndex);
+
+				ImagePatch.PatchImage = GenerateMutableSourceImage(ConnectedImagePin, GenerationContext, ReferenceTextureSize);
+
+				const UEdGraphPin* ImageMaskPin = TypedNodeEdit->GetUsedImageMaskPin(ImageId);
+				check(ImageMaskPin); // Ensured when reconstructing EditMaterial nodes. If it fails, something is wrong.
+
+				if (const UEdGraphPin* ConnectedMaskPin = FollowInputPin(*ImageMaskPin))
+				{
+					ImagePatch.PatchMask = GenerateMutableSourceImage(ConnectedMaskPin, GenerationContext, ReferenceTextureSize);
+				}
+
+				// Add the blocks to patch
+				FIntPoint GridSize = TypedNodeEdit->Layout->GetGridSize();
+				FVector2f GridSizeF = FVector2f(GridSize);
+				ImagePatch.PatchBlocks.Reserve(TypedNodeEdit->Layout->Blocks.Num());
+				for (const FCustomizableObjectLayoutBlock& LayoutBlock : TypedNodeEdit->Layout->Blocks)
+				{
+					FBox2f Rect;
+					Rect.Min = FVector2f(LayoutBlock.Min) / GridSizeF;
+					Rect.Max = FVector2f(LayoutBlock.Max) / GridSizeF;
+					ImagePatch.PatchBlocks.Add(Rect);
+				}
+			}
+		}
+
+		GenerationContext.MeshGenerationFlags.Pop();
+	}
+
+	else if (const UCustomizableObjectNodeModifierMorphMeshSection* TypedNodeMorph = Cast<UCustomizableObjectNodeModifierMorphMeshSection>(Node))
+	{
+		const EMutableMeshConversionFlags ModifiersMeshFlags =
+			EMutableMeshConversionFlags::IgnoreSkinning |
+			EMutableMeshConversionFlags::IgnorePhysics;
+		GenerationContext.MeshGenerationFlags.Push(ModifiersMeshFlags);
+
+		mu::Ptr<mu::NodeModifierSurfaceEdit> SurfNode = new mu::NodeModifierSurfaceEdit();
+		Result = SurfNode;
+
+		// This modifier needs to be applied right after the mesh constant is generated
+		SurfNode->bApplyBeforeNormalOperations = true;
+
+		SurfNode->MultipleTagsPolicy = TypedNodeMorph->MultipleTagPolicy;
+		for (const FString& Tag : TypedNodeMorph->RequiredTags)
+		{
+			SurfNode->RequiredTags.Add(Tag);
+		}
+
+		SurfNode->MeshMorph = TypedNodeMorph->MorphTargetName;
+
+		if (const UEdGraphPin* ConnectedPin = FollowInputPin(*TypedNodeMorph->FactorPin()))
+		{
+			UEdGraphNode* floatNode = ConnectedPin->GetOwningNode();
+			bool validStaticFactor = true;
+			if (const UCustomizableObjectNodeFloatParameter* floatParameterNode = Cast<UCustomizableObjectNodeFloatParameter>(floatNode))
+			{
+				if (floatParameterNode->DefaultValue < -1.0f || floatParameterNode->DefaultValue > 1.0f)
+				{
+					validStaticFactor = false;
+					FString msg = FString::Printf(TEXT("Mesh morph nodes only accept factors between -1.0 and 1.0 inclusive but the default value of the float parameter node is (%f). Factor will be ignored."), floatParameterNode->DefaultValue);
+					GenerationContext.Compiler->CompilerLog(FText::FromString(msg), Node);
+				}
+				if (floatParameterNode->ParamUIMetadata.MinimumValue < -1.0f)
+				{
+					validStaticFactor = false;
+					FString msg = FString::Printf(TEXT("Mesh morph nodes only accept factors between -1.0 and 1.0 inclusive but the minimum UI value for the input float parameter node is (%f). Factor will be ignored."), floatParameterNode->ParamUIMetadata.MinimumValue);
+					GenerationContext.Compiler->CompilerLog(FText::FromString(msg), Node);
+				}
+				if (floatParameterNode->ParamUIMetadata.MaximumValue > 1.0f)
+				{
+					validStaticFactor = false;
+					FString msg = FString::Printf(TEXT("Mesh morph nodes only accept factors between -1.0 and 1.0 inclusive but the maximum UI value for the input float parameter node is (%f). Factor will be ignored."), floatParameterNode->ParamUIMetadata.MaximumValue);
+					GenerationContext.Compiler->CompilerLog(FText::FromString(msg), Node);
+				}
+			}
+			else if (const UCustomizableObjectNodeFloatConstant* floatConstantNode = Cast<UCustomizableObjectNodeFloatConstant>(floatNode))
+			{
+				if (floatConstantNode->Value < -1.0f || floatConstantNode->Value > 1.0f)
+				{
+					validStaticFactor = false;
+					FString msg = FString::Printf(TEXT("Mesh morph nodes only accept factors between -1.0 and 1.0 inclusive but the value of the float constant node is (%f). Factor will be ignored."), floatConstantNode->Value);
+					GenerationContext.Compiler->CompilerLog(FText::FromString(msg), Node);
+				}
+			}
+
+			if (validStaticFactor)
+			{
+				mu::NodeScalarPtr FactorNode = GenerateMutableSourceFloat(ConnectedPin, GenerationContext);
+				SurfNode->MorphFactor = FactorNode;
+			}
+		}
+
+		GenerationContext.MeshGenerationFlags.Pop();
 	}
 
 	else
@@ -236,8 +589,6 @@ mu::Ptr<mu::NodeModifier> GenerateMutableSourceModifier(const UEdGraphPin * Pin,
 		GenerationContext.Compiler->CompilerLog(LOCTEXT("UnimplementedNode", "Node type not implemented yet."), Node);
 	}
 
-	GenerationContext.MeshGenerationFlags.Pop();
-
 	if (!bDoNotAddToGeneratedCache)
 	{
 		GenerationContext.Generated.Add(Key, FGeneratedData(Node, Result));
diff --git a/Engine/Plugins/Experimental/Mutable/Source/CustomizableObjectEditor/Private/MuCOE/GenerateMutableSource/GenerateMutableSourceSurface.cpp b/Engine/Plugins/Experimental/Mutable/Source/CustomizableObjectEditor/Private/MuCOE/GenerateMutableSource/GenerateMutableSourceSurface.cpp
--- a/Engine/Plugins/Experimental/Mutable/Source/CustomizableObjectEditor/Private/MuCOE/GenerateMutableSource/GenerateMutableSourceSurface.cpp
+++ b/Engine/Plugins/Experimental/Mutable/Source/CustomizableObjectEditor/Private/MuCOE/GenerateMutableSource/GenerateMutableSourceSurface.cpp
@@ -30,15 +30,15 @@
 #include "MuCOE/GraphTraversal.h"
 #include "MuCOE/MutableUtils.h"
 #include "MuCOE/Nodes/CustomizableObjectNodeCopyMaterial.h"
-#include "MuCOE/Nodes/CustomizableObjectNodeEditMaterial.h"
-#include "MuCOE/Nodes/CustomizableObjectNodeExtendMaterial.h"
 #include "MuCOE/Nodes/CustomizableObjectNodeFloatConstant.h"
 #include "MuCOE/Nodes/CustomizableObjectNodeFloatParameter.h"
 #include "MuCOE/Nodes/CustomizableObjectNodeMaterialVariation.h"
 #include "MuCOE/Nodes/CustomizableObjectNodeMaterialSwitch.h"
-#include "MuCOE/Nodes/CustomizableObjectNodeMorphMaterial.h"
-#include "MuCOE/Nodes/CustomizableObjectNodeRemoveMesh.h"
-#include "MuCOE/Nodes/CustomizableObjectNodeRemoveMeshBlocks.h"
+#include "MuCOE/Nodes/CustomizableObjectNodeModifierEditMeshSection.h"
+#include "MuCOE/Nodes/CustomizableObjectNodeModifierExtendMeshSection.h"
+#include "MuCOE/Nodes/CustomizableObjectNodeModifierMorphMeshSection.h"
+#include "MuCOE/Nodes/CustomizableObjectNodeModifierRemoveMesh.h"
+#include "MuCOE/Nodes/CustomizableObjectNodeModifierRemoveMeshBlocks.h"
 #include "MuCOE/Nodes/CustomizableObjectNodeSkeletalMesh.h"
 #include "MuCOE/Nodes/CustomizableObjectNodeTable.h"
 #include "MuT/NodeImageFormat.h"
@@ -50,7 +50,6 @@
 #include "MuT/NodeMeshFormat.h"
 #include "MuT/NodeMeshFragment.h"
 #include "MuT/NodeScalarConstant.h"
-#include "MuT/NodeSurfaceEdit.h"
 #include "MuT/NodeSurfaceSwitch.h"
 #include "MuT/NodeSurfaceVariation.h"
 #include "MuT/UnrealPixelFormatOverride.h"
@@ -204,30 +203,6 @@ void SetSurfaceFormat( FMutableGraphGenerationContext& GenerationContext,
 }
 
 
-void AddModifierToSharedSurface(FMutableGraphGenerationContext& GenerationContext, UCustomizableObjectNodeMaterialBase* NodeMaterial, const UCustomizableObjectNode& NodeModifier)
-{
-	if (!NodeMaterial || !NodeMaterial->IsReuseMaterialBetweenLODs())
-	{
-		return;
-	}
-
-	TArray<FMutableGraphGenerationContext::FSharedSurface>* SharedSurfaces = GenerationContext.SharedSurfaceIds.Find(NodeMaterial);
-	if (!SharedSurfaces)
-	{
-		return;
-	}
-
-	// Add the modifier to the key of modifiers per surface. Only add it if there are image operations involved.
-	for (FMutableGraphGenerationContext::FSharedSurface& SharedSurface : *SharedSurfaces)
-	{
-		if (SharedSurface.LOD == GenerationContext.CurrentLOD)
-		{
-			SharedSurface.NodeModifierIDs.Add(reinterpret_cast<SIZE_T>(&NodeModifier));
-		}
-	}
-}
-
-
 mu::Ptr<mu::NodeSurface> GenerateMutableSourceSurface(const UEdGraphPin * Pin, FMutableGraphGenerationContext & GenerationContext)
 {
 	MUTABLE_CPUPROFILER_SCOPE(GenerateMutableSourceSurface);
@@ -977,9 +952,12 @@ mu::Ptr<mu::NodeSurface> GenerateMutableSourceSurface(const UEdGraphPin * Pin, F
 		}
 		
 
-		for (const FString& Tag : TypedNodeMat->GetTags())
+		if (const TArray<FString>* EnableTags = TypedNodeMat->GetEnableTags())
 		{
-			SurfNode->Tags.Add(Tag);
+			for (const FString& Tag : *EnableTags)
+			{
+				SurfNode->Tags.AddUnique(Tag);
+			}
 		}
 
 		TArray<mu::Ptr<mu::NodeSurfaceNew>>* ArraySurfaceNodePtr = GenerationContext.MapMaterialNodeToMutableSurfaceNodeArray.Find(TypedNodeMat->GetMaterialNode());
@@ -1055,357 +1033,6 @@ mu::Ptr<mu::NodeSurface> GenerateMutableSourceSurface(const UEdGraphPin * Pin, F
 		}
 	}
 
-	else if (UCustomizableObjectNodeExtendMaterial* TypedNodeExt = Cast<UCustomizableObjectNodeExtendMaterial>(Node))
-	{
-		mu::Ptr<mu::NodeSurfaceEdit> SurfNode = new mu::NodeSurfaceEdit();
-		Result = SurfNode;
-
-		[&] // Using a lambda so control flow is easier to manage.
-		{
-			UCustomizableObjectNodeMaterialBase* ParentMaterialNode = TypedNodeExt->GetParentMaterialNode();
-			if (!ParentMaterialNode)
-			{
-				GenerationContext.Compiler->CompilerLog(LOCTEXT("ParentMissing", "Parent node not set (or not found)."), Node);
-				return;
-			}
-
-			if (TypedNodeExt->IsNodeOutDatedAndNeedsRefresh())
-			{
-				GenerationContext.Compiler->CompilerLog(LOCTEXT("ParentMissing", "Parent node not set (or not found)."), Node);
-				return;
-			}
-
-			// Parent, probably generated, will be retrieved from the cache
-			mu::Ptr<mu::NodeSurface> ParentNode = GenerateMutableSourceSurface(ParentMaterialNode->OutputPin(), GenerationContext);
-			SurfNode->Parent = ParentNode;
-
-			mu::Ptr<mu::NodeMesh> AddMeshNode;
-			FMutableGraphMeshGenerationData MeshData;
-			if (const UEdGraphPin* ConnectedPin = FollowInputPin(*TypedNodeExt->AddMeshPin()))
-			{
-				AddMeshNode = GenerateMutableSourceMesh(ConnectedPin, GenerationContext, MeshData, true, false);
-			}
-
-			if (AddMeshNode)
-			{
-				mu::Ptr<mu::NodeMesh> MeshPtr = AddMeshNode;
-
-				const TArray<UCustomizableObjectLayout*> Layouts = TypedNodeExt->GetLayouts();
-				
-				if (Layouts.Num())
-				{
-					mu::Ptr<mu::NodeMeshFragment> MeshFrag = new mu::NodeMeshFragment();
-				
-					MeshFrag->SourceMesh = MeshPtr;
-					//TODO: Implement support for multiple UV channels (e.g. Add warning for vertices which have a block in a layout but not in the other)
-					MeshFrag->LayoutIndex = 0;
-
-					// For this case we don't want to create another layout: we will use the one defined in the mesh to be added since we want to add
-					// any block defined there.
-					//bool bWasEmpty = false;
-					//MeshFrag->Layout = CreateMutableLayoutNode(GenerationContext, Layouts[MeshFrag->LayoutIndex], true, bWasEmpty);
-				
-					MeshPtr = MeshFrag;
-				}
-				else
-				{
-					GenerationContext.Compiler->CompilerLog(LOCTEXT("ExtendMaterialLayoutMissing","Skeletal Mesh without Layout Node linked to an Extend Material. A 4x4 layout will be added as default layout."), Node);
-				}
-
-				mu::Ptr<mu::NodeMeshFormat> MeshFormat = new mu::NodeMeshFormat();
-				SetSurfaceFormat( GenerationContext,
-						MeshFormat->GetVertexBuffers(), MeshFormat->GetIndexBuffers(), MeshData,
-						GenerationContext.Options.CustomizableObjectNumBoneInfluences,
-						GenerationContext.Options.b16BitBoneWeightsEnabled);
-
-				MeshFormat->SetSource(MeshPtr.get());
-
-				SurfNode->MeshAdd = MeshFormat;
-			}
-			
-			const int32 NumImages = ParentMaterialNode->GetNumParameters(EMaterialParameterType::Texture);
-			SurfNode->Textures.SetNum(NumImages);
-			for (int32 ImageIndex = 0; ImageIndex < NumImages; ++ImageIndex)
-			{
-				mu::NodeImagePtr ImageNode;
-				
-				if (!ImageNode) // If
-				{
-					const FString MaterialImageId = FGroupProjectorImageInfo::GenerateId(TypedNodeExt, ImageIndex);
-					const FGroupProjectorImageInfo* ProjectorInfo = GenerationContext.GroupProjectorLODCache.Find(MaterialImageId);
-
-					if (ProjectorInfo)
-					{
-						ensure(LOD > GenerationContext.FirstLODAvailable);
-						check(ProjectorInfo->SurfNode->Images[ImageIndex].Image == ProjectorInfo->ImageNode);
-						ImageNode = ProjectorInfo->ImageNode;
-
-						//TextureNameToProjectionResFactor.Add(ProjectorInfo->RealTextureName, ProjectorInfo->AlternateProjectionResolutionFactor);
-						//AlternateResStateName = ProjectorInfo->AlternateResStateName;
-					}
-				}
-
-				if (!ImageNode) // Else if
-				{
-					bool bShareProjectionTexturesBetweenLODs = false;
-					bool bIsGroupProjectorImage = false;
-					UTexture2D * GroupProjectionReferenceTexture = nullptr;
-					TMap<FString, float> TextureNameToProjectionResFactor;
-					FString AlternateResStateName;
-					
-					ImageNode = GenerateMutableGroupProjection(LOD, ImageIndex, AddMeshNode, GenerationContext,
-						nullptr, TypedNodeExt, bShareProjectionTexturesBetweenLODs, bIsGroupProjectorImage,
-						GroupProjectionReferenceTexture, TextureNameToProjectionResFactor, AlternateResStateName);
-				}
-				
-				if (!ImageNode) // Else if
-				{
-					const FNodeMaterialParameterId ImageId = ParentMaterialNode->GetParameterId(EMaterialParameterType::Texture, ImageIndex);
-					
-					if (TypedNodeExt->UsesImage(ImageId))
-					{
-						check(ParentMaterialNode->IsImageMutableMode(ImageIndex)); // Ensured at graph time. If it fails, something is wrong.
-						
-						if (const UEdGraphPin* ConnectedPin = FollowInputPin(*TypedNodeExt->GetUsedImagePin(ImageId)))
-						{
-							// ReferenceTextureSize is used to limit the size of textures contributing to the final image.
-							const int32 ReferenceTextureSize = GetBaseTextureSize(GenerationContext, ParentMaterialNode, ImageIndex);
-
-							ImageNode = GenerateMutableSourceImage(ConnectedPin, GenerationContext, ReferenceTextureSize);
-						}
-					}
-				}
-
-				if (ImageNode)
-				{
-					SurfNode->Textures[ImageIndex].Extend = ImageNode;
-				}
-
-			}
-			AddModifierToSharedSurface(GenerationContext, ParentMaterialNode, *TypedNodeExt);
-		
-			for (const FString& Tag : TypedNodeExt->Tags)
-			{
-				SurfNode->EnableTags.Add(Tag);
-			}
-		}();
-	}
-
-	else if (const UCustomizableObjectNodeRemoveMesh* TypedNodeRem = Cast<UCustomizableObjectNodeRemoveMesh>(Node))
-	{
-		mu::Ptr<mu::NodeSurfaceEdit> SurfNode = new mu::NodeSurfaceEdit();
-		Result = SurfNode;
-
-		UCustomizableObjectNodeMaterialBase* ParentMaterialNode = TypedNodeRem->GetParentMaterialNode();
-		if (!ParentMaterialNode)
-		{
-			GenerationContext.Compiler->CompilerLog(LOCTEXT("ParentMissing", "Parent node not set (or not found)."), Node);
-		}
-		else
-		{
-			// Parent, probably generated, will be retrieved from the cache
-			mu::Ptr<mu::NodeSurface> ParentNode = GenerateMutableSourceSurface(ParentMaterialNode->OutputPin(), GenerationContext);
-			SurfNode->Parent = ParentNode;
-
-			mu::Ptr<mu::NodeMesh> RemoveMeshNode;
-
-			if (const UEdGraphPin* ConnectedPin = FollowInputPin(*TypedNodeRem->RemoveMeshPin()))
-			{
-				FMutableGraphMeshGenerationData DummyMeshData;
-				RemoveMeshNode = GenerateMutableSourceMesh(ConnectedPin, GenerationContext, DummyMeshData, false, true);
-			}
-
-			SurfNode->MeshRemove = RemoveMeshNode;
-
-			AddModifierToSharedSurface(GenerationContext, ParentMaterialNode, *TypedNodeRem);
-		}
-	}
-
-	else if (const UCustomizableObjectNodeRemoveMeshBlocks* TypedNodeRemBlocks = Cast<UCustomizableObjectNodeRemoveMeshBlocks>(Node))
-	{
-		mu::Ptr<mu::NodeSurfaceEdit> SurfNode = new mu::NodeSurfaceEdit();
-		Result = SurfNode;
-
-		UCustomizableObjectNodeMaterialBase* ParentMaterialNode = TypedNodeRemBlocks->GetParentMaterialNode();
-		if (!ParentMaterialNode)
-		{
-			GenerationContext.Compiler->CompilerLog(LOCTEXT("ParentMissing", "Parent node not set (or not found)."), Node);
-		}
-		else
-		{
-			SurfNode->Parent = GenerateMutableSourceSurface(ParentMaterialNode->OutputPin(), GenerationContext);
-
-			const UEdGraphPin* BaseSourcePin = FindMeshBaseSource(*ParentMaterialNode->OutputPin(), false);
-
-			if (!BaseSourcePin)
-			{
-				GenerationContext.Compiler->CompilerLog(LOCTEXT("ParentMissing", "Parent node not set (or not found)."), Node);
-			}
-			else
-			{
-				FMutableGraphMeshGenerationData DummyMeshData;
-				mu::Ptr<mu::NodeMesh> SourceMesh = GenerateMutableSourceMesh(BaseSourcePin, GenerationContext, DummyMeshData, false, false);
-				bool bWasEmpty = false;
-				mu::Ptr<mu::NodeLayout> SourceLayout = CreateMutableLayoutNode(GenerationContext, TypedNodeRemBlocks->Layout, true, bWasEmpty);
-
-				mu::Ptr<mu::NodeMeshFragment> MeshFrag = new mu::NodeMeshFragment();
-				MeshFrag->SourceMesh = SourceMesh;
-				MeshFrag->Layout = SourceLayout;
-				MeshFrag->LayoutIndex = TypedNodeRemBlocks->ParentLayoutIndex;
-				MeshFrag->SetMessageContext(TypedNodeRemBlocks);
-
-				SurfNode->MeshRemove = MeshFrag;
-			}
-
-			AddModifierToSharedSurface(GenerationContext, ParentMaterialNode, *TypedNodeRemBlocks);
-		}
-	}
-
-	else if (UCustomizableObjectNodeEditMaterial* TypedNodeEdit = Cast<UCustomizableObjectNodeEditMaterial>(Node))
-	{
-		mu::Ptr<mu::NodeSurfaceEdit> SurfNode = new mu::NodeSurfaceEdit();
-		Result = SurfNode;
-
-		UCustomizableObjectNodeMaterialBase* ParentMaterialNode = TypedNodeEdit->GetParentMaterialNode();
-		if (!ParentMaterialNode)
-		{
-			GenerationContext.Compiler->CompilerLog(LOCTEXT("ParentMissing", "Parent node not set (or not found)."), Node);
-		}
-		else
-		{
-			// Parent, probably generated, will be retrieved from the cache
-			mu::Ptr<mu::NodeSurface> ParentNode = GenerateMutableSourceSurface(ParentMaterialNode->OutputPin(), GenerationContext);
-			SurfNode->Parent = ParentNode;
-
-			const int32 NumImages = ParentMaterialNode->GetNumParameters(EMaterialParameterType::Texture);
-			SurfNode->Textures.SetNum(NumImages);
-			for (int32 ImageIndex = 0; ImageIndex < NumImages; ++ImageIndex)
-			{
-				const FNodeMaterialParameterId ImageId = ParentMaterialNode->GetParameterId(EMaterialParameterType::Texture, ImageIndex);
-
-				if (TypedNodeEdit->UsesImage(ImageId))
-				{
-					check(ParentMaterialNode->IsImageMutableMode(ImageIndex)); // Ensured at graph time. If it fails, something is wrong.
-					
-					const UEdGraphPin* ConnectedImagePin = FollowInputPin(*TypedNodeEdit->GetUsedImagePin(ImageId));
-					
-					mu::NodeSurfaceEdit::FTexture& ImagePatch = SurfNode->Textures[ImageIndex];
-
-					// \todo: expose these two options?
-					ImagePatch.PatchBlendType = mu::EBlendType::BT_BLEND;
-					ImagePatch.bPatchApplyToAlpha=true;
-
-					// ReferenceTextureSize is used to limit the size of textures contributing to the final image.
-					const int32 ReferenceTextureSize = GetBaseTextureSize(GenerationContext, ParentMaterialNode, ImageIndex);
-
-					ImagePatch.PatchImage = GenerateMutableSourceImage(ConnectedImagePin, GenerationContext, ReferenceTextureSize);
-
-					const UEdGraphPin* ImageMaskPin = TypedNodeEdit->GetUsedImageMaskPin(ImageId);
-					check(ImageMaskPin); // Ensured when reconstructing EditMaterial nodes. If it fails, something is wrong.
-					
-					if (const UEdGraphPin* ConnectedMaskPin = FollowInputPin(*ImageMaskPin))
-					{
-						ImagePatch.PatchMask = GenerateMutableSourceImage(ConnectedMaskPin, GenerationContext, ReferenceTextureSize);
-					}
-
-					// Add the blocks to patch
-					FIntPoint GridSize = TypedNodeEdit->Layout->GetGridSize();
-					FVector2f GridSizeF = FVector2f(GridSize);
-					ImagePatch.PatchBlocks.Reserve(TypedNodeEdit->Layout->Blocks.Num());
-					for (const FCustomizableObjectLayoutBlock& LayoutBlock: TypedNodeEdit->Layout->Blocks)
-					{
-						FBox2f Rect;
-						Rect.Min = FVector2f(LayoutBlock.Min) / GridSizeF;
-						Rect.Max = FVector2f(LayoutBlock.Max) / GridSizeF;
-						ImagePatch.PatchBlocks.Add(Rect);
-					}
-				}
-
-				AddModifierToSharedSurface(GenerationContext, ParentMaterialNode, *TypedNodeEdit);
-			}
-		}
-	}
-
-	else if (const UCustomizableObjectNodeMorphMaterial* TypedNodeMorph = Cast<UCustomizableObjectNodeMorphMaterial>(Node))
-	{
-		mu::Ptr<mu::NodeSurfaceEdit> SurfNode = new mu::NodeSurfaceEdit();
-		Result = SurfNode;
-
-		UCustomizableObjectNodeMaterialBase* ParentMaterialNode = TypedNodeMorph->GetParentMaterialNode();
-		if (!ParentMaterialNode)
-		{
-			GenerationContext.Compiler->CompilerLog(LOCTEXT("ParentMissing", "Parent node not set (or not found)."), Node);
-		}
-		else
-		{
-			// Parent, probably generated, will be retrieved from the cache
-			mu::Ptr<mu::NodeSurface> ParentNode = GenerateMutableSourceSurface(ParentMaterialNode->OutputPin(), GenerationContext);
-			SurfNode->Parent = ParentNode;
-
-			const UEdGraphPin* BaseSourcePin = FindMeshBaseSource(*ParentMaterialNode->OutputPin(), false);
-			if (!BaseSourcePin)
-			{
-				GenerationContext.Compiler->CompilerLog(LOCTEXT("ParentMissing", "Parent node not set (or not found)."), Node);
-			}
-			else
-			{
-				// Morph Mesh
-				// The Mesh Base Source from the Parent Material will always have all LODs since it is not an auxiliary Mesh (clipping, reshape...).
-				mu::MeshPtr MorphedSourceMesh = BuildMorphedMutableMesh(BaseSourcePin, TypedNodeMorph->MorphTargetName, GenerationContext, false); 
-
-				mu::NodeMeshConstantPtr MorphedSourceMeshNode = new mu::NodeMeshConstant;
-				MorphedSourceMeshNode->SetMessageContext(Node);					
-				MorphedSourceMeshNode->SetValue(MorphedSourceMesh);
-				
-				SurfNode->MeshMorph = MorphedSourceMeshNode;
-				
-				if (const UEdGraphPin* ConnectedPin = FollowInputPin(*TypedNodeMorph->FactorPin()))
-				{
-					UEdGraphNode* floatNode = ConnectedPin->GetOwningNode();
-					bool validStaticFactor = true;
-					if (const UCustomizableObjectNodeFloatParameter* floatParameterNode = Cast<UCustomizableObjectNodeFloatParameter>(floatNode))
-					{
-						if (floatParameterNode->DefaultValue < -1.0f || floatParameterNode->DefaultValue > 1.0f)
-						{
-							validStaticFactor = false;
-							FString msg = FString::Printf(TEXT("Mesh morph nodes only accept factors between -1.0 and 1.0 inclusive but the default value of the float parameter node is (%f). Factor will be ignored."), floatParameterNode->DefaultValue);
-							GenerationContext.Compiler->CompilerLog(FText::FromString(msg), Node);
-						}
-						if (floatParameterNode->ParamUIMetadata.MinimumValue < -1.0f)
-						{
-							validStaticFactor = false;
-							FString msg = FString::Printf(TEXT("Mesh morph nodes only accept factors between -1.0 and 1.0 inclusive but the minimum UI value for the input float parameter node is (%f). Factor will be ignored."), floatParameterNode->ParamUIMetadata.MinimumValue);
-							GenerationContext.Compiler->CompilerLog(FText::FromString(msg), Node);
-						}
-						if (floatParameterNode->ParamUIMetadata.MaximumValue > 1.0f)
-						{
-							validStaticFactor = false;
-							FString msg = FString::Printf(TEXT("Mesh morph nodes only accept factors between -1.0 and 1.0 inclusive but the maximum UI value for the input float parameter node is (%f). Factor will be ignored."), floatParameterNode->ParamUIMetadata.MaximumValue);
-							GenerationContext.Compiler->CompilerLog(FText::FromString(msg), Node);
-						}
-					}
-					else if (const UCustomizableObjectNodeFloatConstant* floatConstantNode = Cast<UCustomizableObjectNodeFloatConstant>(floatNode))
-					{
-						if (floatConstantNode->Value < -1.0f || floatConstantNode->Value > 1.0f)
-						{
-							validStaticFactor = false;
-							FString msg = FString::Printf(TEXT("Mesh morph nodes only accept factors between -1.0 and 1.0 inclusive but the value of the float constant node is (%f). Factor will be ignored."), floatConstantNode->Value);
-							GenerationContext.Compiler->CompilerLog(FText::FromString(msg), Node);
-						}
-					}
-
-					if (validStaticFactor)
-					{
-						mu::NodeScalarPtr FactorNode = GenerateMutableSourceFloat(ConnectedPin, GenerationContext);
-						SurfNode->MorphFactor = FactorNode;
-					}
-				}
-			}
-
-			AddModifierToSharedSurface(GenerationContext, ParentMaterialNode, *TypedNodeMorph);
-		}
-	}
-
 	else if (const UCustomizableObjectNodeMaterialVariation* TypedNodeVar = Cast<UCustomizableObjectNodeMaterialVariation>(Node))
 	{
 		mu::Ptr<mu::NodeSurfaceVariation> SurfNode = new mu::NodeSurfaceVariation();
diff --git a/Engine/Plugins/Experimental/Mutable/Source/CustomizableObjectEditor/Private/MuCOE/GenerateMutableSource/GenerateMutableSourceSurface.h b/Engine/Plugins/Experimental/Mutable/Source/CustomizableObjectEditor/Private/MuCOE/GenerateMutableSource/GenerateMutableSourceSurface.h
--- a/Engine/Plugins/Experimental/Mutable/Source/CustomizableObjectEditor/Private/MuCOE/GenerateMutableSource/GenerateMutableSourceSurface.h
+++ b/Engine/Plugins/Experimental/Mutable/Source/CustomizableObjectEditor/Private/MuCOE/GenerateMutableSource/GenerateMutableSourceSurface.h
@@ -2,9 +2,22 @@
 
 #pragma once
 
-#include "MuT/NodeComponentNew.h"
+#include "MuR/Ptr.h"
+#include "Hal/Platform.h"
 
 class UEdGraphPin;
 struct FMutableGraphGenerationContext;
+struct FMutableGraphMeshGenerationData;
+enum class ECustomizableObjectNumBoneInfluences:uint8;
+namespace mu
+{
+	class FMeshBufferSet;
+	class NodeSurface;
+}
+
+extern mu::Ptr<mu::NodeSurface> GenerateMutableSourceSurface(const UEdGraphPin* Pin, FMutableGraphGenerationContext & GenerationContext);
+
+extern void SetSurfaceFormat(FMutableGraphGenerationContext& GenerationContext,
+	mu::FMeshBufferSet& OutVertexBufferFormat, mu::FMeshBufferSet& OutIndexBufferFormat, const FMutableGraphMeshGenerationData& MeshData,
+	ECustomizableObjectNumBoneInfluences ECustomizableObjectNumBoneInfluences, bool bWith16BitWeights);
 
-mu::NodeSurfacePtr GenerateMutableSourceSurface(const UEdGraphPin * Pin, FMutableGraphGenerationContext & GenerationContext);
diff --git a/Engine/Plugins/Experimental/Mutable/Source/CustomizableObjectEditor/Private/MuCOE/GraphTraversal.cpp b/Engine/Plugins/Experimental/Mutable/Source/CustomizableObjectEditor/Private/MuCOE/GraphTraversal.cpp
--- a/Engine/Plugins/Experimental/Mutable/Source/CustomizableObjectEditor/Private/MuCOE/GraphTraversal.cpp
+++ b/Engine/Plugins/Experimental/Mutable/Source/CustomizableObjectEditor/Private/MuCOE/GraphTraversal.cpp
@@ -7,7 +7,7 @@
 #include "MuCOE/EdGraphSchema_CustomizableObject.h"
 #include "MuCOE/Nodes/CustomizableObjectNodeEnumParameter.h"
 #include "MuCOE/Nodes/CustomizableObjectNodeExposePin.h"
-#include "MuCOE/Nodes/CustomizableObjectNodeExtendMaterial.h"
+#include "MuCOE/Nodes/CustomizableObjectNodeModifierExtendMeshSection.h"
 #include "MuCOE/Nodes/CustomizableObjectNodeExternalPin.h"
 #include "MuCOE/Nodes/CustomizableObjectNodeMaterial.h"
 #include "MuCOE/Nodes/CustomizableObjectNodeMaterialVariation.h"
@@ -528,7 +528,7 @@ const UEdGraphPin* FindMeshBaseSource(const UEdGraphPin& Pin, const bool bOnlyLo
 		}
 	}
 
-	else if (const UCustomizableObjectNodeExtendMaterial* TypedNodeExtend = Cast<UCustomizableObjectNodeExtendMaterial>(Node))
+	else if (const UCustomizableObjectNodeModifierExtendMeshSection* TypedNodeExtend = Cast<UCustomizableObjectNodeModifierExtendMeshSection>(Node))
 	{
 		if (const UEdGraphPin* ConnectedPin = FollowInputPin(*TypedNodeExtend->AddMeshPin()))
 		{
